compile fix for accidental implicit conversion in constructor, also removed erroneous new left over from ported code
#rnx

diff --git a/Engine/Plugins/Experimental/GeometryProcessing/Source/DynamicMesh/Private/Remesher.cpp b/Engine/Plugins/Experimental/GeometryProcessing/Source/DynamicMesh/Private/Remesher.cpp
--- a/Engine/Plugins/Experimental/GeometryProcessing/Source/DynamicMesh/Private/Remesher.cpp
+++ b/Engine/Plugins/Experimental/GeometryProcessing/Source/DynamicMesh/Private/Remesher.cpp
@@ -348,8 +348,9 @@ void FRemesher::UpdateAfterSplit(int edgeID, int va, int vb, const FDynamicMesh3
 		{
 			int nSetID = (ca.FixedSetID > 0 && ca.FixedSetID == cb.FixedSetID) ?
 				ca.FixedSetID : FVertexConstraint::InvalidSetID;
+			bool bMovable = ca.Movable && cb.Movable;
 			Constraints->SetOrUpdateVertexConstraint(SplitInfo.NewVertex,
-				new FVertexConstraint(true, nSetID));
+				FVertexConstraint(true, bMovable, nSetID));
 			bPositionFixed = true;
 		}
 
@@ -376,7 +377,7 @@ void FRemesher::UpdateAfterSplit(int edgeID, int va, int vb, const FDynamicMesh3
 			if (set_target != nullptr) 
 			{
 				Constraints->SetOrUpdateVertexConstraint(SplitInfo.NewVertex,
-					new FVertexConstraint(set_target));
+					FVertexConstraint(set_target));
 				ProjectVertex(SplitInfo.NewVertex, set_target);
 				bPositionFixed = true;
 			}
diff --git a/Engine/Plugins/Experimental/GeometryProcessing/Source/DynamicMesh/Public/MeshConstraints.h b/Engine/Plugins/Experimental/GeometryProcessing/Source/DynamicMesh/Public/MeshConstraints.h
--- a/Engine/Plugins/Experimental/GeometryProcessing/Source/DynamicMesh/Public/MeshConstraints.h
+++ b/Engine/Plugins/Experimental/GeometryProcessing/Source/DynamicMesh/Public/MeshConstraints.h
@@ -47,14 +47,14 @@ public:
 		TrackingSetID = -1;
 	}
 
-	FEdgeConstraint(EEdgeRefineFlags ConstraintFlags)
+	explicit FEdgeConstraint(EEdgeRefineFlags ConstraintFlags)
 	{
 		RefineFlags = ConstraintFlags;
 		Target = nullptr;
 		TrackingSetID = -1;
 	}
 
-	FEdgeConstraint(EEdgeRefineFlags ConstraintFlags, IProjectionTarget* TargetIn)
+	explicit FEdgeConstraint(EEdgeRefineFlags ConstraintFlags, IProjectionTarget* TargetIn)
 	{
 		RefineFlags = ConstraintFlags;
 		Target = TargetIn;
@@ -133,7 +133,7 @@ public:
 		Target = nullptr;
 	}
 
-	FVertexConstraint(bool bIsFixed, bool bIsMovable = false, int SetID = InvalidSetID)
+	explicit FVertexConstraint(bool bIsFixed, bool bIsMovable, int SetID = InvalidSetID)
 	{
 		Fixed = bIsFixed;
 		Movable = bIsMovable;
@@ -141,7 +141,7 @@ public:
 		Target = nullptr;
 	}
 
-	FVertexConstraint(IProjectionTarget* TargetIn)
+	explicit FVertexConstraint(IProjectionTarget* TargetIn)
 	{
 		Fixed = false;
 		Movable = false;
