robustness fixes to arrangment and line intersection tolerance logic #jira UE-76167

diff --git a/Engine/Plugins/Experimental/GeometryProcessing/Source/GeometricObjects/Public/Curve/DynamicGraph.h b/Engine/Plugins/Experimental/GeometryProcessing/Source/GeometricObjects/Public/Curve/DynamicGraph.h
--- a/Engine/Plugins/Experimental/GeometryProcessing/Source/GeometricObjects/Public/Curve/DynamicGraph.h
+++ b/Engine/Plugins/Experimental/GeometryProcessing/Source/GeometricObjects/Public/Curve/DynamicGraph.h
@@ -360,9 +360,7 @@ public:
 		// look up primary edge
 		int eab_i = 3 * EAB;
 		int A = edges[EAB].A, B = edges[EAB].B;
-		int GID = edges[EAB].Group;
-
-		// create new vertex
+		int GID = edges[EAB].Group;		
 		int f = append_new_split_vertex(A, B);
 
 		// rewrite edge bc, create edge af
@@ -384,6 +382,43 @@ public:
 		return EMeshResult::Ok;
 	}
 
+	EMeshResult SplitEdgeWithExistingVertex(int EAB, int ExistingMidVert, FEdgeSplitInfo& Split)
+	{
+		if (!IsEdge(EAB))
+		{
+			return EMeshResult::Failed_NotAnEdge;
+		}
+
+		int f = ExistingMidVert;
+		if (!IsVertex(f))
+		{
+			return EMeshResult::Failed_NotAVertex;
+		}
+
+		// look up primary edge
+		int eab_i = 3 * EAB;
+		int A = edges[EAB].A, B = edges[EAB].B;
+		int GID = edges[EAB].Group;		
+
+		// rewrite edge bc, create edge af
+		int eaf = EAB;
+		replace_edge_vertex(eaf, B, f);
+		vertex_edges.Remove(B, EAB);
+		vertex_edges.Insert(f, eaf);
+
+		// create new edge fb
+		int efb = add_edge(f, B, GID);
+
+		// update vertex refcounts
+		vertices_refcount.Increment(f, 2);
+
+		Split.VNew = f;
+		Split.ENewBN = efb;
+
+		updateTimeStamp(true);
+		return EMeshResult::Ok;
+	}
+
 protected:
 	virtual int append_new_split_vertex(int A, int B)
 	{
diff --git a/Engine/Plugins/Experimental/GeometryProcessing/Source/GeometricObjects/Public/Intersection/IntrLine2Line2.h b/Engine/Plugins/Experimental/GeometryProcessing/Source/GeometricObjects/Public/Intersection/IntrLine2Line2.h
--- a/Engine/Plugins/Experimental/GeometryProcessing/Source/GeometricObjects/Public/Intersection/IntrLine2Line2.h
+++ b/Engine/Plugins/Experimental/GeometryProcessing/Source/GeometricObjects/Public/Intersection/IntrLine2Line2.h
@@ -19,6 +19,7 @@ protected:
 	TLine2<RealType> Line1;
 	TLine2<RealType> Line2;
 	RealType dotThresh = TMathUtil<RealType>::ZeroTolerance;
+	RealType DistThresh = TMathUtil<RealType>::ZeroTolerance;
 	
 public:
 	// result data
@@ -65,12 +66,24 @@ public:
 		return dotThresh;
 	}
 
+	RealType GetDistThreshold() const
+	{
+		return DistThresh;
+	}
+
 	void SetDotThreshold(RealType Value)
 	{
 		dotThresh = FMath::Max(Value, (RealType)0);
 		Result = EIntersectionResult::NotComputed;
 	}
 
+	void SetDistThreshold(RealType Value)
+	{
+		DistThresh = FMath::Max(Value, (RealType)0);
+		Result = EIntersectionResult::NotComputed;
+	}
+	
+
 	bool IsSimpleIntersection() const
 	{
 		return Result == EIntersectionResult::Intersects && Type == EIntersectionType::Point;
@@ -102,7 +115,7 @@ public:
 
 		FVector2<RealType> s = FVector2<RealType>::Zero();
 		Type = Classify(Line1.Origin, Line1.Direction,
-			Line2.Origin, Line2.Direction, dotThresh, s);
+			Line2.Origin, Line2.Direction, dotThresh, DistThresh, s);
 
 		if (Type == EIntersectionType::Point) 
 		{
@@ -130,10 +143,10 @@ public:
 	static EIntersectionType Classify(
 		const FVector2<RealType>& P0, const FVector2<RealType>& D0, 
 		const FVector2<RealType>& P1, const FVector2<RealType>& D1, 
-		RealType dotThreshold, FVector2<RealType>& s)
+		RealType DotThreshold, RealType DistThreshold, FVector2<RealType>& s)
 	{
-		// Ensure dotThreshold is nonnegative.
-		dotThreshold = FMath::Max(dotThreshold, (RealType)0);
+		// Ensure DotThreshold is nonnegative.
+		DotThreshold = FMath::Max(DotThreshold, (RealType)0);
 
 		// The intersection of two lines is a solution to P0+s0*D0 = P1+s1*D1.
 		// Rewrite this as s0*D0 - s1*D1 = P1 - P0 = Q.  If D0.Dot(Perp(D1)) = 0,
@@ -145,7 +158,7 @@ public:
 
 		FVector2<RealType> diff = P1 - P0;
 		RealType D0DotPerpD1 = D0.DotPerp(D1);
-		if (FMath::Abs(D0DotPerpD1) > dotThreshold) 
+		if (FMath::Abs(D0DotPerpD1) > DotThreshold) 
 		{
 			// Lines intersect in a single point.
 			RealType invD0DotPerpD1 = 1.0 / D0DotPerpD1;
@@ -156,10 +169,9 @@ public:
 			return EIntersectionType::Point;
 		}
 
-		// Lines are parallel.
-		diff.Normalize();
+		// Lines are parallel; check if they are within DistThresh apart
 		RealType diffNDotPerpD1 = diff.DotPerp(D1);
-		if (FMath::Abs(diffNDotPerpD1) <= dotThreshold) 
+		if (FMath::Abs(diffNDotPerpD1) <= DistThreshold) 
 		{
 			// Lines are collinear.
 			return EIntersectionType::Line;
diff --git a/Engine/Plugins/Experimental/GeometryProcessing/Source/GeometricObjects/Public/Intersection/IntrSegment2Segment2.h b/Engine/Plugins/Experimental/GeometryProcessing/Source/GeometricObjects/Public/Intersection/IntrSegment2Segment2.h
--- a/Engine/Plugins/Experimental/GeometryProcessing/Source/GeometricObjects/Public/Intersection/IntrSegment2Segment2.h
+++ b/Engine/Plugins/Experimental/GeometryProcessing/Source/GeometricObjects/Public/Intersection/IntrSegment2Segment2.h
@@ -152,7 +152,7 @@ public:
 		FVector2<RealType> s = FVector2<RealType>::Zero();
 		Type = TIntrLine2Line2<RealType>::Classify(Segment1.Center, Segment1.Direction,
 			Segment2.Center, Segment2.Direction,
-			DotThreshold, s);
+			DotThreshold, IntervalThreshold, s);
 
 		if (Type == EIntersectionType::Point)
 		{
diff --git a/Engine/Plugins/Experimental/GeometryProcessing/Source/GeometryAlgorithms/Public/Arrangement2d.h b/Engine/Plugins/Experimental/GeometryProcessing/Source/GeometryAlgorithms/Public/Arrangement2d.h
--- a/Engine/Plugins/Experimental/GeometryProcessing/Source/GeometryAlgorithms/Public/Arrangement2d.h
+++ b/Engine/Plugins/Experimental/GeometryProcessing/Source/GeometryAlgorithms/Public/Arrangement2d.h
@@ -357,11 +357,11 @@ protected:
 				continue;
 			}
 
-			// sanity check
-			ensureMsgf(FMath::Abs(points[k].T - points[k + 1].T) >= std::numeric_limits<float>::epsilon(), TEXT("insert_segment: different points have same T??"));
-
 			if (Graph.FindEdge(v0, v1) == FDynamicGraph2d::InvalidID)
 			{
+				// sanity check; technically this can happen and still be correct but it's more likely an error case
+				ensureMsgf(FMath::Abs(points[k].T - points[k + 1].T) >= std::numeric_limits<float>::epsilon(), TEXT("insert_segment: different points have same T??"));
+
 				Graph.AppendEdge(v0, v1, GID);
 			}
 		}
@@ -392,14 +392,25 @@ protected:
 		else
 		{
 			FVector2d Pt = seg.PointAt(T);
-			// ensure(find_existing_vertex(Pt) == -1); // TODO: handle this case!?
 			FDynamicGraph2d::FEdgeSplitInfo splitInfo;
-			EMeshResult result = Graph.SplitEdge(EID, splitInfo);
+			EMeshResult result;
+			int CrossingVert = find_existing_vertex(Pt);
+			if (CrossingVert == -1)
+			{
+				result = Graph.SplitEdge(EID, splitInfo);
+			}
+			else
+			{
+				result = Graph.SplitEdgeWithExistingVertex(EID, CrossingVert, splitInfo);
+			}
 			ensureMsgf(result == EMeshResult::Ok, TEXT("insert_into_segment: edge split failed?"));
 			use_vid = splitInfo.VNew;
 			new_eid = splitInfo.ENewBN;
-			Graph.SetVertex(use_vid, Pt);
-			PointHash.InsertPointUnsafe(splitInfo.VNew, Pt);
+			if (CrossingVert == -1)
+			{	// position + track added vertex
+				Graph.SetVertex(use_vid, Pt);
+				PointHash.InsertPointUnsafe(splitInfo.VNew, Pt);
+			}
 		}
 		return FIndex2i(use_vid, new_eid);
 	}
