Fix invalid graph generation on splitting existing edge #jira UE-76952 #rb none #rnx

diff --git a/Engine/Plugins/Experimental/GeometryProcessing/Source/GeometricObjects/Public/Curve/DynamicGraph.h b/Engine/Plugins/Experimental/GeometryProcessing/Source/GeometricObjects/Public/Curve/DynamicGraph.h
--- a/Engine/Plugins/Experimental/GeometryProcessing/Source/GeometricObjects/Public/Curve/DynamicGraph.h
+++ b/Engine/Plugins/Experimental/GeometryProcessing/Source/GeometricObjects/Public/Curve/DynamicGraph.h
@@ -360,7 +360,7 @@ public:
 		// look up primary edge
 		int eab_i = 3 * EAB;
 		int A = edges[EAB].A, B = edges[EAB].B;
-		int GID = edges[EAB].Group;		
+		int GID = edges[EAB].Group;
 		int f = append_new_split_vertex(A, B);
 
 		// rewrite edge bc, create edge af
@@ -398,22 +398,48 @@ public:
 		// look up primary edge
 		int eab_i = 3 * EAB;
 		int A = edges[EAB].A, B = edges[EAB].B;
-		int GID = edges[EAB].Group;		
+		int GID = edges[EAB].Group;
 
-		// rewrite edge bc, create edge af
-		int eaf = EAB;
-		replace_edge_vertex(eaf, B, f);
-		vertex_edges.Remove(B, EAB);
-		vertex_edges.Insert(f, eaf);
-
-		// create new edge fb
-		int efb = add_edge(f, B, GID);
+		int EAf = FindEdge(A, f);
+		int FIncr = 0; // track how much to increment the reference count of middle vertex, f
+		int BIncr = 0;
+		if (EAf != InvalidID)
+		{
+			RemoveEdge(EAB, false); // this will handle changes to refcounts, so no need to touch bIncr
+			edges[EAf].Group = GID;
+		}
+		else
+		{
+			// rewrite edge ab to create edge af
+			EAf = EAB;
+			replace_edge_vertex(EAf, B, f);
+			vertex_edges.Remove(B, EAB);
+			vertex_edges.Insert(f, EAf);
+			FIncr++;
+			BIncr--;
+		}
 
-		// update vertex refcounts
-		vertices_refcount.Increment(f, 2);
+		int EfB = FindEdge(f, B);
+		if (EfB == InvalidID)
+		{
+			EfB = add_edge(f, B, GID);
+			FIncr++;
+			BIncr++;
+		}
+		else
+		{
+			edges[EfB].Group = GID;
+		}
+		
+		// update middle vertex refcounts
+		vertices_refcount.Increment(f, FIncr);
+		if (BIncr)
+		{
+			vertices_refcount.Increment(B, BIncr);
+		}
 
 		Split.VNew = f;
-		Split.ENewBN = efb;
+		Split.ENewBN = EfB;
 
 		updateTimeStamp(true);
 		return EMeshResult::Ok;
