fix bug in mesh normal computation; bWeightByArea was used instead of bWeightByAngle when deciding to use the angle weights.

#rnx

diff --git a/Engine/Plugins/Experimental/GeometryProcessing/Source/DynamicMesh/Public/MeshNormals.h b/Engine/Plugins/Experimental/GeometryProcessing/Source/DynamicMesh/Public/MeshNormals.h
--- a/Engine/Plugins/Experimental/GeometryProcessing/Source/DynamicMesh/Public/MeshNormals.h
+++ b/Engine/Plugins/Experimental/GeometryProcessing/Source/DynamicMesh/Public/MeshNormals.h
@@ -126,7 +126,7 @@ protected:
 	inline FVector3d GetVertexWeightsOnTriangle(int TriID, double TriArea, bool bWeightByArea, bool bWeightByAngle) const
 	{
 		FVector3d TriNormalWeights = FVector3d::One();
-		if (bWeightByArea)
+		if (bWeightByAngle)
 		{
 			TriNormalWeights = Mesh->GetTriInternalAnglesR(TriID); // component-wise multiply by per-vertex internal angles
 		}
