Merging CL# 9541582 from UE4/Dev-Editor
GeometryProcessing: handle a missed case in DMeshOverlay OnCollapseEdge
#rnx

diff --git a/Engine/Plugins/Experimental/GeometryProcessing/Source/DynamicMesh/Private/DynamicMeshOverlay.cpp b/Engine/Plugins/Experimental/GeometryProcessing/Source/DynamicMesh/Private/DynamicMeshOverlay.cpp
--- a/Engine/Plugins/Experimental/GeometryProcessing/Source/DynamicMesh/Private/DynamicMeshOverlay.cpp
+++ b/Engine/Plugins/Experimental/GeometryProcessing/Source/DynamicMesh/Private/DynamicMeshOverlay.cpp
@@ -484,7 +484,7 @@ void TDynamicMeshOverlay<RealType, ElementSize>::OnReverseTriOrientation(int Tri
 
 
 
-//#pragma optimize( "", off )
+
 template<typename RealType, int ElementSize>
 void TDynamicMeshOverlay<RealType, ElementSize>::OnSplitEdge(const FDynamicMesh3::FEdgeSplitInfo& splitInfo)
 {
@@ -554,7 +554,7 @@ void TDynamicMeshOverlay<RealType, ElementSize>::OnSplitEdge(const FDynamicMesh3
 
 
 
-//#pragma optimize("", off)
+
 template<typename RealType, int ElementSize>
 void TDynamicMeshOverlay<RealType, ElementSize>::OnFlipEdge(const FDynamicMesh3::FEdgeFlipInfo& FlipInfo)
 {
@@ -609,7 +609,7 @@ void TDynamicMeshOverlay<RealType, ElementSize>::OnFlipEdge(const FDynamicMesh3:
 
 
 
-//#pragma optimize("", off)
+
 template<typename RealType, int ElementSize>
 void TDynamicMeshOverlay<RealType, ElementSize>::OnCollapseEdge(const FDynamicMesh3::FEdgeCollapseInfo& collapseInfo)
 {
@@ -639,21 +639,32 @@ void TDynamicMeshOverlay<RealType, ElementSize>::OnCollapseEdge(const FDynamicMe
 		check(bHasSharedUVEdge);
 	}
 
-	// need to find the elementid for the "kept" vertex. Since this isn't a seam, 
-	// there is just one, and we can grab it from "old" removed triangle
+	// need to find the elementid for the "kept" and "removed" vertices. 
+	// Since this isn't a seam, there is just one of each, *unless* either the kept
+	// or removed vertices is on a seam (ie connected to a separate edge that is a seam).
+	// In that case this code may fail. Currently this case must be caught and avoided
+	// at a higher level. 
+	// @todo: handle this case, we would then have removed_elements.Num() > 1
 	int kept_elemid = FDynamicMesh3::InvalidID;
+	int removed_elemid = FDynamicMesh3::InvalidID;
 	for (int j = 0; j < 3; ++j) 
 	{
 		if (BaseTriangle0[j] == vid_base_kept)
 		{
 			kept_elemid = Triangle0[j];
 		}
+		if (ParentVertices[Triangle0[j]] == vid_base_removed)
+		{
+			removed_elemid = Triangle0[j];
+		}
 	}
 	check(kept_elemid != FDynamicMesh3::InvalidID);
+	check(removed_elemid != FDynamicMesh3::InvalidID);
 
 	// look for still-existing triangles that have elements linked to the removed vertex.
 	// in that case, replace with the element we found
 	TArray<int> removed_elements;
+	removed_elements.AddUnique(removed_elemid);
 	for (int onering_tid : ParentMesh->VtxTrianglesItr(vid_base_kept))
 	{
 		FIndex3i elem_tri = GetTriangle(onering_tid);
@@ -663,7 +674,8 @@ void TDynamicMeshOverlay<RealType, ElementSize>::OnCollapseEdge(const FDynamicMe
 			if (ParentVertices[elem_id] == vid_base_removed) 
 			{
 				ElementsRefCounts.Decrement(elem_id);
-				removed_elements.AddUnique(elem_id);
+				//removed_elements.AddUnique(elem_id);
+				check(elem_id == removed_elemid);
 				ElementTriangles[3*onering_tid + j] = kept_elemid;
 				ElementsRefCounts.Increment(kept_elemid);
 			}
@@ -720,6 +732,7 @@ void TDynamicMeshOverlay<RealType, ElementSize>::OnPokeTriangle(const FDynamicMe
 }
 
 
+
 template<typename RealType, int ElementSize>
 void TDynamicMeshOverlay<RealType, ElementSize>::OnMergeEdges(const FDynamicMesh3::FMergeEdgesInfo& MergeInfo)
 {
@@ -785,7 +798,7 @@ void TDynamicMeshOverlay<RealType, ElementSize>::SetElementFromBary(int SetEleme
 }
 
 
-//#pragma optimize("", off)
+
 template<typename RealType, int ElementSize>
 bool TDynamicMeshOverlay<RealType, ElementSize>::CheckValidity(bool bAllowNonManifoldVertices, EValidityCheckFailMode FailMode) const
 {
