fix: AABB raycasts fail to hit sometimes
#rnx

diff --git a/Engine/Plugins/Experimental/GeometryProcessing/Source/GeometricObjects/Public/MeshQueries.h b/Engine/Plugins/Experimental/GeometryProcessing/Source/GeometricObjects/Public/MeshQueries.h
--- a/Engine/Plugins/Experimental/GeometryProcessing/Source/GeometricObjects/Public/MeshQueries.h
+++ b/Engine/Plugins/Experimental/GeometryProcessing/Source/GeometricObjects/Public/MeshQueries.h
@@ -151,8 +151,12 @@ public:
 		double fNearestT = TNumericLimits<double>::Max();
 		FTriangle3d Triangle;
 
-		for (int TriIdx : Mesh.TriangleIndicesItr())
+		for (int TriIdx = 0; TriIdx < Mesh.MaxTriangleID(); TriIdx++)
 		{
+			if (!Mesh.IsTriangle(TriIdx))
+			{
+				continue;
+			}
 			Mesh.GetTriVertices(TriIdx, Triangle.V[0], Triangle.V[1], Triangle.V[2]);
 			FIntrRay3Triangle3d Query(Ray, Triangle);
 			if (Query.Find())
diff --git a/Engine/Plugins/Experimental/GeometryProcessing/Source/GeometricObjects/Public/Spatial/MeshAABBTree3.h b/Engine/Plugins/Experimental/GeometryProcessing/Source/GeometricObjects/Public/Spatial/MeshAABBTree3.h
--- a/Engine/Plugins/Experimental/GeometryProcessing/Source/GeometricObjects/Public/Spatial/MeshAABBTree3.h
+++ b/Engine/Plugins/Experimental/GeometryProcessing/Source/GeometricObjects/Public/Spatial/MeshAABBTree3.h
@@ -409,7 +409,7 @@ protected:
 		return FAxisAlignedBox3d(Min, Max);
 	}
 
-	const double BoxEps = 50.0 * FMathd::Epsilon;
+	const double BoxEps = FMathd::ZeroTolerance;
 
 	double BoxDistanceSqr(int IBox, const FVector3d& V) const
 	{
