fix: Plane Cut does not assign input material to new objects created

#rnx

diff --git a/Engine/Plugins/Experimental/GeometryProcessing/Source/MeshConversion/Private/MeshDescriptionBuilder.cpp b/Engine/Plugins/Experimental/GeometryProcessing/Source/MeshConversion/Private/MeshDescriptionBuilder.cpp
--- a/Engine/Plugins/Experimental/GeometryProcessing/Source/MeshConversion/Private/MeshDescriptionBuilder.cpp
+++ b/Engine/Plugins/Experimental/GeometryProcessing/Source/MeshConversion/Private/MeshDescriptionBuilder.cpp
@@ -230,83 +230,6 @@ void FMeshDescriptionBuilder::SetPolyGroupID(const FPolygonID& PolygonID, int Gr
 
 
 
-void FMeshDescriptionBuilder::AppendMesh(const FDynamicMesh3* Mesh, bool bSetPolyGroups)
-{
-	// TODO: update to support multiple UV layers; & combine w/ extremely similar code in DynamicMeshToMeshDescription.cpp
-
-	// create vertices
-	TArray<FVertexID> MapV;
-	MapV.SetNum(Mesh->MaxVertexID());
-	for (int VertID : Mesh->VertexIndicesItr())
-	{
-		FVector3d v = Mesh->GetVertex(VertID);
-		MapV[VertID] = AppendVertex(v);
-	}
-
-	bool bDoTransferPolyGroups = false;
-	if (bSetPolyGroups)
-	{
-		EnablePolyGroups();
-		bDoTransferPolyGroups = Mesh->HasTriangleGroups();
-	}
-
-	FPolygonGroupID AllGroupID = AppendPolygonGroup();
-
-	bool bHasVertexColors = Mesh->HasVertexColors();
-	bool bHasVertexNormals = Mesh->HasVertexNormals();
-	const FDynamicMeshUVOverlay* UVOverlay = Mesh->HasAttributes() ? Mesh->Attributes()->PrimaryUV() : nullptr;
-	const FDynamicMeshNormalOverlay* NormalOverlay = Mesh->HasAttributes() ? Mesh->Attributes()->PrimaryNormals() : nullptr;
-	//FDynamicMeshUVOverlay* UVOverlay = nullptr;
-	//FDynamicMeshNormalOverlay* NormalOverlay = nullptr;
-
-
-	// create new instances when seen
-	TMap<FIndex3i, FVertexInstanceID> InstanceList;
-	for (int TriID : Mesh->TriangleIndicesItr())
-	{
-		// @todo support additional overlays (requires IndexNi...)
-		FIndex3i Triangle = Mesh->GetTriangle(TriID);
-		FIndex3i UVTriangle = (UVOverlay != nullptr) ? UVOverlay->GetTriangle(TriID) : FIndex3i(-1, -1, -1);
-		FIndex3i NormalTriangle = (NormalOverlay != nullptr) ? NormalOverlay->GetTriangle(TriID) : FIndex3i(-1, -1, -1);
-		FVertexInstanceID InstanceTri[3];
-		for (int j = 0; j < 3; ++j)
-		{
-			FIndex3i InstanceElem(Triangle[j], UVTriangle[j], NormalTriangle[j]);
-			if (InstanceList.Contains(InstanceElem) == false)
-			{
-				FVertexInstanceID NewInstanceID = AppendInstance(MapV[Triangle[j]]);
-				InstanceList.Add(InstanceElem, NewInstanceID);
-				SetInstance(NewInstanceID,
-					(UVTriangle[j] == -1 || UVOverlay == nullptr) ? FVector2f::Zero() : UVOverlay->GetElement(UVTriangle[j]),
-					(NormalTriangle[j] == -1 || NormalOverlay == nullptr) ? FVector3f::UnitY() : NormalOverlay->GetElement(NormalTriangle[j]));
-			}
-			InstanceTri[j] = InstanceList[InstanceElem];
-		}
-
-		FPolygonID NewPolyID = AppendTriangle(InstanceTri[0], InstanceTri[1], InstanceTri[2], AllGroupID);
-		if (bDoTransferPolyGroups)
-		{
-			SetPolyGroupID(NewPolyID, Mesh->GetTriangleGroup(TriID));
-		}
-
-		if (bHasVertexColors)
-		{
-			for (int j = 0; j < 3; ++j)
-			{
-				FVector3f Color = Mesh->GetVertexColor(Triangle[j]);
-				FVector4 Color4(Color.X, Color.Y, Color.Z, 1.0);
-				SetInstanceColor(InstanceTri[j], Color4);
-			}
-		}
-		if (NormalOverlay == nullptr && bHasVertexNormals)
-		{
-			SetInstanceNormal(InstanceTri[0], Mesh->GetVertexNormal(Triangle[0]));
-			SetInstanceNormal(InstanceTri[1], Mesh->GetVertexNormal(Triangle[1]));
-			SetInstanceNormal(InstanceTri[2], Mesh->GetVertexNormal(Triangle[2]));
-		}
-	}
-}
-
 
 
 
diff --git a/Engine/Plugins/Experimental/GeometryProcessing/Source/MeshConversion/Public/MeshDescriptionBuilder.h b/Engine/Plugins/Experimental/GeometryProcessing/Source/MeshConversion/Public/MeshDescriptionBuilder.h
--- a/Engine/Plugins/Experimental/GeometryProcessing/Source/MeshConversion/Public/MeshDescriptionBuilder.h
+++ b/Engine/Plugins/Experimental/GeometryProcessing/Source/MeshConversion/Public/MeshDescriptionBuilder.h
@@ -88,13 +88,6 @@ public:
 	FPolygonID AppendTriangle(const FVertexInstanceID& Instance0, const FVertexInstanceID& Instance1, const FVertexInstanceID& Instance2, const FPolygonGroupID& PolygonGroup);
 
 
-	/**
-	 * Append an entire mesh to this mesh. This will create a minimal set of shared vertex instances.
-	 * @param bSetPolyGroups if true, we transfer the input mesh triangle groups to the PolyTriGroups attribute
-	 */
-	void AppendMesh(const FDynamicMesh3* Mesh, bool bSetPolyGroups);
-
-
 
 	/** Set MeshAttribute::Edge::IsHard to true for all edges */
 	void SetAllEdgesHardness(bool bHard);
diff --git a/Engine/Plugins/Experimental/MeshModelingToolset/Source/MeshModelingTools/Private/PlaneCutTool.cpp b/Engine/Plugins/Experimental/MeshModelingToolset/Source/MeshModelingTools/Private/PlaneCutTool.cpp
--- a/Engine/Plugins/Experimental/MeshModelingToolset/Source/MeshModelingTools/Private/PlaneCutTool.cpp
+++ b/Engine/Plugins/Experimental/MeshModelingToolset/Source/MeshModelingTools/Private/PlaneCutTool.cpp
@@ -346,11 +346,17 @@ void UPlaneCutTool::GenerateAsset(const TArray<TUniquePtr<FDynamicMeshOpResult>>
 		NewSelection.Actors.Add(ComponentTarget->GetOwnerActor());
 
 		check(Results[1]->Mesh.Get() != nullptr);
-		// TODO: copy over material?
+
+		TArray<UMaterialInterface*> Materials;
+		for (int MaterialIdx = 0, NumMaterials = ComponentTarget->GetNumMaterials(); MaterialIdx < NumMaterials; MaterialIdx++)
+		{
+			Materials.Add(ComponentTarget->GetMaterial(MaterialIdx));
+		}
+		
 		AActor* NewActor = AssetGenerationUtil::GenerateStaticMeshActor(
 			AssetAPI, TargetWorld,
 			Results[1]->Mesh.Get(), Results[1]->Transform, TEXT("PlaneCutOtherHalf"),
-			AssetGenerationUtil::GetDefaultAutoGeneratedAssetPath());
+			AssetGenerationUtil::GetDefaultAutoGeneratedAssetPath(), Materials);
 		NewSelection.Actors.Add(NewActor);
 		GetToolManager()->RequestSelectionChange(NewSelection);
 	}
diff --git a/Engine/Plugins/Experimental/MeshModelingToolset/Source/ModelingComponents/Private/AssetGenerationUtil.cpp b/Engine/Plugins/Experimental/MeshModelingToolset/Source/ModelingComponents/Private/AssetGenerationUtil.cpp
--- a/Engine/Plugins/Experimental/MeshModelingToolset/Source/ModelingComponents/Private/AssetGenerationUtil.cpp
+++ b/Engine/Plugins/Experimental/MeshModelingToolset/Source/ModelingComponents/Private/AssetGenerationUtil.cpp
@@ -12,6 +12,7 @@
 #include "StaticMeshComponentBuilder.h"
 #include "MeshDescriptionToDynamicMesh.h"
 #include "MeshDescriptionBuilder.h"
+#include "DynamicMeshToMeshDescription.h"
 
 
 FString AssetGenerationUtil::GetDefaultAutoGeneratedAssetPath()
@@ -21,6 +22,8 @@ FString AssetGenerationUtil::GetDefaultAutoGeneratedAssetPath()
 
 
 #if WITH_EDITOR
+
+// single material case
 AActor* AssetGenerationUtil::GenerateStaticMeshActor(
 	IToolsContextAssetAPI* AssetAPI,
 	UWorld* TargetWorld,
@@ -30,6 +33,20 @@ AActor* AssetGenerationUtil::GenerateStaticMeshActor(
 	FString PackagePath,
 	UMaterialInterface* Material
 )
+{
+	return GenerateStaticMeshActor(AssetAPI, TargetWorld, Mesh, Transform, ObjectName, PackagePath, TArrayView<UMaterialInterface*>(&Material, Material != nullptr ? 1 : 0));
+}
+
+// N-material case
+AActor* AssetGenerationUtil::GenerateStaticMeshActor(
+	IToolsContextAssetAPI* AssetAPI,
+	UWorld* TargetWorld,
+	const FDynamicMesh3* Mesh,
+	const FTransform3d& Transform,
+	FString ObjectName,
+	FString PackagePath,
+	const TArrayView<UMaterialInterface*>& Materials
+)
 {
 	check(AssetAPI);
 	check(TargetWorld);
@@ -53,14 +70,13 @@ AActor* AssetGenerationUtil::GenerateStaticMeshActor(
 
 	// construct new static mesh
 	FStaticMeshComponentBuilder Builder;
-	Builder.Initialize(AssetPackage, FName(*UniqueAssetName));
+	Builder.Initialize(AssetPackage, FName(*UniqueAssetName), Materials.Num());
 
 	// Populate the MeshDescription associated with the asset.
-	FMeshDescriptionBuilder MeshDescripBuilder;
-	MeshDescripBuilder.SetMeshDescription(Builder.MeshDescription);
 	if (Mesh != nullptr)
 	{
-		MeshDescripBuilder.AppendMesh(Mesh, true);
+		FDynamicMeshToMeshDescription Converter;
+		Converter.Convert(Mesh, *Builder.MeshDescription);
 	}
 	else
 	{
@@ -74,11 +90,11 @@ AActor* AssetGenerationUtil::GenerateStaticMeshActor(
 
 	// transform new component to origin
 	Builder.NewMeshComponent->SetWorldTransform((FTransform)Transform);
-
-	if (Material)
+	for (int MatIdx = 0, NumMats = Materials.Num(); MatIdx < NumMats; MatIdx++)
 	{
-		Builder.NewMeshComponent->SetMaterial(0, Material);
+		Builder.NewMeshComponent->SetMaterial(MatIdx, Materials[MatIdx]);
 	}
+	
 	// force save of asset
 	//AssetAPI->InteractiveSaveGeneratedAsset(Builder.NewStaticMesh, AssetPackage);
 	AssetAPI->AutoSaveGeneratedAsset(Builder.NewStaticMesh, AssetPackage);
diff --git a/Engine/Plugins/Experimental/MeshModelingToolset/Source/ModelingComponents/Private/StaticMeshComponentBuilder.cpp b/Engine/Plugins/Experimental/MeshModelingToolset/Source/ModelingComponents/Private/StaticMeshComponentBuilder.cpp
--- a/Engine/Plugins/Experimental/MeshModelingToolset/Source/ModelingComponents/Private/StaticMeshComponentBuilder.cpp
+++ b/Engine/Plugins/Experimental/MeshModelingToolset/Source/ModelingComponents/Private/StaticMeshComponentBuilder.cpp
@@ -15,7 +15,7 @@ FStaticMeshComponentBuilder::FStaticMeshComponentBuilder()
 	NewStaticMesh = nullptr;
 }
 
-void FStaticMeshComponentBuilder::Initialize(UPackage* AssetPackage, FName MeshName)
+void FStaticMeshComponentBuilder::Initialize(UPackage* AssetPackage, FName MeshName, int NumMaterialSlots)
 {
 	// create new UStaticMesh object
 	EObjectFlags flags = EObjectFlags::RF_Public | EObjectFlags::RF_Standalone;
@@ -40,7 +40,10 @@ void FStaticMeshComponentBuilder::Initialize(UPackage* AssetPackage, FName MeshN
 	}
 
 	// add a material slot
-	NewStaticMesh->StaticMaterials.Add(FStaticMaterial());
+	for (int MatIdx = 0; MatIdx < NumMaterialSlots; MatIdx++)
+	{
+		NewStaticMesh->StaticMaterials.Add(FStaticMaterial());
+	}
 }
 
 
diff --git a/Engine/Plugins/Experimental/MeshModelingToolset/Source/ModelingComponents/Public/AssetGenerationUtil.h b/Engine/Plugins/Experimental/MeshModelingToolset/Source/ModelingComponents/Public/AssetGenerationUtil.h
--- a/Engine/Plugins/Experimental/MeshModelingToolset/Source/ModelingComponents/Public/AssetGenerationUtil.h
+++ b/Engine/Plugins/Experimental/MeshModelingToolset/Source/ModelingComponents/Public/AssetGenerationUtil.h
@@ -31,6 +31,7 @@ namespace AssetGenerationUtil
 	 * @param Transform transformation for the new actor
 	 * @param ObjectName name of the new asset
 	 * @param PackagePath path of the new asset
+	 * @param Material optional single material to set on actor
 	 * @return new mesh actor
 	 */
 #if WITH_EDITOR
@@ -43,6 +44,27 @@ namespace AssetGenerationUtil
 		FString PackagePath,
 		UMaterialInterface* Material = nullptr
 	);
+
+	/**
+	* Creates a new StaticMesh actor/component, with a new mesh asset stored at the given PackagePath created via the AssetAPI
+	* @param AssetAPI pointer to context asset API that will be used to create new asset
+	* @param TargetWorld world that Actor will be created in
+	* @param Mesh geometry for the mesh
+	* @param Transform transformation for the new actor
+	* @param ObjectName name of the new asset
+	* @param PackagePath path of the new asset
+	* @param Materials materials to set on generated actor
+	* @return new mesh actor
+	*/
+	MODELINGCOMPONENTS_API AActor* GenerateStaticMeshActor(
+		IToolsContextAssetAPI* AssetAPI,
+		UWorld* TargetWorld,
+		const FDynamicMesh3* Mesh,
+		const FTransform3d& Transform,
+		FString ObjectName,
+		FString PackagePath,
+		const TArrayView<UMaterialInterface*>& Materials
+	);
 #endif
 
 }
diff --git a/Engine/Plugins/Experimental/MeshModelingToolset/Source/ModelingComponents/Public/StaticMeshComponentBuilder.h b/Engine/Plugins/Experimental/MeshModelingToolset/Source/ModelingComponents/Public/StaticMeshComponentBuilder.h
--- a/Engine/Plugins/Experimental/MeshModelingToolset/Source/ModelingComponents/Public/StaticMeshComponentBuilder.h
+++ b/Engine/Plugins/Experimental/MeshModelingToolset/Source/ModelingComponents/Public/StaticMeshComponentBuilder.h
@@ -32,8 +32,9 @@ public:
 	 * Call this to create and configure NewStaticMesh and MeshDescription
 	 * @param AssetPackage The asset package that the UStaticMesh will be saved into
 	 * @param MeshName the name of the mesh asset
+	 * @param NumMaterialSlots the number of material slots to create on the static mesh -- if the staticmesh does not have enough slots, corresponding override materials on the staticmeshcomponent will be discarded
 	 */
-	void Initialize(UPackage* AssetPackage, FName MeshName);
+	void Initialize(UPackage* AssetPackage, FName MeshName, int NumMaterialSlots = 1);
 
 	// (do your mesh setup here)
 
