fix: Normals Tool appears to create spurious sharp vertices

#rnx

diff --git a/Engine/Plugins/Experimental/GeometryProcessing/Source/DynamicMesh/Private/DynamicMeshOverlay.cpp b/Engine/Plugins/Experimental/GeometryProcessing/Source/DynamicMesh/Private/DynamicMeshOverlay.cpp
--- a/Engine/Plugins/Experimental/GeometryProcessing/Source/DynamicMesh/Private/DynamicMeshOverlay.cpp
+++ b/Engine/Plugins/Experimental/GeometryProcessing/Source/DynamicMesh/Private/DynamicMeshOverlay.cpp
@@ -83,7 +83,7 @@ EMeshResult TDynamicMeshOverlay<RealType, ElementSize>::InsertElement(int Elemen
 
 
 template<typename RealType, int ElementSize>
-void TDynamicMeshOverlay<RealType, ElementSize>::CreateFromPredicate(TFunctionRef<bool(int ParentVertexIdx, int TriIDA, int TriIDB)> TrisCanShareVertexPredicate, RealType InitElementValue, bool bNonTransitiveShare)
+void TDynamicMeshOverlay<RealType, ElementSize>::CreateFromPredicate(TFunctionRef<bool(int ParentVertexIdx, int TriIDA, int TriIDB)> TrisCanShareVertexPredicate, RealType InitElementValue)
 {
 	ClearElements(); // deletes all elements and initializes triangles to be 1:1 w/ parentmesh IDs
 	TArray<int> TrisActiveSubGroup, AppendedElements;
@@ -109,83 +109,35 @@ void TDynamicMeshOverlay<RealType, ElementSize>::CreateFromPredicate(TFunctionRe
 			TrisActiveSubGroup.SetNumZeroed(GroupNum, false);
 			int CurrentGroupID = 0;
 			int CurrentGroupRefSubIdx = 0;
-			bool bShouldBreakGroup = false; // tracks whether a contiguous group should be broken up (in bNonTransitiveShare mode); final value after for-loop indicates whether final group was broken
-			bool bFirstGroupWasBroken = false;
 			for (int TriSubIdx = 0; TriSubIdx+1 < GroupNum; TriSubIdx++)
 			{
 				int TriIDA = TriangleIDs[GroupStart + TriSubIdx];
 				int TriIDB = TriangleIDs[GroupStart + TriSubIdx + 1];
 				bool bCanShare = TrisCanShareVertexPredicate(VertexID, TriIDA, TriIDB);
-				if (!bShouldBreakGroup && bCanShare && bNonTransitiveShare && TriSubIdx != CurrentGroupRefSubIdx) // require the triangle match the initial triangle in the group as well as the adjacent one
-				{
-					if (!TrisCanShareVertexPredicate(VertexID, TriangleIDs[GroupStart + CurrentGroupRefSubIdx], TriIDB))
-					{
-						bShouldBreakGroup = true;
-						if (CurrentGroupRefSubIdx == 0)
-						{
-							bFirstGroupWasBroken = true;
-						}
-					}
-				}
 				if (!bCanShare)
 				{
-					if (bShouldBreakGroup)
-					{
-						for (int BreakGroupIdx = CurrentGroupRefSubIdx + 1; BreakGroupIdx <= TriSubIdx; BreakGroupIdx++)
-						{
-							TrisActiveSubGroup[BreakGroupIdx] = ++CurrentGroupID;
-						}
-					}
 					CurrentGroupID++;
 					CurrentGroupRefSubIdx = TriSubIdx + 1;
-					bShouldBreakGroup = false;
 				}
 				
 				TrisActiveSubGroup[TriSubIdx + 1] = CurrentGroupID;
 			}
 
-			// break final subgroup
-			if (bShouldBreakGroup)
-			{
-				for (int BreakGroupIdx = CurrentGroupRefSubIdx + 1; BreakGroupIdx < GroupNum; BreakGroupIdx++)
-				{
-					TrisActiveSubGroup[BreakGroupIdx] = ++CurrentGroupID;
-				}
-			}
-
 			// for loops, merge first and last group if needed
 			int NumGroupID = CurrentGroupID + 1;
 			if (bIsLoop && TrisActiveSubGroup[0] != TrisActiveSubGroup.Last())
 			{
-				// if we can merge the first and last triangles AND their groups weren't already both broken, try to merge them 
-				if ((!bShouldBreakGroup || !bFirstGroupWasBroken) && TrisCanShareVertexPredicate(VertexID, TriangleIDs[GroupStart], TriangleIDs[GroupStart + GroupNum - 1]))
+				if (TrisCanShareVertexPredicate(VertexID, TriangleIDs[GroupStart], TriangleIDs[GroupStart + GroupNum - 1]))
 				{
 					int EndGroupID = TrisActiveSubGroup[GroupNum - 1];
 					int StartGroupID = TrisActiveSubGroup[0];
 					int TriID0 = TriangleIDs[GroupStart];
-					bool bStoppedMerge = bShouldBreakGroup || bFirstGroupWasBroken; // if either group was broken, stop the merge before it even starts
-					for (int Idx = GroupNum - 1; !bStoppedMerge && Idx >= 0 && TrisActiveSubGroup[Idx] == EndGroupID; Idx--)
+					
+					for (int Idx = GroupNum - 1; Idx >= 0 && TrisActiveSubGroup[Idx] == EndGroupID; Idx--)
 					{
-						if (Idx < GroupNum - 1 && bNonTransitiveShare && !TrisCanShareVertexPredicate(VertexID, TriID0, TriangleIDs[GroupStart + Idx]))
-						{
-							bStoppedMerge = true;
-							break;
-						}
 						TrisActiveSubGroup[Idx] = StartGroupID;
 					}
 					NumGroupID--;
-					if (bStoppedMerge)
-					{
-						// Merge failed, break apart both groups
-						for (int BreakGroupIdx = 1; BreakGroupIdx < GroupNum; BreakGroupIdx++)
-						{
-							int Group = TrisActiveSubGroup[BreakGroupIdx];
-							if (Group == StartGroupID || Group == EndGroupID)
-							{
-								TrisActiveSubGroup[BreakGroupIdx] = NumGroupID++;
-							}
-						}
-					}
 				}
 			}
 
@@ -208,7 +160,47 @@ void TDynamicMeshOverlay<RealType, ElementSize>::CreateFromPredicate(TFunctionRe
 }
 
 
-
+template<typename RealType, int ElementSize>
+void TDynamicMeshOverlay<RealType, ElementSize>::SplitVerticesWithPredicate(TFunctionRef<bool(int ElementID, int TriID)> ShouldSplitOutVertex, TFunctionRef<void(int ElementID, int TriID, RealType* FillVect)> GetNewElementValue)
+{
+	for (int TriID : ParentMesh->TriangleIndicesItr())
+	{
+		FIndex3i ElTri = GetTriangle(TriID);
+		if (ElTri.A < 0)
+		{
+			// skip un-set triangles
+			continue;
+		}
+		bool TriChanged = false;
+		for (int SubIdx = 0; SubIdx < 3; SubIdx++)
+		{
+			int ElementID = ElTri[SubIdx];
+			// by convention for overlays, a ref count of 2 means that only one triangle has the element -- can't split it out further
+			if (ElementsRefCounts.GetRefCount(ElementID) <= 2)
+			{
+				// still set the new value though if the function wants to change it
+				if (ShouldSplitOutVertex(ElementID, TriID))
+				{
+					RealType NewElementData[ElementSize];
+					GetNewElementValue(ElementID, TriID, NewElementData);
+					SetElement(ElementID, NewElementData);
+				}
+			}
+			if (ShouldSplitOutVertex(ElementID, TriID))
+			{
+				TriChanged = true;
+				RealType NewElementData[ElementSize];
+				GetNewElementValue(ElementID, TriID, NewElementData);
+				ElTri[SubIdx] = AppendElement(NewElementData, ParentVertices[ElementID]);
+			}
+		}
+		if (TriChanged)
+		{
+			InternalSetTriangle(TriID, ElTri, true);
+		}
+		
+	}
+}
 
 
 template<typename RealType, int ElementSize>
diff --git a/Engine/Plugins/Experimental/GeometryProcessing/Source/DynamicMesh/Public/DynamicMeshOverlay.h b/Engine/Plugins/Experimental/GeometryProcessing/Source/DynamicMesh/Public/DynamicMeshOverlay.h
--- a/Engine/Plugins/Experimental/GeometryProcessing/Source/DynamicMesh/Public/DynamicMeshOverlay.h
+++ b/Engine/Plugins/Experimental/GeometryProcessing/Source/DynamicMesh/Public/DynamicMeshOverlay.h
@@ -110,14 +110,19 @@ public:
 	 * Build overlay topology from a predicate function, e.g. to build topology for sharp normals
 	 * 
 	 * @param TrisCanShareVertexPredicate Indicator function returns true if the given vertex can be shared for the given pair of triangles
-	 * @param bNonTransitiveShare If false, sharing is transitive so if Tris A and B can share a vertex, and B and C can share the vertex, then A, B and C will all use the same vertex, even if A and C could not share.
-	 *							  Otherwise, for a sequence of tris A, B, C, D, E that would all possibly share the same vertex, we require that BOTH:
-	 *									1. each adjacent pair (A,B), (B,C), etc pass the the TrisCanShareVertexPredicate and
-	 *									2. the first triangle in the sequence, A, passes the predicate with all other grouped triangles (A,C), (A,D), etc
-	 *								For a contiguous section that passes the adjacent pair tests but not the second test, we give each triangle in the section a separate normal for that vertex
-	 *								(doing anything 'smarter' like trying to find contiguous subgroups ends up just looking weird)
+	 *									  Note if a vertex can be shared between tris A and B, and B and C, it will be shared between all three
+	 * @param InitElementValue Initial element value, copied into all created elements
 	 */
-	void CreateFromPredicate(TFunctionRef<bool(int ParentVertexIdx, int TriIDA, int TriIDB)> TrisCanShareVertexPredicate, RealType InitElementValue, bool bNonTransitiveShare);
+	void CreateFromPredicate(TFunctionRef<bool(int ParentVertexIdx, int TriIDA, int TriIDB)> TrisCanShareVertexPredicate, RealType InitElementValue);
+
+	/**
+	 * Refine an existing overlay topology.  For any element on a given triangle, if the predicate returns true, it gets topologically split out so it isn't shared by any other triangle.
+	 * Used for creating sharp vertices in the normals overlay.
+	 *
+	 * @param ShouldSplitOutVertex predicate returns true of the element should be split out and not shared w/ any other triangle
+	 * @param GetNewElementValue function to assign a new value to any element that is split out
+	 */
+	void SplitVerticesWithPredicate(TFunctionRef<bool(int ElementIdx, int TriID)> ShouldSplitOutVertex, TFunctionRef<void(int ElementIdx, int TriID, RealType* FillVect)> GetNewElementValue);
 
 
 	//
diff --git a/Engine/Plugins/Experimental/MeshModelingToolset/Source/ModelingOperators/Private/CleaningOps/EditNormalsOp.cpp b/Engine/Plugins/Experimental/MeshModelingToolset/Source/ModelingOperators/Private/CleaningOps/EditNormalsOp.cpp
--- a/Engine/Plugins/Experimental/MeshModelingToolset/Source/ModelingOperators/Private/CleaningOps/EditNormalsOp.cpp
+++ b/Engine/Plugins/Experimental/MeshModelingToolset/Source/ModelingOperators/Private/CleaningOps/EditNormalsOp.cpp
@@ -33,6 +33,11 @@ void FEditNormalsOp::CalculateResult(FProgressCancel* Progress)
 		ResultMesh->EnableAttributes();
 	}
 
+	if (Progress->Cancelled())
+	{
+		return;
+	}
+
 	// if you split normals you must always recompute as well
 	bool bNeedsRecompute = bRecomputeNormals || bSplitNormals;
 
@@ -40,10 +45,20 @@ void FEditNormalsOp::CalculateResult(FProgressCancel* Progress)
 	{
 		FMeshRepairOrientation Repair(ResultMesh.Get());
 		Repair.OrientComponents();
+
+		if (Progress->Cancelled())
+		{
+			return;
+		}
 		FDynamicMeshAABBTree3 Tree(ResultMesh.Get());
 		Repair.SolveGlobalOrientation(&Tree);
 	}
 
+	if (Progress->Cancelled())
+	{
+		return;
+	}
+
 	if (bInvertNormals)
 	{
 		for (int TID : ResultMesh->TriangleIndicesItr())
@@ -63,16 +78,27 @@ void FEditNormalsOp::CalculateResult(FProgressCancel* Progress)
 		}
 	}
 
+	if (Progress->Cancelled())
+	{
+		return;
+	}
+
+	float NormalDotProdThreshold = FMathf::Cos(NormalSplitThreshold * FMathf::DegToRad);
+
+	FMeshNormals FaceNormals(ResultMesh.Get());
 	if (bSplitNormals)
 	{
-		FMeshNormals FaceNormals(ResultMesh.Get());
 		FaceNormals.ComputeTriangleNormals();
 		const TArray<FVector3d>& Normals = FaceNormals.GetNormals();
-		float Threshold = FMathf::Cos(NormalSplitThreshold * FMathf::DegToRad);
-		ResultMesh->Attributes()->PrimaryNormals()->CreateFromPredicate([&Normals, &Threshold](int VID, int TA, int TB)
+		ResultMesh->Attributes()->PrimaryNormals()->CreateFromPredicate([&Normals, &NormalDotProdThreshold](int VID, int TA, int TB)
 		{
-			return Normals[TA].Dot(Normals[TB]) > Threshold;
-		}, 0, bAllowSharpVertices);
+			return Normals[TA].Dot(Normals[TB]) > NormalDotProdThreshold;
+		}, 0);
+	}
+
+	if (Progress->Cancelled())
+	{
+		return;
 	}
 
 	bool bAreaWeight = (NormalCalculationMethod == ENormalCalculationMethod::AreaWeighted || NormalCalculationMethod == ENormalCalculationMethod::AreaAngleWeighting);
@@ -85,4 +111,27 @@ void FEditNormalsOp::CalculateResult(FProgressCancel* Progress)
 		MeshNormals.CopyToOverlay(ResultMesh->Attributes()->PrimaryNormals(), false);
 	}
 
+	if (Progress->Cancelled())
+	{
+		return;
+	}
+
+	if (bAllowSharpVertices)
+	{
+		ResultMesh->Attributes()->PrimaryNormals()->SplitVerticesWithPredicate([this, &FaceNormals, &NormalDotProdThreshold](int ElementID, int TriID)
+		{
+			FVector3f ElNormal;
+			ResultMesh->Attributes()->PrimaryNormals()->GetElement(ElementID, ElNormal);
+			return ElNormal.Dot(FVector3f(FaceNormals.GetNormals()[TriID])) <= NormalDotProdThreshold;
+		},
+			[this, &FaceNormals](int ElementIdx, int TriID, float* FillVect)
+		{
+			FVector3f N(FaceNormals.GetNormals()[TriID]);
+			FillVect[0] = N.X;
+			FillVect[1] = N.Y;
+			FillVect[2] = N.Z;
+		}
+		);
+	}
+
 }
\ No newline at end of file
