safely handle non-finite vertex positions that can be generated during sculpting. Copy CL 9965788 from Dev-Editor.
#rnx

diff --git a/Engine/Plugins/Experimental/GeometryProcessing/Source/DynamicMesh/Private/Remesher.cpp b/Engine/Plugins/Experimental/GeometryProcessing/Source/DynamicMesh/Private/Remesher.cpp
--- a/Engine/Plugins/Experimental/GeometryProcessing/Source/DynamicMesh/Private/Remesher.cpp
+++ b/Engine/Plugins/Experimental/GeometryProcessing/Source/DynamicMesh/Private/Remesher.cpp
@@ -569,7 +569,12 @@ FVector3d FRemesher::ComputeSmoothedVertexPos(int vID,
 	}
 
 	FVector3d vSmoothed = smoothFunc(*Mesh, vID, SmoothSpeedT);
-	check(VectorUtil::IsFinite(vSmoothed));     // this will really catch a lot of bugs...
+	// @todo we should probably make sure that vertex does not move too far here...
+	checkSlow(VectorUtil::IsFinite(vSmoothed));
+	if (VectorUtil::IsFinite(vSmoothed) == false)
+	{
+		return Mesh->GetVertex(vID);
+	}
 
 	// project onto either vtx constraint target, or surface target
 	if (vConstraint.Target != nullptr) 
diff --git a/Engine/Plugins/Experimental/GeometryProcessing/Source/DynamicMesh/Public/DynamicMesh3.h b/Engine/Plugins/Experimental/GeometryProcessing/Source/DynamicMesh/Public/DynamicMesh3.h
--- a/Engine/Plugins/Experimental/GeometryProcessing/Source/DynamicMesh/Public/DynamicMesh3.h
+++ b/Engine/Plugins/Experimental/GeometryProcessing/Source/DynamicMesh/Public/DynamicMesh3.h
@@ -528,13 +528,16 @@ public:
 	/** Set vertex position */
 	inline void SetVertex(int VertexID, const FVector3d& vNewPos)
 	{
-		check(VectorUtil::IsFinite(vNewPos));
+		checkSlow(VectorUtil::IsFinite(vNewPos));
 		check(IsVertex(VertexID));
-		int i = 3 * VertexID;
-		Vertices[i] = vNewPos.X;
-		Vertices[i + 1] = vNewPos.Y;
-		Vertices[i + 2] = vNewPos.Z;
-		UpdateTimeStamp(true, false);
+		if (VectorUtil::IsFinite(vNewPos))
+		{
+			int i = 3 * VertexID;
+			Vertices[i] = vNewPos.X;
+			Vertices[i + 1] = vNewPos.Y;
+			Vertices[i + 2] = vNewPos.Z;
+			UpdateTimeStamp(true, false);
+		}
 	}
 
 	/** Get extended vertex information */
