fix plane cut normal+UVs in degenerate cases where hole fill failed to add some triangles

#rnx

diff --git a/Engine/Plugins/Experimental/GeometryProcessing/Source/DynamicMesh/Private/Operations/MeshPlaneCut.cpp b/Engine/Plugins/Experimental/GeometryProcessing/Source/DynamicMesh/Private/Operations/MeshPlaneCut.cpp
--- a/Engine/Plugins/Experimental/GeometryProcessing/Source/DynamicMesh/Private/Operations/MeshPlaneCut.cpp
+++ b/Engine/Plugins/Experimental/GeometryProcessing/Source/DynamicMesh/Private/Operations/MeshPlaneCut.cpp
@@ -215,21 +215,15 @@ bool FMeshPlaneCut::SimpleHoleFill(int ConstantGroupID)
 	{
 		FSimpleHoleFiller Filler(Mesh, Loop);
 		int GID = ConstantGroupID >= 0 ? ConstantGroupID : Mesh->AllocateTriangleGroup();
-		if (Filler.Fill(GID))
-		{
-			HoleFillTriangles.Add(Filler.NewTriangles);
+		bAllOk = Filler.Fill(GID) && bAllOk;
+		
+		HoleFillTriangles.Add(Filler.NewTriangles);
 
-			if (Mesh->HasAttributes())
-			{
-				FDynamicMeshEditor Editor(Mesh);
-				Editor.SetTriangleNormals(Filler.NewTriangles, (FVector3f)PlaneNormal);
-				Editor.SetTriangleUVsFromProjection(Filler.NewTriangles, ProjectionFrame, UVScaleFactor);
-			}
-
-		}
-		else
+		if (Mesh->HasAttributes())
 		{
-			bAllOk = false;
+			FDynamicMeshEditor Editor(Mesh);
+			Editor.SetTriangleNormals(Filler.NewTriangles, (FVector3f)PlaneNormal);
+			Editor.SetTriangleUVsFromProjection(Filler.NewTriangles, ProjectionFrame, UVScaleFactor);
 		}
 	}
 
@@ -256,11 +250,8 @@ bool FMeshPlaneCut::HoleFill(TFunction<TArray<FIndex3i>(const FGeneralPolygon2d&
 	FPlanarHoleFiller Filler(Mesh, &LoopVertices, PlanarTriangulationFunc, PlaneOrigin, PlaneNormal);
 
 	int GID = ConstantGroupID >= 0 ? ConstantGroupID : Mesh->AllocateTriangleGroup();
-	if (!Filler.Fill(GID))
-	{
-		return false;
-	}
-
+	bool bFullyFilledHole = Filler.Fill(GID);
+	
 	HoleFillTriangles.Add(Filler.NewTriangles);
 	if (Mesh->HasAttributes())
 	{
@@ -271,5 +262,5 @@ bool FMeshPlaneCut::HoleFill(TFunction<TArray<FIndex3i>(const FGeneralPolygon2d&
 		Editor.SetTriangleUVsFromProjection(Filler.NewTriangles, ProjectionFrame, UVScaleFactor);
 	}
 
-	return true;
+	return bFullyFilledHole;
 }
diff --git a/Engine/Plugins/Experimental/GeometryProcessing/Source/DynamicMesh/Private/Operations/PlanarHoleFiller.cpp b/Engine/Plugins/Experimental/GeometryProcessing/Source/DynamicMesh/Private/Operations/PlanarHoleFiller.cpp
--- a/Engine/Plugins/Experimental/GeometryProcessing/Source/DynamicMesh/Private/Operations/PlanarHoleFiller.cpp
+++ b/Engine/Plugins/Experimental/GeometryProcessing/Source/DynamicMesh/Private/Operations/PlanarHoleFiller.cpp
@@ -46,7 +46,10 @@ bool FPlanarHoleFiller::Fill(int GroupID)
 			{
 				bAddedAll = false;
 			}
-			NewTriangles.Add(NewTID);
+			else
+			{
+				NewTriangles.Add(NewTID);
+			}
 		}
 	}
 	
