handle cases where topology had to be modified (e.g. due to nonmanifold input) so we do not crash

#rnx

diff --git a/Engine/Plugins/Experimental/GeometryProcessing/Source/DynamicMesh/Public/DynamicMeshTriangleAttribute.h b/Engine/Plugins/Experimental/GeometryProcessing/Source/DynamicMesh/Public/DynamicMeshTriangleAttribute.h
--- a/Engine/Plugins/Experimental/GeometryProcessing/Source/DynamicMesh/Public/DynamicMeshTriangleAttribute.h
+++ b/Engine/Plugins/Experimental/GeometryProcessing/Source/DynamicMesh/Public/DynamicMeshTriangleAttribute.h
@@ -53,8 +53,8 @@ public:
 	void Initialize(AttribValueType InitialValue = (AttribValueType)0)
 	{
 		check(ParentMesh != nullptr);
-		AttribValues.Resize(0);
-		AttribValues.Resize( ParentMesh->MaxTriangleID() * AttribDimension );
+		AttribValues.Resize(ParentMesh->MaxTriangleID() * AttribDimension);
+		AttribValues.Fill(InitialValue);
 	}
 
 	void SetNewValue(int NewTriangleID, const AttribValueType* Data)
diff --git a/Engine/Plugins/Experimental/GeometryProcessing/Source/MeshConversion/Private/DynamicMeshToMeshDescription.cpp b/Engine/Plugins/Experimental/GeometryProcessing/Source/MeshConversion/Private/DynamicMeshToMeshDescription.cpp
--- a/Engine/Plugins/Experimental/GeometryProcessing/Source/MeshConversion/Private/DynamicMeshToMeshDescription.cpp
+++ b/Engine/Plugins/Experimental/GeometryProcessing/Source/MeshConversion/Private/DynamicMeshToMeshDescription.cpp
@@ -64,8 +64,6 @@ void FDynamicMeshToMeshDescription::Update(const FDynamicMesh3* MeshIn, FMeshDes
 void FDynamicMeshToMeshDescription::UpdateAttributes(const FDynamicMesh3* MeshIn, FMeshDescription& MeshOut, bool bUpdateNormals, bool bUpdateUVs)
 {
 	check(MeshIn->IsCompactV());
-	check(MeshIn->VertexCount() == MeshOut.Vertices().Num());
-	check(MeshIn->TriangleCount() == MeshOut.Triangles().Num())
 
 	if (bUpdateNormals)
 	{
@@ -77,10 +75,12 @@ void FDynamicMeshToMeshDescription::UpdateAttributes(const FDynamicMesh3* MeshIn
 			const FDynamicMeshNormalOverlay* Overlay = MeshIn->HasAttributes() ? MeshIn->Attributes()->PrimaryNormals() : nullptr;
 			if (Overlay)
 			{
+				check(MeshIn->TriangleCount() == MeshOut.Triangles().Num())
 				DynamicMeshToMeshDescriptionConversionHelper::SetAttributesFromOverlay(MeshIn, MeshOut, InstanceAttrib, Overlay);
 			}
 			else
 			{
+				check(MeshIn->VertexCount() == MeshOut.Vertices().Num());
 				for (int VertID : MeshIn->VertexIndicesItr())
 				{
 					FVector Normal = (FVector)MeshIn->GetVertexNormal(VertID);
diff --git a/Engine/Plugins/Experimental/MeshModelingToolset/Source/ModelingComponents/Private/OctreeDynamicMeshComponent.cpp b/Engine/Plugins/Experimental/MeshModelingToolset/Source/ModelingComponents/Private/OctreeDynamicMeshComponent.cpp
--- a/Engine/Plugins/Experimental/MeshModelingToolset/Source/ModelingComponents/Private/OctreeDynamicMeshComponent.cpp
+++ b/Engine/Plugins/Experimental/MeshModelingToolset/Source/ModelingComponents/Private/OctreeDynamicMeshComponent.cpp
@@ -126,7 +126,7 @@ void UOctreeDynamicMeshComponent::ApplyTransform(const FTransform3d& Transform,
 
 void UOctreeDynamicMeshComponent::Bake(FMeshDescription* MeshDescription, bool bHaveModifiedTopology, const FConversionToMeshDescriptionOptions& ConversionOptions)
 {
-	if (bHaveModifiedTopology == false)
+	if (bHaveModifiedTopology == false && Mesh.Get()->VertexCount() == MeshDescription->Vertices().Num())
 	{
 		FDynamicMeshToMeshDescription Converter(ConversionOptions);
 		Converter.Update(Mesh.Get(), *MeshDescription);
diff --git a/Engine/Plugins/Experimental/MeshModelingToolset/Source/ModelingComponents/Private/SimpleDynamicMeshComponent.cpp b/Engine/Plugins/Experimental/MeshModelingToolset/Source/ModelingComponents/Private/SimpleDynamicMeshComponent.cpp
--- a/Engine/Plugins/Experimental/MeshModelingToolset/Source/ModelingComponents/Private/SimpleDynamicMeshComponent.cpp
+++ b/Engine/Plugins/Experimental/MeshModelingToolset/Source/ModelingComponents/Private/SimpleDynamicMeshComponent.cpp
@@ -86,7 +86,7 @@ void USimpleDynamicMeshComponent::InitializeNewMesh()
 
 void USimpleDynamicMeshComponent::Bake(FMeshDescription* MeshDescription, bool bHaveModifiedTopology, const FConversionToMeshDescriptionOptions& ConversionOptions)
 {
-	if (bHaveModifiedTopology == false)
+	if (bHaveModifiedTopology == false && Mesh.Get()->VertexCount() == MeshDescription->Vertices().Num())
 	{
 		FDynamicMeshToMeshDescription Converter(ConversionOptions);
 		Converter.Update(Mesh.Get(), *MeshDescription);
