move GeometryProcessing plugin from NotForLicensees to Experimental
#rnx

diff --git a/Engine/Plugins/Experimental/GeometryProcessing/GeometryProcessing.uplugin b/Engine/Plugins/Experimental/GeometryProcessing/GeometryProcessing.uplugin
new file mode 100644
--- /dev/null
+++ b/Engine/Plugins/Experimental/GeometryProcessing/GeometryProcessing.uplugin
@@ -0,0 +1,43 @@
+{
+	"FileVersion": 3,
+	"Version": 1,
+	"VersionName": "0.1",
+	"FriendlyName": "GeometryProcessing",
+	"Description": "Classes and Algorithms for Processing 2D and 3D Geometry",
+	"Category": "Other",
+	"CreatedBy": "Epic Games, Inc.",
+	"CreatedByURL": "http://epicgames.com",
+	"DocsURL": "",
+	"MarketplaceURL": "",
+	"SupportURL": "",
+	"CanContainContent": true,
+	"IsBetaVersion": true,
+	"Installed": false,
+	"Modules": [
+		{
+			"Name": "GeometricObjects",
+			"Type": "Runtime",
+			"LoadingPhase": "Default"
+		},
+		{
+			"Name": "GeometryAlgorithms",
+			"Type": "Runtime",
+			"LoadingPhase": "Default"
+		},
+		{
+			"Name": "DynamicMesh",
+			"Type": "Runtime",
+			"LoadingPhase": "Default"
+		},
+		{
+			"Name": "MeshConversion",
+			"Type": "Runtime",
+			"LoadingPhase": "Default"
+		},
+		{
+			"Name": "MeshSolverUtilities",
+			"Type": "Runtime",
+			"LoadingPhase": "Default"
+		}			
+	]
+}
\ No newline at end of file
diff --git a/Engine/Plugins/Experimental/GeometryProcessing/Source/DynamicMesh/Private/DynamicMesh3.cpp b/Engine/Plugins/Experimental/GeometryProcessing/Source/DynamicMesh/Private/DynamicMesh3.cpp
new file mode 100644
--- /dev/null
+++ b/Engine/Plugins/Experimental/GeometryProcessing/Source/DynamicMesh/Private/DynamicMesh3.cpp
@@ -0,0 +1,1103 @@
+// Copyright 1998-2019 Epic Games, Inc. All Rights Reserved.
+
+#include "DynamicMesh3.h"
+#include "DynamicMeshAttributeSet.h"
+#include "Generators/MeshShapeGenerator.h"
+
+
+const int FDynamicMesh3::InvalidID = IndexConstants::InvalidID; // -1;
+const int FDynamicMesh3::NonManifoldID = -2;
+const int FDynamicMesh3::InvalidGroupID = IndexConstants::InvalidID; // -1;
+
+
+FDynamicMesh3::~FDynamicMesh3() 
+{ 
+	if (VertexNormals != nullptr)
+	{
+		delete VertexNormals;
+	}
+	if (VertexColors != nullptr)
+	{
+		delete VertexColors;
+	}
+	if (VertexUVs != nullptr)
+	{
+		delete VertexUVs;
+	}
+	if (TriangleGroups != nullptr)
+	{
+		delete TriangleGroups;
+	}
+	if (AttributeSet != nullptr)
+	{
+		delete AttributeSet;
+	}
+}
+
+
+
+
+FDynamicMesh3::FDynamicMesh3(bool bWantNormals, bool bWantColors, bool bWantUVs, bool bWantTriGroups)
+{
+	Vertices = FDynamicVector<double>();
+	VertexNormals = (bWantNormals) ? new FDynamicVector<float>() : nullptr;
+	VertexColors = (bWantColors) ? new FDynamicVector<float>() : nullptr;
+	VertexUVs = (bWantUVs) ? new FDynamicVector<float>() : nullptr;
+
+	VertexEdgeLists = FSmallListSet();
+
+	VertexRefCounts = FRefCountVector();
+
+	Triangles = FDynamicVector<int>();
+	TriangleEdges = FDynamicVector<int>();
+	TriangleRefCounts = FRefCountVector();
+	TriangleGroups = (bWantTriGroups) ? new FDynamicVector<int>() : nullptr;
+	GroupIDCounter = 0;
+
+	Edges = FDynamicVector<int>();
+	EdgeRefCounts = FRefCountVector();
+
+	AttributeSet = nullptr;
+}
+
+
+FDynamicMesh3::FDynamicMesh3(const FDynamicMesh3& CopyMesh, bool bCompact, EMeshComponents flags)
+	: FDynamicMesh3(CopyMesh, bCompact, ((int)flags & (int)EMeshComponents::VertexNormals) != 0, ((int)flags & (int)EMeshComponents::VertexColors) != 0,
+	((int)flags & (int)EMeshComponents::VertexUVs) != 0)
+{
+}
+
+
+
+
+// normals/colors/uvs will only be copied if they exist
+FDynamicMesh3::FDynamicMesh3(const FDynamicMesh3& CopyMesh, bool bCompact, bool bWantNormals, bool bWantColors, bool bWantUVs, bool bWantAttributes)
+{
+	VertexNormals = nullptr;
+	VertexColors = nullptr;
+	VertexUVs = nullptr;
+	TriangleGroups = nullptr;
+	AttributeSet = nullptr;
+
+	if (bCompact)
+	{
+		CompactCopy(CopyMesh, bWantNormals, bWantColors, bWantUVs, bWantAttributes);
+	}
+	else
+	{
+		Copy(CopyMesh, bWantNormals, bWantColors, bWantUVs, bWantAttributes);
+	}
+}
+
+
+
+FDynamicMesh3::FDynamicMesh3(const FMeshShapeGenerator* Generator)
+{
+	Copy(Generator);
+}
+void FDynamicMesh3::Copy(const FMeshShapeGenerator* Generator)
+{
+	Clear();
+
+	Vertices = FDynamicVector<double>();
+	VertexEdgeLists = FSmallListSet();
+	VertexRefCounts = FRefCountVector();
+	Triangles = FDynamicVector<int>();
+	TriangleEdges = FDynamicVector<int>();
+	TriangleRefCounts = FRefCountVector();
+	TriangleGroups = new FDynamicVector<int>();
+	GroupIDCounter = 0;
+	Edges = FDynamicVector<int>();
+	EdgeRefCounts = FRefCountVector();
+
+	EnableAttributes();
+	FDynamicMeshUVOverlay* UVOverlay = Attributes()->PrimaryUV();
+	FDynamicMeshNormalOverlay* NormalOverlay = Attributes()->PrimaryNormals();
+
+	int NumVerts = Generator->Vertices.Num();
+	for (int i = 0; i < NumVerts; ++i)
+	{
+		AppendVertex(Generator->Vertices[i]);
+	}
+	int NumUVs = Generator->UVs.Num();
+	for (int i = 0; i < NumUVs; ++i)
+	{
+		UVOverlay->AppendElement(Generator->UVs[i], Generator->UVParentVertex[i]);
+	}
+	int NumNormals = Generator->Normals.Num();
+	for (int i = 0; i < NumNormals; ++i)
+	{
+		NormalOverlay->AppendElement(Generator->Normals[i], Generator->NormalParentVertex[i]);
+	}
+
+	int NumTris = Generator->Triangles.Num();
+	for (int i = 0; i < NumTris; ++i)
+	{
+		int tid = AppendTriangle(Generator->Triangles[i], Generator->TrianglePolygonIDs[i]);
+		check(tid == i);
+		UVOverlay->SetTriangle(tid, Generator->TriangleUVs[i]);
+		NormalOverlay->SetTriangle(tid, Generator->TriangleNormals[i]);
+	}
+}
+
+
+
+
+
+void FDynamicMesh3::Copy(const FDynamicMesh3& copy, bool bNormals, bool bColors, bool bUVs, bool bAttributes)
+{
+	// currently cannot re-use existing attribute buffers
+	DiscardVertexNormals();
+	DiscardVertexColors();
+	DiscardVertexUVs();
+	DiscardTriangleGroups();
+	DiscardAttributes();
+
+	Vertices = FDynamicVector<double>(copy.Vertices);
+
+	VertexNormals = (bNormals && copy.HasVertexNormals()) ? new FDynamicVector<float>(*copy.VertexNormals) : nullptr;
+	VertexColors = (bColors && copy.HasVertexColors()) ? new FDynamicVector<float>(*copy.VertexColors) : nullptr;
+	VertexUVs = (bUVs && copy.HasVertexUVs()) ? new FDynamicVector<float>(*copy.VertexUVs) : nullptr;
+
+	VertexRefCounts = FRefCountVector(copy.VertexRefCounts);
+
+	VertexEdgeLists = FSmallListSet(copy.VertexEdgeLists);
+
+	Triangles = FDynamicVector<int>(copy.Triangles);
+	TriangleEdges = FDynamicVector<int>(copy.TriangleEdges);
+	TriangleRefCounts = FRefCountVector(copy.TriangleRefCounts);
+	TriangleGroups = (copy.HasTriangleGroups()) ? new FDynamicVector<int>(*copy.TriangleGroups) : nullptr;
+	GroupIDCounter = copy.GroupIDCounter;
+
+	Edges = FDynamicVector<int>(copy.Edges);
+	EdgeRefCounts = FRefCountVector(copy.EdgeRefCounts);
+
+	if (bAttributes && copy.HasAttributes())
+	{
+		EnableAttributes();
+		AttributeSet->Copy(*copy.Attributes());
+	}
+
+	Timestamp = FMath::Max(Timestamp + 1, copy.Timestamp);
+	ShapeTimestamp = TopologyTimestamp = Timestamp;
+}
+
+
+
+
+
+void FDynamicMesh3::CompactCopy(const FDynamicMesh3& copy, bool bNormals, bool bColors, bool bUVs, bool bAttributes, FCompactMaps* CompactInfo)
+{
+	// @todo these are not compacted and hence code cannot work. fix!
+	check(copy.AttributeSet == nullptr);
+	check(bNormals == false && bColors == false && bUVs == false);
+
+	// TODO can't do until CompactInfo works
+	//if ( copy.IsCompact() ) {
+	//    Copy(copy, bNormals, bColors, bUVs);
+	//    CompactInfo ci = CompactInfo() { MapV = IdentityIndexMap() };
+	//    return ci;
+	//}
+
+	// currently cannot re-use existing attribute buffers
+	DiscardVertexNormals();
+	DiscardVertexColors();
+	DiscardVertexUVs();
+	DiscardTriangleGroups();
+
+	Vertices = FDynamicVector<double>();
+	VertexEdgeLists = FSmallListSet();
+	VertexRefCounts = FRefCountVector();
+	Triangles = FDynamicVector<int>();
+	TriangleEdges = FDynamicVector<int>();
+	TriangleRefCounts = FRefCountVector();
+	Edges = FDynamicVector<int>();
+	EdgeRefCounts = FRefCountVector();
+	GroupIDCounter = 0;
+
+	TriangleGroups = (copy.HasTriangleGroups()) ? new FDynamicVector<int>(*copy.TriangleGroups) : nullptr;
+
+	// [TODO] if we knew some of these were dense we could copy directly...
+
+	FVertexInfo vinfo;
+	TArray<int> mapV; mapV.SetNum(copy.MaxVertexID());
+	for (int vid : copy.VertexIndicesItr()) 
+	{
+		copy.GetVertex(vid, vinfo, bNormals, bColors, bUVs);
+		mapV[vid] = AppendVertex(vinfo);
+	}
+
+	// [TODO] would be much faster to explicitly copy triangle & edge data structures!!
+	if (copy.HasTriangleGroups()) 
+	{
+		EnableTriangleGroups(0);
+	}
+	for (int tid : copy.TriangleIndicesItr()) 
+	{
+		FIndex3i t = copy.GetTriangle(tid);
+		t = FIndex3i(mapV[t.A], mapV[t.B], mapV[t.C]);
+		int g = (copy.HasTriangleGroups()) ? copy.GetTriangleGroup(tid) : InvalidID;
+		AppendTriangle(t, g);
+		GroupIDCounter = FMath::Max(GroupIDCounter, g + 1);
+	}
+
+	// generate compact info...?
+	if (CompactInfo != nullptr) 
+	{
+		for (int vid : copy.VertexIndicesItr()) 
+		{
+			CompactInfo->MapV.Add(vid, mapV[vid]);
+		}
+	}
+
+	Timestamp = FMath::Max(Timestamp + 1, copy.Timestamp);
+	ShapeTimestamp = TopologyTimestamp = Timestamp;
+}
+
+
+
+
+void FDynamicMesh3::Clear()
+{
+	Vertices.Clear();
+	VertexEdgeLists = FSmallListSet();
+	VertexRefCounts = FRefCountVector();
+	Triangles.Clear();
+	TriangleEdges.Clear();
+	TriangleRefCounts = FRefCountVector();
+	Edges.Clear();
+	EdgeRefCounts = FRefCountVector();
+
+	DiscardTriangleGroups();
+	DiscardVertexColors();
+	DiscardVertexUVs();
+	DiscardVertexNormals();
+	DiscardAttributes();
+
+	Timestamp = ShapeTimestamp = TopologyTimestamp = 0;
+	GroupIDCounter = 0;
+	CachedBoundingBoxTimestamp = -1;
+	bIsClosedCached = false;
+	CachedIsClosedTimestamp = -1;
+}
+
+
+
+int FDynamicMesh3::GetComponentsFlags() const
+{
+	int c = 0;
+	if (HasVertexNormals())
+	{
+		c |= (int)EMeshComponents::VertexNormals;
+	}
+	if (HasVertexColors())
+	{
+		c |= (int)EMeshComponents::VertexColors;
+	}
+	if (HasVertexUVs())
+	{
+		c |= (int)EMeshComponents::VertexUVs;
+	}
+	if (HasTriangleGroups())
+	{
+		c |= (int)EMeshComponents::FaceGroups;
+	}
+	return c;
+}
+
+
+
+
+
+
+
+
+
+
+void FDynamicMesh3::EnableVertexNormals(const FVector3f& InitialNormal)
+{
+	if (HasVertexNormals())
+	{
+		return;
+	}
+	VertexNormals = new FDynamicVector<float>();
+	int NV = MaxVertexID();
+	VertexNormals->Resize(3 * NV);
+	for (int i = 0; i < NV; ++i)
+	{
+		int vi = 3 * i;
+		(*VertexNormals)[vi] = InitialNormal.X;
+		(*VertexNormals)[vi + 1] = InitialNormal.Y;
+		(*VertexNormals)[vi + 2] = InitialNormal.Z;
+	}
+}
+
+void FDynamicMesh3::DiscardVertexNormals()
+{
+	if (VertexNormals != nullptr)
+	{
+		delete VertexNormals;
+		VertexNormals = nullptr;
+	}
+}
+
+
+
+
+
+void FDynamicMesh3::EnableVertexColors(const FVector3f& InitialColor)
+{
+	if (HasVertexColors())
+	{
+		return;
+	}
+	VertexColors = new FDynamicVector<float>();
+	int NV = MaxVertexID();
+	VertexColors->Resize(3 * NV);
+	for (int i = 0; i < NV; ++i)
+	{
+		int vi = 3 * i;
+		(*VertexColors)[vi] = InitialColor.X;
+		(*VertexColors)[vi + 1] = InitialColor.Y;
+		(*VertexColors)[vi + 2] = InitialColor.Z;
+	}
+}
+
+void FDynamicMesh3::DiscardVertexColors()
+{
+	if (VertexColors != nullptr)
+	{
+		delete VertexColors;
+		VertexColors = nullptr;
+	}
+}
+
+
+
+
+void FDynamicMesh3::EnableVertexUVs(const FVector2f& InitialUV)
+{
+	if (HasVertexUVs())
+	{
+		return;
+	}
+	VertexUVs = new FDynamicVector<float>();
+	int NV = MaxVertexID();
+	VertexUVs->Resize(2 * NV);
+	for (int i = 0; i < NV; ++i)
+	{
+		int vi = 2 * i;
+		(*VertexUVs)[vi] = InitialUV.X;
+		(*VertexUVs)[vi + 1] = InitialUV.Y;
+	}
+}
+
+void FDynamicMesh3::DiscardVertexUVs()
+{
+	if (VertexUVs != nullptr)
+	{
+		delete VertexUVs;
+		VertexUVs = nullptr;
+	}
+}
+
+
+
+void FDynamicMesh3::EnableTriangleGroups(int InitialGroup)
+{
+	if (HasTriangleGroups())
+	{
+		return;
+	}
+	TriangleGroups = new FDynamicVector<int>();
+	int NT = MaxTriangleID();
+	TriangleGroups->Resize(NT);
+	for (int i = 0; i < NT; ++i)
+	{
+		(*TriangleGroups)[i] = InitialGroup;
+	}
+	GroupIDCounter = 0;
+}
+
+void FDynamicMesh3::DiscardTriangleGroups()
+{
+	if (TriangleGroups != nullptr)
+	{
+		delete TriangleGroups;
+		TriangleGroups = nullptr;
+	}
+	GroupIDCounter = 0;
+}
+
+
+
+
+void FDynamicMesh3::EnableAttributes()
+{
+	if (HasAttributes())
+	{
+		return;
+	}
+	AttributeSet = new FDynamicMeshAttributeSet(this);
+	AttributeSet->Initialize(MaxVertexID(), MaxTriangleID());
+}
+
+
+void FDynamicMesh3::DiscardAttributes()
+{
+	if (AttributeSet != nullptr) 
+	{
+		delete AttributeSet;
+		AttributeSet = nullptr;
+	}
+}
+
+
+
+
+
+
+bool FDynamicMesh3::GetVertex(int vID, FVertexInfo& vinfo, bool bWantNormals, bool bWantColors, bool bWantUVs) const
+{
+	if (VertexRefCounts.IsValid(vID) == false)
+	{
+		return false;
+	}
+	int vi = 3 * vID;
+	vinfo.Position = FVector3d(Vertices[vi], Vertices[vi + 1], Vertices[vi + 2]);
+	vinfo.bHaveN = vinfo.bHaveUV = vinfo.bHaveC = false;
+	if (HasVertexNormals() && bWantNormals) 
+	{
+		vinfo.bHaveN = true;
+		vinfo.Normal = FVector3f((*VertexNormals)[vi], (*VertexNormals)[vi + 1], (*VertexNormals)[vi + 2]);
+	}
+	if (HasVertexColors() && bWantColors) 
+	{
+		vinfo.bHaveC = true;
+		vinfo.Color = FVector3f((*VertexColors)[vi], (*VertexColors)[vi + 1], (*VertexColors)[vi + 2]);
+	}
+	if (HasVertexUVs() && bWantUVs) 
+	{
+		vinfo.bHaveUV = true;
+		vinfo.UV = FVector2f((*VertexUVs)[2*vID], (*VertexUVs)[2*vID + 1]);
+	}
+	return true;
+}
+
+
+int FDynamicMesh3::GetMaxVtxEdgeCount() const 
+{
+	int max = 0;
+	for (int vid : VertexIndicesItr())
+	{
+		max = FMath::Max(max, VertexEdgeLists.GetCount(vid));
+	}
+	return max;
+}
+
+
+
+FVertexInfo FDynamicMesh3::GetVertexInfo(int i) const
+{
+	FVertexInfo vi = FVertexInfo();
+	vi.Position = GetVertex(i);
+	vi.bHaveN = vi.bHaveC = vi.bHaveUV = false;
+	if (HasVertexNormals()) 
+	{
+		vi.bHaveN = true;
+		vi.Normal = GetVertexNormal(i);
+	}
+	if (HasVertexColors()) 
+	{
+		vi.bHaveC = true;
+		vi.Color = GetVertexColor(i);
+	}
+	if (HasVertexUVs()) 
+	{
+		vi.bHaveUV = true;
+		vi.UV = GetVertexUV(i);
+	}
+	return vi;
+}
+
+
+
+FIndex3i FDynamicMesh3::GetTriNeighbourTris(int tID) const
+{
+	if (TriangleRefCounts.IsValid(tID))
+	{
+		int tei = 3 * tID;
+		FIndex3i nbr_t = FIndex3i::Zero();
+		for (int j = 0; j < 3; ++j)
+		{
+			int ei = 4 * TriangleEdges[tei + j];
+			nbr_t[j] = (Edges[ei + 2] == tID) ? Edges[ei + 3] : Edges[ei + 2];
+		}
+		return nbr_t;
+	}
+	else
+	{
+		return InvalidTriangle();
+	}
+}
+
+
+
+/** Enumerate triangle one ring at a vertex */
+FDynamicMesh3::vtx_triangles_enumerable FDynamicMesh3::VtxTrianglesItr(int vID) const
+{
+	check(VertexRefCounts.IsValid(vID));
+
+	TFunction<int(int, int&)> expand_f = [this,vID](int eid, int& k) 
+	{
+		int vOther = GetOtherEdgeVertex(eid, vID);
+		int i = 4 * eid;
+		if (k == -1) 
+		{
+			int et0 = Edges[i + 2];
+			if (TriHasSequentialVertices(et0, vID, vOther))
+			{
+				k = 0;
+				return et0;
+			}
+		}
+		if (k != 1) 
+		{
+			int et1 = Edges[i + 3];
+			if (et1 != InvalidID && TriHasSequentialVertices(et1, vID, vOther))
+			{
+				k = 1;
+				return et1;
+			}
+		}
+		k = -1;
+		return -1;
+	};
+
+	return vtx_triangles_enumerable(VertexEdgeLists.Values(vID), expand_f);
+}
+
+
+
+
+
+
+
+
+FString FDynamicMesh3::MeshInfoString()
+{
+	FString VtxString = FString::Printf(TEXT("Vertices count %d max %d  %s  VtxEdges %s"),
+		VertexCount(), MaxVertexID(), *VertexRefCounts.UsageStats(), *(VertexEdgeLists.MemoryUsage()));
+	FString TriString = FString::Printf(TEXT("Triangles count %d max %d  %s"),
+		TriangleCount(), MaxTriangleID(), *TriangleRefCounts.UsageStats());
+	FString EdgeString = FString::Printf(TEXT("Edges count %d max %d  %s"),
+		EdgeCount(), MaxEdgeID(), *EdgeRefCounts.UsageStats());
+	FString AttribString = FString::Printf(TEXT("VtxNormals %d  VtxColors %d  VtxUVs %d  TriGroups %d  Attributes %d"),
+		HasVertexNormals(), HasVertexColors(), HasVertexUVs(), HasTriangleGroups(), HasAttributes());
+	FString InfoString = FString::Printf(TEXT("Closed %d  Compact %d  Timestamp %d  ShapeTimestamp %d  TopologyTimestamp %d  MaxGroupID %d"),
+		GetCachedIsClosed(), IsCompact(), GetTimestamp(), GetShapeTimestamp(), GetTopologyTimestamp(), MaxGroupID());
+
+	return VtxString + "\n" + TriString + "\n" + EdgeString + "\n" + AttribString + "\n" + InfoString;
+}
+
+
+
+
+bool FDynamicMesh3::IsSameMesh(const FDynamicMesh3& m2, bool bCheckConnectivity, bool bCheckEdgeIDs,
+	bool bCheckNormals, bool bCheckColors, bool bCheckUVs, bool bCheckGroups,
+	float Epsilon)
+{
+	if (VertexCount() != m2.VertexCount())
+	{
+		return false;
+	}
+	if (TriangleCount() != m2.TriangleCount())
+	{
+		return false;
+	}
+	for (int vid : VertexIndicesItr()) 
+	{
+		if (m2.IsVertex(vid) == false || VectorUtil::EpsilonEqual(GetVertex(vid), m2.GetVertex(vid), (double)Epsilon) == false)
+		{
+			return false;
+		}
+	}
+	for (int tid : TriangleIndicesItr()) 
+	{
+		if (m2.IsTriangle(tid) == false || (GetTriangle(tid) != m2.GetTriangle(tid)))
+		{
+			return false;
+		}
+	}
+	if (bCheckConnectivity) 
+	{
+		for (int eid : EdgeIndicesItr()) 
+		{
+			FIndex4i e = GetEdge(eid);
+			int other_eid = m2.FindEdge(e[0], e[1]);
+			if (other_eid == InvalidID)
+			{
+				return false;
+			}
+			FIndex4i oe = m2.GetEdge(other_eid);
+			if (FMath::Min(e[2], e[3]) != FMath::Min(oe[2], oe[3]) || FMath::Max(e[2], e[3]) != FMath::Max(oe[2], oe[3]))
+			{
+				return false;
+			}
+		}
+	}
+	if (bCheckEdgeIDs) 
+	{
+		if (EdgeCount() != m2.EdgeCount())
+		{
+			return false;
+		}
+		for (int eid : EdgeIndicesItr()) 
+		{
+			if (m2.IsEdge(eid) == false || GetEdge(eid) != m2.GetEdge(eid))
+			{
+				return false;
+			}
+		}
+	}
+	if (bCheckNormals) 
+	{
+		if (HasVertexNormals() != m2.HasVertexNormals())
+		{
+			return false;
+		}
+		if (HasVertexNormals()) 
+		{
+			for (int vid : VertexIndicesItr()) 
+			{
+				if (VectorUtil::EpsilonEqual(GetVertexNormal(vid), m2.GetVertexNormal(vid), Epsilon) == false)
+				{
+					return false;
+				}
+			}
+		}
+	}
+	if (bCheckColors) 
+	{
+		if (HasVertexColors() != m2.HasVertexColors())
+		{
+			return false;
+		}
+		if (HasVertexColors()) 
+		{
+			for (int vid : VertexIndicesItr()) 
+			{
+				if (VectorUtil::EpsilonEqual(GetVertexColor(vid), m2.GetVertexColor(vid), Epsilon) == false)
+				{
+					return false;
+				}
+			}
+		}
+	}
+	if (bCheckUVs) 
+	{
+		if (HasVertexUVs() != m2.HasVertexUVs())
+		{
+			return false;
+		}
+		if (HasVertexUVs()) 
+		{
+			for (int vid : VertexIndicesItr()) 
+			{
+				if (VectorUtil::EpsilonEqual(GetVertexUV(vid), m2.GetVertexUV(vid), Epsilon) == false)
+				{
+					return false;
+				}
+			}
+		}
+	}
+	if (bCheckGroups) 
+	{
+		if (HasTriangleGroups() != m2.HasTriangleGroups())
+		{
+			return false;
+		}
+		if (HasTriangleGroups()) 
+		{
+			for (int tid : TriangleIndicesItr()) 
+			{
+				if (GetTriangleGroup(tid) != m2.GetTriangleGroup(tid))
+				{
+					return false;
+				}
+			}
+		}
+	}
+	return true;
+}
+
+
+
+
+
+
+
+bool FDynamicMesh3::CheckValidity(bool bAllowNonManifoldVertices, EValidityCheckFailMode FailMode) const
+{
+
+	TArray<int> triToVtxRefs;
+	triToVtxRefs.SetNum(MaxVertexID());
+	//int[] triToVtxRefs = new int[this.MaxVertexID];
+
+	bool is_ok = true;
+	TFunction<void(bool)> CheckOrFailF = [&](bool b)
+	{
+		is_ok = is_ok && b;
+	};
+	if (FailMode == EValidityCheckFailMode::Check)
+	{
+		CheckOrFailF = [&](bool b)
+		{
+			checkf(b, TEXT("FEdgeLoop::CheckValidity failed!"));
+			is_ok = is_ok && b;
+		};
+	}
+	else if (FailMode == EValidityCheckFailMode::Ensure)
+	{
+		CheckOrFailF = [&](bool b)
+		{
+			ensureMsgf(b, TEXT("FEdgeLoop::CheckValidity failed!"));
+			is_ok = is_ok && b;
+		};
+	}
+
+
+	for (int tID : TriangleIndicesItr()) 
+	{
+		CheckOrFailF(IsTriangle(tID));
+		CheckOrFailF(TriangleRefCounts.GetRefCount(tID) == 1);
+
+		// vertices must exist
+		FIndex3i tv = GetTriangle(tID);
+		for (int j = 0; j < 3; ++j) 
+		{
+			CheckOrFailF(IsVertex(tv[j]));
+			triToVtxRefs[tv[j]] += 1;
+		}
+
+		// edges must exist and reference this tri
+		FIndex3i e;
+		for (int j = 0; j < 3; ++j) 
+		{
+			int a = tv[j], b = tv[(j + 1) % 3];
+			e[j] = FindEdge(a, b);
+			CheckOrFailF(e[j] != InvalidID);
+			CheckOrFailF(EdgeHasTriangle(e[j], tID));
+			CheckOrFailF(e[j] == FindEdgeFromTri(a, b, tID));
+		}
+		CheckOrFailF(e[0] != e[1] && e[0] != e[2] && e[1] != e[2]);
+
+		// tri nbrs must exist and reference this tri, or same edge must be boundary edge
+		FIndex3i te = GetTriEdges(tID);
+		for (int j = 0; j < 3; ++j) 
+		{
+			int eid = te[j];
+			CheckOrFailF(IsEdge(eid));
+			int tOther = GetOtherEdgeTriangle(eid, tID);
+			if (tOther == InvalidID) 
+			{
+				CheckOrFailF(IsBoundaryTriangle(tID));
+				continue;
+			}
+
+			CheckOrFailF(TriHasNeighbourTri(tOther, tID) == true);
+
+			// edge must have same two verts as tri for same index
+			int a = tv[j], b = tv[(j + 1) % 3];
+			FIndex2i ev = GetEdgeV(te[j]);
+			CheckOrFailF(IndexUtil::SamePairUnordered(a, b, ev[0], ev[1]));
+
+			// also check that nbr edge has opposite orientation
+			FIndex3i othertv = GetTriangle(tOther);
+			int found = IndexUtil::FindTriOrderedEdge(b, a, othertv);
+			CheckOrFailF(found != InvalidID);
+		}
+	}
+
+
+	// edge verts/tris must exist
+	for (int eID : EdgeIndicesItr()) 
+	{
+		CheckOrFailF(IsEdge(eID));
+		CheckOrFailF(EdgeRefCounts.GetRefCount(eID) == 1);
+		FIndex2i ev = GetEdgeV(eID);
+		FIndex2i et = GetEdgeT(eID);
+		CheckOrFailF(IsVertex(ev[0]));
+		CheckOrFailF(IsVertex(ev[1]));
+		CheckOrFailF(et[0] != InvalidID);
+		CheckOrFailF(ev[0] < ev[1]);
+		CheckOrFailF(IsTriangle(et[0]));
+		if (et[1] != InvalidID) 
+		{
+			CheckOrFailF(IsTriangle(et[1]));
+		}
+	}
+
+	// verify compact check
+	bool is_compact = VertexRefCounts.IsDense();
+	if (is_compact) 
+	{
+		for (int vid = 0; vid < (int)Vertices.GetLength() / 3; ++vid) 
+		{
+			CheckOrFailF(VertexRefCounts.IsValid(vid));
+		}
+	}
+
+	// vertex edges must exist and reference this vert
+	for (int vID : VertexIndicesItr()) 
+	{
+		CheckOrFailF(IsVertex(vID));
+
+		FVector3d v = GetVertex(vID);
+		CheckOrFailF(FMathd::IsNaN(v.SquaredLength()) == false);
+		CheckOrFailF(FMathd::IsFinite(v.SquaredLength()));
+
+		for (int edgeid : VertexEdgeLists.Values(vID)) 
+		{
+			CheckOrFailF(IsEdge(edgeid));
+			CheckOrFailF(EdgeHasVertex(edgeid, vID));
+
+			int otherV = GetOtherEdgeVertex(edgeid, vID);
+			int e2 = FindEdge(vID, otherV);
+			CheckOrFailF(e2 != InvalidID);
+			CheckOrFailF(e2 == edgeid);
+			e2 = FindEdge(otherV, vID);
+			CheckOrFailF(e2 != InvalidID);
+			CheckOrFailF(e2 == edgeid);
+		}
+
+		for (int nbr_vid : VtxVerticesItr(vID)) 
+		{
+			CheckOrFailF(IsVertex(nbr_vid));
+			int edge = FindEdge(vID, nbr_vid);
+			CheckOrFailF(IsEdge(edge));
+		}
+
+		TArray<int> vTris, vTris2;
+		GetVtxTriangles(vID, vTris, false);
+		GetVtxTriangles(vID, vTris2, true);
+		CheckOrFailF(vTris.Num() == vTris2.Num());
+		//System.Console.WriteLine(string.Format("{0} {1} {2}", vID, vTris.Count, GetVtxEdges(vID).Count));
+		if (bAllowNonManifoldVertices)
+		{
+			CheckOrFailF(vTris.Num() <= GetVtxEdgeCount(vID));
+		}
+		else
+		{
+			CheckOrFailF(vTris.Num() == GetVtxEdgeCount(vID) || vTris.Num() == GetVtxEdgeCount(vID) - 1);
+		}
+		CheckOrFailF(VertexRefCounts.GetRefCount(vID) == vTris.Num() + 1);
+		CheckOrFailF(triToVtxRefs[vID] == vTris.Num());
+		for (int tID : vTris) 
+		{
+			CheckOrFailF(TriangleHasVertex(tID, vID));
+		}
+
+		// check that edges around vert only references tris above, and reference all of them!
+		TArray<int> vRemoveTris(vTris);
+		for (int edgeid : VertexEdgeLists.Values(vID)) 
+		{
+			FIndex2i edget = GetEdgeT(edgeid);
+			CheckOrFailF(vTris.Contains(edget[0]));
+			if (edget[1] != InvalidID)
+			{
+				CheckOrFailF(vTris.Contains(edget[1]));
+			}
+			vRemoveTris.Remove(edget[0]);
+			if (edget[1] != InvalidID)
+			{
+				vRemoveTris.Remove(edget[1]);
+			}
+		}
+		CheckOrFailF(vRemoveTris.Num() == 0);
+	}
+
+	if (HasAttributes())
+	{
+		Attributes()->PrimaryUV()->CheckValidity(true, FailMode);
+		Attributes()->PrimaryNormals()->CheckValidity(true, FailMode);
+	}
+
+	return is_ok;
+}
+
+
+
+
+
+
+
+
+
+int FDynamicMesh3::AddEdgeInternal(int vA, int vB, int tA, int tB)
+{
+	if (vB < vA) {
+		int t = vB; vB = vA; vA = t;
+	}
+	int eid = EdgeRefCounts.Allocate();
+	int i = 4 * eid;
+	Edges.InsertAt(vA, i);
+	Edges.InsertAt(vB, i + 1);
+	Edges.InsertAt(tA, i + 2);
+	Edges.InsertAt(tB, i + 3);
+
+	VertexEdgeLists.Insert(vA, eid);
+	VertexEdgeLists.Insert(vB, eid);
+	return eid;
+}
+
+
+int FDynamicMesh3::AddTriangleInternal(int a, int b, int c, int e0, int e1, int e2)
+{
+	int tid = TriangleRefCounts.Allocate();
+	int i = 3 * tid;
+	Triangles.InsertAt(c, i + 2);
+	Triangles.InsertAt(b, i + 1);
+	Triangles.InsertAt(a, i);
+	TriangleEdges.InsertAt(e2, i + 2);
+	TriangleEdges.InsertAt(e1, i + 1);
+	TriangleEdges.InsertAt(e0, i + 0);
+	return tid;
+}
+
+
+int FDynamicMesh3::ReplaceEdgeVertex(int eID, int vOld, int vNew)
+{
+	int i = 4 * eID;
+	int a = Edges[i], b = Edges[i + 1];
+	if (a == vOld) 
+	{
+		Edges[i] = FMath::Min(b, vNew);
+		Edges[i + 1] = FMath::Max(b, vNew);
+		return 0;
+	}
+	else if (b == vOld) 
+	{
+		Edges[i] = FMath::Min(a, vNew);
+		Edges[i + 1] = FMath::Max(a, vNew);
+		return 1;
+	}
+	else
+	{
+		return -1;
+	}
+}
+
+
+int FDynamicMesh3::ReplaceEdgeTriangle(int eID, int tOld, int tNew)
+{
+	int i = 4 * eID;
+	int a = Edges[i + 2], b = Edges[i + 3];
+	if (a == tOld) {
+		if (tNew == InvalidID) 
+		{
+			Edges[i + 2] = b;
+			Edges[i + 3] = InvalidID;
+		}
+		else
+		{
+			Edges[i + 2] = tNew;
+		}
+		return 0;
+	}
+	else if (b == tOld) 
+	{
+		Edges[i + 3] = tNew;
+		return 1;
+	}
+	else
+	{
+		return -1;
+	}
+}
+
+int FDynamicMesh3::ReplaceTriangleEdge(int tID, int eOld, int eNew)
+{
+	int i = 3 * tID;
+	if (TriangleEdges[i] == eOld) 
+	{
+		TriangleEdges[i] = eNew;
+		return 0;
+	}
+	else if (TriangleEdges[i + 1] == eOld) 
+	{
+		TriangleEdges[i + 1] = eNew;
+		return 1;
+	}
+	else if (TriangleEdges[i + 2] == eOld) 
+	{
+		TriangleEdges[i + 2] = eNew;
+		return 2;
+	}
+	else
+	{
+		return -1;
+	}
+}
+
+
+
+//! returns edge ID
+int FDynamicMesh3::FindTriangleEdge(int tID, int vA, int vB) const
+{
+	int i = 3 * tID;
+	int tv0 = Triangles[i], tv1 = Triangles[i + 1];
+	if (IndexUtil::SamePairUnordered(tv0, tv1, vA, vB)) return TriangleEdges[3 * tID];
+	int tv2 = Triangles[i + 2];
+	if (IndexUtil::SamePairUnordered(tv1, tv2, vA, vB)) return TriangleEdges[3 * tID + 1];
+	if (IndexUtil::SamePairUnordered(tv2, tv0, vA, vB)) return TriangleEdges[3 * tID + 2];
+	return InvalidID;
+}
+
+
+
+
+
+int FDynamicMesh3::FindEdge(int vA, int vB) const
+{
+	check(IsVertex(vA));
+	check(IsVertex(vB));
+
+	// [RMS] edge vertices must be sorted (min,max),
+	//   that means we only need one index-check in inner loop.
+	//   commented out code is robust to incorrect ordering, but slower.
+	int vO = FMath::Max(vA, vB);
+	int vI = FMath::Min(vA, vB);
+	for (int eid : VertexEdgeLists.Values(vI)) 
+	{
+		if (Edges[4 * eid + 1] == vO)
+		{
+			//if (DoesEdgeHaveVertex(eid, vO))
+			return eid;
+		}
+	}
+	return InvalidID;
+
+	// this is slower, likely because it creates func<> every time. can we do w/o that?
+	//return VertexEdgeLists.Find(vI, (eid) => { return Edges[4 * eid + 1] == vO; }, InvalidID);
+}
+
+int FDynamicMesh3::FindEdgeFromTri(int vA, int vB, int tID) const
+{
+	int i = 3 * tID;
+	int t0 = Triangles[i], t1 = Triangles[i + 1];
+	if (IndexUtil::SamePairUnordered(vA, vB, t0, t1))
+	{
+		return TriangleEdges[i];
+	}
+	int t2 = Triangles[i + 2];
+	if (IndexUtil::SamePairUnordered(vA, vB, t1, t2))
+	{
+		return TriangleEdges[i + 1];
+	}
+	if (IndexUtil::SamePairUnordered(vA, vB, t2, t0))
+	{
+		return TriangleEdges[i + 2];
+	}
+	return InvalidID;
+}
\ No newline at end of file
diff --git a/Engine/Plugins/Experimental/GeometryProcessing/Source/DynamicMesh/Private/DynamicMesh3_Queries.cpp b/Engine/Plugins/Experimental/GeometryProcessing/Source/DynamicMesh/Private/DynamicMesh3_Queries.cpp
new file mode 100644
--- /dev/null
+++ b/Engine/Plugins/Experimental/GeometryProcessing/Source/DynamicMesh/Private/DynamicMesh3_Queries.cpp
@@ -0,0 +1,836 @@
+// Copyright 1998-2019 Epic Games, Inc. All Rights Reserved.
+
+#include "DynamicMesh3.h"
+
+
+FIndex2i FDynamicMesh3::GetEdgeOpposingV(int eID) const
+{
+	// [TODO] there was a comment here saying this does more work than necessary??
+	// ** it is important that verts returned maintain [c,d] order!!
+	int i = 4 * eID;
+	int a = Edges[i], b = Edges[i + 1];
+	int t0 = Edges[i + 2], t1 = Edges[i + 3];
+	int c = IndexUtil::FindTriOtherVtx(a, b, Triangles, t0);
+	if (t1 != InvalidID) 
+	{
+		int d = IndexUtil::FindTriOtherVtx(a, b, Triangles, t1);
+		return FIndex2i(c, d);
+	}
+	else
+	{
+		return FIndex2i(c, InvalidID);
+	}
+}
+
+
+int FDynamicMesh3::GetVtxBoundaryEdges(int vID, int& e0, int& e1) const
+{
+	if (VertexRefCounts.IsValid(vID)) 
+	{
+		int count = 0;
+		for (int eid : VertexEdgeLists.Values(vID)) 
+		{
+			int ei = 4 * eid;
+			if (Edges[ei + 3] == InvalidID) 
+			{
+				if (count == 0)
+				{
+					e0 = eid;
+				}
+				else if (count == 1)
+				{
+					e1 = eid;
+				}
+				count++;
+			}
+		}
+		return count;
+	}
+	check(false);
+	return -1;
+}
+
+
+int FDynamicMesh3::GetAllVtxBoundaryEdges(int vID, TArray<int>& EdgeListOut) const
+{
+	if (VertexRefCounts.IsValid(vID)) 
+	{
+		int count = 0;
+		for (int eid : VertexEdgeLists.Values(vID)) 
+		{
+			int ei = 4 * eid;
+			if (Edges[ei + 3] == InvalidID)
+			{
+				EdgeListOut.Add(eid);
+				count++;
+			}
+		}
+		return count;
+	}
+	check(false);
+	return -1;
+}
+
+
+
+void FDynamicMesh3::GetVtxNbrhood(int eID, int vID, int& vOther, int& oppV1, int& oppV2, int& t1, int& t2) const
+{
+	int i = 4 * eID;
+	vOther = (Edges[i] == vID) ? Edges[i + 1] : Edges[i];
+	t1 = Edges[i + 2];
+	oppV1 = IndexUtil::FindTriOtherVtx(vID, vOther, Triangles, t1);
+	t2 = Edges[i + 3];
+	if (t2 != InvalidID)
+	{
+		oppV2 = IndexUtil::FindTriOtherVtx(vID, vOther, Triangles, t2);
+	}
+	else
+	{
+		t2 = InvalidID;
+	}
+}
+
+
+int FDynamicMesh3::GetVtxTriangleCount(int vID, bool bBruteForce) const
+{
+	if (bBruteForce) 
+	{
+		TArray<int> vTriangles;
+		if (GetVtxTriangles(vID, vTriangles, false) != EMeshResult::Ok)
+		{
+			return -1;
+		}
+		return (int)vTriangles.Num();
+	}
+
+	if (!IsVertex(vID))
+	{
+		return -1;
+	}
+	int N = 0;
+	for (int eid : VertexEdgeLists.Values(vID)) 
+	{
+		int vOther = GetOtherEdgeVertex(eid, vID);
+		int i = 4 * eid;
+		int et0 = Edges[i + 2];
+		if (TriHasSequentialVertices(et0, vID, vOther))
+		{
+			N++;
+		}
+		int et1 = Edges[i + 3];
+		if (et1 != InvalidID && TriHasSequentialVertices(et1, vID, vOther))
+		{
+			N++;
+		}
+	}
+	return N;
+}
+
+
+
+EMeshResult FDynamicMesh3::GetVtxTriangles(int vID, TArray<int>& TrianglesOut, bool bUseOrientation) const
+{
+	if (!IsVertex(vID))
+	{
+		return EMeshResult::Failed_NotAVertex;
+	}
+
+	if (bUseOrientation) 
+	{
+		for (int eid : VertexEdgeLists.Values(vID)) 
+		{
+			int vOther = GetOtherEdgeVertex(eid, vID);
+			int i = 4 * eid;
+			int et0 = Edges[i + 2];
+			if (TriHasSequentialVertices(et0, vID, vOther))
+			{
+				TrianglesOut.Add(et0);
+			}
+			int et1 = Edges[i + 3];
+			if (et1 != InvalidID && TriHasSequentialVertices(et1, vID, vOther))
+			{
+				TrianglesOut.Add(et1);
+			}
+		}
+	}
+	else 
+	{
+		// brute-force method
+		for (int eid : VertexEdgeLists.Values(vID)) 
+		{
+			int i = 4 * eid;
+			int t0 = Edges[i + 2];
+			TrianglesOut.AddUnique(t0);
+
+			int t1 = Edges[i + 3];
+			if (t1 != InvalidID)
+			{
+				TrianglesOut.AddUnique(t1);
+			}
+		}
+	}
+	return EMeshResult::Ok;
+}
+
+
+
+
+
+bool FDynamicMesh3::IsBoundaryVertex(int vID) const 
+{
+	check(IsVertex(vID));
+	for (int eid : VertexEdgeLists.Values(vID)) 
+	{
+		if (Edges[4 * eid + 3] == InvalidID)
+		{
+			return true;
+		}
+	}
+	return false;
+}
+
+
+
+bool FDynamicMesh3::IsBoundaryTriangle(int tID) const
+{
+	check(IsTriangle(tID));
+	int i = 3 * tID;
+	return IsBoundaryEdge(TriangleEdges[i]) || IsBoundaryEdge(TriangleEdges[i + 1]) || IsBoundaryEdge(TriangleEdges[i + 2]);
+}
+
+
+
+
+FIndex2i FDynamicMesh3::GetOrientedBoundaryEdgeV(int eID) const
+{
+	if (EdgeRefCounts.IsValid(eID)) 
+	{
+		int ei = 4 * eID;
+		if (Edges[ei + 3] == InvalidID) 
+		{
+			int a = Edges[ei], b = Edges[ei + 1];
+			int ti = 3 * Edges[ei + 2];
+			FIndex3i tri(Triangles[ti], Triangles[ti + 1], Triangles[ti + 2]);
+			int ai = IndexUtil::FindEdgeIndexInTri(a, b, tri);
+			return FIndex2i(tri[ai], tri[(ai + 1) % 3]);
+		}
+	}
+	check(false);
+	return InvalidEdge();
+}
+
+
+bool FDynamicMesh3::IsGroupBoundaryEdge(int eID) const
+{
+	check(IsEdge(eID));
+	check(HasTriangleGroups());
+
+	int et1 = Edges[4 * eID + 3];
+	if (et1 == InvalidID)
+	{
+		return false;
+	}
+	int g1 = (*TriangleGroups)[et1];
+	int et0 = Edges[4 * eID + 2];
+	int g0 = (*TriangleGroups)[et0];
+	return g1 != g0;
+}
+
+
+
+bool FDynamicMesh3::IsGroupBoundaryVertex(int vID) const
+{
+	check(IsVertex(vID));
+	check(HasTriangleGroups());
+
+	int group_id = InvalidGroupID;
+	for (int eID : VertexEdgeLists.Values(vID)) 
+	{
+		int et0 = Edges[4 * eID + 2];
+		int g0 = (*TriangleGroups)[et0];
+		if (group_id != g0) 
+		{
+			if (group_id == InvalidGroupID)
+			{
+				group_id = g0;
+			}
+			else
+			{
+				return true;        // saw multiple group IDs
+			}
+		}
+		int et1 = Edges[4 * eID + 3];
+		if (et1 != InvalidID) 
+		{
+			int g1 = (*TriangleGroups)[et1];
+			if (group_id != g1)
+			{
+				return true;        // saw multiple group IDs
+			}
+		}
+	}
+	return false;
+}
+
+
+
+bool FDynamicMesh3::IsGroupJunctionVertex(int vID) const
+{
+	check(IsVertex(vID));
+	check(HasTriangleGroups());
+
+	FIndex2i groups(InvalidGroupID, InvalidGroupID);
+	for (int eID : VertexEdgeLists.Values(vID)) 
+	{
+		FIndex2i et(Edges[4 * eID + 2], Edges[4 * eID + 3]);
+		for (int k = 0; k < 2; ++k) 
+		{
+			if (et[k] == InvalidID)
+			{
+				continue;
+			}
+			int g0 = (*TriangleGroups)[et[k]];
+			if (g0 != groups[0] && g0 != groups[1]) 
+			{
+				if (groups[0] != InvalidGroupID && groups[1] != InvalidGroupID)
+				{
+					return true;
+				}
+				if (groups[0] == InvalidGroupID)
+				{
+					groups[0] = g0;
+				}
+				else
+				{
+					groups[1] = g0;
+				}
+			}
+		}
+	}
+	return false;
+}
+
+
+bool FDynamicMesh3::GetVertexGroups(int vID, FIndex4i& groups) const
+{
+	check(IsVertex(vID));
+	check(HasTriangleGroups());
+
+	groups = FIndex4i(InvalidGroupID, InvalidGroupID, InvalidGroupID, InvalidGroupID);
+	int ng = 0;
+
+	for (int eID : VertexEdgeLists.Values(vID)) 
+	{
+		int et0 = Edges[4 * eID + 2];
+		int g0 = (*TriangleGroups)[et0];
+		if (groups.Contains(g0) == false)
+		{
+			groups[ng++] = g0;
+		}
+		if (ng == 4)
+		{
+			return false;
+		}
+		int et1 = Edges[4 * eID + 3];
+		if (et1 != InvalidID) 
+		{
+			int g1 = (*TriangleGroups)[et1];
+			if (groups.Contains(g1) == false)
+			{
+				groups[ng++] = g1;
+			}
+			if (ng == 4)
+			{
+				return false;
+			}
+		}
+	}
+	return true;
+}
+
+
+
+bool FDynamicMesh3::GetAllVertexGroups(int vID, TArray<int>& GroupsOut) const
+{
+	check(IsVertex(vID));
+	check(HasTriangleGroups());
+
+	for (int eID : VertexEdgeLists.Values(vID)) 
+	{
+		int et0 = Edges[4 * eID + 2];
+		int g0 = (*TriangleGroups)[et0];
+		GroupsOut.AddUnique(g0);
+
+		int et1 = Edges[4 * eID + 3];
+		if (et1 != InvalidID) 
+		{
+			int g1 = (*TriangleGroups)[et1];
+			GroupsOut.AddUnique(g1);
+		}
+	}
+	return true;
+}
+
+
+
+
+/**
+ * returns true if vID is a "bowtie" vertex, ie multiple disjoint triangle sets in one-ring
+ */
+bool FDynamicMesh3::IsBowtieVertex(int vID) const
+{
+	check(VertexRefCounts.IsValid(vID));
+
+	int nEdges = VertexEdgeLists.GetCount(vID);
+	if (nEdges == 0)
+	{
+		return false;
+	}
+
+	// find a boundary edge to start at
+	int start_eid = -1;
+	bool start_at_boundary = false;
+	for (int eid : VertexEdgeLists.Values(vID)) 
+	{
+		if (Edges[4 * eid + 3] == InvalidID) 
+		{
+			start_at_boundary = true;
+			start_eid = eid;
+			break;
+		}
+	}
+	// if no boundary edge, start at arbitrary edge
+	if (start_eid == -1)
+	{
+		start_eid = VertexEdgeLists.First(vID);
+	}
+	// initial triangle
+	int start_tid = Edges[4 * start_eid + 2];
+
+	int prev_tid = start_tid;
+	int prev_eid = start_eid;
+
+	// walk forward to next edge. if we hit start edge or boundary edge,
+	// we are done the walk. count number of edges as we go.
+	int count = 1;
+	while (true) 
+	{
+		int i = 3 * prev_tid;
+		FIndex3i tv(Triangles[i], Triangles[i + 1], Triangles[i + 2]);
+		FIndex3i te(TriangleEdges[i], TriangleEdges[i + 1], TriangleEdges[i + 2]);
+		int vert_idx = IndexUtil::FindTriIndex(vID, tv);
+		int e1 = te[vert_idx], e2 = te[(vert_idx + 2) % 3];
+		int next_eid = (e1 == prev_eid) ? e2 : e1;
+		if (next_eid == start_eid)
+		{
+			break;
+		}
+		FIndex2i next_eid_tris = GetEdgeT(next_eid);
+		int next_tid = (next_eid_tris[0] == prev_tid) ? next_eid_tris[1] : next_eid_tris[0];
+		if (next_tid == InvalidID) 
+		{
+			break;
+		}
+		prev_eid = next_eid;
+		prev_tid = next_tid;
+		count++;
+	}
+
+	// if we did not see all edges at vertex, we have a bowtie
+	int target_count = (start_at_boundary) ? nEdges - 1 : nEdges;
+	bool is_bowtie = (target_count != count);
+	return is_bowtie;
+}
+
+
+
+
+int FDynamicMesh3::FindTriangle(int a, int b, int c) const
+{
+	int eid = FindEdge(a, b);
+	if (eid == InvalidID)
+	{
+		return InvalidID;
+	}
+	int ei = 4 * eid;
+
+	// triangles attached to edge [a,b] must contain verts a and b...
+	int ti = 3 * Edges[ei + 2];
+	if (Triangles[ti] == c || Triangles[ti + 1] == c || Triangles[ti + 2] == c)
+	{
+		return Edges[ei + 2];
+	}
+	if (Edges[ei + 3] != InvalidID) 
+	{
+		ti = 3 * Edges[ei + 3];
+		if (Triangles[ti] == c || Triangles[ti + 1] == c || Triangles[ti + 2] == c)
+		{
+			return Edges[ei + 3];
+		}
+	}
+
+	return InvalidID;
+}
+
+
+
+/**
+ * Computes bounding box of all vertices.
+ */
+FAxisAlignedBox3d FDynamicMesh3::GetBounds() const
+{
+	double x = 0, y = 0, z = 0;
+	for (int vi : VertexIndicesItr())     // find initial valid vertex
+	{
+		int k = 3 * vi;
+		x = Vertices[k]; y = Vertices[k + 1]; z = Vertices[k + 2];
+		break;
+	}
+	double minx = x, maxx = x, miny = y, maxy = y, minz = z, maxz = z;
+	for (int vi : VertexIndicesItr()) 
+	{
+		int k = 3 * vi;
+		x = Vertices[k]; y = Vertices[k + 1]; z = Vertices[k + 2];
+		if (x < minx) minx = x; else if (x > maxx) maxx = x;
+		if (y < miny) miny = y; else if (y > maxy) maxy = y;
+		if (z < minz) minz = z; else if (z > maxz) maxz = z;
+	}
+	return FAxisAlignedBox3d(FVector3d(minx, miny, minz), FVector3d(maxx, maxy, maxz));
+}
+
+
+FAxisAlignedBox3d FDynamicMesh3::GetCachedBounds()
+{
+	if (CachedBoundingBoxTimestamp != GetShapeTimestamp()) 
+	{
+		CachedBoundingBox = GetBounds();
+		CachedBoundingBoxTimestamp = GetShapeTimestamp();
+	}
+	return CachedBoundingBox;
+}
+
+
+
+
+bool FDynamicMesh3::IsClosed() const 
+{
+	if (TriangleCount() == 0) 
+	{
+		return false;
+	}
+
+	int N = MaxEdgeID();
+	for (int i = 0; i < N; ++i)
+	{
+		if (EdgeRefCounts.IsValid(i) && IsBoundaryEdge(i))
+		{
+			return false;
+		}
+	}
+	return true;
+}
+
+
+bool FDynamicMesh3::GetCachedIsClosed() 
+{
+	if (CachedIsClosedTimestamp != GetTopologyTimestamp()) 
+	{
+		bIsClosedCached = IsClosed();
+		CachedIsClosedTimestamp = GetTopologyTimestamp();
+	}
+	return bIsClosedCached;
+}
+
+
+
+
+
+// average of 1 or 2 face normals
+FVector3d FDynamicMesh3::GetEdgeNormal(int eID) const
+{
+	if (EdgeRefCounts.IsValid(eID)) 
+	{
+		int ei = 4 * eID;
+		FVector3d n = GetTriNormal(Edges[ei + 2]);
+		if (Edges[ei + 3] != InvalidID) 
+		{
+			n += GetTriNormal(Edges[ei + 3]);
+			n.Normalize();
+		}
+		return n;
+	}
+	check(false);
+	return FVector3d::Zero();
+}
+
+FVector3d FDynamicMesh3::GetEdgePoint(int eID, double t) const
+{
+	t = VectorUtil::Clamp(t, 0.0, 1.0);
+	if (EdgeRefCounts.IsValid(eID)) 
+	{
+		int ei = 4 * eID;
+		int iv0 = 3 * Edges[ei];
+		int iv1 = 3 * Edges[ei + 1];
+		double mt = 1.0 - t;
+		return FVector3d(
+			mt*Vertices[iv0] + t * Vertices[iv1],
+			mt*Vertices[iv0 + 1] + t * Vertices[iv1 + 1],
+			mt*Vertices[iv0 + 2] + t * Vertices[iv1 + 2]);
+	}
+	check(false);
+	return FVector3d::Zero();
+}
+
+
+
+void FDynamicMesh3::GetVtxOneRingCentroid(int vID, FVector3d& centroid) const
+{
+	centroid = FVector3d::Zero();
+	if (VertexRefCounts.IsValid(vID)) 
+	{
+		int n = 0;
+		for (int eid : VertexEdgeLists.Values(vID)) 
+		{
+			int other_idx = 3 * GetOtherEdgeVertex(eid, vID);
+			centroid.X += Vertices[other_idx];
+			centroid.Y += Vertices[other_idx + 1];
+			centroid.Z += Vertices[other_idx + 2];
+			n++;
+		}
+		if (n > 0) 
+		{
+			centroid *= 1.0 / n;
+		}
+	}
+}
+
+
+
+FFrame3d FDynamicMesh3::GetVertexFrame(int vID, bool bFrameNormalY) const
+{
+	check(HasVertexNormals());
+
+	int vi = 3 * vID;
+	FVector3d v(Vertices[vi], Vertices[vi + 1], Vertices[vi + 2]);
+	FVector3d normal((*VertexNormals)[vi], (*VertexNormals)[vi + 1], (*VertexNormals)[vi + 2]);
+	int eid = VertexEdgeLists.First(vID);
+	int ovi = 3 * GetOtherEdgeVertex(eid, vID);
+	FVector3d ov(Vertices[ovi], Vertices[ovi + 1], Vertices[ovi + 2]);
+	FVector3d edge = (ov - v);
+	edge.Normalize();
+
+	FVector3d other = normal.Cross(edge);
+	edge = other.Cross(normal);
+	if (bFrameNormalY)
+	{
+		return FFrame3d(v, edge, normal, -other);
+	}
+	else
+	{
+		return FFrame3d(v, edge, other, normal);
+	}
+}
+
+
+
+FVector3d FDynamicMesh3::GetTriNormal(int tID) const
+{
+	FVector3d v0, v1, v2;
+	GetTriVertices(tID, v0, v1, v2);
+	return VectorUtil::Normal(v0, v1, v2);
+}
+
+double FDynamicMesh3::GetTriArea(int tID) const
+{
+	FVector3d v0, v1, v2;
+	GetTriVertices(tID, v0, v1, v2);
+	return VectorUtil::Area(v0, v1, v2);
+}
+
+
+
+void FDynamicMesh3::GetTriInfo(int tID, FVector3d& Normal, double& Area, FVector3d& Centroid) const
+{
+	FVector3d v0, v1, v2;
+	GetTriVertices(tID, v0, v1, v2);
+	Centroid = (v0 + v1 + v2) * (1.0 / 3.0);
+	Area = VectorUtil::Area(v0, v1, v2);
+	Normal = VectorUtil::Normal(v0, v1, v2);
+	//normal = FastNormalArea(ref v0, ref v1, ref v2, out fArea);
+}
+
+
+FVector3d FDynamicMesh3::GetTriBaryPoint(int tID, double bary0, double bary1, double bary2) const 
+{
+	int ai = 3 * Triangles[3 * tID],
+		bi = 3 * Triangles[3 * tID + 1],
+		ci = 3 * Triangles[3 * tID + 2];
+	return FVector3d(
+		(bary0*Vertices[ai] + bary1 * Vertices[bi] + bary2 * Vertices[ci]),
+		(bary0*Vertices[ai + 1] + bary1 * Vertices[bi + 1] + bary2 * Vertices[ci + 1]),
+		(bary0*Vertices[ai + 2] + bary1 * Vertices[bi + 2] + bary2 * Vertices[ci + 2]));
+}
+
+
+FVector3d FDynamicMesh3::GetTriBaryNormal(int tID, double bary0, double bary1, double bary2) const
+{
+	check(HasVertexNormals());
+	int ai = 3 * Triangles[3 * tID],
+		bi = 3 * Triangles[3 * tID + 1],
+		ci = 3 * Triangles[3 * tID + 2];
+	const FDynamicVector<float>& normalsR = *(this->VertexNormals);
+	FVector3d n = FVector3d(
+		(bary0*normalsR[ai] + bary1 * normalsR[bi] + bary2 * normalsR[ci]),
+		(bary0*normalsR[ai + 1] + bary1 * normalsR[bi + 1] + bary2 * normalsR[ci + 1]),
+		(bary0*normalsR[ai + 2] + bary1 * normalsR[bi + 2] + bary2 * normalsR[ci + 2]));
+	n.Normalize();
+	return n;
+}
+
+FVector3d FDynamicMesh3::GetTriCentroid(int tID) const
+{
+	int ai = 3 * Triangles[3 * tID],
+		bi = 3 * Triangles[3 * tID + 1],
+		ci = 3 * Triangles[3 * tID + 2];
+	double f = (1.0 / 3.0);
+	return FVector3d(
+		(Vertices[ai] + Vertices[bi] + Vertices[ci]) * f,
+		(Vertices[ai + 1] + Vertices[bi + 1] + Vertices[ci + 1]) * f,
+		(Vertices[ai + 2] + Vertices[bi + 2] + Vertices[ci + 2]) * f);
+}
+
+
+void FDynamicMesh3::GetTriBaryPoint(int tID, double bary0, double bary1, double bary2, FVertexInfo& vinfo) const
+{
+	vinfo = FVertexInfo();
+	int ai = 3 * Triangles[3 * tID],
+		bi = 3 * Triangles[3 * tID + 1],
+		ci = 3 * Triangles[3 * tID + 2];
+	vinfo.Position = FVector3d(
+		(bary0 * Vertices[ai] + bary1 * Vertices[bi] + bary2 * Vertices[ci]),
+		(bary0 * Vertices[ai + 1] + bary1 * Vertices[bi + 1] + bary2 * Vertices[ci + 1]),
+		(bary0 * Vertices[ai + 2] + bary1 * Vertices[bi + 2] + bary2 * Vertices[ci + 2]));
+	vinfo.bHaveN = HasVertexNormals();
+	if (vinfo.bHaveN) 
+	{
+		FDynamicVector<float>& normalsR = *(this->VertexNormals);
+		vinfo.Normal = FVector3f(
+			(float)(bary0 * normalsR[ai] + bary1 * normalsR[bi] + bary2 * normalsR[ci]),
+			(float)(bary0 * normalsR[ai + 1] + bary1 * normalsR[bi + 1] + bary2 * normalsR[ci + 1]),
+			(float)(bary0 * normalsR[ai + 2] + bary1 * normalsR[bi + 2] + bary2 * normalsR[ci + 2]));
+		vinfo.Normal.Normalize();
+	}
+	vinfo.bHaveC = HasVertexColors();
+	if (vinfo.bHaveC) 
+	{
+		FDynamicVector<float>& colorsR = *(this->VertexColors);
+		vinfo.Color = FVector3f(
+			(float)(bary0 * colorsR[ai] + bary1 * colorsR[bi] + bary2 * colorsR[ci]),
+			(float)(bary0 * colorsR[ai + 1] + bary1 * colorsR[bi + 1] + bary2 * colorsR[ci + 1]),
+			(float)(bary0 * colorsR[ai + 2] + bary1 * colorsR[bi + 2] + bary2 * colorsR[ci + 2]));
+	}
+	vinfo.bHaveUV = HasVertexUVs();
+	if (vinfo.bHaveUV) 
+	{
+		FDynamicVector<float>& uvR = *(this->VertexUVs);
+		ai = 2 * Triangles[3 * tID];
+		bi = 2 * Triangles[3 * tID + 1];
+		ci = 2 * Triangles[3 * tID + 2];
+		vinfo.UV = FVector2f(
+			(float)(bary0 * uvR[ai] + bary1 * uvR[bi] + bary2 * uvR[ci]),
+			(float)(bary0 * uvR[ai + 1] + bary1 * uvR[bi + 1] + bary2 * uvR[ci + 1]));
+	}
+}
+
+
+FAxisAlignedBox3d FDynamicMesh3::GetTriBounds(int tID) const
+{
+	int vi = 3 * Triangles[3 * tID];
+	double x = Vertices[vi], y = Vertices[vi + 1], z = Vertices[vi + 2];
+	double minx = x, maxx = x, miny = y, maxy = y, minz = z, maxz = z;
+	for (int i = 1; i < 3; ++i) 
+	{
+		vi = 3 * Triangles[3 * tID + i];
+		x = Vertices[vi]; y = Vertices[vi + 1]; z = Vertices[vi + 2];
+		if (x < minx) minx = x; else if (x > maxx) maxx = x;
+		if (y < miny) miny = y; else if (y > maxy) maxy = y;
+		if (z < minz) minz = z; else if (z > maxz) maxz = z;
+	}
+	return FAxisAlignedBox3d(FVector3d(minx, miny, minz), FVector3d(maxx, maxy, maxz));
+}
+
+
+FFrame3d FDynamicMesh3::GetTriFrame(int tID, int nEdge) const
+{
+	int ti = 3 * tID;
+	int a = 3 * Triangles[ti + (nEdge % 3)];
+	int b = 3 * Triangles[ti + ((nEdge + 1) % 3)];
+	int c = 3 * Triangles[ti + ((nEdge + 2) % 3)];
+	FVector3d v1(Vertices[a], Vertices[a + 1], Vertices[a + 2]);
+	FVector3d v2(Vertices[b], Vertices[b + 1], Vertices[b + 2]);
+	FVector3d v3(Vertices[c], Vertices[c + 1], Vertices[c + 2]);
+
+	FVector3d edge1 = v2 - v1;  edge1.Normalize();
+	FVector3d edge2 = v3 - v2;  edge2.Normalize();
+	FVector3d normal = edge2.Cross(edge1); normal.Normalize();
+
+	FVector3d other = normal.Cross(edge1);
+
+	FVector3d center = (v1 + v2 + v3) / 3;
+	return FFrame3d(center, edge1, other, normal);
+}
+
+
+
+double FDynamicMesh3::GetTriSolidAngle(int tID, const FVector3d& p) const
+{
+	// inlined version of GetTriVertices & VectorUtil::TriSolidAngle
+	int ti = 3 * tID;
+	int ta = 3 * Triangles[ti];
+	FVector3d a(Vertices[ta] - p.X, Vertices[ta + 1] - p.Y, Vertices[ta + 2] - p.Z);
+	int tb = 3 * Triangles[ti + 1];
+	FVector3d b(Vertices[tb] - p.X, Vertices[tb + 1] - p.Y, Vertices[tb + 2] - p.Z);
+	int tc = 3 * Triangles[ti + 2];
+	FVector3d c(Vertices[tc] - p.X, Vertices[tc + 1] - p.Y, Vertices[tc + 2] - p.Z);
+	double la = a.Length(), lb = b.Length(), lc = c.Length();
+	double top = (la * lb * lc) + a.Dot(b) * lc + b.Dot(c) * la + c.Dot(a) * lb;
+	double bottom = a.X * (b.Y * c.Z - c.Y * b.Z) - a.Y * (b.X * c.Z - c.X * b.Z) + a.Z * (b.X * c.Y - c.X * b.Y);
+	// -2 instead of 2 to account for UE winding
+	return -2.0 * atan2(bottom, top);
+}
+
+
+
+double FDynamicMesh3::GetTriInternalAngleR(int tID, int i)
+{
+	int ti = 3 * tID;
+	int ta = 3 * Triangles[ti];
+	FVector3d a(Vertices[ta], Vertices[ta + 1], Vertices[ta + 2]);
+	int tb = 3 * Triangles[ti + 1];
+	FVector3d b(Vertices[tb], Vertices[tb + 1], Vertices[tb + 2]);
+	int tc = 3 * Triangles[ti + 2];
+	FVector3d c(Vertices[tc], Vertices[tc + 1], Vertices[tc + 2]);
+	if (i == 0)
+	{
+		return (b - a).Normalized().AngleR((c - a).Normalized());
+	}
+	else if (i == 1)
+	{
+		return (a - b).Normalized().AngleR((c - b).Normalized());
+	}
+	else
+	{
+		return (a - c).Normalized().AngleR((b - c).Normalized());
+	}
+}
+
+
+double FDynamicMesh3::CalculateWindingNumber(const FVector3d& QueryPoint) const
+{
+	double sum = 0;
+	for (int tid : TriangleIndicesItr())
+	{
+		sum += GetTriSolidAngle(tid, QueryPoint);
+	}
+	return sum / FMathd::FourPi;
+}
\ No newline at end of file
diff --git a/Engine/Plugins/Experimental/GeometryProcessing/Source/DynamicMesh/Private/DynamicMeshAttributeSet.cpp b/Engine/Plugins/Experimental/GeometryProcessing/Source/DynamicMesh/Private/DynamicMeshAttributeSet.cpp
new file mode 100644
--- /dev/null
+++ b/Engine/Plugins/Experimental/GeometryProcessing/Source/DynamicMesh/Private/DynamicMeshAttributeSet.cpp
@@ -0,0 +1,63 @@
+// Copyright 1998-2019 Epic Games, Inc. All Rights Reserved.
+
+#include "DynamicMeshAttributeSet.h"
+
+
+
+
+bool FDynamicMeshAttributeSet::IsSeamEdge(int eid) const
+{
+	return UV0.IsSeamEdge(eid) || Normals0.IsSeamEdge(eid);
+}
+
+
+
+void FDynamicMeshAttributeSet::OnNewTriangle(int TriangleID, bool bInserted)
+{
+	UV0.InitializeNewTriangle(TriangleID);
+	Normals0.InitializeNewTriangle(TriangleID);
+}
+
+void FDynamicMeshAttributeSet::OnRemoveTriangle(int TriangleID, bool bRemoveIsolatedVertices)
+{
+	UV0.OnRemoveTriangle(TriangleID, bRemoveIsolatedVertices);
+	Normals0.OnRemoveTriangle(TriangleID, bRemoveIsolatedVertices);
+}
+
+void FDynamicMeshAttributeSet::OnReverseTriOrientation(int TriangleID)
+{
+	UV0.OnReverseTriOrientation(TriangleID);
+	Normals0.OnReverseTriOrientation(TriangleID);
+}
+
+void FDynamicMeshAttributeSet::OnSplitEdge(const FDynamicMesh3::FEdgeSplitInfo & splitInfo)
+{
+	UV0.OnSplitEdge(splitInfo);
+	Normals0.OnSplitEdge(splitInfo);
+}
+
+void FDynamicMeshAttributeSet::OnFlipEdge(const FDynamicMesh3::FEdgeFlipInfo & flipInfo)
+{
+	UV0.OnFlipEdge(flipInfo);
+	Normals0.OnFlipEdge(flipInfo);
+}
+
+
+void FDynamicMeshAttributeSet::OnCollapseEdge(const FDynamicMesh3::FEdgeCollapseInfo & collapseInfo)
+{
+	UV0.OnCollapseEdge(collapseInfo);
+	Normals0.OnCollapseEdge(collapseInfo);
+}
+
+void FDynamicMeshAttributeSet::OnPokeTriangle(const FDynamicMesh3::FPokeTriangleInfo & pokeInfo)
+{
+	UV0.OnPokeTriangle(pokeInfo);
+	Normals0.OnPokeTriangle(pokeInfo);
+}
+
+void FDynamicMeshAttributeSet::OnMergeEdges(const FDynamicMesh3::FMergeEdgesInfo & mergeInfo)
+{
+	UV0.OnMergeEdges(mergeInfo);
+	Normals0.OnMergeEdges(mergeInfo);
+}
+
diff --git a/Engine/Plugins/Experimental/GeometryProcessing/Source/DynamicMesh/Private/DynamicMeshEditor.cpp b/Engine/Plugins/Experimental/GeometryProcessing/Source/DynamicMesh/Private/DynamicMeshEditor.cpp
new file mode 100644
--- /dev/null
+++ b/Engine/Plugins/Experimental/GeometryProcessing/Source/DynamicMesh/Private/DynamicMeshEditor.cpp
@@ -0,0 +1,484 @@
+// Copyright 1998-2019 Epic Games, Inc. All Rights Reserved.
+
+
+#include "DynamicMeshEditor.h"
+#include "DynamicMeshAttributeSet.h"
+#include "Util/BufferUtil.h"
+
+
+
+void FMeshIndexMappings::Initialize(FDynamicMesh3* Mesh)
+{
+	if (Mesh->HasAttributes())
+	{
+		FDynamicMeshAttributeSet* Attribs = Mesh->Attributes();
+		UVMaps.SetNum(Attribs->NumUVLayers());
+		NormalMaps.SetNum(Attribs->NumNormalLayers());
+	}
+}
+
+
+
+
+
+void FDynamicMeshEditResult::GetAllTriangles(TArray<int>& TrianglesOut) const
+{
+	BufferUtil::AppendElements(TrianglesOut, NewTriangles);
+
+	int NumQuads = NewQuads.Num();
+	for (int k = 0; k < NumQuads; ++k)
+	{
+		TrianglesOut.Add(NewQuads[k].A);
+		TrianglesOut.Add(NewQuads[k].B);
+	}
+	int NumPolys = NewPolygons.Num();
+	for (int k = 0; k < NumPolys; ++k)
+	{
+		BufferUtil::AppendElements(TrianglesOut, NewPolygons[k]);
+	}
+}
+
+
+
+
+
+
+bool FDynamicMeshEditor::StitchVertexLoopsMinimal(const TArray<int>& Loop1, const TArray<int>& Loop2, FDynamicMeshEditResult& ResultOut)
+{
+	int N = Loop1.Num();
+	checkf(N == Loop2.Num(), TEXT("FDynamicMeshEditor::StitchLoop: loops are not the same length!"));
+	if (N != Loop2.Num())
+	{
+		return false;
+	}
+
+	ResultOut.NewQuads.Reserve(N);
+	ResultOut.NewGroups.Reserve(N);
+
+	int i = 0;
+	for (; i < N; ++i) 
+	{
+		int a = Loop1[i];
+		int b = Loop1[(i + 1) % N];
+		int c = Loop2[i];
+		int d = Loop2[(i + 1) % N];
+
+		int NewGroupID = Mesh->AllocateTriangleGroup();
+		ResultOut.NewGroups.Add(NewGroupID);
+
+		FIndex3i t1(b, a, d);
+		int tid1 = Mesh->AppendTriangle(t1, NewGroupID);
+
+		FIndex3i t2(a, c, d);
+		int tid2 = Mesh->AppendTriangle(t2, NewGroupID);
+
+		ResultOut.NewQuads.Add(FIndex2i(tid1, tid2));
+
+		if (tid1 < 0 || tid2 < 0)
+		{
+			goto operation_failed;
+		}
+	}
+
+	return true;
+
+operation_failed:
+	// remove what we added so far
+	if (i > 0) 
+	{
+		ResultOut.NewTriangles.SetNum(2 * i + 1);
+		if (RemoveTriangles(ResultOut.NewTriangles, false) == false)
+		{
+			checkf(false, TEXT("FDynamicMeshEditor::StitchLoop: failed to add all triangles, and also failed to back out changes."));
+		}
+	}
+	return false;
+}
+
+
+
+
+bool FDynamicMeshEditor::RemoveTriangles(const TArray<int>& Triangles, bool bRemoveIsolatedVerts)
+{
+	bool bAllOK = true;
+	int NumTriangles = Triangles.Num();
+	for (int i = 0; i < NumTriangles; ++i) 
+	{
+		if (Mesh->IsTriangle(Triangles[i]) == false)
+		{
+			continue;
+		}
+
+		EMeshResult result = Mesh->RemoveTriangle(Triangles[i], bRemoveIsolatedVerts, false);
+		if (result != EMeshResult::Ok)
+		{
+			bAllOK = false;
+		}
+	}
+	return bAllOK;
+}
+
+
+
+
+
+/**
+ * Make a copy of provided triangles, with new vertices. You provide IndexMaps because
+ * you know if you are doing a small subset or a full-mesh-copy.
+ */
+void FDynamicMeshEditor::DuplicateTriangles(const TArray<int>& Triangles, FMeshIndexMappings& IndexMaps, FDynamicMeshEditResult& ResultOut)
+{
+	ResultOut.Reset();
+	IndexMaps.Initialize(Mesh);
+
+	for (int TriangleID : Triangles) 
+	{
+		FIndex3i Tri = Mesh->GetTriangle(TriangleID);
+
+		int NewGroupID = FindOrCreateDuplicateGroup(TriangleID, IndexMaps, ResultOut);
+
+		FIndex3i NewTri;
+		NewTri[0] = FindOrCreateDuplicateVertex(Tri[0], IndexMaps, ResultOut);
+		NewTri[1] = FindOrCreateDuplicateVertex(Tri[1], IndexMaps, ResultOut);
+		NewTri[2] = FindOrCreateDuplicateVertex(Tri[2], IndexMaps, ResultOut);
+
+		int NewTriangleID = Mesh->AppendTriangle(NewTri, NewGroupID);
+		IndexMaps.SetTriangle(TriangleID, NewTriangleID);
+		ResultOut.NewTriangles.Add(NewTriangleID);
+
+		CopyAttributes(TriangleID, NewTriangleID, IndexMaps, ResultOut);
+
+		Mesh->CheckValidity(true);
+	}
+
+}
+
+
+
+FVector3f FDynamicMeshEditor::ComputeAndSetQuadNormal(const FIndex2i& QuadTris, bool bIsPlanar)
+{
+	FVector3f Normal(0, 0, 1);
+	if (bIsPlanar)
+	{
+		Normal = (FVector3f)Mesh->GetTriNormal(QuadTris.A);
+	}
+	else
+	{
+		Normal = (FVector3f)Mesh->GetTriNormal(QuadTris.A);
+		Normal += (FVector3f)Mesh->GetTriNormal(QuadTris.B);
+		Normal.Normalize();
+	}
+	SetQuadNormals(QuadTris, Normal);
+	return Normal;
+}
+
+
+
+
+void FDynamicMeshEditor::SetQuadNormals(const FIndex2i& QuadTris, const FVector3f& Normal)
+{
+	check(Mesh->HasAttributes());
+	FDynamicMeshNormalOverlay* Normals = Mesh->Attributes()->PrimaryNormals();
+
+	FIndex3i Triangle1 = Mesh->GetTriangle(QuadTris.A);
+
+	FIndex3i NormalTriangle1;
+	NormalTriangle1[0] = Normals->AppendElement(Normal, Triangle1[0]);
+	NormalTriangle1[1] = Normals->AppendElement(Normal, Triangle1[1]);
+	NormalTriangle1[2] = Normals->AppendElement(Normal, Triangle1[2]);
+	Normals->SetTriangle(QuadTris.A, NormalTriangle1);
+
+	if (Mesh->IsTriangle(QuadTris.B))
+	{
+		FIndex3i Triangle2 = Mesh->GetTriangle(QuadTris.B);
+		FIndex3i NormalTriangle2;
+		for (int j = 0; j < 3; ++j)
+		{
+			int i = Triangle1.IndexOf(Triangle2[j]);
+			if (i == -1)
+			{
+				NormalTriangle2[j] = Normals->AppendElement(Normal, Triangle2[j]);
+			}
+			else
+			{
+				NormalTriangle2[j] = NormalTriangle1[i];
+			}
+		}
+		Normals->SetTriangle(QuadTris.B, NormalTriangle2);
+	}
+
+}
+
+
+void FDynamicMeshEditor::SetTriangleNormals(const TArray<int>& Triangles, const FVector3f& Normal)
+{
+	check(Mesh->HasAttributes());
+	FDynamicMeshNormalOverlay* Normals = Mesh->Attributes()->PrimaryNormals();
+
+	TMap<int, int> Vertices;
+
+	for (int tid : Triangles)
+	{
+		FIndex3i BaseTri = Mesh->GetTriangle(tid);
+		FIndex3i ElemTri;
+		for (int j = 0; j < 3; ++j)
+		{
+			const int* FoundElementID = Vertices.Find(BaseTri[j]);
+			if (FoundElementID == nullptr)
+			{
+				ElemTri[j] = Normals->AppendElement(Normal, BaseTri[j]);
+				Vertices.Add(BaseTri[j], ElemTri[j]);
+			}
+			else
+			{
+				ElemTri[j] = *FoundElementID;
+			}
+		}
+		Normals->SetTriangle(tid, ElemTri);
+	}
+}
+
+
+
+
+
+void FDynamicMeshEditor::SetQuadUVsFromProjection(const FIndex2i& QuadTris, const FFrame3f& ProjectionFrame, float UVScaleFactor, int UVLayerIndex)
+{
+	check(Mesh->HasAttributes() && Mesh->Attributes()->NumUVLayers() > UVLayerIndex );
+	FDynamicMeshUVOverlay* UVs = Mesh->Attributes()->GetUVLayer(UVLayerIndex);
+
+	FIndex4i AllUVIndices(-1, -1, -1, -1);
+	FVector2f AllUVs[4];
+
+	// project first triangle
+	FIndex3i Triangle1 = Mesh->GetTriangle(QuadTris.A);
+	FIndex3i UVTriangle1;
+	for (int j = 0; j < 3; ++j)
+	{
+		FVector2f UV = ProjectionFrame.ToPlaneUV( (FVector3f)Mesh->GetVertex(Triangle1[j]), 2);
+		UVTriangle1[j] = UVs->AppendElement(UV, Triangle1[j]);
+		AllUVs[j] = UV;
+		AllUVIndices[j] = UVTriangle1[j];
+	}
+	UVs->SetTriangle(QuadTris.A, UVTriangle1);
+
+	// project second triangle
+	if (Mesh->IsTriangle(QuadTris.B))
+	{
+		FIndex3i Triangle2 = Mesh->GetTriangle(QuadTris.B);
+		FIndex3i UVTriangle2;
+		for (int j = 0; j < 3; ++j)
+		{
+			int i = Triangle1.IndexOf(Triangle2[j]);
+			if (i == -1)
+			{
+				FVector2f UV = ProjectionFrame.ToPlaneUV( (FVector3f)Mesh->GetVertex(Triangle2[j]), 2);
+				UVTriangle2[j] = UVs->AppendElement(UV, Triangle2[j]);
+				AllUVs[3] = UV;
+				AllUVIndices[3] = UVTriangle2[j];
+			}
+			else
+			{
+				UVTriangle2[j] = UVTriangle1[i];
+			}
+		}
+		UVs->SetTriangle(QuadTris.B, UVTriangle2);
+	}
+
+	// shift UVs so that their bbox min-corner is at origin and scaled by external scale factor
+	FAxisAlignedBox2f UVBounds(FAxisAlignedBox2f::Empty());
+	UVBounds.Contain(AllUVs[0]);  UVBounds.Contain(AllUVs[1]);  UVBounds.Contain(AllUVs[2]);
+	if (AllUVIndices[3] != -1)
+	{
+		UVBounds.Contain(AllUVs[3]);
+	}
+	for (int j = 0; j < 4; ++j)
+	{
+		if (AllUVIndices[j] != -1)
+		{
+			FVector2f TransformedUV = (AllUVs[j] - UVBounds.Min) * UVScaleFactor;
+			UVs->SetElement(AllUVIndices[j], TransformedUV);
+		}
+	}
+}
+
+
+
+
+
+
+void FDynamicMeshEditor::ReverseTriangleOrientations(const TArray<int>& Triangles, bool bInvertNormals)
+{
+	for (int tid : Triangles)
+	{
+		Mesh->ReverseTriOrientation(tid);
+	}
+	if (bInvertNormals)
+	{
+		InvertTriangleNormals(Triangles);
+	}
+}
+
+
+void FDynamicMeshEditor::InvertTriangleNormals(const TArray<int>& Triangles)
+{
+	// @todo re-use the TBitA
+
+	if (Mesh->HasVertexNormals())
+	{
+		TBitArray<FDefaultBitArrayAllocator> DoneVertices(false, Mesh->MaxVertexID());
+		for (int TriangleID : Triangles)
+		{
+			FIndex3i Tri = Mesh->GetTriangle(TriangleID);
+			for (int j = 0; j < 3; ++j)
+			{
+				if (DoneVertices[Tri[j]] == false)
+				{
+					Mesh->SetVertexNormal(Tri[j], -Mesh->GetVertexNormal(Tri[j]));
+					DoneVertices[Tri[j]] = true;
+				}
+			}
+		}
+	}
+
+
+	if (Mesh->HasAttributes())
+	{
+		for (FDynamicMeshNormalOverlay* Normals : Mesh->Attributes()->GetAllNormalLayers())
+		{
+			TBitArray<FDefaultBitArrayAllocator> DoneNormals(false, Normals->MaxElementID());
+			for (int TriangleID : Triangles)
+			{
+				FIndex3i ElemTri = Normals->GetTriangle(TriangleID);
+				for (int j = 0; j < 3; ++j)
+				{
+					if (DoneNormals[ElemTri[j]] == false)
+					{
+						Normals->SetElement(ElemTri[j], -Normals->GetElement(ElemTri[j]));
+						DoneNormals[ElemTri[j]] = true;
+					}
+				}
+			}
+		}
+	}
+}
+
+
+
+
+
+void FDynamicMeshEditor::CopyAttributes(int FromTriangleID, int ToTriangleID, FMeshIndexMappings& IndexMaps, FDynamicMeshEditResult& ResultOut)
+{
+	if (Mesh->HasAttributes() == false)
+	{
+		return;
+	}
+
+	int UVLayerIndex = 0;
+	for (FDynamicMeshUVOverlay* UVOverlay : Mesh->Attributes()->GetAllUVLayers())
+	{
+		FIndex3i FromElemTri = UVOverlay->GetTriangle(FromTriangleID);
+		FIndex3i ToElemTri = UVOverlay->GetTriangle(ToTriangleID);
+		for (int j = 0; j < 3; ++j)
+		{
+			if (FromElemTri[j] != FDynamicMesh3::InvalidID )
+			{
+				int NewElemID = FindOrCreateDuplicateUV(FromElemTri[j], UVLayerIndex, IndexMaps);
+				ToElemTri[j] = NewElemID;
+			}
+		}
+		UVOverlay->SetTriangle(ToTriangleID, ToElemTri);
+		UVLayerIndex++;
+	}
+
+
+	int NormalLayerIndex = 0;
+	for (FDynamicMeshNormalOverlay* NormalOverlay : Mesh->Attributes()->GetAllNormalLayers())
+	{
+		FIndex3i FromElemTri = NormalOverlay->GetTriangle(FromTriangleID);
+		FIndex3i ToElemTri = NormalOverlay->GetTriangle(ToTriangleID);
+		for (int j = 0; j < 3; ++j)
+		{
+			if (FromElemTri[j] != FDynamicMesh3::InvalidID)
+			{
+				int NewElemID = FindOrCreateDuplicateNormal(FromElemTri[j], NormalLayerIndex, IndexMaps);
+				ToElemTri[j] = NewElemID;
+			}
+		}
+		NormalOverlay->SetTriangle(ToTriangleID, ToElemTri);
+		NormalLayerIndex++;
+	}
+}
+
+
+
+int FDynamicMeshEditor::FindOrCreateDuplicateUV(int ElementID, int UVLayerIndex, FMeshIndexMappings& IndexMaps)
+{
+	int NewElementID = IndexMaps.GetNewUV(UVLayerIndex, ElementID);
+	if (NewElementID == IndexMaps.InvalidID())
+	{
+		FDynamicMeshUVOverlay* UVOverlay = Mesh->Attributes()->GetUVLayer(UVLayerIndex);
+
+		// need to determine new parent vertex. It should be in the map already!
+		int ParentVertexID = UVOverlay->GetParentVertex(ElementID);
+		int NewParentVertexID = IndexMaps.GetNewVertex(ParentVertexID);
+		check(NewParentVertexID != IndexMaps.InvalidID());
+
+		NewElementID = UVOverlay->AppendElement(
+			UVOverlay->GetElement(ElementID), NewParentVertexID);
+
+		IndexMaps.SetUV(UVLayerIndex, ElementID, NewElementID);
+	}
+	return NewElementID;
+}
+
+
+
+int FDynamicMeshEditor::FindOrCreateDuplicateNormal(int ElementID, int NormalLayerIndex, FMeshIndexMappings& IndexMaps)
+{
+	int NewElementID = IndexMaps.GetNewNormal(NormalLayerIndex, ElementID);
+	if (NewElementID == IndexMaps.InvalidID())
+	{
+		FDynamicMeshNormalOverlay* NormalOverlay = Mesh->Attributes()->GetNormalLayer(NormalLayerIndex);
+
+		// need to determine new parent vertex. It should be in the map already!
+		int ParentVertexID = NormalOverlay->GetParentVertex(ElementID);
+		int NewParentVertexID = IndexMaps.GetNewVertex(ParentVertexID);
+		check(NewParentVertexID != IndexMaps.InvalidID());
+
+		NewElementID = NormalOverlay->AppendElement(
+			NormalOverlay->GetElement(ElementID), NewParentVertexID);
+
+		IndexMaps.SetNormal(NormalLayerIndex, ElementID, NewElementID);
+	}
+	return NewElementID;
+}
+
+
+
+int FDynamicMeshEditor::FindOrCreateDuplicateVertex(int VertexID, FMeshIndexMappings& IndexMaps, FDynamicMeshEditResult& ResultOut)
+{
+	int NewVertexID = IndexMaps.GetNewVertex(VertexID);
+	if (NewVertexID == IndexMaps.InvalidID())
+	{
+		NewVertexID = Mesh->AppendVertex(Mesh, VertexID);
+		IndexMaps.SetVertex(VertexID, NewVertexID);
+		ResultOut.NewVertices.Add(NewVertexID);
+	}
+	return NewVertexID;
+}
+
+
+
+int FDynamicMeshEditor::FindOrCreateDuplicateGroup(int TriangleID, FMeshIndexMappings& IndexMaps, FDynamicMeshEditResult& ResultOut)
+{
+	int GroupID = Mesh->GetTriangleGroup(TriangleID);
+	int NewGroupID = IndexMaps.GetNewGroup(GroupID);
+	if (NewGroupID == IndexMaps.InvalidID())
+	{
+		NewGroupID = Mesh->AllocateTriangleGroup();
+		IndexMaps.SetGroup(GroupID, NewGroupID);
+		ResultOut.NewGroups.Add(NewGroupID);
+	}
+	return NewGroupID;
+}
\ No newline at end of file
diff --git a/Engine/Plugins/Experimental/GeometryProcessing/Source/DynamicMesh/Private/DynamicMeshOverlay.cpp b/Engine/Plugins/Experimental/GeometryProcessing/Source/DynamicMesh/Private/DynamicMeshOverlay.cpp
new file mode 100644
--- /dev/null
+++ b/Engine/Plugins/Experimental/GeometryProcessing/Source/DynamicMesh/Private/DynamicMeshOverlay.cpp
@@ -0,0 +1,656 @@
+// Copyright 1998-2019 Epic Games, Inc. All Rights Reserved.
+
+#include "DynamicMeshOverlay.h"
+#include "DynamicMesh3.h"
+
+
+
+template<typename RealType, int ElementSize>
+int TDynamicMeshOverlay<RealType, ElementSize>::AppendElement(RealType ConstantValue, int SourceVertex)
+{
+	int vid = ElementsRefCounts.Allocate();
+	int i = ElementSize * vid;
+	for (int k = ElementSize - 1; k >= 0; --k)
+	{
+		Elements.InsertAt(ConstantValue, i + k);
+	}
+	ParentVertices.InsertAt(SourceVertex, vid);
+
+	//updateTimeStamp(true);
+	return vid;
+}
+
+
+template<typename RealType, int ElementSize>
+int TDynamicMeshOverlay<RealType, ElementSize>::AppendElement(const RealType* Value, int SourceVertex)
+{
+	int vid = ElementsRefCounts.Allocate();
+	int i = ElementSize * vid;
+
+	// insert in reverse order so that Resize() is only called once
+	for (int k = ElementSize - 1; k >= 0; --k)
+	{
+		Elements.InsertAt(Value[k], i + k);
+	}
+
+	ParentVertices.InsertAt(SourceVertex, vid);
+
+	//updateTimeStamp(true);
+	return vid;
+}
+
+
+
+
+template<typename RealType, int ElementSize>
+void TDynamicMeshOverlay<RealType, ElementSize>::InitializeTriangles(int MaxTriangleID)
+{
+	ElementTriangles.Resize(0);
+	ElementTriangles.Resize(MaxTriangleID * 3, FDynamicMesh3::InvalidID);
+}
+
+
+
+
+template<typename RealType, int ElementSize>
+EMeshResult TDynamicMeshOverlay<RealType, ElementSize>::SetTriangle(int tid, const FIndex3i& tv)
+{
+	if (IsElement(tv[0]) == false || IsElement(tv[1]) == false || IsElement(tv[2]) == false)
+	{
+		check(false);
+		return EMeshResult::Failed_NotAVertex;
+	}
+	if (tv[0] == tv[1] || tv[0] == tv[2] || tv[1] == tv[2]) 
+	{
+		check(false);
+		return EMeshResult::Failed_InvalidNeighbourhood;
+	}
+
+	InternalSetTriangle(tid, tv, true);
+
+	//updateTimeStamp(true);
+	return EMeshResult::Ok;
+}
+
+
+template<typename RealType, int ElementSize>
+void TDynamicMeshOverlay<RealType, ElementSize>::InternalSetTriangle(int tid, const FIndex3i& tv, bool bIncrementRefCounts)
+{
+	int i = 3 * tid;
+	ElementTriangles.InsertAt(tv[2], i + 2);
+	ElementTriangles.InsertAt(tv[1], i + 1);
+	ElementTriangles.InsertAt(tv[0], i);
+
+	if (bIncrementRefCounts)
+	{
+		ElementsRefCounts.Increment(tv[0]);
+		ElementsRefCounts.Increment(tv[1]);
+		ElementsRefCounts.Increment(tv[2]);
+	}
+}
+
+
+
+template<typename RealType, int ElementSize>
+void TDynamicMeshOverlay<RealType, ElementSize>::InitializeNewTriangle(int tid)
+{
+	int i = 3 * tid;
+	ElementTriangles.InsertAt(FDynamicMesh3::InvalidID, i + 2);
+	ElementTriangles.InsertAt(FDynamicMesh3::InvalidID, i + 1);
+	ElementTriangles.InsertAt(FDynamicMesh3::InvalidID, i);
+
+	//updateTimeStamp(true);
+}
+
+
+
+
+template<typename RealType, int ElementSize>
+bool TDynamicMeshOverlay<RealType,ElementSize>::IsSeamEdge(int eid) const
+{
+	FIndex2i et = ParentMesh->GetEdgeT(eid);
+	if (et.B == FDynamicMesh3::InvalidID)
+	{
+		return true;
+	}
+
+	FIndex2i ev = ParentMesh->GetEdgeV(eid);
+	int base_a = ev.A, base_b = ev.B;
+
+	FIndex3i Triangle0 = GetTriangle(et.A);
+	FIndex3i BaseTriangle0(ParentVertices[Triangle0.A], ParentVertices[Triangle0.B], ParentVertices[Triangle0.C]);
+	int idx_base_a1 = BaseTriangle0.IndexOf(base_a);
+	int idx_base_b1 = BaseTriangle0.IndexOf(base_b);
+
+	FIndex3i Triangle1 = GetTriangle(et.B);
+	FIndex3i BaseTriangle1(ParentVertices[Triangle1.A], ParentVertices[Triangle1.B], ParentVertices[Triangle1.C]);
+	int idx_base_a2 = BaseTriangle1.IndexOf(base_a);
+	int idx_base_b2 = BaseTriangle1.IndexOf(base_b);
+
+	return !IndexUtil::SamePairUnordered(Triangle0[idx_base_a1], Triangle0[idx_base_b1], Triangle1[idx_base_a2], Triangle1[idx_base_b2]);
+
+
+	// [RMS] this doesn't seem to work but it should, and would be more efficient. 
+	//   - add ParentMesh->FindTriEdgeIndex(tid,eid)
+	//   - SamePairUnordered query could directly index into ElementTriangles[]
+	//FIndex3i TriangleA = GetTriangle(et.A);
+	//FIndex3i TriangleB = GetTriangle(et.B);
+
+	//FIndex3i BaseTriEdgesA = ParentMesh->GetTriEdges(et.A);
+	//int WhichA = (BaseTriEdgesA.A == eid) ? 0 :
+	//	((BaseTriEdgesA.B == eid) ? 1 : 2);
+
+	//FIndex3i BaseTriEdgesB = ParentMesh->GetTriEdges(et.B);
+	//int WhichB = (BaseTriEdgesB.A == eid) ? 0 :
+	//	((BaseTriEdgesB.B == eid) ? 1 : 2);
+
+	//return SamePairUnordered(
+	//	TriangleA[WhichA], TriangleA[(WhichA + 1) % 3],
+	//	TriangleB[WhichB], TriangleB[(WhichB + 1) % 3]);
+}
+
+
+
+template<typename RealType, int ElementSize>
+bool TDynamicMeshOverlay<RealType, ElementSize>::IsSeamVertex(int vid) const
+{
+	// @todo can we do this more efficiently? At minimum we are looking up each triangle twice...
+	for (int edgeid : ParentMesh->VtxEdgesItr(vid))
+	{
+		if (IsSeamEdge(edgeid))
+		{
+			return true;
+		}
+	}
+	return false;
+}
+
+
+template<typename RealType, int ElementSize>
+void TDynamicMeshOverlay<RealType, ElementSize>::GetVertexElements(int vid, TArray<int>& OutElements) const
+{
+	OutElements.Reset();
+	for (int tid : ParentMesh->VtxTrianglesItr(vid))
+	{
+		FIndex3i Triangle = GetTriangle(tid);
+		for (int j = 0; j < 3; ++j)
+		{
+			if (ParentVertices[Triangle[j]] == vid)
+			{
+				OutElements.AddUnique(Triangle[j]);
+			}
+		}
+	}
+}
+
+
+
+template<typename RealType, int ElementSize>
+int TDynamicMeshOverlay<RealType, ElementSize>::CountVertexElements(int vid, bool bBruteForce) const
+{
+	TArray<int> VertexElements;
+	if (bBruteForce) 
+	{
+		for (int tid : ParentMesh->TriangleIndicesItr())
+		{
+			FIndex3i Triangle = GetTriangle(tid);
+			for (int j = 0; j < 3; ++j)
+			{
+				if (ParentVertices[Triangle[j]] == vid)
+				{
+					VertexElements.AddUnique(Triangle[j]);
+				}
+			}
+		}
+	}
+	else
+	{
+		for (int tid : ParentMesh->VtxTrianglesItr(vid))
+		{
+			FIndex3i Triangle = GetTriangle(tid);
+			for (int j = 0; j < 3; ++j)
+			{
+				if (ParentVertices[Triangle[j]] == vid)
+				{
+					VertexElements.AddUnique(Triangle[j]);
+				}
+			}
+		}
+	}
+
+	return VertexElements.Num();
+}
+
+
+
+
+template<typename RealType, int ElementSize>
+void TDynamicMeshOverlay<RealType, ElementSize>::OnRemoveTriangle(int TriangleID, bool bRemoveIsolatedVertices)
+{
+	FIndex3i Triangle = GetTriangle(TriangleID);
+	InitializeNewTriangle(TriangleID);
+
+	// decrement element refcounts, and free element if it is now unreferenced
+	for (int j = 0; j < 3; ++j) 
+	{
+		int elemid = Triangle[j];
+		ElementsRefCounts.Decrement(elemid);
+		if (bRemoveIsolatedVertices && ElementsRefCounts.GetRefCount(elemid) == 1) 
+		{
+			ElementsRefCounts.Decrement(elemid);
+			ParentVertices[elemid] = FDynamicMesh3::InvalidID;
+			check(ElementsRefCounts.IsValid(elemid) == false);
+		}
+	}
+}
+
+
+template<typename RealType, int ElementSize>
+void TDynamicMeshOverlay<RealType, ElementSize>::OnReverseTriOrientation(int TriangleID)
+{
+	FIndex3i Triangle = GetTriangle(TriangleID);
+	int i = 3 * TriangleID;
+	ElementTriangles[i] = Triangle[1];			// mirrors order in FDynamicMesh3::ReverseTriOrientationInternal
+	ElementTriangles[i + 1] = Triangle[0];
+	ElementTriangles[i + 2] = Triangle[2];
+}
+
+
+
+//#pragma optimize( "", off )
+template<typename RealType, int ElementSize>
+void TDynamicMeshOverlay<RealType, ElementSize>::OnSplitEdge(const FDynamicMesh3::FEdgeSplitInfo& splitInfo)
+{
+	int orig_t0 = splitInfo.OriginalTriangles.A;
+	int orig_t1 = splitInfo.OriginalTriangles.B;
+	int base_a = splitInfo.OriginalVertices.A;
+	int base_b = splitInfo.OriginalVertices.B;
+
+	// look up current triangle 0, and infer base triangle 0
+	// @todo handle case where these are InvalidID because no UVs exist for this triangle
+	FIndex3i Triangle0 = GetTriangle(orig_t0);
+	FIndex3i BaseTriangle0(ParentVertices[Triangle0.A], ParentVertices[Triangle0.B], ParentVertices[Triangle0.C]);
+	int idx_base_a1 = BaseTriangle0.IndexOf(base_a);
+	int idx_base_b1 = BaseTriangle0.IndexOf(base_b);
+	int idx_base_c = IndexUtil::GetOtherTriIndex(idx_base_a1, idx_base_b1);
+
+	// create new element at lerp position
+	int NewElemID = AppendElement((RealType)0, splitInfo.NewVertex);
+	SetElementFromLerp(NewElemID, Triangle0[idx_base_a1], Triangle0[idx_base_b1], (RealType)splitInfo.SplitT);
+
+	// rewrite triangle 0
+	ElementTriangles[3*orig_t0 + idx_base_b1] = NewElemID;
+
+	// create new triangle 2 w/ correct winding order
+	FIndex3i NewTriangle2(NewElemID, Triangle0[idx_base_b1], Triangle0[idx_base_c]);  // mirrors DMesh3::SplitEdge [f,b,c]
+	InternalSetTriangle(splitInfo.NewTriangles.A, NewTriangle2, false);
+
+	// update ref counts
+	ElementsRefCounts.Increment(NewElemID, 2);
+	ElementsRefCounts.Increment(Triangle0[idx_base_c]);
+
+	if (orig_t1 == FDynamicMesh3::InvalidID)
+	{
+		return;  // we are done if this is a boundary triangle
+	}
+
+	// look up current triangle1 and infer base triangle 1
+	// @todo handle case where these are InvalidID because no UVs exist for this triangle
+	FIndex3i Triangle1 = GetTriangle(orig_t1);
+	FIndex3i BaseTriangle1(ParentVertices[Triangle1.A], ParentVertices[Triangle1.B], ParentVertices[Triangle1.C]);
+	int idx_base_a2 = BaseTriangle1.IndexOf(base_a);
+	int idx_base_b2 = BaseTriangle1.IndexOf(base_b);
+	int idx_base_d = IndexUtil::GetOtherTriIndex(idx_base_a2, idx_base_b2);
+
+	int OtherNewElemID = NewElemID;
+
+	// if we don't have a shared edge, we need to create another new UV for the other side
+	bool bHasSharedUVEdge = IndexUtil::SamePairUnordered(Triangle0[idx_base_a1], Triangle0[idx_base_b1], Triangle1[idx_base_a2], Triangle1[idx_base_b2]);
+	if (bHasSharedUVEdge == false)
+	{
+		// create new element at lerp position
+		OtherNewElemID = AppendElement((RealType)0, splitInfo.NewVertex);
+		SetElementFromLerp(OtherNewElemID, Triangle1[idx_base_a2], Triangle1[idx_base_b2], (RealType)splitInfo.SplitT);
+	}
+
+	// rewrite triangle 1
+	ElementTriangles[3*orig_t1 + idx_base_b2] = OtherNewElemID;
+
+	// create new triangle 3 w/ correct winding order
+	FIndex3i NewTriangle3(OtherNewElemID, Triangle1[idx_base_d], Triangle1[idx_base_b2]);  // mirrors DMesh3::SplitEdge [f,d,b]
+	InternalSetTriangle(splitInfo.NewTriangles.B, NewTriangle3, false);
+
+	// update ref counts
+	ElementsRefCounts.Increment(OtherNewElemID, 2);
+	ElementsRefCounts.Increment(Triangle1[idx_base_d]);
+}
+
+
+
+//#pragma optimize("", off)
+template<typename RealType, int ElementSize>
+void TDynamicMeshOverlay<RealType, ElementSize>::OnFlipEdge(const FDynamicMesh3::FEdgeFlipInfo& FlipInfo)
+{
+	int orig_t0 = FlipInfo.Triangles.A;
+	int orig_t1 = FlipInfo.Triangles.B;
+	int base_a = FlipInfo.OriginalVerts.A;
+	int base_b = FlipInfo.OriginalVerts.B;
+	int base_c = FlipInfo.OpposingVerts.A;
+	int base_d = FlipInfo.OpposingVerts.B;
+
+	// look up triangle 0
+	FIndex3i Triangle0 = GetTriangle(orig_t0);
+	FIndex3i BaseTriangle0(ParentVertices[Triangle0.A], ParentVertices[Triangle0.B], ParentVertices[Triangle0.C]);
+	int idx_base_a1 = BaseTriangle0.IndexOf(base_a);
+	int idx_base_b1 = BaseTriangle0.IndexOf(base_b);
+	int idx_base_c = IndexUtil::GetOtherTriIndex(idx_base_a1, idx_base_b1);
+
+	// look up triangle 1 (must exist because base mesh would never flip a boundary edge)
+	FIndex3i Triangle1 = GetTriangle(orig_t1);
+	FIndex3i BaseTriangle1(ParentVertices[Triangle1.A], ParentVertices[Triangle1.B], ParentVertices[Triangle1.C]);
+	int idx_base_a2 = BaseTriangle1.IndexOf(base_a);
+	int idx_base_b2 = BaseTriangle1.IndexOf(base_b);
+	int idx_base_d = IndexUtil::GetOtherTriIndex(idx_base_a2, idx_base_b2);
+
+	// sanity checks
+	check(idx_base_c == BaseTriangle0.IndexOf(base_c));
+	check(idx_base_d == BaseTriangle1.IndexOf(base_d));
+
+	// we should not have been called on a non-shared edge!!
+	bool bHasSharedUVEdge = IndexUtil::SamePairUnordered(Triangle0[idx_base_a1], Triangle0[idx_base_b1], Triangle1[idx_base_a2], Triangle1[idx_base_b2]);
+	check(bHasSharedUVEdge);
+
+	int A = Triangle0[idx_base_a1];
+	int B = Triangle0[idx_base_b1];
+	int C = Triangle0[idx_base_c];
+	int D = Triangle1[idx_base_d];
+
+	// set triangles to same index order as in FDynamicMesh::FlipEdge
+	int i0 = 3 * orig_t0;
+	ElementTriangles[i0] = C; ElementTriangles[i0+1] = D; ElementTriangles[i0+2] = B;
+	int i1 = 3 * orig_t1;
+	ElementTriangles[i1] = D; ElementTriangles[i1+1] = C; ElementTriangles[i1+2] = A;
+
+	// update reference counts
+	ElementsRefCounts.Decrement(A);
+	ElementsRefCounts.Decrement(B);
+	ElementsRefCounts.Increment(C);
+	ElementsRefCounts.Increment(D);
+}
+
+
+
+
+
+//#pragma optimize("", off)
+template<typename RealType, int ElementSize>
+void TDynamicMeshOverlay<RealType, ElementSize>::OnCollapseEdge(const FDynamicMesh3::FEdgeCollapseInfo& collapseInfo)
+{
+	int vid_base_kept = collapseInfo.KeptVertex;
+	int vid_base_removed = collapseInfo.RemovedVertex;
+	int tid_removed0 = collapseInfo.RemovedTris.A;
+	int tid_removed1 = collapseInfo.RemovedTris.B;
+
+	// look up triangle 0
+	FIndex3i Triangle0 = GetTriangle(tid_removed0);
+	FIndex3i BaseTriangle0(ParentVertices[Triangle0.A], ParentVertices[Triangle0.B], ParentVertices[Triangle0.C]);
+	int idx_removed0_a = BaseTriangle0.IndexOf(vid_base_kept);
+	int idx_removed0_b = BaseTriangle0.IndexOf(vid_base_removed);
+
+	// look up triangle 1 if this is not a boundary edge
+	FIndex3i Triangle1, BaseTriangle1;
+	if (collapseInfo.bIsBoundary == false)
+	{
+		Triangle1 = GetTriangle(tid_removed1);
+		BaseTriangle1 = FIndex3i(ParentVertices[Triangle1.A], ParentVertices[Triangle1.B], ParentVertices[Triangle1.C]);
+
+		int idx_removed1_a = BaseTriangle1.IndexOf(vid_base_kept);
+		int idx_removed1_b = BaseTriangle1.IndexOf(vid_base_removed);
+
+		// if this is an internal edge it cannot be a seam or we cannot collapse
+		bool bHasSharedUVEdge = IndexUtil::SamePairUnordered(Triangle0[idx_removed0_a], Triangle0[idx_removed0_b], Triangle1[idx_removed1_a], Triangle1[idx_removed1_b]);
+		check(bHasSharedUVEdge);
+	}
+
+	// need to find the elementid for the "kept" vertex. Since this isn't a seam, 
+	// there is just one, and we can grab it from "old" removed triangle
+	int kept_elemid = FDynamicMesh3::InvalidID;
+	for (int j = 0; j < 3; ++j) 
+	{
+		if (BaseTriangle0[j] == vid_base_kept)
+		{
+			kept_elemid = Triangle0[j];
+		}
+	}
+	check(kept_elemid != FDynamicMesh3::InvalidID);
+
+	// look for still-existing triangles that have elements linked to the removed vertex.
+	// in that case, replace with the element we found
+	TArray<int> removed_elements;
+	for (int onering_tid : ParentMesh->VtxTrianglesItr(vid_base_kept))
+	{
+		FIndex3i elem_tri = GetTriangle(onering_tid);
+		for (int j = 0; j < 3; ++j) 
+		{
+			int elem_id = elem_tri[j];
+			if (ParentVertices[elem_id] == vid_base_removed) 
+			{
+				ElementsRefCounts.Decrement(elem_id);
+				removed_elements.AddUnique(elem_id);
+				ElementTriangles[3*onering_tid + j] = kept_elemid;
+				ElementsRefCounts.Increment(kept_elemid);
+			}
+		}
+	}
+	check(removed_elements.Num() == 1);   // should always be true for non-seam edges...
+
+	// update position of kept element
+	SetElementFromLerp(kept_elemid, kept_elemid, removed_elements[0], collapseInfo.CollapseT);
+
+	// clear the two triangles we removed
+	for (int j = 0; j < 3; ++j) 
+	{
+		ElementsRefCounts.Decrement(Triangle0[j]);
+		ElementTriangles[3*tid_removed0 + j] = FDynamicMesh3::InvalidID;
+		if (collapseInfo.bIsBoundary == false) 
+		{
+			ElementsRefCounts.Decrement(Triangle1[j]);
+			ElementTriangles[3*tid_removed1 + j] = FDynamicMesh3::InvalidID;
+		}
+	}
+
+	// remove the elements associated with the removed vertex
+	for (int k = 0; k < removed_elements.Num(); ++k) 
+	{
+		check(ElementsRefCounts.GetRefCount(removed_elements[k]) == 1);
+		ElementsRefCounts.Decrement(removed_elements[k]);
+		ParentVertices[removed_elements[k]] = FDynamicMesh3::InvalidID;
+	}
+
+}
+
+
+
+template<typename RealType, int ElementSize>
+void TDynamicMeshOverlay<RealType, ElementSize>::OnPokeTriangle(const FDynamicMesh3::FPokeTriangleInfo& PokeInfo)
+{
+	FIndex3i Triangle = GetTriangle(PokeInfo.OriginalTriangle);
+
+	// create new element at barycentric position
+	int CenterElemID = AppendElement((RealType)0, PokeInfo.NewVertex);
+	FVector3<RealType> BaryCoords((RealType)PokeInfo.BaryCoords.X, (RealType)PokeInfo.BaryCoords.Y, (RealType)PokeInfo.BaryCoords.Z);
+	SetElementFromBary(CenterElemID, Triangle[0], Triangle[1], Triangle[2], BaryCoords);
+
+	// update orig triangle and two new ones. Winding orders here mirror FDynamicMesh3::PokeTriangle
+	SetTriangle(PokeInfo.OriginalTriangle, FIndex3i(Triangle[0], Triangle[1], CenterElemID) );
+	SetTriangle(PokeInfo.NewTriangles.A, FIndex3i(Triangle[1], Triangle[2], CenterElemID));
+	SetTriangle(PokeInfo.NewTriangles.B, FIndex3i(Triangle[2], Triangle[0], CenterElemID));
+
+	ElementsRefCounts.Increment(Triangle[0]);
+	ElementsRefCounts.Increment(Triangle[1]);
+	ElementsRefCounts.Increment(Triangle[2]);
+	ElementsRefCounts.Increment(CenterElemID, 3);
+}
+
+
+template<typename RealType, int ElementSize>
+void TDynamicMeshOverlay<RealType, ElementSize>::OnMergeEdges(const FDynamicMesh3::FMergeEdgesInfo& MergeInfo)
+{
+	// MergeEdges just merges vertices. For now we will not also merge UVs. So all we need to
+	// do is rewrite the UV parent vertices
+
+	FIndex3i ModifiedEdges(MergeInfo.KeptEdge, MergeInfo.ExtraKeptEdges.A, MergeInfo.ExtraKeptEdges.B);
+	for (int EdgeIdx = 0; EdgeIdx < 3; ++EdgeIdx)
+	{
+		if (ParentMesh->IsEdge(ModifiedEdges[EdgeIdx]) == false)
+		{
+			continue;
+		}
+
+		FIndex2i EdgeTris = ParentMesh->GetEdgeT(ModifiedEdges[EdgeIdx]);
+		for (int j = 0; j < 2; ++j)
+		{
+			FIndex3i ElemTriangle = GetTriangle(EdgeTris[j]);
+			for (int k = 0; k < 3; ++k)
+			{
+				int ParentVID = ParentVertices[ElemTriangle[k]];
+				if (ParentVID == MergeInfo.RemovedVerts.A)
+				{
+					ParentVertices[ElemTriangle[k]] = MergeInfo.KeptVerts.A;
+				}
+				else if (ParentVID == MergeInfo.RemovedVerts.B)
+				{
+					ParentVertices[ElemTriangle[k]] = MergeInfo.KeptVerts.B;
+				}
+			}
+		}
+	}
+}
+
+
+
+
+template<typename RealType, int ElementSize>
+void TDynamicMeshOverlay<RealType, ElementSize>::SetElementFromLerp(int SetElement, int ElementA, int ElementB, RealType Alpha)
+{
+	int IndexSet = ElementSize * SetElement;
+	int IndexA = ElementSize * ElementA;
+	int IndexB = ElementSize * ElementB;
+	RealType Beta = ((RealType)1 - Alpha);
+	for (int i = 0; i < ElementSize; ++i)
+	{
+		Elements[IndexSet+i] = Beta*Elements[IndexA+i] + Alpha*Elements[IndexB+i];
+	}
+}
+
+template<typename RealType, int ElementSize>
+void TDynamicMeshOverlay<RealType, ElementSize>::SetElementFromBary(int SetElement, int ElementA, int ElementB, int ElementC, const FVector3<RealType>& BaryCoords)
+{
+	int IndexSet = ElementSize * SetElement;
+	int IndexA = ElementSize * ElementA;
+	int IndexB = ElementSize * ElementB;
+	int IndexC = ElementSize * ElementC;
+	for (int i = 0; i < ElementSize; ++i)
+	{
+		Elements[IndexSet + i] = 
+			BaryCoords.X*Elements[IndexA+i] + BaryCoords.Y*Elements[IndexB+i] + BaryCoords.Z*Elements[IndexC+i];
+	}
+}
+
+
+//#pragma optimize("", off)
+template<typename RealType, int ElementSize>
+bool TDynamicMeshOverlay<RealType, ElementSize>::CheckValidity(bool bAllowNonManifoldVertices, EValidityCheckFailMode FailMode) const
+{
+	bool is_ok = true;
+	TFunction<void(bool)> CheckOrFailF = [&](bool b)
+	{
+		is_ok = is_ok && b;
+	};
+	if (FailMode == EValidityCheckFailMode::Check)
+	{
+		CheckOrFailF = [&](bool b)
+		{
+			checkf(b, TEXT("TDynamicMeshOverlay::CheckValidity failed!"));
+			is_ok = is_ok && b;
+		};
+	}
+	else if (FailMode == EValidityCheckFailMode::Ensure)
+	{
+		CheckOrFailF = [&](bool b)
+		{
+			ensureMsgf(b, TEXT("TDynamicMeshOverlay::CheckValidity failed!"));
+			is_ok = is_ok && b;
+		};
+	}
+
+
+	// @todo: check that all connected element-pairs are also edges in parentmesh
+
+	// check that parent vtx of each element is actually a vertex
+	for (int elemid : ElementIndicesItr())
+	{
+		int ParentVID = GetParentVertex(elemid);
+		CheckOrFailF(ParentMesh->IsVertex(ParentVID));
+	}
+
+	// check that parent vertices of each element triangle are the same as base triangle
+	for (int tid : ParentMesh->TriangleIndicesItr())
+	{
+		FIndex3i ElemTri = GetTriangle(tid);
+		FIndex3i BaseTri = ParentMesh->GetTriangle(tid);
+		for (int j = 0; j < 3; ++j)
+		{
+			if (ElemTri[j] != FDynamicMesh3::InvalidID)
+			{
+				CheckOrFailF(GetParentVertex(ElemTri[j]) == BaseTri[j]);
+			}
+		}
+	}
+
+	// count references to each element
+	TArray<int> RealRefCounts; RealRefCounts.Init(0, MaxElementID());
+	for (int tid : ParentMesh->TriangleIndicesItr())
+	{
+		FIndex3i Tri = GetTriangle(tid);
+		for (int j = 0; j < 3; ++j)
+		{
+			if (Tri[j] != FDynamicMesh3::InvalidID)
+			{
+				RealRefCounts[Tri[j]] += 1;
+			}
+		}
+	}
+	// verify that refcount list counts are same as actual reference counts
+	for (int elemid : ElementsRefCounts.Indices())
+	{
+		int CurRefCount = ElementsRefCounts.GetRefCount(elemid);
+		CheckOrFailF(CurRefCount == RealRefCounts[elemid] + 1);
+	}
+
+	return is_ok;
+}
+
+
+
+
+
+// These are explicit instantiations of the templates that are exported from the shared lib.
+// Only these instantiations of the template can be used.
+// This is necessary because we have placed most of the templated functions in this .cpp file, instead of the header.
+template class DYNAMICMESH_API TDynamicMeshOverlay<float, 1>;
+template class DYNAMICMESH_API TDynamicMeshOverlay<double, 1>;
+template class DYNAMICMESH_API TDynamicMeshOverlay<int, 1>;
+template class DYNAMICMESH_API TDynamicMeshOverlay<float, 2>;
+template class DYNAMICMESH_API TDynamicMeshOverlay<double, 2>;
+template class DYNAMICMESH_API TDynamicMeshOverlay<int, 2>;
+template class DYNAMICMESH_API TDynamicMeshOverlay<float, 3>;
+template class DYNAMICMESH_API TDynamicMeshOverlay<double, 3>;
+template class DYNAMICMESH_API TDynamicMeshOverlay<int, 3>;
+
+template class DYNAMICMESH_API TDynamicMeshVectorOverlay<float, 2, FVector2f>;
+template class DYNAMICMESH_API TDynamicMeshVectorOverlay<double, 2, FVector2d>;
+template class DYNAMICMESH_API TDynamicMeshVectorOverlay<int, 2, FVector2i>;
+template class DYNAMICMESH_API TDynamicMeshVectorOverlay<float, 3, FVector3f>;
+template class DYNAMICMESH_API TDynamicMeshVectorOverlay<double, 3, FVector3d>;
+template class DYNAMICMESH_API TDynamicMeshVectorOverlay<int, 3, FVector3i>;
\ No newline at end of file
diff --git a/Engine/Plugins/Experimental/GeometryProcessing/Source/DynamicMesh/Private/MeshConstraintsUtil.cpp b/Engine/Plugins/Experimental/GeometryProcessing/Source/DynamicMesh/Private/MeshConstraintsUtil.cpp
new file mode 100644
--- /dev/null
+++ b/Engine/Plugins/Experimental/GeometryProcessing/Source/DynamicMesh/Private/MeshConstraintsUtil.cpp
@@ -0,0 +1,27 @@
+// Copyright 1998-2019 Epic Games, Inc. All Rights Reserved.
+
+#include "MeshConstraintsUtil.h"
+
+
+void FMeshConstraintsUtil::ConstrainAllSeams(FMeshConstraints& Constraints, const FDynamicMesh3& Mesh, bool bAllowSplits, bool bAllowSmoothing)
+{
+	if (Mesh.HasAttributes() == false)
+	{
+		return;
+	}
+	const FDynamicMeshAttributeSet* Attributes = Mesh.Attributes();
+
+	FEdgeConstraint EdgeConstraint = (bAllowSplits) ? FEdgeConstraint::SplitsOnly() : FEdgeConstraint::FullyConstrained();
+	FVertexConstraint VtxConstraint = (bAllowSmoothing) ? FVertexConstraint::PinnedMovable() : FVertexConstraint::Pinned();
+
+	for (int EdgeID : Mesh.EdgeIndicesItr())
+	{
+		if (Attributes->IsSeamEdge(EdgeID))
+		{
+			Constraints.SetOrUpdateEdgeConstraint(EdgeID, EdgeConstraint);
+			FIndex2i EdgeVerts = Mesh.GetEdgeV(EdgeID);
+			Constraints.SetOrUpdateVertexConstraint(EdgeVerts.A, VtxConstraint);
+			Constraints.SetOrUpdateVertexConstraint(EdgeVerts.B, VtxConstraint);
+		}
+	}
+}
\ No newline at end of file
