move GeometryProcessing plugin from NotForLicensees to Experimental

diff --git a/Engine/Plugins/Experimental/GeometryProcessing/Source/DynamicMesh/Private/MeshRegionBoundaryLoops.cpp b/Engine/Plugins/Experimental/GeometryProcessing/Source/DynamicMesh/Private/MeshRegionBoundaryLoops.cpp
new file mode 100644
--- /dev/null
+++ b/Engine/Plugins/Experimental/GeometryProcessing/Source/DynamicMesh/Private/MeshRegionBoundaryLoops.cpp
@@ -0,0 +1,476 @@
+// Copyright 1998-2019 Epic Games, Inc. All Rights Reserved.
+
+
+#include "MeshRegionBoundaryLoops.h"
+#include "MeshBoundaryLoops.h"   // has a set of internal static functions we re-use
+#include "VectorUtil.h"
+
+
+
+FMeshRegionBoundaryLoops::FMeshRegionBoundaryLoops(const FDynamicMesh3* MeshIn, const TArray<int>& RegionTris, bool bAutoCompute)
+{
+	this->Mesh = MeshIn;
+
+	// make flag set for included triangles
+	triangles.Init(false, Mesh->MaxTriangleID());
+	for (int i = 0; i < RegionTris.Num(); ++i)
+	{
+		triangles[RegionTris[i]] = true;
+	}
+
+	// make flag set for included edges
+	// NOTE: this currently processes non-boundary-edges twice. Could
+	// avoid w/ another IndexFlagSet, but the check is inexpensive...
+	edges.Init(false, Mesh->MaxEdgeID());
+	for (int i = 0; i < RegionTris.Num(); ++i)
+	{
+		int tid = RegionTris[i];
+		FIndex3i te = Mesh->GetTriEdges(tid);
+		for (int j = 0; j < 3; ++j)
+		{
+			int eid = te[j];
+			if (!ContainsElement(edges, eid))
+			{
+				FIndex2i et = Mesh->GetEdgeT(eid);
+				if (et.B == IndexConstants::InvalidID || triangles[et.A] != triangles[et.B])
+				{
+					edges_roi.Add(eid);
+					edges[eid] = true;
+				}
+			}
+		}
+	}
+
+	if (bAutoCompute)
+	{
+		Compute();
+	}
+}
+
+
+
+int FMeshRegionBoundaryLoops::GetMaxVerticesLoopIndex() const
+{
+	int j = 0;
+	for (int i = 1; i < Loops.Num(); ++i)
+	{
+		if (Loops[i].Vertices.Num() > Loops[j].Vertices.Num())
+		{
+			j = i;
+		}
+	}
+	return j;
+}
+
+
+
+bool FMeshRegionBoundaryLoops::Compute()
+{
+	// This algorithm assumes that triangles are oriented consistently, 
+	// so closed boundary-loop can be followed by walking edges in-order
+	Loops.SetNum(0);
+
+	// Temporary memory used to indicate when we have "used" an edge.
+	TArray<bool> used_edge;
+	used_edge.Init(false, Mesh->MaxEdgeID());
+
+	// current loop is stored here, cleared after each loop extracted
+	TArray<int> loop_edges;
+	TArray<int> loop_verts;
+	TArray<int> bowties;
+
+	// Temp buffer for reading back all boundary edges of a vertex.
+	// probably always small but : pathological cases it could be large...
+	TArray<int> all_e;
+	all_e.SetNum(16);
+
+	// process all edges of mesh
+	for (int eid : edges_roi)
+	{
+		if (used_edge[eid] == true)
+		{
+			continue;
+		}
+		if (IsEdgeOnBoundary(eid) == false)
+		{
+			continue;
+		}
+
+		// ok this is start of a boundary chain
+		int eStart = eid;
+		used_edge[eStart] = true;
+		loop_edges.Add(eStart);
+
+		int eCur = eid;
+
+		// follow the chain : order of oriented edges
+		bool bClosed = false;
+		while (!bClosed)
+		{
+
+			// [TODO] can do this more efficiently?
+			int tid_in = IndexConstants::InvalidID, tid_out = IndexConstants::InvalidID;
+			IsEdgeOnBoundary(eCur, tid_in, tid_out);
+
+			FIndex2i ev = GetOrientedEdgeVerts(eCur, tid_in, tid_out);
+			int cure_a = ev.A, cure_b = ev.B;
+			loop_verts.Add(cure_a);
+
+			int e0 = -1, e1 = 1;
+			int bdry_nbrs = GetVertexBoundaryEdges(cure_b, e0, e1);
+
+			check(bdry_nbrs >= 2); //  if (bdry_nbrs < 2) throw new MeshBoundaryLoopsException("MeshRegionBoundaryLoops.Compute: found broken neighbourhood at vertex " + cure_b){ UnclosedLoop = true };
+
+			int eNext = -1;
+			if (bdry_nbrs > 2)
+			{
+				// found "bowtie" vertex...things just got complicated!
+
+				if (cure_b == loop_verts[0])
+				{
+					// The "end" of the current edge is the same as the start vertex.
+					// This means we can close the loop here. Might as well!
+					eNext = -2;   // sentinel value used below
+
+				}
+				else
+				{
+					// try to find an unused outgoing edge that is oriented properly.
+					// This could create sub-loops, we will handle those later
+					if (bdry_nbrs >= all_e.Num())
+						all_e.SetNum(bdry_nbrs);
+					int num_be = GetAllVertexBoundaryEdges(cure_b, all_e);
+
+					check(num_be == bdry_nbrs);
+
+					// Try to pick the best "turn left" vertex.
+					eNext = FindLeftTurnEdge(eCur, cure_b, all_e, num_be, used_edge);
+
+					check(eNext != -1); // throw new MeshBoundaryLoopsException("MeshRegionBoundaryLoops.Compute: cannot find valid outgoing edge at bowtie vertex " + cure_b){ BowtieFailure = true };
+				}
+
+				if (bowties.Contains(cure_b) == false)
+				{
+					bowties.Add(cure_b);
+				}
+
+			}
+			else
+			{
+				check(e0 == eCur || e1 == eCur);
+				eNext = (e0 == eCur) ? e1 : e0;
+			}
+
+			if (eNext == -2)
+			{
+				// found a bowtie vert that is the same as start-of-loop, so we
+				// are just closing it off explicitly
+				bClosed = true;
+			}
+			else if (eNext == eStart)
+			{
+				// found edge at start of loop, so loop is done.
+				bClosed = true;
+			}
+			else
+			{
+				// push onto accumulated list
+				check(used_edge[eNext] == false);
+				loop_edges.Add(eNext);
+				eCur = eNext;
+				used_edge[eCur] = true;
+			}
+		}
+
+		// if we saw a bowtie vertex, we might need to break up this loop,
+		// so call ExtractSubloops
+		if (bowties.Num() > 0)
+		{
+			TArray<FEdgeLoop> subloops = ExtractSubloops(loop_verts, loop_edges, bowties);
+			for (int i = 0; i < subloops.Num(); ++i)
+			{
+				Loops.Add(subloops[i]);
+			}
+		}
+		else
+		{
+			// clean simple loop, convert to FEdgeLoop instance
+			FEdgeLoop loop(Mesh);
+			loop.Vertices = loop_verts;
+			loop.Edges = loop_edges;
+			Loops.Add(loop);
+		}
+
+		// reset these lists
+		loop_edges.SetNum(0);
+		loop_verts.SetNum(0);
+		bowties.SetNum(0);
+	}
+
+	return true;
+}
+
+
+
+
+
+
+// returns true for both internal and mesh boundary edges
+// tid_in and tid_out are triangles 'in' and 'out' of set, respectively
+bool FMeshRegionBoundaryLoops::IsEdgeOnBoundary(int eid, int& tid_in, int& tid_out) const
+{
+	if (ContainsElement(edges, eid) == false)
+	{
+		return false;
+	}
+
+	tid_in = tid_out = IndexConstants::InvalidID;
+	FIndex2i et = Mesh->GetEdgeT(eid);
+	if (et.B == IndexConstants::InvalidID)	// boundary edge!
+	{
+		tid_in = et.A;
+		tid_out = et.B;
+		return true;
+	}
+
+	bool in0 = triangles[et.A];
+	bool in1 = triangles[et.B];
+	if (in0 != in1)
+	{
+		tid_in = (in0) ? et.A : et.B;
+		tid_out = (in0) ? et.B : et.A;
+		return true;
+	}
+	return false;
+}
+
+
+
+// return same indices as GetEdgeV, but oriented based on attached triangle
+FIndex2i FMeshRegionBoundaryLoops::GetOrientedEdgeVerts(int eID, int tid_in, int tid_out)
+{
+	FIndex2i edgev = Mesh->GetEdgeV(eID);
+	int a = edgev.A, b = edgev.B;
+	FIndex3i tri = Mesh->GetTriangle(tid_in);
+	int ai = IndexUtil::FindEdgeIndexInTri(a, b, tri);
+	return FIndex2i(tri[ai], tri[(ai + 1) % 3]);
+}
+
+
+int FMeshRegionBoundaryLoops::GetVertexBoundaryEdges(int vID, int& e0, int& e1)
+{
+	int count = 0;
+	for (int eid : Mesh->VtxEdgesItr(vID))
+	{
+		if (IsEdgeOnBoundary(eid))
+		{
+			if (count == 0)
+			{
+				e0 = eid;
+			}
+			else if (count == 1)
+			{
+				e1 = eid;
+			}
+			count++;
+		}
+	}
+	return count;
+}
+
+
+int FMeshRegionBoundaryLoops::GetAllVertexBoundaryEdges(int vID, TArray<int>& e)
+{
+	int count = 0;
+	for (int eid : Mesh->VtxEdgesItr(vID))
+	{
+		if (IsEdgeOnBoundary(eid))
+		{
+			e[count++] = eid;
+		}
+	}
+	return count;
+}
+
+
+FVector3d FMeshRegionBoundaryLoops::GetVertexNormal(int vid)
+{
+	FVector3d n = FVector3d::Zero();
+	for (int ti : Mesh->VtxTrianglesItr(vid))
+	{
+		n += Mesh->GetTriNormal(ti);
+	}
+	n.Normalize();
+	return n;
+}
+
+
+
+//
+// [TODO] for internal vertices, there is no ambiguity : which is the left-turn edge,
+//   we should be using 'closest' left-neighbour edge.
+//
+// ok, bdry_edges[0...bdry_edges_count] contains the boundary edges coming out of bowtie_v.
+// We want to pick the best one to continue the loop that came : to bowtie_v on incoming_e.
+// If the loops are all sane, then we will get the smallest loops by "turning left" at bowtie_v.
+// So, we compute the tangent plane at bowtie_v, and then the signed angle for each
+// viable edge : this plane. 
+int FMeshRegionBoundaryLoops::FindLeftTurnEdge(int incoming_e, int bowtie_v, TArray<int>& bdry_edges, int bdry_edges_count, TArray<bool>& used_edges)
+{
+	// compute normal and edge [a,bowtie]
+	FVector3d n = GetVertexNormal(bowtie_v);
+	//int other_v = Mesh->edge_other_v(incoming_e, bowtie_v);
+	FIndex2i ev = Mesh->GetEdgeV(incoming_e);
+	int other_v = (ev.A == bowtie_v) ? ev.B : ev.A;
+	FVector3d ab = Mesh->GetVertex(bowtie_v) - Mesh->GetVertex(other_v);
+
+	// our winner
+	int best_e = -1;
+	double best_angle = TNumericLimits<double>::Max();
+
+	for (int i = 0; i < bdry_edges_count; ++i)
+	{
+		int bdry_eid = bdry_edges[i];
+		if (used_edges[bdry_eid] == true)
+			continue;       // this edge is already used
+
+		// [TODO] can do this more efficiently?
+		int tid_in = IndexConstants::InvalidID, tid_out = IndexConstants::InvalidID;
+		IsEdgeOnBoundary(bdry_eid, tid_in, tid_out);
+		FIndex2i bdry_ev = GetOrientedEdgeVerts(bdry_eid, tid_in, tid_out);
+		//FIndex2i bdry_ev = Mesh.GetOrientedBoundaryEdgeV(bdry_eid);
+
+		if (bdry_ev.A != bowtie_v) {
+			continue;       // have to be able to chain to end of current edge, orientation-wise
+		}
+
+		// compute projected angle
+		FVector3d bc = Mesh->GetVertex(bdry_ev.B) - Mesh->GetVertex(bowtie_v);
+		double fAngleS = VectorUtil::PlaneAngleSignedD(ab, bc, n);
+
+		// turn left!
+		if (best_angle == TNumericLimits<double>::Max() || fAngleS < best_angle)
+		{
+			best_angle = fAngleS;
+			best_e = bdry_eid;
+		}
+	}
+	check(best_e != -1);
+
+	return best_e;
+}
+
+
+
+
+// This is called when loopV contains one or more "bowtie" vertices.
+// These vertices *might* be duplicated : loopV (but not necessarily)
+// If they are, we have to break loopV into subloops that don't contain duplicates.
+//
+// The list bowties contains all the possible duplicates 
+// (all v in bowties occur in loopV at least once)
+//
+// Currently loopE is not used, and the returned FEdgeLoop objects do not have their Edges
+// arrays initialized. Perhaps to improve : future.
+TArray<FEdgeLoop> FMeshRegionBoundaryLoops::ExtractSubloops(TArray<int>& loopV, const TArray<int>& loopE, const TArray<int>& bowties)
+{
+	TArray<FEdgeLoop> subs;
+
+	// figure out which bowties we saw are actually duplicated : loopV
+	TArray<int> dupes;
+	for (int bv : bowties)
+	{
+		if (FMeshBoundaryLoops::CountInList(loopV, bv) > 1)
+		{
+			dupes.Add(bv);
+		}
+	}
+
+	// we might not actually have any duplicates, if we got luck. Early out : that case
+	if (dupes.Num() == 0)
+	{
+		FEdgeLoop NewLoop(Mesh);
+		NewLoop.Vertices = loopV;
+		NewLoop.Edges = loopE;
+		NewLoop.BowtieVertices = bowties;
+		subs.Add(NewLoop);
+		return subs;
+	}
+
+	// This loop extracts subloops until we have dealt with all the
+	// duplicate vertices : loopV
+	while (dupes.Num() > 0)
+	{
+
+		// Find shortest "simple" loop, ie a loop from a bowtie to itself that
+		// does not contain any other bowties. This is an independent loop.
+		// We're doing a lot of extra work here if we only have one element : dupes...
+		int bi = 0, bv = 0;
+		int start_i = -1, end_i = -1;
+		int bv_shortest = -1; int shortest = TNumericLimits<int>::Max();
+		for (; bi < dupes.Num(); ++bi)
+		{
+			bv = dupes[bi];
+			if (FMeshBoundaryLoops::IsSimpleBowtieLoop(loopV, dupes, bv, start_i, end_i))
+			{
+				int len = FMeshBoundaryLoops::CountSpan(loopV, start_i, end_i);
+				if (len < shortest)
+				{
+					bv_shortest = bv;
+					shortest = len;
+				}
+			}
+		}
+		check(bv_shortest != -1); //  throw new MeshBoundaryLoopsException("MeshRegionBoundaryLoops.Compute: Cannot find a valid simple loop");
+		if (bv != bv_shortest)
+		{
+			bv = bv_shortest;
+			// running again just to get start_i and end_i...
+			FMeshBoundaryLoops::IsSimpleBowtieLoop(loopV, dupes, bv, start_i, end_i);
+		}
+
+		check(loopV[start_i] == bv && loopV[end_i] == bv);
+
+		FEdgeLoop loop(Mesh);
+		FMeshBoundaryLoops::ExtractSpan(loopV, start_i, end_i, true, loop.Vertices);
+		FEdgeLoop::VertexLoopToEdgeLoop(Mesh, loop.Vertices, loop.Edges);
+		loop.BowtieVertices = bowties;
+		subs.Add(loop);
+
+		// If there are no more duplicates of this bowtie, we can treat
+		// it like a regular vertex now
+		if (FMeshBoundaryLoops::CountInList(loopV, bv) < 2)
+		{
+			dupes.Remove(bv);
+		}
+	}
+
+	// Should have one loop left that contains duplicates. 
+	// Extract this as a separate loop
+	int nLeft = 0;
+	for (int i = 0; i < loopV.Num(); ++i)
+	{
+		if (loopV[i] != -1)
+		{
+			nLeft++;
+		}
+	}
+	if (nLeft > 0)
+	{
+		FEdgeLoop loop(Mesh);
+		loop.Vertices.SetNum(nLeft);
+		int vi = 0;
+		for (int i = 0; i < loopV.Num(); ++i)
+		{
+			if (loopV[i] != -1)
+			{
+				loop.Vertices[vi++] = loopV[i];
+			}
+		}
+		FEdgeLoop::VertexLoopToEdgeLoop(Mesh, loop.Vertices, loop.Edges);
+		loop.BowtieVertices = bowties;
+		subs.Add(loop);
+	}
+
+	return subs;
+}
\ No newline at end of file
diff --git a/Engine/Plugins/Experimental/GeometryProcessing/Source/DynamicMesh/Private/MeshSimplification.cpp b/Engine/Plugins/Experimental/GeometryProcessing/Source/DynamicMesh/Private/MeshSimplification.cpp
new file mode 100644
--- /dev/null
+++ b/Engine/Plugins/Experimental/GeometryProcessing/Source/DynamicMesh/Private/MeshSimplification.cpp
@@ -0,0 +1,623 @@
+// Copyright 1998-2019 Epic Games, Inc. All Rights Reserved.
+
+#include "MeshSimplification.h"
+#include "DynamicMeshAttributeSet.h"
+#include "Util/IndexUtil.h"
+
+
+
+
+void FMeshSimplification::InitializeVertexQuadrics()
+{
+	int NT = Mesh->MaxTriangleID();
+	TArray<FQuadricErrord> triQuadrics;
+	triQuadrics.SetNum(NT);
+	TArray<double> triAreas;
+	triAreas.SetNum(NT);
+	//@todo parallel version
+	//gParallel.BlockStartEnd(0, Mesh->MaxTriangleID - 1, (start_tid, end_tid) = > {
+	FVector3d c, n;
+	for (int tid : Mesh->TriangleIndicesItr())
+	{
+		Mesh->GetTriInfo(tid, n, triAreas[tid], c);
+		triQuadrics[tid] = FQuadricErrord(n, c);
+	}
+
+
+	int NV = Mesh->MaxVertexID();
+	vertQuadrics.SetNum(NV);
+	//gParallel.BlockStartEnd(0, Mesh->MaxVertexID - 1, (start_vid, end_vid) = > {
+	for (int vid : Mesh->VertexIndicesItr())
+	{
+		vertQuadrics[vid] = FQuadricErrord::Zero();
+		for (int tid : Mesh->VtxTrianglesItr(vid))
+		{
+			vertQuadrics[vid].Add(triAreas[tid], triQuadrics[tid]);
+		}
+		//check(TMathUtil.EpsilonEqual(0, vertQuadrics[i].Evaluate(Mesh->GetVertex(i)), TMathUtil.Epsilon * 10));
+	}
+}
+
+
+
+void FMeshSimplification::InitializeQueue()
+{
+	int NE = Mesh->EdgeCount();
+	int MaxEID = Mesh->MaxEdgeID();
+
+	EdgeQuadrics.SetNum(MaxEID);
+	EdgeQueue.Initialize(MaxEID);
+	TArray<FEdgeError> EdgeErrors;
+	EdgeErrors.SetNum(MaxEID);
+
+	// @todo vertex quadrics can be computed in parallel
+	//gParallel.BlockStartEnd(0, MaxEID - 1, (start_eid, end_eid) = > {
+	//for (int eid = start_eid; eid <= end_eid; eid++) {
+	for (int eid : Mesh->EdgeIndicesItr())
+	{
+		FIndex2i ev = Mesh->GetEdgeV(eid);
+		FQuadricErrord Q(vertQuadrics[ev.A], vertQuadrics[ev.B]);
+		FVector3d opt = OptimalPoint(eid, Q, ev.A, ev.B);
+		EdgeErrors[eid] = { (float)Q.Evaluate(opt), eid };
+		EdgeQuadrics[eid] = QEdge(eid, Q, opt);
+	}
+
+	// sorted pq insert is faster, so sort edge errors array and index map
+	EdgeErrors.Sort();
+
+	// now do inserts
+	int N = EdgeErrors.Num();
+	for (int i = 0; i < N; ++i) 
+	{
+		int eid = EdgeErrors[i].eid;
+		if (Mesh->IsEdge(eid))
+		{
+			QEdge edge = EdgeQuadrics[eid];
+			EdgeQueue.Insert(edge.eid, EdgeErrors[i].error);
+		}
+	}
+
+	/*
+	// previous code that does unsorted insert. This is marginally slower, but
+	// might get even slower on larger meshes? have only tried up to about 350k.
+	// (still, this function is not the bottleneck...)
+	int cur_eid = StartEdges();
+	bool done = false;
+	do {
+		if (Mesh->IsEdge(cur_eid)) {
+			QEdge edge = EdgeQuadrics[cur_eid];
+			double err = errList[cur_eid];
+			EdgeQueue.Enqueue(cur_eid, (float)err);
+		}
+		cur_eid = GetNextEdge(cur_eid, out done);
+	} while (done == false);
+	*/
+}
+
+
+
+
+
+FVector3d FMeshSimplification::OptimalPoint(int eid, const FQuadricErrord& q, int ea, int eb)
+{
+	// if we would like to preserve boundary, we need to know that here
+	// so that we properly score these edges
+	if (bHaveBoundary && bPreserveBoundaryShape)
+	{
+		if (Mesh->IsBoundaryEdge(eid))
+		{
+			return (Mesh->GetVertex(ea) + Mesh->GetVertex(eb)) * 0.5;
+		}
+		else
+		{
+			if (IsBoundaryVertex(ea))
+			{
+				return Mesh->GetVertex(ea);
+			}
+			else if (IsBoundaryVertex(eb))
+			{
+				return Mesh->GetVertex(eb);
+			}
+		}
+	}
+
+	// [TODO] if we have constraints, we should apply them here, for same reason as bdry above...
+
+	if (bMinimizeQuadricPositionError == false)
+	{
+		return GetProjectedPoint((Mesh->GetVertex(ea) + Mesh->GetVertex(eb)) * 0.5);
+	}
+	else
+	{
+		FVector3d result = FVector3d::Zero();
+		if (q.OptimalPoint(result))
+		{
+			return GetProjectedPoint(result);
+		}
+
+		// degenerate matrix, evaluate quadric at edge end and midpoints
+		// (could do line search here...)
+		FVector3d va = Mesh->GetVertex(ea);
+		FVector3d vb = Mesh->GetVertex(eb);
+		FVector3d c = GetProjectedPoint((va + vb) * 0.5);
+		double fa = q.Evaluate(va);
+		double fb = q.Evaluate(vb);
+		double fc = q.Evaluate(c);
+		double m = FMath::Min3(fa, fb, fc);
+		if (m == fa)
+		{
+			return va;
+		}
+		else if (m == fb)
+		{
+			return vb;
+		}
+		return c;
+	}
+}
+
+
+
+
+// update queue weight for each edge in vertex one-ring
+void FMeshSimplification::UpdateNeighbours(int vid)
+{
+	for (int eid : Mesh->VtxEdgesItr(vid))
+	{
+		FIndex2i nev = Mesh->GetEdgeV(eid);
+		FQuadricErrord Q(vertQuadrics[nev.A], vertQuadrics[nev.B]);
+		FVector3d opt = OptimalPoint(eid, Q, nev.A, nev.B);
+		double err = Q.Evaluate(opt);
+		EdgeQuadrics[eid] = QEdge(eid, Q, opt);
+		if (EdgeQueue.Contains(eid))
+		{
+			EdgeQueue.Update(eid, (float)err);
+		}
+		else
+		{
+			EdgeQueue.Insert(eid, (float)err);
+		}
+	}
+}
+
+
+
+
+void FMeshSimplification::Precompute(bool bMeshIsClosed)
+{
+	bHaveBoundary = false;
+	IsBoundaryVtxCache.SetNum(Mesh->MaxVertexID());
+	if (bMeshIsClosed == false)
+	{
+		for (int eid : Mesh->BoundaryEdgeIndicesItr())
+		{
+			FIndex2i ev = Mesh->GetEdgeV(eid);
+			IsBoundaryVtxCache[ev.A] = true;
+			IsBoundaryVtxCache[ev.B] = true;
+			bHaveBoundary = true;
+		}
+	}
+}
+
+
+
+
+void FMeshSimplification::DoSimplify()
+{
+	if (Mesh->TriangleCount() == 0)    // badness if we don't catch this...
+	{
+		return;
+	}
+
+	ProfileBeginPass();
+
+	ProfileBeginSetup();
+	Precompute();
+	if (Cancelled())
+	{
+		return;
+	}
+	InitializeVertexQuadrics();
+	if (Cancelled())
+	{
+		return;
+	}
+	InitializeQueue();
+	if (Cancelled())
+	{
+		return;
+	}
+	ProfileEndSetup();
+
+	ProfileBeginOps();
+
+	ProfileBeginCollapse();
+	while (EdgeQueue.GetCount() > 0) 
+	{
+		// termination criteria
+		if (SimplifyMode == ETargetModes::VertexCount)
+		{
+			if (Mesh->VertexCount() <= TargetCount)
+			{
+				break;
+			}
+		}
+		else 
+		{
+			if (Mesh->TriangleCount() <= TargetCount)
+			{
+				break;
+			}
+		}
+
+		COUNT_ITERATIONS++;
+		int eid = EdgeQueue.Dequeue();
+		if (Mesh->IsEdge(eid) == false)
+		{
+			continue;
+		}
+		if (Cancelled())
+		{
+			return;
+		}
+
+		int vKeptID;
+		EProcessResult result = CollapseEdge(eid, EdgeQuadrics[eid].collapse_pt, vKeptID);
+		if (result == EProcessResult::Ok_Collapsed) 
+		{
+			vertQuadrics[vKeptID] = EdgeQuadrics[eid].q;
+			UpdateNeighbours(vKeptID);
+		}
+	}
+	ProfileEndCollapse();
+	ProfileEndOps();
+
+	if (Cancelled())
+	{
+		return;
+	}
+
+	Reproject();
+
+	ProfileEndPass();
+}
+
+
+void FMeshSimplification::SimplifyToTriangleCount(int nCount)
+{
+	SimplifyMode = ETargetModes::TriangleCount;
+	TargetCount = FMath::Max(1, nCount);
+	MinEdgeLength = FMathd::MaxReal;
+	DoSimplify();
+}
+
+void FMeshSimplification::SimplifyToVertexCount(int nCount)
+{
+	SimplifyMode = ETargetModes::VertexCount;
+	TargetCount = FMath::Max(3, nCount);
+	MinEdgeLength = FMathd::MaxReal;
+	DoSimplify();
+}
+
+void FMeshSimplification::SimplifyToEdgeLength(double minEdgeLen)
+{
+	SimplifyMode = ETargetModes::MinEdgeLength;
+	TargetCount = 1;
+	MinEdgeLength = minEdgeLen;
+	DoSimplify();
+}
+
+
+
+
+void FMeshSimplification::FastCollapsePass(double fMinEdgeLength, int nRounds, bool MeshIsClosedHint)
+{
+	if (Mesh->TriangleCount() == 0)    // badness if we don't catch this...
+	{
+		return;
+	}
+
+	MinEdgeLength = fMinEdgeLength;
+	double min_sqr = MinEdgeLength * MinEdgeLength;
+
+	// we don't collapse on the boundary
+	bHaveBoundary = false;
+
+	ProfileBeginPass();
+
+	ProfileBeginSetup();
+	Precompute(MeshIsClosedHint);
+	if (Cancelled())
+	{
+		return;
+	}
+	ProfileEndSetup();
+
+	ProfileBeginOps();
+
+	ProfileBeginCollapse();
+
+	int N = Mesh->MaxEdgeID();
+	int num_last_pass = 0;
+	for (int ri = 0; ri < nRounds; ++ri) 
+	{
+		num_last_pass = 0;
+
+		FVector3d va = FVector3d::Zero(), vb = FVector3d::Zero();
+		for (int eid = 0; eid < N; ++eid) 
+		{
+			if ( ( ! Mesh->IsEdge(eid) ) || Mesh->IsBoundaryEdge(eid) )
+			{
+				continue;
+			}
+			if (Cancelled())
+			{
+				return;
+			}
+
+			Mesh->GetEdgeV(eid, va, vb);
+			if (va.DistanceSquared(vb) > min_sqr)
+			{
+				continue;
+			}
+
+			COUNT_ITERATIONS++;
+
+			FVector3d midpoint = (va + vb) * 0.5;
+			int vKeptID;
+			EProcessResult result = CollapseEdge(eid, midpoint, vKeptID);
+			if (result == EProcessResult::Ok_Collapsed) 
+			{
+				++num_last_pass;
+			}
+		}
+
+		if (num_last_pass == 0)     // converged
+		{
+			break;
+		}
+	}
+	ProfileEndCollapse();
+	ProfileEndOps();
+
+	if (Cancelled())
+	{
+		return;
+	}
+
+	Reproject();
+
+	ProfileEndPass();
+}
+
+
+
+
+
+
+
+
+
+
+
+
+
+
+FMeshSimplification::EProcessResult FMeshSimplification::CollapseEdge(int edgeID, FVector3d vNewPos, int& collapseToV)
+{
+	collapseToV = FDynamicMesh3::InvalidID;
+	RuntimeDebugCheck(edgeID);
+
+	FEdgeConstraint constraint =
+		(Constraints == nullptr) ? FEdgeConstraint::Unconstrained() : Constraints->GetEdgeConstraint(edgeID);
+	if (constraint.NoModifications())
+	{
+		return EProcessResult::Ignored_EdgeIsFullyConstrained;
+	}
+	if (constraint.CanCollapse() == false)
+	{
+		return EProcessResult::Ignored_EdgeIsFullyConstrained;
+	}
+
+
+	// look up verts and tris for this edge
+	if (Mesh->IsEdge(edgeID) == false)
+	{
+		return EProcessResult::Failed_NotAnEdge;
+	}
+	FIndex4i edgeInfo = Mesh->GetEdge(edgeID);
+	int a = edgeInfo.A, b = edgeInfo.B, t0 = edgeInfo.C, t1 = edgeInfo.D;
+	bool bIsBoundaryEdge = (t1 == FDynamicMesh3::InvalidID);
+
+	// look up 'other' verts c (from t0) and d (from t1, if it exists)
+	FIndex3i T0tv = Mesh->GetTriangle(t0);
+	int c = IndexUtil::FindTriOtherVtx(a, b, T0tv);
+	FIndex3i T1tv = (bIsBoundaryEdge) ? FDynamicMesh3::InvalidTriangle() : Mesh->GetTriangle(t1);
+	int d = (bIsBoundaryEdge) ? FDynamicMesh3::InvalidID : IndexUtil::FindTriOtherVtx(a, b, T1tv);
+
+	FVector3d vA = Mesh->GetVertex(a);
+	FVector3d vB = Mesh->GetVertex(b);
+	double edge_len_sqr = (vA - vB).SquaredLength();
+	if (edge_len_sqr > MinEdgeLength * MinEdgeLength)
+	{
+		return EProcessResult::Ignored_EdgeTooLong;
+	}
+
+	ProfileBeginCollapse();
+
+	// check if we should collapse, and also find which vertex we should collapse to,
+	// in cases where we have constraints/etc
+	int collapse_to = -1;
+	bool bCanCollapse = CanCollapseEdge(edgeID, a, b, c, d, t0, t1, collapse_to);
+	if (bCanCollapse == false)
+	{
+		return EProcessResult::Ignored_Constrained;
+	}
+
+	// if we have a boundary, we want to collapse to boundary
+	if (bPreserveBoundaryShape && bHaveBoundary) 
+	{
+		if (collapse_to != -1) 
+		{
+			if ((IsBoundaryVertex(b) && collapse_to != b) ||
+				(IsBoundaryVertex(a) && collapse_to != a))
+			{
+				return EProcessResult::Ignored_Constrained;
+			}
+		}
+		if (IsBoundaryVertex(b))
+		{
+			collapse_to = b;
+		}
+		else if (IsBoundaryVertex(a))
+		{
+			collapse_to = a;
+		}
+	}
+
+	// optimization: if edge cd exists, we cannot collapse or flip. look that up here?
+	//  funcs will do it internally...
+	//  (or maybe we can collapse if cd exists? edge-collapse doesn't check for it explicitly...)
+	EProcessResult retVal = EProcessResult::Failed_OpNotSuccessful;
+
+	int iKeep = b, iCollapse = a;
+
+	// if either vtx is fixed, collapse to that position
+	double collapse_t = 0;
+	if (collapse_to == b) 
+	{
+		vNewPos = vB;
+		collapse_t = 0;
+	}
+	else if (collapse_to == a) 
+	{
+		iKeep = a; iCollapse = b;
+		vNewPos = vA;
+		collapse_t = 0;
+	}
+	else
+	{
+		vNewPos = GetProjectedCollapsePosition(iKeep, vNewPos);
+		double div = vA.Distance(vB);
+		collapse_t = (div < FMathd::ZeroTolerance) ? 0.5 : (vNewPos.Distance(Mesh->GetVertex(iKeep))) / div;
+		collapse_t = VectorUtil::Clamp(collapse_t, 0.0, 1.0);
+	}
+
+	// check if this collapse will create a normal flip. Also checks
+	// for invalid collapse nbrhood, since we are doing one-ring iter anyway.
+	// [TODO] could we skip this one-ring check in CollapseEdge? pass in hints?
+	if (CheckIfCollapseCreatesFlipOrInvalid(a, b, vNewPos, t0, t1) || CheckIfCollapseCreatesFlipOrInvalid(b, a, vNewPos, t0, t1)) 
+	{
+		ProfileEndCollapse();
+		return EProcessResult::Ignored_CreatesFlip;
+	}
+
+	// lots of cases where we cannot collapse, but we should just let
+	// Mesh sort that out, right?
+	COUNT_COLLAPSES++;
+	FDynamicMesh3::FEdgeCollapseInfo collapseInfo;
+	EMeshResult result = Mesh->CollapseEdge(iKeep, iCollapse, collapse_t, collapseInfo);
+	if (result == EMeshResult::Ok) 
+	{
+		collapseToV = iKeep;
+		Mesh->SetVertex(iKeep, vNewPos);
+		if (Constraints != nullptr) 
+		{
+			Constraints->ClearEdgeConstraint(edgeID);
+			Constraints->ClearEdgeConstraint(collapseInfo.RemovedEdges.A);
+			if (collapseInfo.RemovedEdges.B != FDynamicMesh3::InvalidID)
+			{
+				Constraints->ClearEdgeConstraint(collapseInfo.RemovedEdges.B);
+			}
+			Constraints->ClearVertexConstraint(iCollapse);
+		}
+		OnEdgeCollapse(edgeID, iKeep, iCollapse, collapseInfo);
+		DoDebugChecks();
+
+		retVal = EProcessResult::Ok_Collapsed;
+	}
+
+	ProfileEndCollapse();
+	return retVal;
+}
+
+
+
+
+
+
+
+
+// Project vertices onto projection target. 
+// We can do projection in parallel if we have .net 
+void FMeshSimplification::FullProjectionPass()
+{
+	auto project = [&](int vID) 
+	{
+		if (IsVertexConstrained(vID))
+		{
+			return;
+		}
+		if (VertexControlF != nullptr && ((int)VertexControlF(vID) & (int)EVertexControl::NoProject) != 0)
+		{
+			return;
+		}
+		FVector3d curpos = Mesh->GetVertex(vID);
+		FVector3d projected = ProjTarget->Project(curpos, vID);
+		Mesh->SetVertex(vID, projected);
+	};
+
+	ApplyToProjectVertices(project);
+
+	// [RMS] not sure how to do this...
+	//if (EnableParallelProjection) {
+	//    gParallel.ForEach<int>(project_vertices(), project);
+	//} else {
+	//    foreach (int vid in project_vertices())
+	//        project(vid);
+	//}
+}
+
+
+void FMeshSimplification::ApplyToProjectVertices(const TFunction<void(int)>& apply_f)
+{
+	for (int vid : Mesh->VertexIndicesItr())
+	{
+		apply_f(vid);
+	}
+}
+
+
+void FMeshSimplification::ProjectVertex(int vID, IProjectionTarget* targetIn)
+{
+	FVector3d curpos = Mesh->GetVertex(vID);
+	FVector3d projected = targetIn->Project(curpos, vID);
+	Mesh->SetVertex(vID, projected);
+}
+
+// used by collapse-edge to get projected position for new vertex
+FVector3d FMeshSimplification::GetProjectedCollapsePosition(int vid, const FVector3d& vNewPos)
+{
+	if (Constraints != nullptr) 
+	{
+		FVertexConstraint vc = Constraints->GetVertexConstraint(vid);
+		if (vc.Target != nullptr)
+		{
+			return vc.Target->Project(vNewPos, vid);
+		}
+		if (vc.Fixed)
+		{
+			return vNewPos;
+		}
+	}
+	// no constraint applied, so if we have a target surface, project to that
+	if (EnableInlineProjection() && ProjTarget != nullptr)
+	{
+		if (VertexControlF == nullptr || ((int)VertexControlF(vid) & (int)EVertexControl::NoProject) == 0)
+		{
+			return ProjTarget->Project(vNewPos, vid);
+		}
+	}
+	return vNewPos;
+}
+
diff --git a/Engine/Plugins/Experimental/GeometryProcessing/Source/DynamicMesh/Private/MeshWeights.cpp b/Engine/Plugins/Experimental/GeometryProcessing/Source/DynamicMesh/Private/MeshWeights.cpp
new file mode 100644
--- /dev/null
+++ b/Engine/Plugins/Experimental/GeometryProcessing/Source/DynamicMesh/Private/MeshWeights.cpp
@@ -0,0 +1,147 @@
+// Copyright 1998-2019 Epic Games, Inc. All Rights Reserved.
+
+#include "MeshWeights.h"
+#include "VectorUtil.h"
+
+
+FVector3d FMeshWeights::UniformCentroid(const FDynamicMesh3 & mesh, int VertexIndex)
+{
+	FVector3d Centroid;
+	mesh.GetVtxOneRingCentroid(VertexIndex, Centroid);
+	return Centroid;
+}
+
+
+FVector3d FMeshWeights::MeanValueCentroid(const FDynamicMesh3 & mesh, int v_i)
+{
+	// based on equations in https://www.inf.usi.ch/hormann/papers/Floater.2006.AGC.pdf (formula 9)
+	// refer to that paper for variable names/etc
+
+	FVector3d vSum = FVector3d::Zero();
+	double wSum = 0;
+	FVector3d Vi = mesh.GetVertex(v_i);
+
+	int v_j = FDynamicMesh3::InvalidID, opp_v1 = FDynamicMesh3::InvalidID, opp_v2 = FDynamicMesh3::InvalidID;
+	int t1 = FDynamicMesh3::InvalidID, t2 = FDynamicMesh3::InvalidID;
+	for (int eid : mesh.VtxEdgesItr(v_i) ) 
+	{
+		opp_v2 = FDynamicMesh3::InvalidID;
+		mesh.GetVtxNbrhood(eid, v_i, v_j, opp_v1, opp_v2, t1, t2);
+
+		FVector3d Vj = mesh.GetVertex(v_j);
+		FVector3d vVj = (Vj - Vi);
+		double len_vVj = vVj.Normalize();
+		// [RMS] is this the right thing to do? if vertices are coincident,
+		//   weight of this vertex should be very high!
+		if (len_vVj < FMathd::ZeroTolerance)
+			continue;
+		FVector3d vVdelta = (mesh.GetVertex(opp_v1) - Vi).Normalized();
+		double w_ij = VectorUtil::VectorTanHalfAngle(vVj, vVdelta);
+
+		if (opp_v2 != FDynamicMesh3::InvalidID) {
+			FVector3d vVgamma = (mesh.GetVertex(opp_v2) - Vi).Normalized();
+			w_ij += VectorUtil::VectorTanHalfAngle(vVj, vVgamma);
+		}
+
+		w_ij /= len_vVj;
+
+		vSum += w_ij * Vj;
+		wSum += w_ij;
+	}
+	if (wSum < FMathd::ZeroTolerance)
+		return Vi;
+	return vSum / wSum;
+}
+
+
+
+
+FVector3d FMeshWeights::CotanCentroid(const FDynamicMesh3& mesh, int v_i)
+{
+	// based on equations in http://www.geometry.caltech.edu/pubs/DMSB_III.pdf
+
+	FVector3d vSum = FVector3d::Zero();
+	double wSum = 0;
+	FVector3d Vi = mesh.GetVertex(v_i);
+
+	int v_j = FDynamicMesh3::InvalidID, opp_v1 = FDynamicMesh3::InvalidID, opp_v2 = FDynamicMesh3::InvalidID;
+	int t1 = FDynamicMesh3::InvalidID, t2 = FDynamicMesh3::InvalidID;
+	bool bAborted = false;
+	for (int eid : mesh.VtxEdgesItr(v_i)) 
+	{
+		opp_v2 = FDynamicMesh3::InvalidID;
+		mesh.GetVtxNbrhood(eid, v_i, v_j, opp_v1, opp_v2, t1, t2);
+		FVector3d Vj = mesh.GetVertex(v_j);
+
+		FVector3d Vo1 = mesh.GetVertex(opp_v1);
+		double cot_alpha_ij = VectorUtil::VectorCot(
+			(Vi - Vo1).Normalized(), (Vj - Vo1).Normalized());
+		if (cot_alpha_ij == 0) {
+			bAborted = true;
+			break;
+		}
+		double w_ij = cot_alpha_ij;
+
+		if (opp_v2 != FDynamicMesh3::InvalidID) {
+			FVector3d Vo2 = mesh.GetVertex(opp_v2);
+			double cot_beta_ij = VectorUtil::VectorCot(
+				(Vi - Vo2).Normalized(), (Vj - Vo2).Normalized());
+			if (cot_beta_ij == 0) {
+				bAborted = true;
+				break;
+			}
+			w_ij += cot_beta_ij;
+		}
+
+		vSum += w_ij * Vj;
+		wSum += w_ij;
+	}
+	if (bAborted || fabs(wSum) < FMathd::ZeroTolerance)
+		return Vi;
+	return vSum / wSum;
+}
+
+
+
+double FMeshWeights::VoronoiArea(const FDynamicMesh3& mesh, int v_i)
+{
+	// based on equations in http://www.geometry.caltech.edu/pubs/DMSB_III.pdf
+
+	double areaSum = 0;
+	FVector3d Vi = mesh.GetVertex(v_i);
+
+	for (int tid : mesh.VtxTrianglesItr(v_i) ) 
+	{
+		FIndex3i t = mesh.GetTriangle(tid);
+		int ti = (t[0] == v_i) ? 0 : ((t[1] == v_i) ? 1 : 2);
+		FVector3d Vj = mesh.GetVertex(t[(ti + 1) % 3]);
+		FVector3d Vk = mesh.GetVertex(t[(ti + 2) % 3]);
+
+		if (VectorUtil::IsObtuse(Vi, Vj, Vk)) 
+		{
+			// if triangle is obtuse voronoi area is undefind and we just return portion of triangle area
+			FVector3d Vij = Vj - Vi;
+			FVector3d Vik = Vk - Vi;
+			Vij.Normalize(); Vik.Normalize();
+			double areaT = 0.5 * Vij.Cross(Vik).Length();
+			areaSum += ( Vij.AngleR(Vik) > FMathd::HalfPi ) ?    // obtuse at v_i ?
+				(areaT * 0.5) : (areaT * 0.25);
+
+		}
+		else 
+		{
+			// voronoi area
+			FVector3d Vji = Vi - Vj;
+			double dist_ji = Vji.Normalize();
+			FVector3d Vki = Vi - Vk;
+			double dist_ki = Vki.Normalize();
+			FVector3d Vkj = (Vj - Vk).Normalized();
+
+			double cot_alpha_ij = VectorUtil::VectorCot(Vki, Vkj);
+			double cot_alpha_ik = VectorUtil::VectorCot(Vji, -Vkj);
+			areaSum += dist_ji * dist_ji * cot_alpha_ij * 0.125;
+			areaSum += dist_ki * dist_ki * cot_alpha_ik * 0.125;
+		}
+	}
+	return areaSum;
+}
\ No newline at end of file
diff --git a/Engine/Plugins/Experimental/GeometryProcessing/Source/DynamicMesh/Private/Operations/ExtrudeMesh.cpp b/Engine/Plugins/Experimental/GeometryProcessing/Source/DynamicMesh/Private/Operations/ExtrudeMesh.cpp
new file mode 100644
--- /dev/null
+++ b/Engine/Plugins/Experimental/GeometryProcessing/Source/DynamicMesh/Private/Operations/ExtrudeMesh.cpp
@@ -0,0 +1,121 @@
+// Copyright 1998-2019 Epic Games, Inc. All Rights Reserved.
+
+#include "Operations/ExtrudeMesh.h"
+#include "MeshNormals.h"
+#include "DynamicMeshEditor.h"
+
+
+FExtrudeMesh::FExtrudeMesh(FDynamicMesh3* mesh) : Mesh(mesh)
+{
+	ExtrudedPositionFunc = [this](const FVector3d& Position, const FVector3f& Normal, int VertexID) 
+	{
+		return Position + this->DefaultExtrudeDistance * Normal;
+	};
+}
+
+
+bool FExtrudeMesh::Apply()
+{
+	//@todo should apply per-connected-component? will then handle bowties properly
+
+	FMeshNormals Normals;
+	bool bHaveVertexNormals = Mesh->HasVertexNormals();
+	if (!bHaveVertexNormals)
+	{
+		Normals = FMeshNormals(Mesh);
+		Normals.ComputeVertexNormals();
+	}
+
+	InitialLoops.SetMesh(Mesh);
+	InitialLoops.Compute();
+	int NumInitialLoops = InitialLoops.GetLoopCount();
+
+	BufferUtil::AppendElements(InitialTriangles, Mesh->TriangleIndicesItr());
+	BufferUtil::AppendElements(InitialVertices, Mesh->VertexIndicesItr());
+
+	// duplicate triangles of mesh
+
+	FDynamicMeshEditor Editor(Mesh);
+
+	FMeshIndexMappings IndexMap;
+	FDynamicMeshEditResult DuplicateResult;
+	Editor.DuplicateTriangles(InitialTriangles, IndexMap, DuplicateResult);
+	OffsetTriangles = DuplicateResult.NewTriangles;
+	OffsetTriGroups = DuplicateResult.NewGroups;
+	InitialToOffsetMapV = IndexMap.GetVertexMap().GetForwardMap();
+
+	// set vertices to new positions
+	for (int vid : InitialVertices)
+	{
+		int newvid = InitialToOffsetMapV[vid];
+		if ( ! Mesh->IsVertex(newvid) )
+		{
+			continue;
+		}
+
+		FVector3d v = Mesh->GetVertex(vid);
+		FVector3f n = (bHaveVertexNormals) ? Mesh->GetVertexNormal(vid) : (FVector3f)Normals[vid];
+		FVector3d newv = ExtrudedPositionFunc(v, n, vid);
+
+		Mesh->SetVertex(newvid, newv);
+	}
+
+	// we need to reverse one side
+	if (IsPositiveOffset)
+	{
+		Editor.ReverseTriangleOrientations(InitialTriangles, true);
+	}
+	else
+	{
+		Editor.ReverseTriangleOrientations(OffsetTriangles, true);
+	}
+
+	// stitch each loop
+	NewLoops.SetNum(NumInitialLoops);
+	StitchTriangles.SetNum(NumInitialLoops);
+	StitchPolygonIDs.SetNum(NumInitialLoops);
+	int LoopIndex = 0;
+	for (FEdgeLoop& BaseLoop : InitialLoops.Loops)
+	{
+		int LoopCount = BaseLoop.GetVertexCount();
+
+		TArray<int> OffsetLoop;
+		OffsetLoop.SetNum(LoopCount);
+		for (int k = 0; k < LoopCount; ++k)
+		{
+			OffsetLoop[k] = InitialToOffsetMapV[BaseLoop.Vertices[k]];
+		}
+
+		FDynamicMeshEditResult StitchResult;
+		if (IsPositiveOffset)
+		{
+			Editor.StitchVertexLoopsMinimal(OffsetLoop, BaseLoop.Vertices, StitchResult);
+		}
+		else
+		{
+			Editor.StitchVertexLoopsMinimal(BaseLoop.Vertices, OffsetLoop, StitchResult);
+		}
+		StitchResult.GetAllTriangles(StitchTriangles[LoopIndex]);
+		StitchPolygonIDs[LoopIndex] = StitchResult.NewGroups;
+
+		// for each polygon we created in stitch, set UVs and normals
+		int NumNewQuads = StitchResult.NewQuads.Num();
+		for (int k = 0; k < NumNewQuads; k++)
+		{
+			FVector3f Normal = Editor.ComputeAndSetQuadNormal(StitchResult.NewQuads[k], true);
+
+			// @todo is there a simpler way to construct rotation from 3 known axes (third id Normal.Cross(UnitY))
+			//  (converting from matrix might end up being more efficient due to trig in ConstrainedAlignAxis?)
+			FFrame3f ProjectFrame(FVector3f::Zero(), Normal);
+			ProjectFrame.ConstrainedAlignAxis(1, FVector3f::UnitY(), ProjectFrame.Z());
+			Editor.SetQuadUVsFromProjection(StitchResult.NewQuads[k], ProjectFrame, UVScaleFactor);
+		}
+
+		NewLoops[LoopIndex].InitializeFromVertices(Mesh, OffsetLoop);
+		LoopIndex++;
+	}
+
+	return true;
+}
+
+
diff --git a/Engine/Plugins/Experimental/GeometryProcessing/Source/DynamicMesh/Public/DynamicMesh3.h b/Engine/Plugins/Experimental/GeometryProcessing/Source/DynamicMesh/Public/DynamicMesh3.h
new file mode 100644
--- /dev/null
+++ b/Engine/Plugins/Experimental/GeometryProcessing/Source/DynamicMesh/Public/DynamicMesh3.h
@@ -0,0 +1,1411 @@
+// Copyright 1998-2019 Epic Games, Inc. All Rights Reserved.
+
+// Port of geometry3cpp DMesh3
+
+#pragma once
+
+#include "BoxTypes.h"
+#include "FrameTypes.h"
+#include "GeometryTypes.h"
+#include "MathUtil.h"
+#include "Quaternion.h"
+#include "Util/DynamicVector.h"
+#include "Util/IndexUtil.h"
+#include "Util/IteratorUtil.h"
+#include "Util/RefCountVector.h"
+#include "Util/SmallListSet.h"
+#include "VectorTypes.h"
+#include "VectorUtil.h"
+
+
+class FDynamicMeshAttributeSet;
+class FMeshShapeGenerator;
+
+enum class EMeshComponents
+{
+	None = 0,
+	VertexNormals = 1,
+	VertexColors = 2,
+	VertexUVs = 4,
+	FaceGroups = 8
+};
+
+/**
+ * FVertexInfo stores information about vertex attributes - position, normal, color, UV
+ */
+struct FVertexInfo
+{
+	FVector3d Position;
+	FVector3f Normal;
+	FVector3f Color;
+	FVector2f UV;
+	bool bHaveN, bHaveUV, bHaveC;
+
+	FVertexInfo()
+	{
+		Position = FVector3d::Zero();
+		Normal = Color = FVector3f::Zero();
+		UV = FVector2f::Zero();
+		bHaveN = bHaveC = bHaveUV = false;
+	}
+	FVertexInfo(const FVector3d& PositionIn)
+	{
+		Position = PositionIn;
+		Normal = Color = FVector3f::Zero();
+		UV = FVector2f::Zero();
+		bHaveN = bHaveC = bHaveUV = false;
+	}
+	FVertexInfo(const FVector3d& PositionIn, const FVector3f& NormalIn)
+	{
+		Position = PositionIn;
+		Normal = NormalIn;
+		Color = FVector3f::Zero();
+		UV = FVector2f::Zero();
+		bHaveN = true;
+		bHaveC = bHaveUV = false;
+	}
+	FVertexInfo(const FVector3d& PositionIn, const FVector3f& NormalIn, const FVector3f& ColorIn)
+	{
+		Position = PositionIn;
+		Normal = NormalIn;
+		Color = ColorIn;
+		UV = FVector2f::Zero();
+		bHaveN = bHaveC = true;
+		bHaveUV = false;
+	}
+	FVertexInfo(const FVector3d& PositionIn, const FVector3f& NormalIn, const FVector3f& ColorIn, const FVector2f& UVIn)
+	{
+		Position = PositionIn;
+		Normal = NormalIn;
+		Color = ColorIn;
+		UV = UVIn;
+		bHaveN = bHaveC = bHaveUV = true;
+	}
+};
+
+/**
+* FDynamicMesh3 is a dynamic triangle mesh class. The mesh has has connectivity,
+* is an indexed mesh, and allows for gaps in the index space.
+*
+* internally, all data is stored in POD-type buffers, except for the vertex->edge
+* links, which are stored as List<int>'s. The arrays of POD data are stored in
+* FDynamicVector's, so they grow in chunks, which is relatively efficient. The actual
+* blocks are arrays, so they can be efficiently mem-copied into larger buffers
+* if necessary.
+*
+* Reference counts for verts/tris/edges are stored as separate FRefCountVector
+* instances.
+*
+* Vertices are stored as doubles, although this should be easily changed
+* if necessary, as the internal data structure is not exposed
+*
+* Per-vertex Vertex Normals, Colors, and UVs are optional and stored as floats.
+*
+* For each vertex, VertexEdgeLists[i] is the unordered list of connected edges. The
+* elements of the list are indices into the edges list.
+* This list is unsorted but can be traversed in-order (ie cw/ccw) at some additional cost.
+*
+* Triangles are stored as 3 ints, with optionally a per-triangle integer group id.
+*
+* The edges of a triangle are similarly stored as 3 ints, in triangle_edes. If the
+* triangle is [v1,v2,v3], then the triangle edges [e1,e2,e3] are
+* e1=edge(v1,v2), e2=edge(v2,v3), e3=edge(v3,v1), where the e# are indexes into edges.
+*
+* Edges are stored as tuples of 4 ints. If the edge is between v1 and v2, with neighbour
+* tris t1 and t2, then the edge is [min(v1,v2), max(v1,v2), t1, t2]. For a boundary
+* edge, t2 is InvalidID. t1 is never InvalidID.
+*
+* Most of the class assumes that the mesh is manifold. Many functions will
+* work if the topology is non-manifold, but behavior of operators like Split/Flip/Collapse
+* edge is untested.
+*
+* The function CheckValidity() does extensive sanity checking on the mesh data structure.
+* Use this to test your code, both for mesh construction and editing!!
+*
+*
+* TODO:
+*  - Many of the iterators depend on lambda functions, can we replace these with calls to
+*    internal/static functions that do the same thing?
+*  - efficient TriTrianglesItr() implementation?
+*  - additional Topology timestamp?
+*  - CompactInPlace() does not compact VertexEdgeLists
+*  - CompactCopy does not support per-vertex or extended attributes
+*  ? FDynamicVector w/ 'stride' option, so that we can guarantee that tuples are in single block.
+*    The can have custom accessor that looks up entire tuple
+*/
+class DYNAMICMESH_API FDynamicMesh3
+{
+public:
+	/** InvalidID indicates that a vertex/edge/triangle ID is invalid */
+	static const int InvalidID; // = IndexConstants::InvalidID = -1
+	/** NonManifoldID is returned by AppendTriangle() to indicate that the added triangle would result in nonmanifold geometry and hence was ignored */
+	static const int NonManifoldID; // = -2
+	/** InvalidGroupID indicates that a group ID is invalid */
+	static const int InvalidGroupID; // = IndexConstants::InvalidID = -1
+
+	static FVector3d InvalidVertex()
+	{
+		return FVector3d(TNumericLimits<double>::Max(), 0, 0);
+	}
+	static FIndex3i InvalidTriangle()
+	{
+		return FIndex3i(InvalidID, InvalidID, InvalidID);
+	}
+	static FIndex2i InvalidEdge()
+	{
+		return FIndex2i(InvalidID, InvalidID);
+	}
+
+protected:
+	/** Reference counts of vertex indices. Iterate over this to find out which vertex indices are valid. */
+	FRefCountVector VertexRefCounts;
+	/** List of vertex positions */
+	FDynamicVector<double> Vertices;
+	/** (optional) List of per-vertex normals */
+	FDynamicVector<float>* VertexNormals = nullptr;
+	/** (optional) List of per-vertex colors */
+	FDynamicVector<float>* VertexColors = nullptr;
+	/** (optional) List of per-vertex uvs */
+	FDynamicVector<float>* VertexUVs = nullptr;
+
+	/** List of per-vertex edge one-rings */
+	FSmallListSet VertexEdgeLists;
+
+	/** Reference counts of triangle indices. Iterate over this to find out which triangle indices are valid. */
+	FRefCountVector TriangleRefCounts;
+	/** List of triangle vertex-index triplets [Vert0 Vert1 Vert2]*/
+	FDynamicVector<int> Triangles;
+	/** List of triangle edge triplets [Edge0 Edge1 Edge2] */
+	FDynamicVector<int> TriangleEdges;
+	/** (optional) List of per-triangle group identifiers */
+	FDynamicVector<int>* TriangleGroups = nullptr;
+
+	/** Reference counts of edge indices. Iterate over this to find out which edge indices are valid. */
+	FRefCountVector EdgeRefCounts;
+	/** List of edge elements. An edge is four elements [VertA, VertB, Tri0, Tri1], where VertA < VertB, and Tri1 may be InvalidID (if the edge is a boundary edge) */
+	FDynamicVector<int> Edges;
+
+	FDynamicMeshAttributeSet* AttributeSet = nullptr;
+
+	/** The mesh timestamp is incremented any time a function that modifies the mesh is called */
+	int Timestamp = 0;
+	/** The shape timestamp is incremented any time a function that modifies the mesh shape or topology is called */
+	int ShapeTimestamp = 0;
+	/** The topology timestamp is incremented any time a function that modifies the mesh topology is called */
+	int TopologyTimestamp = 0;
+
+	/** Upper bound on the triangle group IDs used in the mesh (may be larger than the actual maximum if triangles have been deleted) */
+	int GroupIDCounter = 0;
+
+	/** Cached vertex bounding box (includes un-referenced vertices) */
+	FAxisAlignedBox3d CachedBoundingBox;
+	/** timestamp for CachedBoundingBox, if less than current timestamp, cache is invalid */
+	int CachedBoundingBoxTimestamp = -1;
+	/** Cached value of IsClosed() */
+	bool bIsClosedCached = false;
+	/** timestamp for bIsClosedCached, if less than current timestamp, cache is invalid */
+	int CachedIsClosedTimestamp = -1;
+
+public:
+	virtual ~FDynamicMesh3();
+
+	FDynamicMesh3(bool bWantNormals = true, bool bWantColors = false, bool bWantUVs = false, bool bWantTriGroups = false);
+
+	FDynamicMesh3(EMeshComponents flags)
+		: FDynamicMesh3(((int)flags & (int)EMeshComponents::VertexNormals) != 0, ((int)flags & (int)EMeshComponents::VertexColors) != 0,
+						((int)flags & (int)EMeshComponents::VertexUVs) != 0, ((int)flags & (int)EMeshComponents::FaceGroups) != 0)
+	{
+	}
+
+	/**
+	 * @param CopyMesh mesh to copy
+	 * @param bCompact if true, compact CopyMesh on the fly
+	 * @param bWantNormals should we copy per-vertex normals, if they exist
+	 * @param bWantColors should we copy per-vertex colors, if they exist
+	 * @param bWantUVs should we copy per-vertex uvs, if they exist
+	 */
+	FDynamicMesh3(const FDynamicMesh3& CopyMesh, bool bCompact = false, bool bWantNormals = true, bool bWantColors = true, bool bWantUVs = true, bool bWantAttributes = true);
+
+	/**
+	 * @param CopyMesh mesh to copy
+	 * @param bCompact if true, compact CopyMesh on the fly
+	 * @param Flags which components of CopyMesh to copy, if they exist
+	 */
+	FDynamicMesh3(const FDynamicMesh3& CopyMesh, bool bCompact, EMeshComponents Flags);
+
+	/** Initialize mesh from the output of a MeshShapeGenerator (assumes Generate() was already called) */
+	FDynamicMesh3(const FMeshShapeGenerator* Generator);
+
+	//
+	// Copy functions to construct a mesh from an input mesh
+	//
+public:
+	/** Set internal data structures to be a copy of input mesh */
+	void Copy(const FDynamicMesh3& CopyMesh, bool bNormals = true, bool bColors = true, bool bUVs = true, bool bAttributes = true);
+
+	/** Initialize mesh from the output of a MeshShapeGenerator (assumes Generate() was already called) */
+	void Copy(const FMeshShapeGenerator* Generator);
+
+	// @todo make this work
+	struct FCompactMaps
+	{
+		TMap<int, int> MapV;
+	};
+
+	/** Copy input mesh while compacting, ie removing unused vertices/triangles/edges */
+	void CompactCopy(const FDynamicMesh3& CopyMesh, bool bNormals = true, bool bColors = true, bool bUVs = true, bool bAttributes = true, FCompactMaps* CompactInfo = nullptr);
+
+	/** Discard all data */
+	void Clear();
+
+	
+public:
+
+	/** @return number of vertices in the mesh */
+	int VertexCount() const
+	{
+		return (int)VertexRefCounts.GetCount();
+	}
+	/** @return number of triangles in the mesh */
+	int TriangleCount() const
+	{
+		return (int)TriangleRefCounts.GetCount();
+	}
+	/** @return number of edges in the mesh */
+	int EdgeCount() const
+	{
+		return (int)EdgeRefCounts.GetCount();
+	}
+
+	/** @return upper bound on vertex IDs used in the mesh, ie all vertex IDs in use are < MaxVertexID */
+	int MaxVertexID() const
+	{
+		return (int)VertexRefCounts.GetMaxIndex();
+	}
+	/** @return upper bound on triangle IDs used in the mesh, ie all triangle IDs in use are < MaxTriangleID */
+	int MaxTriangleID() const
+	{
+		return (int)TriangleRefCounts.GetMaxIndex();
+	}
+	/** @return upper bound on edge IDs used in the mesh, ie all edge IDs in use are < MaxEdgeID */
+	int MaxEdgeID() const
+	{
+		return (int)EdgeRefCounts.GetMaxIndex();
+	}
+	/** @return upper bound on group IDs used in the mesh, ie all group IDs in use are < MaxGroupID */
+	int MaxGroupID() const
+	{
+		return GroupIDCounter;
+	}
+
+
+	/** @return true if this mesh has per-vertex normals */
+	bool HasVertexNormals() const
+	{
+		return VertexNormals != nullptr;
+	}
+	/** @return true if this mesh has per-vertex colors */
+	bool HasVertexColors() const
+	{
+		return VertexColors != nullptr;
+	}
+	/** @return true if this mesh has per-vertex UVs */
+	bool HasVertexUVs() const
+	{
+		return VertexUVs != nullptr;
+	}
+	/** @return true if this mesh has per-triangle groups */
+	bool HasTriangleGroups() const
+	{
+		return TriangleGroups != nullptr;
+	}
+
+
+	/** @return true if this mesh has attribute layers */
+	bool HasAttributes() const
+	{
+		return AttributeSet != nullptr;
+	}
+
+	/** @return bitwise-or of EMeshComponents flags specifying which extra data this mesh has */
+	int GetComponentsFlags() const;
+
+
+	/** @return true if VertexID is a valid vertex in this mesh */
+	inline bool IsVertex(int VertexID) const
+	{
+		return VertexRefCounts.IsValid(VertexID);
+	}
+	/** @return true if TriangleID is a valid triangle in this mesh */
+	inline bool IsTriangle(int TriangleID) const
+	{
+		return TriangleRefCounts.IsValid(TriangleID);
+	}
+	/** @return true if EdgeID is a valid edge in this mesh */
+	inline bool IsEdge(int EdgeID) const
+	{
+		return EdgeRefCounts.IsValid(EdgeID);
+	}
+
+
+	//
+	// Mesh Element Iterators
+	//   The functions VertexIndicesItr() / TriangleIndicesItr() / EdgeIndicesItr() allow you to do:
+	//      for ( int eid : EdgeIndicesItr() ) { ... }
+	//   and other related begin() / end() idioms
+public:
+	// simplify names for iterations
+	typedef typename FRefCountVector::IndexEnumerable vertex_iterator;
+	typedef typename FRefCountVector::IndexEnumerable triangle_iterator;
+	typedef typename FRefCountVector::IndexEnumerable edge_iterator;
+	template <typename T>
+	using value_iteration = FRefCountVector::MappedEnumerable<T>;
+	using vtx_triangles_enumerable = ExpandEnumerable<int, int, FSmallListSet::ValueIterator>;
+
+	/** @return enumerable object for valid vertex indices suitable for use with range-based for, ie for ( int i : VertexIndicesItr() ) */
+	vertex_iterator VertexIndicesItr() const
+	{
+		return VertexRefCounts.Indices();
+	}
+
+	/** @return enumerable object for valid triangle indices suitable for use with range-based for, ie for ( int i : TriangleIndicesItr() ) */
+	triangle_iterator TriangleIndicesItr() const
+	{
+		return TriangleRefCounts.Indices();
+	}
+
+	/** @return enumerable object for valid edge indices suitable for use with range-based for, ie for ( int i : EdgeIndicesItr() ) */
+	edge_iterator EdgeIndicesItr() const
+	{
+		return EdgeRefCounts.Indices();
+	}
+
+	// TODO: write helper functions that allow us to do these iterations w/o lambdas
+
+	/** @return enumerable object for boundary edge indices suitable for use with range-based for, ie for ( int i : BoundaryEdgeIndicesItr() ) */
+	FRefCountVector::FilteredEnumerable BoundaryEdgeIndicesItr() const
+	{
+		return EdgeRefCounts.FilteredIndices([this](int EdgeID) {\
+			return Edges[4*EdgeID + 3] == InvalidID;
+		});
+	}
+
+	/** Enumerate positions of all vertices in mesh */
+	value_iteration<FVector3d> VerticesItr() const
+	{
+		return VertexRefCounts.MappedIndices<FVector3d>([this](int VertexID) {
+			int i = 3 * VertexID;
+			return FVector3d(Vertices[i], Vertices[i + 1], Vertices[i + 2]);
+		});
+	}
+
+	/** Enumerate all triangles in the mesh */
+	value_iteration<FIndex3i> TrianglesItr() const
+	{
+		return TriangleRefCounts.MappedIndices<FIndex3i>([this](int TriangleID) {
+			int i = 3 * TriangleID;
+			return FIndex3i(Triangles[i], Triangles[i + 1], Triangles[i + 2]);
+		});
+	}
+
+	/** Enumerate edges. Each returned element is [v0,v1,t0,t1], where t1 will be InvalidID if this is a boundary edge */
+	value_iteration<FIndex4i> EdgesItr() const
+	{
+		return EdgeRefCounts.MappedIndices<FIndex4i>([this](int EdgeID) {
+			int i = 4 * EdgeID;
+			return FIndex4i(Edges[i], Edges[i + 1], Edges[i + 2], Edges[i + 3]);
+		});
+	}
+
+	/** @return enumerable object for one-ring vertex neighbours of a vertex, suitable for use with range-based for, ie for ( int i : VtxVerticesItr(VertexID) ) */
+	FSmallListSet::ValueEnumerable VtxVerticesItr(int VertexID) const
+	{
+		check(VertexRefCounts.IsValid(VertexID));
+		return VertexEdgeLists.Values(VertexID, [VertexID, this](int eid) { return GetOtherEdgeVertex(eid, VertexID); });
+	}
+
+	/** @return enumerable object for one-ring edges of a vertex, suitable for use with range-based for, ie for ( int i : VtxEdgesItr(VertexID) ) */
+	FSmallListSet::ValueEnumerable VtxEdgesItr(int VertexID) const
+	{
+		check(VertexRefCounts.IsValid(VertexID));
+		return VertexEdgeLists.Values(VertexID);
+	}
+
+	/** @return enumerable object for one-ring triangles of a vertex, suitable for use with range-based for, ie for ( int i : VtxTrianglesItr(VertexID) ) */
+	vtx_triangles_enumerable VtxTrianglesItr(int VertexID) const;
+
+
+
+	//
+	// Mesh Construction
+	//
+public:
+	/** Append vertex at position and other fields, returns vid */
+	int AppendVertex(const FVertexInfo& VertInfo);
+
+	/** Append vertex at position, returns vid */
+	int AppendVertex(const FVector3d& Position)
+	{
+		return AppendVertex(FVertexInfo(Position));
+	}
+
+	/** Copy vertex SourceVertexID from existing SourceMesh, returns new vertex id */
+	int AppendVertex(const FDynamicMesh3& SourceMesh, int SourceVertexID);
+
+	int AppendTriangle(const FIndex3i& TriVertices, int GroupID = -1);
+
+	inline int AppendTriangle(int Vertex0, int Vertex1, int Vertex2, int GroupID = -1)
+	{
+		return AppendTriangle(FIndex3i(Vertex0, Vertex1, Vertex2), GroupID);
+	}
+
+	//
+	// Support for inserting vertex and triangle at specific IDs. This is a bit tricky
+	// because we likely will need to update the free lists in the RefCountVectors, which
+	// can be expensive. If you are going to do many inserts (eg inside a loop), wrap in
+	// BeginUnsafe / EndUnsafe calls, and pass bUnsafe = true to the InsertX() calls, to
+	// the defer free list rebuild until you are done.
+	//
+
+	/** Call this before a set of unsafe InsertVertex() calls */
+	virtual void BeginUnsafeVerticesInsert()
+	{
+		// do nothing...
+	}
+
+	/** Call after a set of unsafe InsertVertex() calls to rebuild free list */
+	virtual void EndUnsafeVerticesInsert()
+	{
+		VertexRefCounts.RebuildFreeList();
+	}
+
+	/**
+	 * Insert vertex at given index, assuming it is unused.
+	 * If bUnsafe, we use fast id allocation that does not update free list.
+	 * You should only be using this between BeginUnsafeVerticesInsert() / EndUnsafeVerticesInsert() calls
+	 */
+	EMeshResult InsertVertex(int VertexID, const FVertexInfo& VertInfo, bool bUnsafe = false);
+
+	/** Call this before a set of unsafe InsertTriangle() calls */
+	virtual void BeginUnsafeTrianglesInsert()
+	{
+		// do nothing...
+	}
+
+	/** Call after a set of unsafe InsertTriangle() calls to rebuild free list */
+	virtual void EndUnsafeTrianglesInsert()
+	{
+		TriangleRefCounts.RebuildFreeList();
+	}
+
+	/**
+	 * Insert triangle at given index, assuming it is unused.
+	 * If bUnsafe, we use fast id allocation that does not update free list.
+	 * You should only be using this between BeginUnsafeTrianglesInsert() / EndUnsafeTrianglesInsert() calls
+	 */
+	EMeshResult InsertTriangle(int TriangleID, const FIndex3i& TriVertices, int GroupID = -1, bool bUnsafe = false);
+
+	//
+	// Vertex/Tri/Edge accessors
+	//
+public:
+
+	/** @return the vertex position */
+	inline FVector3d GetVertex(int VertexID) const
+	{
+		check(IsVertex(VertexID));
+		int i = 3 * VertexID;
+		return FVector3d(Vertices[i], Vertices[i + 1], Vertices[i + 2]);
+	}
+
+	/** Set vertex position */
+	inline void SetVertex(int VertexID, const FVector3d& vNewPos)
+	{
+		check(VectorUtil::IsFinite(vNewPos));
+		check(IsVertex(VertexID));
+		int i = 3 * VertexID;
+		Vertices[i] = vNewPos.X;
+		Vertices[i + 1] = vNewPos.Y;
+		Vertices[i + 2] = vNewPos.Z;
+		UpdateTimeStamp(true, false);
+	}
+
+	/** Get extended vertex information */
+	bool GetVertex(int VertexID, FVertexInfo& VertInfo, bool bWantNormals, bool bWantColors, bool bWantUVs) const;
+
+	/** Get all vertex information available */
+	FVertexInfo GetVertexInfo(int VertexID) const;
+
+	/** @return the valence of a vertex (the number of connected edges) */
+	int GetVtxEdgeCount(int VertexID) const
+	{
+		return VertexRefCounts.IsValid(VertexID) ? VertexEdgeLists.GetCount(VertexID) : -1;
+	}
+
+	/** @return the max valence of all vertices in the mesh */
+	int GetMaxVtxEdgeCount() const;
+
+	/** Get triangle vertices */
+	inline FIndex3i GetTriangle(int TriangleID) const
+	{
+		check(IsTriangle(TriangleID));
+		int i = 3 * TriangleID;
+		return FIndex3i(Triangles[i], Triangles[i + 1], Triangles[i + 2]);
+	}
+
+	/** Get triangle edges */
+	inline FIndex3i GetTriEdges(int TriangleID) const
+	{
+		check(IsTriangle(TriangleID));
+		int i = 3 * TriangleID;
+		return FIndex3i(TriangleEdges[i], TriangleEdges[i + 1], TriangleEdges[i + 2]);
+	}
+
+	/** Get one of the edges of a triangle */
+	inline int GetTriEdge(int TriangleID, int j) const
+	{
+		check(IsTriangle(TriangleID));
+		return TriangleEdges[3 * TriangleID + j];
+	}
+
+	/** Find the neighbour triangles of a triangle (any of them might be InvalidID) */
+	FIndex3i GetTriNeighbourTris(int TriangleID) const;
+
+	/** Get the three vertex positions of a triangle */
+	inline void GetTriVertices(int TriangleID, FVector3d& v0, FVector3d& v1, FVector3d& v2) const
+	{
+		int i = 3 * TriangleID;
+		int ai = 3 * Triangles[i];
+		v0.X = Vertices[ai];		v0.Y = Vertices[ai + 1];		v0.Z = Vertices[ai + 2];
+		int bi = 3 * Triangles[i + 1];
+		v1.X = Vertices[bi];		v1.Y = Vertices[bi + 1];		v1.Z = Vertices[bi + 2];
+		int ci = 3 * Triangles[i + 2];
+		v2.X = Vertices[ci];		v2.Y = Vertices[ci + 1];		v2.Z = Vertices[ci + 2];
+	}
+
+	/** Get the position of one of the vertices of a triangle */
+	inline FVector3d GetTriVertex(int TriangleID, int j) const
+	{
+		int vi = 3 * Triangles[3 * TriangleID + j];
+		return FVector3d(Vertices[vi], Vertices[vi + 1], Vertices[vi + 2]);
+	}
+
+	/** Get the vertices and triangles of an edge, returned as [v0,v1,t0,t1], where t1 may be InvalidID */
+	inline FIndex4i GetEdge(int EdgeID) const
+	{
+		check(IsEdge(EdgeID));
+		int i = 4 * EdgeID;
+		return FIndex4i(Edges[i], Edges[i + 1], Edges[i + 2], Edges[i + 3]);
+	}
+
+	/** Get the vertex pair for an edge */
+	inline FIndex2i GetEdgeV(int EdgeID) const
+	{
+		check(IsEdge(EdgeID));
+		int i = 4 * EdgeID;
+		return FIndex2i(Edges[i], Edges[i + 1]);
+	}
+
+	/** Get the vertex positions of an edge */
+	inline bool GetEdgeV(int EdgeID, FVector3d& a, FVector3d& b) const
+	{
+		check(IsEdge(EdgeID));
+		int iv0 = 3 * Edges[4 * EdgeID];
+		a.X = Vertices[iv0];
+		a.Y = Vertices[iv0 + 1];
+		a.Z = Vertices[iv0 + 2];
+		int iv1 = 3 * Edges[4 * EdgeID + 1];
+		b.X = Vertices[iv1];
+		b.Y = Vertices[iv1 + 1];
+		b.Z = Vertices[iv1 + 2];
+		return true;
+	}
+
+	/** Get the triangle pair for an edge. The second triangle may be InvalidID */
+	inline FIndex2i GetEdgeT(int EdgeID) const
+	{
+		check(IsEdge(EdgeID));
+		int i = 4 * EdgeID;
+		return FIndex2i(Edges[i + 2], Edges[i + 3]);
+	}
+
+	/** Return edge vertex indices, but oriented based on attached triangle (rather than min-sorted) */
+	FIndex2i GetOrientedBoundaryEdgeV(int EdgeID) const;
+
+	//
+	// Vertex and Triangle attribute arrays
+	//
+public:
+	void EnableVertexNormals(const FVector3f& InitialNormal);
+	void DiscardVertexNormals();
+
+	FVector3f GetVertexNormal(int vID) const
+	{
+		if (HasVertexNormals() == false)
+		{
+			return FVector3f::UnitY();
+		}
+		check(IsVertex(vID));
+		int i = 3 * vID;
+		return FVector3f((*VertexNormals)[i], (*VertexNormals)[i + 1], (*VertexNormals)[i + 2]);
+	}
+
+	void SetVertexNormal(int vID, const FVector3f& vNewNormal)
+	{
+		if (HasVertexNormals())
+		{
+			check(IsVertex(vID));
+			int i = 3 * vID;
+			(*VertexNormals)[i] = vNewNormal.X;
+			(*VertexNormals)[i + 1] = vNewNormal.Y;
+			(*VertexNormals)[i + 2] = vNewNormal.Z;
+			UpdateTimeStamp(false, false);
+		}
+	}
+
+	void EnableVertexColors(const FVector3f& InitialColor);
+	void DiscardVertexColors();
+
+
+	FVector3f GetVertexColor(int vID) const
+	{
+		if (HasVertexColors() == false)
+		{
+			return FVector3f::One();
+		}
+		check(IsVertex(vID));
+		int i = 3 * vID;
+		return FVector3f((*VertexColors)[i], (*VertexColors)[i + 1], (*VertexColors)[i + 2]);
+	}
+
+	void SetVertexColor(int vID, const FVector3f& vNewColor)
+	{
+		if (HasVertexColors())
+		{
+			check(IsVertex(vID));
+			int i = 3 * vID;
+			(*VertexColors)[i] = vNewColor.X;
+			(*VertexColors)[i + 1] = vNewColor.Y;
+			(*VertexColors)[i + 2] = vNewColor.Z;
+			UpdateTimeStamp(false, false);
+		}
+	}
+
+	void EnableVertexUVs(const FVector2f& InitialUV);
+	void DiscardVertexUVs();
+
+	FVector2f GetVertexUV(int vID) const
+	{
+		if (HasVertexUVs() == false)
+		{
+			return FVector2f::Zero();
+		}
+		check(IsVertex(vID));
+		int i = 2 * vID;
+		return FVector2f((*VertexUVs)[i], (*VertexUVs)[i + 1]);
+	}
+
+	void SetVertexUV(int vID, const FVector2f& vNewUV)
+	{
+		if (HasVertexUVs())
+		{
+			check(IsVertex(vID));
+			int i = 2 * vID;
+			(*VertexUVs)[i] = vNewUV.X;
+			(*VertexUVs)[i + 1] = vNewUV.Y;
+			UpdateTimeStamp(false, false);
+		}
+	}
+
+	void EnableTriangleGroups(int InitialGroupID = 0);
+	void DiscardTriangleGroups();
+
+	int AllocateTriangleGroup()	{ return ++GroupIDCounter; }
+
+	int GetTriangleGroup(int tID) const
+	{
+		return (HasTriangleGroups() == false) ? -1
+			: (TriangleRefCounts.IsValid(tID) ? (*TriangleGroups)[tID] : 0);
+	}
+
+	void SetTriangleGroup(int tid, int group_id)
+	{
+		if (HasTriangleGroups())
+		{
+			check(IsTriangle(tid));
+			(*TriangleGroups)[tid] = group_id;
+			GroupIDCounter = FMath::Max(GroupIDCounter, group_id + 1);
+			UpdateTimeStamp(false, false);
+		}
+	}
+
+	FDynamicMeshAttributeSet* Attributes()
+	{
+		return AttributeSet;
+	}
+	const FDynamicMeshAttributeSet* Attributes() const
+	{
+		return AttributeSet;
+	}
+
+	void EnableAttributes();
+
+	void DiscardAttributes();
+
+
+	//
+	// topological queries
+	//
+public:
+	/** Returns true if edge is on the mesh boundary, ie only connected to one triangle */
+	inline bool IsBoundaryEdge(int EdgeID) const
+	{
+		check(IsEdge(EdgeID));
+		return Edges[4 * EdgeID + 3] == InvalidID;
+	}
+
+	/** Returns true if the vertex is part of any boundary edges */
+	bool IsBoundaryVertex(int VertexID) const;
+
+	/** Returns true if any edge of triangle is a boundary edge */
+	bool IsBoundaryTriangle(int TriangleID) const;
+
+	/** Find id of edge connecting A and B */
+	int FindEdge(int VertexA, int VertexB) const;
+
+	/** Find edgeid for edge [a,b] from triangle that contains the edge. Faster than FindEdge() because it is constant-time. */
+	int FindEdgeFromTri(int VertexA, int VertexB, int TriangleID) const;
+
+	/** Find triangle made up of any permutation of vertices [a,b,c] */
+	int FindTriangle(int A, int B, int C) const;
+
+	/**
+	 * If edge has vertices [a,b], and is connected two triangles [a,b,c] and [a,b,d],
+	 * this returns [c,d], or [c,InvalidID] for a boundary edge
+	 */
+	FIndex2i GetEdgeOpposingV(int EdgeID) const;
+
+	/**
+	 * Given an edge and vertex on that edge, returns other vertex of edge, the two opposing verts, and the two connected triangles (OppVert2Out and Tri2Out are be InvalidID for boundary edge)
+	 */
+	void GetVtxNbrhood(int EdgeID, int VertexID, int& OtherVertOut, int& OppVert1Out, int& OppVert2Out, int& Tri1Out, int& Tri2Out) const;
+
+	/**
+	 * Returns count of boundary edges at vertex, and the first two boundary
+	 * edges if found. If return is > 2, call GetAllVtxBoundaryEdges
+	 */
+	int GetVtxBoundaryEdges(int VertexID, int& Edge0Out, int& Edge1Out) const;
+
+	/**
+	 * Find edge ids of boundary edges connected to vertex.
+	 * @param vID Vertex ID
+	 * @param EdgeListOut boundary edge IDs are appended to this list
+	 * @return count of number of elements of e that were filled
+	 */
+	int GetAllVtxBoundaryEdges(int VertexID, TArray<int>& EdgeListOut) const;
+
+	/**
+	 * return # of triangles attached to vID, or -1 if invalid vertex
+	 * if bBruteForce = true, explicitly checks, which creates a list and is expensive
+	 * default is false, uses orientation, no memory allocation
+	 */
+	int GetVtxTriangleCount(int VertexID, bool bBruteForce = false) const;
+
+	/**
+	 * Get triangle one-ring at vertex.
+	 * bUseOrientation is more efficient but returns incorrect result if vertex is a bowtie
+	 */
+	EMeshResult GetVtxTriangles(int VertexID, TArray<int>& TrianglesOut, bool bUseOrientation) const;
+
+	/** Returns true if the two triangles connected to edge have different group IDs */
+	bool IsGroupBoundaryEdge(int EdgeID) const;
+
+	/** Returns true if vertex has more than one tri group in its tri nbrhood */
+	bool IsGroupBoundaryVertex(int VertexID) const;
+
+	/** Returns true if more than two group boundary edges meet at vertex (ie 3+ groups meet at this vertex) */
+	bool IsGroupJunctionVertex(int VertexID) const;
+
+	/** Returns up to 4 group IDs at vertex. Returns false if > 4 encountered */
+	bool GetVertexGroups(int VertexID, FIndex4i& GroupsOut) const;
+
+	/** Returns all group IDs at vertex */
+	bool GetAllVertexGroups(int VertexID, TArray<int>& GroupsOut) const;
+
+	/** returns true if vID is a "bowtie" vertex, ie multiple disjoint triangle sets in one-ring */
+	bool IsBowtieVertex(int VertexID) const;
+
+	/** returns true if vertices, edges, and triangles are all dense (Count == MaxID) **/
+	bool IsCompact() const
+	{
+		return VertexRefCounts.IsDense() && EdgeRefCounts.IsDense() && TriangleRefCounts.IsDense();
+	}
+
+	/** @return true if vertex count == max vertex id */
+	bool IsCompactV() const
+	{
+		return VertexRefCounts.IsDense();
+	}
+
+	/** @return true if triangle count == max triangle id */
+	bool IsCompactT() const
+	{
+		return TriangleRefCounts.IsDense();
+	}
+
+	/** returns measure of compactness in range [0,1], where 1 is fully compacted */
+	double CompactMetric() const
+	{
+		return ((double)VertexCount() / (double)MaxVertexID() + (double)TriangleCount() / (double)MaxTriangleID()) * 0.5;
+	}
+
+	/** @return true if mesh has no boundary edges */
+	bool IsClosed() const;
+
+	/** @return value of IsClosed() and cache result. cache is invalidated and recomputed if topology has changed since last call */
+	bool GetCachedIsClosed();
+
+
+	/** Timestamp is incremented any time any change is made to the mesh */
+	inline int GetTimestamp() const
+	{
+		return Timestamp;
+	}
+
+	/** ShapeTimestamp is incremented any time any vertex position is changed or the mesh topology is modified */
+	inline int GetShapeTimestamp() const
+	{
+		return ShapeTimestamp;
+	}
+
+	/** TopologyTimestamp is incremented any time any vertex position is changed or the mesh topology is modified */
+	inline int GetTopologyTimestamp() const
+	{
+		return TopologyTimestamp;
+	}
+
+	//
+	// Geometric queries
+	//
+public:
+	/** Returns bounding box of all mesh vertices (including unreferenced vertices) */
+	FAxisAlignedBox3d GetBounds() const;
+
+	/** Returns GetBounds() and saves result, cache is invalidated and recomputed if topology has changed since last call */
+	FAxisAlignedBox3d GetCachedBounds();
+
+	/**
+	 * Compute a normal/tangent frame at vertex that is "stable" as long as
+	 * the mesh topology doesn't change, meaning that one axis of the frame
+	 * will be computed from projection of outgoing edge. Requires that vertex normals are available.
+	 * By default, frame.Z is normal, and .X points along mesh edge.
+	 * If bFrameNormalY, then frame.Y is normal (X still points along mesh edge)
+	 */
+	FFrame3d GetVertexFrame(int VertexID, bool bFrameNormalY = false) const;
+
+	/** Calculate face normal of triangle */
+	FVector3d GetTriNormal(int TriangleID) const;
+
+	/** Calculate area triangle */
+	double GetTriArea(int TriangleID) const;
+
+	/**
+	 * Compute triangle normal, area, and centroid all at once. Re-uses vertex
+	 * lookups and computes normal & area simultaneously. *However* does not produce
+	 * the same normal/area as separate calls, because of this.
+	 */
+	void GetTriInfo(int TriangleID, FVector3d& Normal, double& Area, FVector3d& Centroid) const;
+
+	/** Compute centroid of triangle */
+	FVector3d GetTriCentroid(int TriangleID) const;
+
+	/** Interpolate vertex positions of triangle using barycentric coordinates */
+	FVector3d GetTriBaryPoint(int TriangleID, double Bary0, double Bary1, double Bary2) const;
+
+	/** Interpolate vertex normals of triangle using barycentric coordinates */
+	FVector3d GetTriBaryNormal(int TriangleID, double Bary0, double Bary1, double Bary2) const;
+
+	/** Compute interpolated vertex attributes at point of triangle */
+	void GetTriBaryPoint(int TriangleID, double Bary0, double Bary1, double Bary2, FVertexInfo& VertInfo) const;
+
+	/** Construct bounding box of triangle as efficiently as possible */
+	FAxisAlignedBox3d GetTriBounds(int TriangleID) const;
+
+	/** Construct stable frame at triangle centroid, where frame.Z is face normal, and frame.X is aligned with edge nEdge of triangle. */
+	FFrame3d GetTriFrame(int TriangleID, int Edge = 0) const;
+
+	/** Compute solid angle of oriented triangle tID relative to point p - see WindingNumber() */
+	double GetTriSolidAngle(int TriangleID, const FVector3d& p) const;
+
+	/** Compute internal angle at vertex i of triangle (where i is 0,1,2); */
+	double GetTriInternalAngleR(int TriangleID, int i);
+
+	/** Returns average normal of connected face normals */
+	FVector3d GetEdgeNormal(int EdgeID) const;
+
+	/** Get point along edge, t clamped to range [0,1] */
+	FVector3d GetEdgePoint(int EdgeID, double ParameterT) const;
+
+	/**
+	 * Fastest possible one-ring centroid. This is used inside many other algorithms
+	 * so it helps to have it be maximally efficient
+	 */
+	void GetVtxOneRingCentroid(int VertexID, FVector3d& CentroidOut) const;
+
+	/**
+	 * Compute mesh winding number, from Jacobson et al, Robust Inside-Outside Segmentation using Generalized Winding Numbers
+	 * http://igl.ethz.ch/projects/winding-number/
+	 * returns ~0 for points outside a closed, consistently oriented mesh, and a positive or negative integer
+	 * for points inside, with value > 1 depending on how many "times" the point inside the mesh (like in 2D polygon winding)
+	 */
+	double CalculateWindingNumber(const FVector3d& QueryPoint) const;
+
+	//
+	// direct buffer access
+	//
+public:
+	const FDynamicVector<double>& GetVerticesBuffer()
+	{
+		return Vertices;
+	}
+	const FRefCountVector& GetVerticesRefCounts()
+	{
+		return VertexRefCounts;
+	}
+	const FDynamicVector<float>* GetNormalsBuffer()
+	{
+		return VertexNormals;
+	}
+	const FDynamicVector<float>* GetColorsBuffer()
+	{
+		return VertexColors;
+	}
+	const FDynamicVector<float>* GetUVBuffer()
+	{
+		return VertexUVs;
+	}
+
+	const FDynamicVector<int>& GetTrianglesBuffer()
+	{
+		return Triangles;
+	}
+	const FRefCountVector& GetTrianglesRefCounts()
+	{
+		return TriangleRefCounts;
+	}
+	const FDynamicVector<int>* GetTriangleGroupsBuffer()
+	{
+		return TriangleGroups;
+	}
+
+	const FDynamicVector<int>& GetEdgesBuffer()
+	{
+		return Edges;
+	}
+	const FRefCountVector& GetEdgesRefCounts()
+	{
+		return EdgeRefCounts;
+	}
+	const FSmallListSet& GetVertexEdges()
+	{
+		return VertexEdgeLists;
+	}
+
+	//
+	// Mesh Edit operations
+	//
+public:
+	/**
+	 * Compact mesh in-place, by moving vertices around and rewriting indices.
+	 * Should be faster if the amount of compacting is not too significant, and is useful in some places.
+	 *
+	 * @param bComputeCompactInfo if false, then returned CompactMaps is not initialized
+	 * @todo VertexEdgeLists is not compacted. does not affect indices, but does keep memory.
+	 * @todo returned CompactMaps is currently not valid
+	 */
+	void CompactInPlace(FCompactMaps* CompactInfo = nullptr);
+
+	/**
+	 * Reverse the ccw/cw orientation of all triangles in the mesh, and
+	 * optionally flip the vertex normals if they exist
+	 */
+	void ReverseOrientation(bool bFlipNormals = true);
+
+	/**
+	 * Reverse the ccw/cw orientation of a triangle
+	 */
+	EMeshResult ReverseTriOrientation(int TriangleID);
+
+	/**
+	 * Remove vertex vID, and all connected triangles if bRemoveAllTriangles = true
+	 * Returns Failed_VertexStillReferenced if vertex is still referenced by triangles.
+	 * if bPreserveManifold, checks that we will not create a bowtie vertex first
+	 */
+	EMeshResult RemoveVertex(int VertexID, bool bRemoveAllTriangles = true, bool bPreserveManifold = false);
+
+	/**
+	* Remove a triangle from the mesh. Also removes any unreferenced edges after tri is removed.
+	* If bRemoveIsolatedVertices is false, then if you remove all tris from a vert, that vert is also removed.
+	* If bPreserveManifold, we check that you will not create a bowtie vertex (and return false).
+	* If this check is not done, you have to make sure you don't create a bowtie, because other
+	* code assumes we don't have bowties, and will not handle it properly
+	*/
+	EMeshResult RemoveTriangle(int TriangleID, bool bRemoveIsolatedVertices = true, bool bPreserveManifold = false);
+
+	/**
+	 * Rewrite the triangle to reference the new tuple of vertices.
+	 *
+	 * @todo this function currently does not guarantee that the returned mesh is well-formed. Only call if you know it's OK.
+	 */
+	virtual EMeshResult SetTriangle(int TriangleID, const FIndex3i& NewVertices, bool bRemoveIsolatedVertices = true);
+
+
+
+	/** Information about the mesh elements created by a call to SplitEdge() */
+	struct FEdgeSplitInfo
+	{
+		int OriginalEdge;					// the edge that was split
+		FIndex2i OriginalVertices;			// original edge vertices [a,b]
+		FIndex2i OtherVertices;				// original opposing vertices [c,d] - d is InvalidID for boundary edges
+		FIndex2i OriginalTriangles;			// original edge triangles [t0,t1]
+		bool bIsBoundary;					// was the split edge a boundary edge?  (redundant)
+
+		int NewVertex;						// new vertex f that was created
+		FIndex2i NewTriangles;				// new triangles [t2,t3], oriented as explained below
+		FIndex3i NewEdges;					// new edges are [f,b], [f,c] and [f,d] if this is not a boundary edge
+		
+		double SplitT;						// parameter value for NewVertex along original edge
+	};
+
+	/**
+	 * Split an edge of the mesh by inserting a vertex. This creates a new triangle on either side of the edge (ie a 2-4 split).
+	 * If the original edge had vertices [a,b], with triangles t0=[a,b,c] and t1=[b,a,d],  then the split inserts new vertex f.
+	 * After the split t0=[a,f,c] and t1=[f,a,d], and we have t2=[f,b,c] and t3=[f,d,b]  (it's best to draw it out on paper...)
+	 * 
+	 * @param EdgeAB index of the edge to be split
+	 * @param SplitInfo returned information about new and modified mesh elements
+	 * @param SplitParameterT defines the position along the edge that we split at, must be between 0 and 1, and is assumed to be based on the order of vertices returned by GetEdgeV()
+	 * @return Ok on success, or enum value indicates why operation cannot be applied. Mesh remains unmodified on error.
+	 */
+	virtual EMeshResult SplitEdge(int EdgeAB, FEdgeSplitInfo& SplitInfo, double SplitParameterT = 0.5);
+
+	/** 
+	 * Splits the edge between two vertices at the midpoint, if this edge exists 
+	 * @param EdgeVertA index of first vertex
+	 * @param EdgeVertB index of second vertex
+	 * @param SplitInfo returned information about new and modified mesh elements
+	 * @return Ok on success, or enum value indicates why operation cannot be applied. Mesh remains unmodified on error.
+	 */
+	EMeshResult SplitEdge(int EdgeVertA, int EdgeVertB, FEdgeSplitInfo& SplitInfo);
+
+
+
+	/** Information about the mesh elements modified by a call to FlipEdge() */
+	struct FEdgeFlipInfo
+	{
+		int EdgeID;						// the edge that was flipped
+		FIndex2i OriginalVerts;			// original verts of the flipped edge, that are no longer connected
+		FIndex2i OpposingVerts;			// the opposing verts of the flipped edge, that are now connected
+		FIndex2i Triangles;				// the two triangle IDs. Original tris vert [Vert0,Vert1,OtherVert0] and [Vert1,Vert0,OtherVert1].
+										// New triangles are [OtherVert0, OtherVert1, Vert1] and [OtherVert1, OtherVert0, Vert0]
+	};
+
+	/** 
+	 * Flip/Rotate an edge of the mesh. This does not change the number of edges, vertices, or triangles.
+	 * Boundary edges of the mesh cannot be flipped.
+	 * @param EdgeAB index of edge to be flipped
+	 * @param FlipInfo returned information about new and modified mesh elements
+	 * @return Ok on success, or enum value indicates why operation cannot be applied. Mesh remains unmodified on error.
+	 */
+	virtual EMeshResult FlipEdge(int EdgeAB, FEdgeFlipInfo& FlipInfo);
+
+	/** calls FlipEdge() on the edge between two vertices, if it exists
+	 * @param EdgeVertA index of first vertex
+	 * @param EdgeVertB index of second vertex
+	 * @param FlipInfo returned information about new and modified mesh elements
+	 * @return Ok on success, or enum value indicates why operation cannot be applied. Mesh remains unmodified on error.
+	 */
+	virtual EMeshResult FlipEdge(int EdgeVertA, int EdgeVertB, FEdgeFlipInfo& FlipInfo);
+
+
+
+	/** Information about mesh elements modified/removed by CollapseEdge() */
+	struct FEdgeCollapseInfo
+	{
+		int KeptVertex;					// the vertex that was kept (ie collapsed "to")
+		int RemovedVertex;				// the vertex that was removed
+		FIndex2i OpposingVerts;			// the opposing vertices [c,d]. If the edge was a boundary edge, d is InvalidID
+		bool bIsBoundary;				// was the edge a boundary edge
+
+		int CollapsedEdge;				// the edge that was collapsed/removed
+		FIndex2i RemovedTris;			// the triangles that were removed in the collapse (second is InvalidID for boundary edge)
+		FIndex2i RemovedEdges;			// the edges that were removed (second is InvalidID for boundary edge)
+		FIndex2i KeptEdges;				// the edges that were kept (second is InvalidID for boundary edge)
+
+		double CollapseT;				// interpolation parameter along edge for new vertex in range [0,1]
+	};
+
+	/**
+	 * Collapse the edge between the two vertices, if topologically possible.
+	 * @param KeepVertID index of the vertex that should be kept
+	 * @param RemoveVertID index of the vertex that should be removed
+	 * @param EdgeParameterT vKeep is moved to Lerp(KeepPos, RemovePos, collapse_t)
+	 * @param CollapseInfo returned information about new and modified mesh elements
+	 * @return Ok on success, or enum value indicates why operation cannot be applied. Mesh remains unmodified on error.
+	 */
+	virtual EMeshResult CollapseEdge(int KeepVertID, int RemoveVertID, double EdgeParameterT, FEdgeCollapseInfo& CollapseInfo);
+	virtual EMeshResult CollapseEdge(int KeepVertID, int RemoveVertID, FEdgeCollapseInfo& CollapseInfo)
+	{
+		return CollapseEdge(KeepVertID, RemoveVertID, 0, CollapseInfo);
+	}
+
+
+
+	/** Information about mesh elements modified by MergeEdges() */
+	struct FMergeEdgesInfo
+	{
+		int KeptEdge;				// the edge that was kept
+		int RemovedEdge;			// the edge that was removed
+
+		FIndex2i KeptVerts;			// The two vertices that were kept (redundant w/ KeptEdge?)
+		FIndex2i RemovedVerts;		// The removed vertices of RemovedEdge. Either may be InvalidID if it was same as the paired KeptVert
+
+		FIndex2i ExtraRemovedEdges; // extra removed edges, see description below. Either may be or InvalidID
+		FIndex2i ExtraKeptEdges;	// extra kept edges, paired with ExtraRemovedEdges
+	};
+
+	/**
+	 * Given two edges of the mesh, weld both their vertices, so that one edge is removed.
+	 * This could result in one neighbour edge-pair attached to each vertex also collapsing,
+	 * so those cases are detected and handled (eg middle edge-pair in abysmal ascii drawing below)
+	 *
+	 *   ._._._.    (dots are vertices)
+	 *    \._./
+	 *
+	 * @param KeepEdgeID index of the edge that should be kept
+	 * @param DiscardEdgeID index of the edge that should be removed
+	 * @param MergeInfo returned information about new and modified mesh elements
+	 * @return Ok on success, or enum value indicates why operation cannot be applied. Mesh remains unmodified on error.
+	 */
+	virtual EMeshResult MergeEdges(int KeepEdgeID, int DiscardEdgeID, FMergeEdgesInfo& MergeInfo);
+
+
+
+	/** Information about mesh elements modified/created by PokeTriangle() */
+	struct FPokeTriangleInfo
+	{
+		int OriginalTriangle;				// the triangle that was poked
+		FIndex3i TriVertices;				// vertices of the original triangle
+			
+		int NewVertex;						// the new vertex that was inserted
+		FIndex2i NewTriangles;				// the two new triangles that were added (OriginalTriangle is re-used, see code for vertex orders)
+		FIndex3i NewEdges;					// the three new edges connected to NewVertex
+
+		FVector3d BaryCoords;				// barycentric coords that NewVertex was inserted at
+	};
+
+	/**
+	 * Insert a new vertex inside a triangle, ie do a 1 to 3 triangle split
+	 * @param TriangleID index of triangle to poke
+	 * @param BaryCoordinates barycentric coordinates of poke position
+	 * @param PokeInfo returned information about new and modified mesh elements
+	 * @return Ok on success, or enum value indicates why operation cannot be applied. Mesh remains unmodified on error.
+	 */
+	virtual EMeshResult PokeTriangle(int TriangleID, const FVector3d& BaryCoordinates, FPokeTriangleInfo& PokeInfo);
+
+	/** Call PokeTriangle at the centroid of the triangle */
+	virtual EMeshResult PokeTriangle(int TriangleID, FPokeTriangleInfo& PokeInfo)
+	{
+		return PokeTriangle(TriangleID, FVector3d::One() / 3.0, PokeInfo);
+	}
+
+	//
+	// Debug utility functions
+	//
+public:
+	/**
+	 * Returns a debug string that contains mesh statistics and other information
+	 */
+	virtual FString MeshInfoString();
+
+	/**
+	 * Check if another mesh is the same as this mesh. By default only checks
+	 * vertices and triangles, turn on other parameters w/ flags
+	 */
+	virtual bool IsSameMesh(const FDynamicMesh3& OtherMesh, bool bCheckConnectivity, bool bCheckEdgeIDs = false,
+							bool bCheckNormals = false, bool bCheckColors = false, bool bCheckUVs = false,
+							bool bCheckGroups = false,
+							float Epsilon = TMathUtil<float>::Epsilon);
+
+
+	/**
+	 * Checks that the mesh is well-formed, ie all internal data structures are consistent
+	 */
+	virtual bool CheckValidity(bool bAllowNonManifoldVertices = false, EValidityCheckFailMode FailMode = EValidityCheckFailMode::Check) const;
+
+	//
+	// Internal functions
+	//
+protected:
+
+	inline void SetTriangleInternal(int TriangleID, int v0, int v1, int v2)
+	{
+		int i = 3 * TriangleID;
+		Triangles[i] = v0;
+		Triangles[i + 1] = v1;
+		Triangles[i + 2] = v2;
+	}
+	inline void SetTriangleEdgesInternal(int TriangleID, int e0, int e1, int e2)
+	{
+		int i = 3 * TriangleID;
+		TriangleEdges[i] = e0;
+		TriangleEdges[i + 1] = e1;
+		TriangleEdges[i + 2] = e2;
+	}
+
+	int AddEdgeInternal(int vA, int vB, int tA, int tB = InvalidID);
+	int AddTriangleInternal(int a, int b, int c, int e0, int e1, int e2);
+
+	inline int ReplaceTriangleVertex(int TriangleID, int vOld, int vNew)
+	{
+		int i = 3 * TriangleID;
+		if (Triangles[i] == vOld)
+		{
+			Triangles[i] = vNew;
+			return 0;
+		}
+		if (Triangles[i + 1] == vOld)
+		{
+			Triangles[i + 1] = vNew;
+			return 1;
+		}
+		if (Triangles[i + 2] == vOld)
+		{
+			Triangles[i + 2] = vNew;
+			return 2;
+		}
+		return -1;
+	}
+
+	inline void AllocateEdgesList(int VertexID)
+	{
+		if (VertexID < (int)VertexEdgeLists.Size())
+		{
+			VertexEdgeLists.Clear(VertexID);
+		}
+		VertexEdgeLists.AllocateAt(VertexID);
+	}
+
+	void GetVertexEdgesList(int VertexID, TArray<int>& EdgesOut) const
+	{
+		for (int eid : VertexEdgeLists.Values(VertexID))
+		{
+			EdgesOut.Add(eid);
+		}
+	}
+
+	inline void SetEdgeVerticesInternal(int EdgeID, int a, int b)
+	{
+		int i = 4 * EdgeID;
+		if (a < b)
+		{
+			Edges[i] = a;
+			Edges[i + 1] = b;
+		}
+		else
+		{
+			Edges[i] = b;
+			Edges[i + 1] = a;
+		}
+	}
+
+	inline void SetEdgeTrianglesInternal(int EdgeID, int t0, int t1)
+	{
+		int i = 4 * EdgeID;
+		Edges[i + 2] = t0;
+		Edges[i + 3] = t1;
+	}
+
+	int ReplaceEdgeVertex(int EdgeID, int vOld, int vNew);
+	int ReplaceEdgeTriangle(int EdgeID, int tOld, int tNew);
+	int ReplaceTriangleEdge(int EdgeID, int eOld, int eNew);
+
+	inline bool TriangleHasVertex(int TriangleID, int VertexID) const
+	{
+		int i = 3 * TriangleID;
+		return Triangles[i] == VertexID || Triangles[i + 1] == VertexID || Triangles[i + 2] == VertexID;
+	}
+
+	inline bool TriHasNeighbourTri(int CheckTriID, int NbrTriID) const
+	{
+		int i = 3 * CheckTriID;
+		return EdgeHasTriangle(TriangleEdges[i], NbrTriID) || EdgeHasTriangle(TriangleEdges[i + 1], NbrTriID) || EdgeHasTriangle(TriangleEdges[i + 2], NbrTriID);
+	}
+
+	inline bool TriHasSequentialVertices(int TriangleID, int vA, int vB) const
+	{
+		int i = 3 * TriangleID;
+		int v0 = Triangles[i], v1 = Triangles[i + 1], v2 = Triangles[i + 2];
+		return ((v0 == vA && v1 == vB) || (v1 == vA && v2 == vB) || (v2 == vA && v0 == vB));
+	}
+
+	int FindTriangleEdge(int TriangleID, int vA, int vB) const;
+
+	inline bool EdgeHasVertex(int EdgeID, int VertexID) const
+	{
+		int i = 4 * EdgeID;
+		return (Edges[i] == VertexID) || (Edges[i + 1] == VertexID);
+	}
+	inline bool EdgeHasTriangle(int EdgeID, int TriangleID) const
+	{
+		int i = 4 * EdgeID;
+		return (Edges[i + 2] == TriangleID) || (Edges[i + 3] == TriangleID);
+	}
+
+	inline int GetOtherEdgeVertex(int EdgeID, int VertexID) const
+	{
+		int i = 4 * EdgeID;
+		int ev0 = Edges[i], ev1 = Edges[i + 1];
+		return (ev0 == VertexID) ? ev1 : ((ev1 == VertexID) ? ev0 : InvalidID);
+	}
+	inline int GetOtherEdgeTriangle(int EdgeID, int TriangleID) const
+	{
+		int i = 4 * EdgeID;
+		int et0 = Edges[i + 2], et1 = Edges[i + 3];
+		return (et0 == TriangleID) ? et1 : ((et1 == TriangleID) ? et0 : InvalidID);
+	}
+
+	inline void AddTriangleEdge(int TriangleID, int v0, int v1, int j, int EdgeID)
+	{
+		if (EdgeID != InvalidID)
+		{
+			Edges[4 * EdgeID + 3] = TriangleID;
+			TriangleEdges.InsertAt(EdgeID, 3 * TriangleID + j);
+		}
+		else
+		{
+			TriangleEdges.InsertAt(AddEdgeInternal(v0, v1, TriangleID), 3 * TriangleID + j);
+		}
+	}
+
+	void ReverseTriOrientationInternal(int TriangleID);
+
+	void UpdateTimeStamp(bool bShapeChange, bool bTopologyChange)
+	{
+		Timestamp++;
+		if (bShapeChange)
+		{
+			ShapeTimestamp++;
+		}
+		if (bTopologyChange)
+		{
+			check(bShapeChange);   // we consider topology change to be a shape change!
+			TopologyTimestamp++;
+		}
+	}
+
+
+
+};
+
+
diff --git a/Engine/Plugins/Experimental/GeometryProcessing/Source/DynamicMesh/Public/DynamicMeshAABBTree3.h b/Engine/Plugins/Experimental/GeometryProcessing/Source/DynamicMesh/Public/DynamicMeshAABBTree3.h
new file mode 100644
--- /dev/null
+++ b/Engine/Plugins/Experimental/GeometryProcessing/Source/DynamicMesh/Public/DynamicMeshAABBTree3.h
@@ -0,0 +1,6 @@
+// Copyright 1998-2019 Epic Games, Inc. All Rights Reserved.
+
+#include "Spatial/MeshAABBTree3.h"
+#include "DynamicMesh3.h"
+
+typedef TMeshAABBTree3<FDynamicMesh3> FDynamicMeshAABBTree3;
\ No newline at end of file
diff --git a/Engine/Plugins/Experimental/GeometryProcessing/Source/DynamicMesh/Public/DynamicMeshAttributeSet.h b/Engine/Plugins/Experimental/GeometryProcessing/Source/DynamicMesh/Public/DynamicMeshAttributeSet.h
new file mode 100644
--- /dev/null
+++ b/Engine/Plugins/Experimental/GeometryProcessing/Source/DynamicMesh/Public/DynamicMeshAttributeSet.h
@@ -0,0 +1,171 @@
+// Copyright 1998-2019 Epic Games, Inc. All Rights Reserved.
+
+
+#pragma once
+
+#include "DynamicMeshOverlay.h"
+
+/** Standard UV overlay type - 2-element float */
+typedef TDynamicMeshVectorOverlay<float, 2, FVector2f> FDynamicMeshUVOverlay;
+/** Standard Normal overlay type - 3-element float */
+typedef TDynamicMeshVectorOverlay<float, 3, FVector3f> FDynamicMeshNormalOverlay;
+
+
+/**
+ * FDynamicMeshAttributeSet manages a set of extended attributes for a FDynamicMesh3.
+ * This includes UV and Normal overlays, etc.
+ * 
+ * Currently the default is to always have one UV layer and one Normal layer
+ * 
+ * @todo current internal structure is a work-in-progress
+ */
+class DYNAMICMESH_API FDynamicMeshAttributeSet
+{
+public:
+
+	FDynamicMeshAttributeSet(FDynamicMesh3* Mesh) 
+		: ParentMesh(Mesh), UV0(Mesh), Normals0(Mesh)
+	{
+		UVLayers.Add(&UV0);
+		NormalLayers.Add(&Normals0);
+	}
+
+	virtual ~FDynamicMeshAttributeSet()
+	{
+	}
+
+	void Copy(const FDynamicMeshAttributeSet& Copy)
+	{
+		UV0.Copy(Copy.UV0);
+		Normals0.Copy(Copy.Normals0);
+		// parent mesh is *not* copied!
+	}
+
+	/** @return the parent mesh for this overlay */
+	const FDynamicMesh3* GetParentMesh() const { return ParentMesh; }
+	/** @return the parent mesh for this overlay */
+	FDynamicMesh3* GetParentMesh() { return ParentMesh; }
+
+
+	/** @return number of UV layers */
+	virtual int NumUVLayers() const 
+	{
+		return 1;
+	}
+
+	/** @return number of Normals layers */
+	virtual int NumNormalLayers() const
+	{
+		return 1;
+	}
+
+	/** @return true if the given edge is a seam edge in any overlay */
+	virtual bool IsSeamEdge(int EdgeID) const;
+
+	
+	//
+	// UV Layers 
+	//
+
+
+	/** @return the UV layer at the given Index */
+	FDynamicMeshUVOverlay* GetUVLayer(int Index)
+	{
+		return (Index == 0) ? &UV0 : nullptr;
+	}
+
+	/** @return the UV layer at the given Index */
+	const FDynamicMeshUVOverlay* GetUVLayer(int Index) const
+	{
+		return (Index == 0) ? &UV0 : nullptr;
+	}
+
+	/** @return list of all UV layers */
+	const TArray<FDynamicMeshUVOverlay*>& GetAllUVLayers() const
+	{
+		return UVLayers;
+	}
+
+	/** @return the primary UV layer (layer 0) */
+	FDynamicMeshUVOverlay* PrimaryUV() 
+	{
+		return &UV0;
+	}
+	/** @return the primary UV layer (layer 0) */
+	const FDynamicMeshUVOverlay* PrimaryUV() const
+	{
+		return &UV0;
+	}
+
+
+	//
+	// Normal Layers 
+	//
+
+	/** @return the Normal layer at the given Index */
+	FDynamicMeshNormalOverlay* GetNormalLayer(int Index)
+	{
+		return (Index == 0) ? &Normals0 : nullptr;
+	}
+
+	/** @return the Normal layer at the given Index */
+	const FDynamicMeshNormalOverlay* GetNormalLayer(int Index) const
+	{
+		return (Index == 0) ? &Normals0 : nullptr;
+	}
+
+	/** @return list of all Normal layers */
+	const TArray<FDynamicMeshNormalOverlay*>& GetAllNormalLayers() const
+	{
+		return NormalLayers;
+	}
+
+	/** @return the primary Normal layer (layer 0) */
+	FDynamicMeshNormalOverlay* PrimaryNormals()
+	{
+		return &Normals0;
+	}
+	/** @return the primary Normal layer (layer 0) */
+	const FDynamicMeshNormalOverlay* PrimaryNormals() const
+	{
+		return &Normals0;
+	}
+
+
+protected:
+	/** Parent mesh of this attribute set */
+	FDynamicMesh3* ParentMesh;
+
+	/** Default UV layer */
+	FDynamicMeshUVOverlay UV0;
+	/** Default Normals layer */
+	FDynamicMeshNormalOverlay Normals0;
+
+	TArray<FDynamicMeshUVOverlay*> UVLayers;
+	TArray<FDynamicMeshNormalOverlay*> NormalLayers;
+
+
+protected:
+	friend class FDynamicMesh3;
+
+	/**
+	 * Initialize the existing attribute layers with the given vertex and triangle sizes
+	 */
+	void Initialize(int MaxVertexID, int MaxTriangleID)
+	{
+		UV0.InitializeTriangles(MaxTriangleID);
+		Normals0.InitializeTriangles(MaxTriangleID);
+	}
+
+	// These functions are called by the FDynamicMesh3 to update the various
+	// attributes when the parent mesh topology has been modified.
+	virtual void OnNewTriangle(int TriangleID, bool bInserted);
+	virtual void OnRemoveTriangle(int TriangleID, bool bRemoveIsolatedVertices);
+	virtual void OnReverseTriOrientation(int TriangleID);
+	virtual void OnSplitEdge(const FDynamicMesh3::FEdgeSplitInfo & splitInfo);
+	virtual void OnFlipEdge(const FDynamicMesh3::FEdgeFlipInfo & flipInfo);
+	virtual void OnCollapseEdge(const FDynamicMesh3::FEdgeCollapseInfo & collapseInfo);
+	virtual void OnPokeTriangle(const FDynamicMesh3::FPokeTriangleInfo & pokeInfo);
+	virtual void OnMergeEdges(const FDynamicMesh3::FMergeEdgesInfo & mergeInfo);
+};
+
