move GeometryProcessing plugin from NotForLicensees to Experimental

diff --git a/Engine/Plugins/Experimental/GeometryProcessing/Source/DynamicMesh/Public/DynamicMeshEditor.h b/Engine/Plugins/Experimental/GeometryProcessing/Source/DynamicMesh/Public/DynamicMeshEditor.h
new file mode 100644
--- /dev/null
+++ b/Engine/Plugins/Experimental/GeometryProcessing/Source/DynamicMesh/Public/DynamicMeshEditor.h
@@ -0,0 +1,264 @@
+// Copyright 1998-2019 Epic Games, Inc. All Rights Reserved.
+
+// Port of geometry3Sharp MeshEditor
+
+#pragma once
+
+#include "DynamicMesh3.h"
+
+
+
+
+/**
+ * FMeshIndexMappings stores a set of integer IndexMaps for a mesh
+ * This is a convenient object to have, to avoid passing around large numbers of separate maps.
+ * The individual maps are not necessarily all filled by every operation.
+ */
+struct DYNAMICMESH_API FMeshIndexMappings
+{
+protected:
+	FIndexMapi VertexMap;
+	FIndexMapi TriangleMap;
+	FIndexMapi GroupMap;
+
+	TArray<FIndexMapi> UVMaps;
+	TArray<FIndexMapi> NormalMaps;
+
+public:
+	/** Size internal arrays-of-maps to be suitable for this Mesh */
+	void Initialize(FDynamicMesh3* Mesh);
+
+	/** @return the value used to indicate "invalid" in the mapping */
+	int InvalidID() { return VertexMap.GetInvalidID(); }
+
+	FIndexMapi& GetVertexMap() { return VertexMap; }
+	inline void SetVertex(int FromID, int ToID) { VertexMap.Add(FromID, ToID); }
+	inline int GetNewVertex(int FromID) const { return VertexMap.GetTo(FromID); }
+	inline bool ContainsVertex(int FromID) const { return VertexMap.ContainsFrom(FromID); }
+
+	FIndexMapi& GetTriangleMap() { return TriangleMap; }
+	void SetTriangle(int FromID, int ToID) { TriangleMap.Add(FromID, ToID); }
+	int GetNewTriangle(int FromID) const { return TriangleMap.GetTo(FromID); }
+	inline bool ContainsTriangle(int FromID) const { return TriangleMap.ContainsFrom(FromID); }
+
+	FIndexMapi& GetGroupMap() { return GroupMap; }
+	void SetGroup(int FromID, int ToID) { GroupMap.Add(FromID, ToID); }
+	int GetNewGroup(int FromID) const { return GroupMap.GetTo(FromID); }
+	inline bool ContainsGroup(int FromID) const { return GroupMap.ContainsFrom(FromID); }
+
+	FIndexMapi& GetUVMap(int UVLayer) { return UVMaps[UVLayer]; }
+	void SetUV(int UVLayer, int FromID, int ToID) { UVMaps[UVLayer].Add(FromID, ToID); }
+	int GetNewUV(int UVLayer, int FromID) const { return UVMaps[UVLayer].GetTo(FromID); }
+	inline bool ContainsUV(int UVLayer, int FromID) const { return UVMaps[UVLayer].ContainsFrom(FromID); }
+
+	FIndexMapi& GetNormalMap(int NormalLayer) { return NormalMaps[NormalLayer]; }
+	void SetNormal(int NormalLayer, int FromID, int ToID) { NormalMaps[NormalLayer].Add(FromID, ToID); }
+	int GetNewNormal(int NormalLayer, int FromID) const { return NormalMaps[NormalLayer].GetTo(FromID); }
+	inline bool ContainsNormal(int NormalLayer, int FromID) const { return NormalMaps[NormalLayer].ContainsFrom(FromID); }
+
+};
+
+
+/**
+ * FDynamicMeshEditResult is used to return information about new mesh elements
+ * created by mesh changes, primarily in FDynamicMeshEditor
+ */
+struct DYNAMICMESH_API FDynamicMeshEditResult
+{
+	/** New vertices created by an edit */
+	TArray<int> NewVertices;
+
+	/** New triangles created by an edit. Note that this list may be empty if the operation created quads or polygons */
+	TArray<int> NewTriangles;
+	/** New quads created by an edit, where each quad is a pair of triangle IDs */
+	TArray<FIndex2i> NewQuads;
+	/** New polygons created by an edit, where each polygon is a list of triangle IDs */
+	TArray<TArray<int>> NewPolygons;
+
+	/** New triangle groups created by an edit */
+	TArray<int> NewGroups;
+
+	/** clear this data structure */
+	void Reset()
+	{
+		NewVertices.Reset();
+		NewTriangles.Reset();
+		NewQuads.Reset();
+		NewPolygons.Reset();
+		NewGroups.Reset();
+	}
+
+	/** Flatten the triangle/quad/polygon lists into a single list of all triangles */
+	void GetAllTriangles(TArray<int>& TrianglesOut) const;
+};
+
+
+
+
+/**
+ * FDynamicMeshEditor implements low-level mesh editing operations. These operations
+ * can be used to construct higher-level operations. For example an Extrude operation
+ * could be implemented via DuplicateTriangles() and StitchLoopMinimal().
+ */
+class DYNAMICMESH_API FDynamicMeshEditor
+{
+public:
+	/** The mesh we will be editing */
+	FDynamicMesh3* Mesh;
+
+	FDynamicMeshEditor(FDynamicMesh3* MeshIn)
+	{
+		Mesh = MeshIn;
+	}
+
+
+	//////////////////////////////////////////////////////////////////////////
+	// Create and Remove Triangle Functions
+	//////////////////////////////////////////////////////////////////////////
+
+	/**
+	 * Stitch together two loops of vertices with a quad-strip of triangles.
+	 * Loops must be oriented (ordered) correctly for your use case.
+	 * @param Loop1 first loop of sequential vertices
+	 * @param Loop2 second loop of sequential vertices
+	 * @param ResultOut lists of newly created triangles/vertices/etc
+	 * @return true if operation suceeded. If a failure occurs, any added triangles are removed via RemoveTriangles
+	 */
+	bool StitchVertexLoopsMinimal(const TArray<int>& VertexLoop1, const TArray<int>& VertexLoop2, FDynamicMeshEditResult& ResultOut);
+
+
+	/**
+	 * Duplicate triangles of a mesh. This duplicates the current groups and also any attributes existing on the triangles.
+	 * @param Triangles the triangles to duplicate
+	 * @param IndexMaps returned mappings from old to new triangles/vertices/etc (you may initialize to optimize memory usage, etc)
+	 * @param ResultOut lists of newly created triangles/vertices/etc
+	 */
+	void DuplicateTriangles(const TArray<int>& Triangles, FMeshIndexMappings& IndexMaps, FDynamicMeshEditResult& ResultOut);
+
+
+	/**
+	 * Remove a list of triangles from the mesh, and optionally any vertices that are now orphaned
+	 * @param Triangles the triangles to remove
+	 * @param bRemoveIsolatedVerts if true, remove vertices that end up with no triangles
+	 * @return true if all removes succeeded
+	 */
+	bool RemoveTriangles(const TArray<int>& Triangles, bool bRemoveIsolatedVerts);
+
+
+	//////////////////////////////////////////////////////////////////////////
+	// Normal utility functions
+	//////////////////////////////////////////////////////////////////////////
+
+	/**
+	 * Reverse the orientation of the given triangles, and optionally flip relevant normals
+	 * @param Triangles the triangles to modify
+	 * @param bInvertNormals if ture we call InvertTriangleNormals()
+	 */
+	void ReverseTriangleOrientations(const TArray<int>& Triangles, bool bInvertNormals);
+
+	/**
+	 * Flip the normals of the given triangles. This includes their vertex normals, if they
+	 * exist, as well as any per-triangle attribute normals. @todo currently creates full-mesh bit arrays, could be more efficient on subsets
+	 * @param Triangles the triangles to modify
+	 */
+	void InvertTriangleNormals(const TArray<int>& Triangles);
+
+
+	/**
+	 * Calculate and set the per-triangle normals of the two input quads.
+	 * Average of the two face normals is used unless the quad is planar
+	 * @param QuadTris pair of triangle IDs. If second ID is invalid, it is ignored
+	 * @param bIsPlanar if the quad is known to be planar, operation is more efficient
+	 * @return the normal vector that was set
+	 */
+	FVector3f ComputeAndSetQuadNormal(const FIndex2i& QuadTris, bool bIsPlanar = false);
+
+
+	/**
+	 * Create and set new shared per-triangle normals for a pair of triangles that share one edge (ie a quad)
+	 * @param QuadTris pair of triangle IDs. If second ID is invalid, it is ignored
+	 * @param Normal normal vector to set
+	 */
+	void SetQuadNormals(const FIndex2i& QuadTris, const FVector3f& Normal);
+
+	/**
+	 * Create and set new shared per-triangle normals for a list of triangles
+	 * @param Triangles list of triangle IDs
+	 * @param Normal normal vector to set
+	 */
+	void SetTriangleNormals(const TArray<int>& Triangles, const FVector3f& Normal);
+
+
+	//////////////////////////////////////////////////////////////////////////
+	// UV utility functions
+	//////////////////////////////////////////////////////////////////////////
+
+
+	/**
+	 * Project the two triangles of the quad onto a plane defined by the normal and use that to create/set new shared per-triangle UVs.
+	 * UVs are translated so that their bbox min-corner is at origin, and scaled by given scale factor
+	 * @param QuadTris pair of triangle IDs. If second ID is invalid, it is ignored
+	 * @param ProjectFrame vertices are projected into XY axes of this frame
+	 * @param UVScaleFactor UVs are scaled
+	 * @param UVLayerIndex which UV layer to operate on (must exist)
+	 */
+	void SetQuadUVsFromProjection(const FIndex2i& QuadTris, const FFrame3f& ProjectFrame, float UVScaleFactor = 1.0f, int UVLayerIndex = 0);
+
+
+
+	//////////////////////////////////////////////////////////////////////////
+	// mesh element copying / duplication
+	//////////////////////////////////////////////////////////////////////////
+
+
+	/**
+	 * Find "new" vertex for input vertex under Index mapping, or create new if missing
+	 * @param VertexID the source vertex we want a copy of
+	 * @param IndexMaps source/destination mapping of already-duplicated vertices
+	 * @param ResultOut newly-created vertices are stored here
+	 * @return index of duplicate vertex
+	 */	
+	int FindOrCreateDuplicateVertex(int VertexID, FMeshIndexMappings& IndexMaps, FDynamicMeshEditResult& ResultOut);
+
+	/**
+	 * Find "new" group for input group under Index mapping, or create new if missing
+	 * @param TriangleID the source triangle whose group we want a copy of
+	 * @param IndexMaps source/destination mapping of already-duplicated groups
+	 * @param ResultOut newly-created groups are stored here
+	 * @return index of duplicate group
+	 */
+	int FindOrCreateDuplicateGroup(int TriangleID, FMeshIndexMappings& IndexMaps, FDynamicMeshEditResult& ResultOut);
+
+	/**
+	 * Find "new" UV for input UV element under Index mapping, or create new if missing
+	 * @param ElementID the source UV we want a duplicate of
+	 * @param UVLayerIndex which UV layer to consider
+	 * @param IndexMaps source/destination mapping of already-duplicated UVs
+	 * @return index of duplicate UV in given UV layer
+	 */
+	int FindOrCreateDuplicateUV(int ElementID, int UVLayerIndex, FMeshIndexMappings& IndexMaps);
+
+	/**
+	 * Find "new" normal for input normal element under Index mapping, or create new if missing
+	 * @param ElementID the source normal we want a duplicate of
+	 * @param NormalLayerIndex which normal layer to consider
+	 * @param IndexMaps source/destination mapping of already-duplicated normals
+	 * @return index of duplicate normal in given normal layer
+	 */
+	int FindOrCreateDuplicateNormal(int ElementID, int NormalLayerIndex, FMeshIndexMappings& IndexMaps);
+
+
+	/**
+	 * Copy all attribute-layer values from one triangle to another, using the IndexMaps to track and re-use shared attribute values.
+	 * @param FromTriangleID source triangle
+	 * @param ToTriangleID destination triangle
+	 * @param IndexMaps mappings passed to FindOrCreateDuplicateX functions to track already-created attributes
+	 * @param ResultOut information about new attributes is stored here (@todo populate this, at time of writing there are no attribute fields)
+	 */
+	void CopyAttributes(int FromTriangleID, int ToTriangleID, FMeshIndexMappings& IndexMaps, FDynamicMeshEditResult& ResultOut);
+
+
+
+};
+
+
diff --git a/Engine/Plugins/Experimental/GeometryProcessing/Source/DynamicMesh/Public/MeshAdapterUtil.h b/Engine/Plugins/Experimental/GeometryProcessing/Source/DynamicMesh/Public/MeshAdapterUtil.h
new file mode 100644
--- /dev/null
+++ b/Engine/Plugins/Experimental/GeometryProcessing/Source/DynamicMesh/Public/MeshAdapterUtil.h
@@ -0,0 +1,35 @@
+// Copyright 1998-2019 Epic Games, Inc. All Rights Reserved.
+
+#pragma once
+
+#include "CoreTypes.h"
+#include "DynamicMesh3.h"
+#include "PointSetAdapter.h"
+
+/**
+ * Utility functions for constructing PointSetAdapter instances from various 
+ * point sets on a mesh
+ */
+namespace MeshAdapterUtil
+{
+	/**
+	 * @return Mesh vertices as a point set
+	 */
+	FPointSetAdapterd DYNAMICMESH_API MakeVerticesAdapter(const FDynamicMesh3* Mesh);
+
+	/**
+	 * @return Mesh triangle centroids as a point set
+	 */
+	FPointSetAdapterd DYNAMICMESH_API MakeTriCentroidsAdapter(const FDynamicMesh3* Mesh);
+
+	/**
+	 * @return mesh edge midpoints as a point set
+	 */
+	FPointSetAdapterd DYNAMICMESH_API MakeEdgeMidpointsAdapter(const FDynamicMesh3* Mesh);
+
+	/**
+	 * @return Mesh boundary edge midpoints as a point set
+	 */
+	FPointSetAdapterd DYNAMICMESH_API MakeBoundaryEdgeMidpointsAdapter(const FDynamicMesh3* Mesh);
+}
+
diff --git a/Engine/Plugins/Experimental/GeometryProcessing/Source/DynamicMesh/Public/MeshNormals.h b/Engine/Plugins/Experimental/GeometryProcessing/Source/DynamicMesh/Public/MeshNormals.h
new file mode 100644
--- /dev/null
+++ b/Engine/Plugins/Experimental/GeometryProcessing/Source/DynamicMesh/Public/MeshNormals.h
@@ -0,0 +1,114 @@
+// Copyright 1998-2019 Epic Games, Inc. All Rights Reserved.
+
+// Port of geometry3Sharp MeshNormals
+
+#pragma once
+
+#include "DynamicMesh3.h"
+#include "DynamicMeshAttributeSet.h"
+
+/**
+ * FMeshNormals is a utility class that can calculate and store various types of
+ * normal vectors for a FDynamicMesh. 
+ */
+class DYNAMICMESH_API FMeshNormals
+{
+protected:
+	/** Target Mesh */
+	const FDynamicMesh3* Mesh;
+	/** Set of computed normals */
+	TArray<FVector3d> Normals;
+
+public:
+	FMeshNormals()
+	{
+		Mesh = nullptr;
+	}
+
+	FMeshNormals(const FDynamicMesh3* Mesh)
+	{
+		SetMesh(Mesh);
+	}
+
+
+	void SetMesh(const FDynamicMesh3* MeshIn)
+	{
+		this->Mesh = MeshIn;
+	}
+
+	const TArray<FVector3d>& GetNormals() const { return Normals; }
+
+	FVector3d& operator[](int i) { return Normals[i]; }
+	const FVector3d& operator[](int i) const { return Normals[i]; }
+
+
+	/**
+	 * Set the size of the Normals array to Count, and optionally clear all values to (0,0,0)
+	 */
+	void SetCount(int Count, bool bClearToZero);
+
+	/**
+	 * Compute standard per-vertex normals by averaging one-ring face normals
+	 */
+	void ComputeVertexNormals()
+	{
+		Compute_FaceAvg_AreaWeighted();
+	}
+
+	/**
+	 * Compute per-triangle normals
+	 */
+	void ComputeTriangleNormals()
+	{
+		Compute_Triangle();
+	}
+
+	/**
+	 * Recompute the per-element normals of the given overlay by averaging one-ring face normals
+	 * @warning NormalOverlay must be attached to ParentMesh or an exact copy
+	 */
+	void RecomputeOverlayNormals(FDynamicMeshNormalOverlay* NormalOverlay)
+	{
+		Compute_Overlay_FaceAvg_AreaWeighted(NormalOverlay);
+	}
+
+
+
+	/**
+	 * Copy the current set of normals to the vertex normals of SetMesh
+	 * @warning assumes that the computed normals are vertex normals
+	 * @param bInvert if true, normals are flipped
+	 */
+	void CopyToVertexNormals(FDynamicMesh3* SetMesh, bool bInvert) const;
+
+	/**
+	 * Copy the current set of normals to the NormalOverlay attribute layer
+	 * @warning assumes that the computed normals are attribute normals
+	 * @param bInvert if true, normals are flipped
+	 */
+	void CopyToOverlay(FDynamicMeshNormalOverlay* NormalOverlay, bool bInvert) const;
+
+
+
+	/**
+	 * Compute per-vertex normals for the given Mesh
+	 * @param bInvert if true, normals are flipped
+	 */
+	static void QuickComputeVertexNormals(FDynamicMesh3& Mesh, bool bInvert);
+
+
+	/**
+	 * @return the vertex normal at vertex VertIDx of Mesh
+	 */
+	static FVector3d ComputeVertexNormal(const FDynamicMesh3& Mesh, int VertIdx);
+
+
+protected:
+	/** Compute per-vertex normals using area-weighted averaging of one-ring triangles */
+	void Compute_FaceAvg_AreaWeighted();
+	/** Compute per-triangle normals */
+	void Compute_Triangle();
+	/** Recompute the element Normals of the given attribute overlay using area-weighted averaging of one-ring triangles */
+	void Compute_Overlay_FaceAvg_AreaWeighted(const FDynamicMeshNormalOverlay* NormalOverlay);
+
+};
\ No newline at end of file
diff --git a/Engine/Plugins/Experimental/GeometryProcessing/Source/DynamicMesh/Public/MeshRegionBoundaryLoops.h b/Engine/Plugins/Experimental/GeometryProcessing/Source/DynamicMesh/Public/MeshRegionBoundaryLoops.h
new file mode 100644
--- /dev/null
+++ b/Engine/Plugins/Experimental/GeometryProcessing/Source/DynamicMesh/Public/MeshRegionBoundaryLoops.h
@@ -0,0 +1,97 @@
+// Copyright 1998-2019 Epic Games, Inc. All Rights Reserved.
+
+// Port of geometry3cpp MeshRegionBoundaryLoops
+
+#pragma once
+
+#include "DynamicMesh3.h"
+#include "EdgeLoop.h"
+
+
+/**
+ * Extract FEdgeLoops on the boundary of a set of triangles of a mesh.
+ * @todo this was an early port and possibly can be optimized
+ */
+class DYNAMICMESH_API FMeshRegionBoundaryLoops
+{
+public:
+	/** Mesh we are finding loops on */
+	const FDynamicMesh3* Mesh = nullptr;
+	/** Resulting set of loops filled by Compute() */
+	TArray<FEdgeLoop> Loops;
+
+public:
+	FMeshRegionBoundaryLoops() {}
+
+	FMeshRegionBoundaryLoops(const FDynamicMesh3* MeshIn, const TArray<int>& RegionTris, bool bAutoCompute = true);
+
+	/**
+	 * Find set of FEdgeLoops on the border of the input triangle set
+	 * @return false if errors occurred, in this case output set is incomplete
+	 */
+	bool Compute();
+
+
+	/** @return number of loops found by Compute() */
+	int GetLoopCount() const
+	{
+		return Loops.Num();
+	}
+
+	/** @return Loop at the given index */
+	const FEdgeLoop& operator[](int Index) const
+	{
+		return Loops[Index];
+	}
+
+	/** @return index of loop with maximum number of vertices */
+	int GetMaxVerticesLoopIndex() const;
+	
+
+
+protected:
+
+
+	// TODO: C# code used IndexFlagSet here, which is sparse if region size << mesh size
+	// list of included triangles and edges
+	TArray<bool> triangles;
+	TArray<bool> edges;
+	TArray<int> edges_roi;
+
+
+	static bool ContainsElement(const TArray<bool>& flagset, int index) { return flagset[index] == true; }
+
+	bool IsEdgeOnBoundary(int eid) const { return ContainsElement(edges, eid); }
+
+	// returns true for both internal and mesh boundary edges
+	// tid_in and tid_out are triangles 'in' and 'out' of set, respectively
+	bool IsEdgeOnBoundary(int eid, int& tid_in, int& tid_out) const;
+
+	// return same indices as GetEdgeV, but oriented based on attached triangle
+	FIndex2i GetOrientedEdgeVerts(int eID, int tid_in, int tid_out);
+
+	// returns first two boundary edges, and count of total boundary edges
+	int GetVertexBoundaryEdges(int vID, int& e0, int& e1);
+
+	// e needs to be large enough (ie call GetVtxBoundaryEdges, or as large as max one-ring)
+	// returns count, ie number of elements of e that were filled
+	int GetAllVertexBoundaryEdges(int vID, TArray<int>& e);
+
+
+	// [TODO] cache this : a dictionary? we will not need very many, but we will
+	//   need each multiple times!
+	FVector3d GetVertexNormal(int vid);
+
+	// ok, bdry_edges[0...bdry_edges_count] contains the boundary edges coming out of bowtie_v.
+	// We want to pick the best one to continue the loop that came : to bowtie_v on incoming_e.
+	// If the loops are all sane, then we will get the smallest loops by "turning left" at bowtie_v.
+	int FindLeftTurnEdge(int incoming_e, int bowtie_v, TArray<int>& bdry_edges, int bdry_edges_count, TArray<bool>& used_edges);
+
+
+	// This is called when loopV contains one or more "bowtie" vertices.
+	// These vertices *might* be duplicated : loopV (but not necessarily)
+	// If they are, we have to break loopV into subloops that don't contain duplicates.
+	TArray<FEdgeLoop> ExtractSubloops(TArray<int>& loopV, const TArray<int>& loopE, const TArray<int>& bowties);
+
+
+};
diff --git a/Engine/Plugins/Experimental/GeometryProcessing/Source/DynamicMesh/Public/MeshSimplification.h b/Engine/Plugins/Experimental/GeometryProcessing/Source/DynamicMesh/Public/MeshSimplification.h
new file mode 100644
--- /dev/null
+++ b/Engine/Plugins/Experimental/GeometryProcessing/Source/DynamicMesh/Public/MeshSimplification.h
@@ -0,0 +1,307 @@
+// Copyright 1998-2019 Epic Games, Inc. All Rights Reserved.
+
+// Port of geometry3cpp Reducer
+
+#pragma once
+
+#include "MeshRefinerBase.h"
+#include "QuadricError.h"
+#include "Util/IndexPriorityQueue.h"
+
+/**
+ * Implementation of Garland & Heckbert Quadric Error Metric (QEM) Triangle Mesh Simplification
+ */
+class DYNAMICMESH_API FMeshSimplification : public FMeshRefinerBase
+{
+public:
+
+	/**
+	 * If true, we try to find position for collapsed vertices that minimizes quadric error.
+	 * If false we just use midpoints, which is actually significantly slower, because it results
+	 * in may more points that would cause a triangle flip, which are then rejected.
+	 */
+	bool bMinimizeQuadricPositionError = true;
+
+	/** if true, we try to keep boundary vertices on boundary. You probably want this. */
+	bool bPreserveBoundaryShape = true;
+
+
+
+	FMeshSimplification(FDynamicMesh3* m) : FMeshRefinerBase(m)
+	{
+	}
+
+
+	/**
+	 * Simplify mesh until we reach a specific triangle count.
+	 * Note that because triangles are removed in pairs, the resulting count may be TriangleCount-1
+	 * @param TriangleCount the target triangle count
+	 */
+	virtual void SimplifyToTriangleCount(int TriangleCount);
+
+	/**
+	 * Simplify mesh until it has a specific vertex count
+	 * @param VertexCount the target vertex count
+	 */
+	virtual void SimplifyToVertexCount(int VertexCount);
+
+	/**
+	 * Simplify mesh until no edges smaller than min length remain. This is not a great criteria.
+	 * @param MinEdgeLength collapse any edge longer than this
+	 */
+	virtual void SimplifyToEdgeLength(double MinEdgeLength);
+
+	/** 
+	 * Does N rounds of collapsing edges longer than fMinEdgeLength. Does not use Quadrics or priority queue.
+	 * This is a quick way to get triangle count down on huge meshes (eg like marching cubes output). 
+	 * @param MinEdgeLength collapse any edge longer than this
+	 * @param Rounds number of collapse rounds
+	 * @param MeshIsClosedHint if you know the mesh is closed, this pass this true to avoid some precomputes
+	 */
+	virtual void FastCollapsePass(double MinEdgeLength, int Rounds = 1, bool bMeshIsClosedHint = false);
+
+
+
+protected:
+
+	FMeshSimplification()		// for subclasses that extend our behavior
+	{
+	}
+
+	// this just lets us write more concise code
+	bool EnableInlineProjection() const { return ProjectionMode == ETargetProjectionMode::Inline; }
+
+
+	double MinEdgeLength = FMathd::MaxReal;
+	int TargetCount = INT_MAX;
+	enum class ETargetModes
+	{
+		TriangleCount = 0,
+		VertexCount = 1,
+		MinEdgeLength = 2
+	};
+	ETargetModes SimplifyMode = ETargetModes::TriangleCount;
+
+
+
+	/** Top-level function that does the simplification */
+	virtual void DoSimplify();
+
+
+
+	// StartEdges() and GetNextEdge() control the iteration over edges that will be refined.
+	// Default here is to iterate over entire mesh->
+	// Subclasses can override these two functions to restrict the affected edges (eg EdgeLoopRemesher)
+
+	// We are using a modulo-index loop to break symmetry/pathological conditions. 
+	const int ModuloPrime = 31337;     // any prime will do...
+	int MaxEdgeID = 0;
+	virtual int StartEdges() 
+	{
+		MaxEdgeID = Mesh->MaxEdgeID();
+		return 0;
+	}
+
+	virtual int GetNextEdge(int CurEdgeID, bool& bDoneOut) 
+	{
+		int new_eid = (CurEdgeID + ModuloPrime) % MaxEdgeID;
+		bDoneOut = (new_eid == 0);
+		return new_eid;
+	}
+
+
+
+
+	TArray<FQuadricErrord> vertQuadrics;
+	virtual void InitializeVertexQuadrics();
+
+
+	// internal class for priority queue
+	struct QEdge 
+	{
+		int eid;
+		FQuadricErrord q;
+		FVector3d collapse_pt;
+
+		QEdge() { eid = 0; }
+
+		QEdge(int edge_id, const FQuadricErrord& qin, const FVector3d& pt) 
+		{
+			eid = edge_id;
+			q = qin;
+			collapse_pt = pt;
+		}
+	};
+
+	TArray<QEdge> EdgeQuadrics;
+	FIndexPriorityQueue EdgeQueue;
+
+	struct FEdgeError
+	{
+		float error;
+		int eid;
+		bool operator<(const FEdgeError& e2) const
+		{
+			return error < e2.error;
+		}
+	};
+
+	virtual void InitializeQueue();
+
+	// return point that minimizes quadric error for edge [ea,eb]
+	FVector3d OptimalPoint(int eid, const FQuadricErrord& q, int ea, int eb);
+	
+	FVector3d GetProjectedPoint(const FVector3d& pos)
+	{
+		if (EnableInlineProjection() && ProjTarget != nullptr)
+		{
+			return ProjTarget->Project(pos);
+		}
+		return pos;
+	}
+
+
+	// update queue weight for each edge in vertex one-ring
+	virtual void UpdateNeighbours(int vid);
+
+
+	virtual void Reproject() 
+	{
+		ProfileBeginProject();
+		if (ProjTarget != nullptr && ProjectionMode == ETargetProjectionMode::AfterRefinement)
+		{
+			FullProjectionPass();
+			DoDebugChecks();
+		}
+		ProfileEndProject();
+	}
+
+
+
+
+
+	bool bHaveBoundary;
+	TArray<bool> IsBoundaryVtxCache;
+	void Precompute(bool bMeshIsClosed = false);
+
+	inline bool IsBoundaryVertex(int vid) const
+	{
+		return IsBoundaryVtxCache[vid];
+	}
+
+
+
+
+	enum class EProcessResult 
+	{
+		Ok_Collapsed = 0,
+		Ignored_CannotCollapse = 1,
+		Ignored_EdgeIsFullyConstrained = 2,
+		Ignored_EdgeTooLong = 3,
+		Ignored_Constrained = 4,
+		Ignored_CreatesFlip = 5,
+		Failed_OpNotSuccessful = 6,
+		Failed_NotAnEdge = 7
+	};
+	EProcessResult CollapseEdge(int edgeID, FVector3d vNewPos, int& collapseToV);
+
+
+
+	// subclasses can override these to implement custom behavior...
+	virtual void OnEdgeCollapse(int edgeID, int va, int vb, const FDynamicMesh3::FEdgeCollapseInfo& collapseInfo)
+	{
+		// this is for subclasses...
+	}
+
+
+
+
+
+	// Project vertices onto projection target. 
+	virtual void FullProjectionPass();
+
+	virtual void ProjectVertex(int vID, IProjectionTarget* targetIn);
+
+	// used by collapse-edge to get projected position for new vertex
+	virtual FVector3d GetProjectedCollapsePosition(int vid, const FVector3d& vNewPos);
+
+	virtual void ApplyToProjectVertices(const TFunction<void(int)>& apply_f);
+
+
+
+	/*
+	 * testing/debug/profiling stuff
+	 */
+protected:
+
+
+	//
+	// profiling functions, turn on ENABLE_PROFILING to see output in console
+	// 
+	int COUNT_COLLAPSES;
+	int COUNT_ITERATIONS;
+	//Stopwatch AllOpsW, SetupW, ProjectW, CollapseW;
+
+	virtual void ProfileBeginPass() 
+	{
+		if (ENABLE_PROFILING) 
+		{
+			COUNT_COLLAPSES = 0;
+			COUNT_ITERATIONS = 0;
+			//AllOpsW = new Stopwatch();
+			//SetupW = new Stopwatch();
+			//ProjectW = new Stopwatch();
+			//CollapseW = new Stopwatch();
+		}
+	}
+
+	virtual void ProfileEndPass()
+	{
+		if (ENABLE_PROFILING) 
+		{
+			//System.Console.WriteLine(string.Format(
+			//	"ReducePass: T {0} V {1} collapses {2}  iterations {3}", mesh->TriangleCount, mesh->VertexCount, COUNT_COLLAPSES, COUNT_ITERATIONS
+			//));
+			//System.Console.WriteLine(string.Format(
+			//	"           Timing1: setup {0} ops {1} project {2}", Util.ToSecMilli(SetupW.Elapsed), Util.ToSecMilli(AllOpsW.Elapsed), Util.ToSecMilli(ProjectW.Elapsed)
+			//));
+		}
+	}
+
+	virtual void ProfileBeginOps()
+	{
+		//if (ENABLE_PROFILING) AllOpsW.Start();
+	}
+	virtual void ProfileEndOps()
+	{
+		//if (ENABLE_PROFILING) AllOpsW.Stop();
+	}
+	virtual void ProfileBeginSetup()
+	{
+		//if (ENABLE_PROFILING) SetupW.Start();
+	}
+	virtual void ProfileEndSetup()
+	{
+		//if (ENABLE_PROFILING) SetupW.Stop();
+	}
+
+	virtual void ProfileBeginProject()
+	{
+		//if (ENABLE_PROFILING) ProjectW.Start();
+	}
+	virtual void ProfileEndProject() 
+	{
+		//if (ENABLE_PROFILING) ProjectW.Stop();
+	}
+
+	virtual void ProfileBeginCollapse() 
+	{
+		//if (ENABLE_PROFILING) CollapseW.Start();
+	}
+	virtual void ProfileEndCollapse() 
+	{
+		//if (ENABLE_PROFILING) CollapseW.Stop();
+	}
+
+
+};
\ No newline at end of file
diff --git a/Engine/Plugins/Experimental/GeometryProcessing/Source/DynamicMesh/Public/MeshWeights.h b/Engine/Plugins/Experimental/GeometryProcessing/Source/DynamicMesh/Public/MeshWeights.h
new file mode 100644
--- /dev/null
+++ b/Engine/Plugins/Experimental/GeometryProcessing/Source/DynamicMesh/Public/MeshWeights.h
@@ -0,0 +1,44 @@
+// Copyright 1998-2019 Epic Games, Inc. All Rights Reserved.
+
+// Port of geometry3cpp MeshWeights
+
+#pragma once
+
+#include "DynamicMesh3.h"
+
+/**
+ * FMeshWeights implements various techniques for computing local weights of a mesh,
+ * for example one-ring weights like Cotangent or Mean-Value.
+ */
+class DYNAMICMESH_API FMeshWeights
+{
+public:
+
+	/**
+	 * Compute uniform centroid of a vertex one-ring.
+	 * These weights are strictly positive and all equal to 1 / valence
+	 */
+	static FVector3d UniformCentroid(const FDynamicMesh3 & mesh, int VertexIndex);
+
+
+	/**
+	 * Compute mean-value centroid of a vertex one-ring.
+	 * These weights are strictly positive.
+	 */
+	static FVector3d MeanValueCentroid(const FDynamicMesh3 & mesh, int VertexIndex);
+
+	/**
+	 * Compute cotan-weighted centroid of a vertex one-ring.
+	 * These weights are numerically unstable if any of the triangles are degenerate.
+	 * We catch these problems and return input vertex as centroid
+	 */
+	static FVector3d CotanCentroid(const FDynamicMesh3 & mesh, int VertexIndex);
+
+	/**
+	 * Compute the voronoi area associated with a vertex.
+	 */
+	static double VoronoiArea(const FDynamicMesh3& mesh, int VertexIndex);
+
+protected:
+	FMeshWeights() = delete;
+};
\ No newline at end of file
diff --git a/Engine/Plugins/Experimental/GeometryProcessing/Source/DynamicMesh/Public/ProjectionTargets.h b/Engine/Plugins/Experimental/GeometryProcessing/Source/DynamicMesh/Public/ProjectionTargets.h
new file mode 100644
--- /dev/null
+++ b/Engine/Plugins/Experimental/GeometryProcessing/Source/DynamicMesh/Public/ProjectionTargets.h
@@ -0,0 +1,65 @@
+// Copyright 1998-2019 Epic Games, Inc. All Rights Reserved.
+
+#pragma once
+
+// Port of geometry3cpp ProjectionTargets
+
+#include "DynamicMeshAABBTree3.h"
+#include "Distance/DistPoint3Triangle3.h"
+
+
+/**
+ * MeshProjectionTarget provides an IProjectionTarget interface to a FDynamicMesh + FDynamicMeshAABBTree3
+ * Use to project points to mesh surface.
+ */
+class MeshProjectionTarget : public IOrientedProjectionTarget
+{
+public:
+	/** The mesh to project onto */
+	FDynamicMesh3* Mesh;
+	/** An AABBTree for Mesh */
+	FDynamicMeshAABBTree3* Spatial;
+
+	~MeshProjectionTarget() {}
+
+	MeshProjectionTarget(FDynamicMesh3* MeshIn, FDynamicMeshAABBTree3* SpatialIn)
+	{
+		Mesh = MeshIn;
+		Spatial = SpatialIn;
+	}
+
+
+	/**
+	 * @return Projection of Point onto this target
+	 */
+	virtual FVector3d Project(const FVector3d& Point, int Identifier = -1) override
+	{
+		double fDistSqr;
+		int tNearestID = Spatial->FindNearestTriangle(Point, fDistSqr);
+		FTriangle3d Triangle;
+		Mesh->GetTriVertices(tNearestID, Triangle.V[0], Triangle.V[1], Triangle.V[2]);
+
+		FDistPoint3Triangle3d DistanceQuery(Point, Triangle);
+		DistanceQuery.GetSquared();
+		return DistanceQuery.ClosestTrianglePoint;
+	}
+
+	/**
+	 * @return Projection of Point onto this target, and set ProjectNormalOut to the triangle normal at the returned point (*not* interpolated vertex normal)
+	 */
+	virtual FVector3d Project(const FVector3d& Point, FVector3d& ProjectNormalOut, int Identifier = -1) override
+	{
+		double fDistSqr;
+		int tNearestID = Spatial->FindNearestTriangle(Point, fDistSqr);
+		FTriangle3d Triangle;
+		Mesh->GetTriVertices(tNearestID, Triangle.V[0], Triangle.V[1], Triangle.V[2]);
+
+		ProjectNormalOut = Triangle.Normal();
+
+		FDistPoint3Triangle3d DistanceQuery(Point, Triangle);
+		DistanceQuery.GetSquared();
+		return DistanceQuery.ClosestTrianglePoint;
+	}
+
+
+};
\ No newline at end of file
diff --git a/Engine/Plugins/Experimental/GeometryProcessing/Source/GeometricObjects/GeometricObjects.Build.cs b/Engine/Plugins/Experimental/GeometryProcessing/Source/GeometricObjects/GeometricObjects.Build.cs
new file mode 100644
--- /dev/null
+++ b/Engine/Plugins/Experimental/GeometryProcessing/Source/GeometricObjects/GeometricObjects.Build.cs
@@ -0,0 +1,17 @@
+// Copyright 1998-2019 Epic Games, Inc. All Rights Reserved.
+
+using UnrealBuildTool;
+
+public class GeometricObjects : ModuleRules
+{	
+	public GeometricObjects(ReadOnlyTargetRules Target) : base(Target)
+	{
+        PCHUsage = ModuleRules.PCHUsageMode.UseExplicitOrSharedPCHs;
+
+        PublicDependencyModuleNames.AddRange(
+			new string[] {
+				"Core"
+			}
+			);
+	}
+}
diff --git a/Engine/Plugins/Experimental/GeometryProcessing/Source/GeometricObjects/Private/CompGeom/PolygonTriangulation.cpp b/Engine/Plugins/Experimental/GeometryProcessing/Source/GeometricObjects/Private/CompGeom/PolygonTriangulation.cpp
new file mode 100644
--- /dev/null
+++ b/Engine/Plugins/Experimental/GeometryProcessing/Source/GeometricObjects/Private/CompGeom/PolygonTriangulation.cpp
@@ -0,0 +1,157 @@
+// Copyright 1998-2019 Epic Games, Inc. All Rights Reserved.
+
+#include "CompGeom/PolygonTriangulation.h"
+#include "TriangleTypes.h"
+
+
+// explicit instantiations
+namespace PolygonTriangulation
+{
+	template GEOMETRICOBJECTS_API void TriangulateSimplePolygon<float>(const TArray<FVector2<float>>& VertexPositions, TArray<FIndex3i>& OutTriangles);
+	template GEOMETRICOBJECTS_API void TriangulateSimplePolygon<double>(const TArray<FVector2<double>>& VertexPositions, TArray<FIndex3i>& OutTriangles);
+}
+
+
+
+
+//
+// Triangulate using ear clipping
+// This is based on the 3D triangulation code from MeshDescription.cpp, simplified for 2D polygons
+// 
+template<typename T>
+void PolygonTriangulation::TriangulateSimplePolygon(const TArray<FVector2<T>>& VertexPositions, TArray<FIndex3i>& OutTriangles)
+{
+	struct Local
+	{
+		static inline bool IsTriangleFlipped(T OrientationSign, const FVector2<T>& VertexPositionA, const FVector2<T>& VertexPositionB, const FVector2<T>& VertexPositionC)
+		{
+			T TriSignedArea = FTriangle2<T>::SignedArea(VertexPositionA, VertexPositionB, VertexPositionC);
+			return TriSignedArea * OrientationSign < 0;
+		}
+	};
+
+	// Polygon must have at least three vertices/edges
+	int32 PolygonVertexCount = VertexPositions.Num();
+	check(PolygonVertexCount >= 3);
+
+
+	// compute signed area of polygon
+	double PolySignedArea = 0;
+	for (int i = 0; i < PolygonVertexCount; ++i)
+	{
+		const FVector2<T>& v1 = VertexPositions[i];
+		const FVector2<T>& v2 = VertexPositions[(i + 1) % PolygonVertexCount];
+		PolySignedArea += v1.X*v2.Y - v1.Y*v2.X;
+	}
+	PolySignedArea *= 0.5;
+	bool bIsClockwise = PolySignedArea < 0;
+	double OrientationSign = (bIsClockwise) ? -1.0 : 1.0;
+
+
+	OutTriangles.Reset();
+
+
+	// If perimeter has 3 vertices, just copy content of perimeter out 
+	if (PolygonVertexCount == 3)
+	{
+		OutTriangles.Add(FIndex3i(0, 1, 2));
+		return;
+	}
+
+	// Make a simple linked list array of the previous and next vertex numbers, for each vertex number
+	// in the polygon.  This will just save us having to iterate later on.
+	static TArray<int32> PrevVertexNumbers, NextVertexNumbers;
+
+	PrevVertexNumbers.SetNumUninitialized(PolygonVertexCount, false);
+	NextVertexNumbers.SetNumUninitialized(PolygonVertexCount, false);
+
+	for (int32 VertexNumber = 0; VertexNumber < PolygonVertexCount; ++VertexNumber)
+	{
+		PrevVertexNumbers[VertexNumber] = VertexNumber - 1;
+		NextVertexNumbers[VertexNumber] = VertexNumber + 1;
+	}
+	PrevVertexNumbers[0] = PolygonVertexCount - 1;
+	NextVertexNumbers[PolygonVertexCount - 1] = 0;
+
+
+	int32 EarVertexNumber = 0;
+	int32 EarTestCount = 0;
+	for (int32 RemainingVertexCount = PolygonVertexCount; RemainingVertexCount >= 3; )
+	{
+		bool bIsEar = true;
+
+		// If we're down to only a triangle, just treat it as an ear.  Also, if we've tried every possible candidate
+		// vertex looking for an ear, go ahead and just treat the current vertex as an ear.  This can happen when 
+		// vertices are collinear or other degenerate cases.
+		if (RemainingVertexCount > 3 && EarTestCount < RemainingVertexCount)
+		{
+			const FVector2<T>& PrevVertexPosition = VertexPositions[PrevVertexNumbers[EarVertexNumber]];
+			const FVector2<T>& EarVertexPosition = VertexPositions[EarVertexNumber];
+			const FVector2<T>& NextVertexPosition = VertexPositions[NextVertexNumbers[EarVertexNumber]];
+
+			// Figure out whether the potential ear triangle is facing the same direction as the polygon
+			// itself.  If it's facing the opposite direction, then we're dealing with a concave triangle
+			// and we'll skip it for now.
+			if (!Local::IsTriangleFlipped(
+				OrientationSign, PrevVertexPosition, EarVertexPosition, NextVertexPosition))
+			{
+				int32 TestVertexNumber = NextVertexNumbers[NextVertexNumbers[EarVertexNumber]];
+
+				do
+				{
+					// Test every other remaining vertex to make sure that it doesn't lie inside our potential ear
+					// triangle.  If we find a vertex that's inside the triangle, then it cannot actually be an ear.
+					const FVector2<T>& TestVertexPosition = VertexPositions[TestVertexNumber];
+					if (FTriangle2<T>::IsInside(PrevVertexPosition, EarVertexPosition, NextVertexPosition, TestVertexPosition))
+					{
+						bIsEar = false;
+						break;
+					}
+
+					TestVertexNumber = NextVertexNumbers[TestVertexNumber];
+				} while (TestVertexNumber != PrevVertexNumbers[EarVertexNumber]);
+			}
+			else
+			{
+				bIsEar = false;
+			}
+		}
+
+		if (bIsEar)
+		{
+			// OK, we found an ear!  Let's save this triangle in our output buffer.
+			{
+				OutTriangles.Emplace();
+				FIndex3i& Triangle = OutTriangles.Last();
+				Triangle.A = PrevVertexNumbers[EarVertexNumber];
+				Triangle.B = EarVertexNumber;
+				Triangle.C = NextVertexNumbers[EarVertexNumber];
+			}
+
+			// Update our linked list.  We're effectively cutting off the ear by pointing the ear vertex's neighbors to
+			// point at their next sequential neighbor, and reducing the remaining vertex count by one.
+			{
+				NextVertexNumbers[PrevVertexNumbers[EarVertexNumber]] = NextVertexNumbers[EarVertexNumber];
+				PrevVertexNumbers[NextVertexNumbers[EarVertexNumber]] = PrevVertexNumbers[EarVertexNumber];
+				--RemainingVertexCount;
+			}
+
+			// Move on to the previous vertex in the list, now that this vertex was cut
+			EarVertexNumber = PrevVertexNumbers[EarVertexNumber];
+
+			EarTestCount = 0;
+		}
+		else
+		{
+			// The vertex is not the ear vertex, because it formed a triangle that either had a normal which pointed in the opposite direction
+			// of the polygon, or at least one of the other polygon vertices was found to be inside the triangle.  Move on to the next vertex.
+			EarVertexNumber = NextVertexNumbers[EarVertexNumber];
+
+			// Keep track of how many ear vertices we've tested, so that if we exhaust all remaining vertices, we can
+			// fall back to clipping the triangle and adding it to our mesh anyway.  This is important for degenerate cases.
+			++EarTestCount;
+		}
+	}
+
+	check(OutTriangles.Num() > 0);
+}
diff --git a/Engine/Plugins/Experimental/GeometryProcessing/Source/GeometricObjects/Private/GeometricObjectsModule.cpp b/Engine/Plugins/Experimental/GeometryProcessing/Source/GeometricObjects/Private/GeometricObjectsModule.cpp
new file mode 100644
--- /dev/null
+++ b/Engine/Plugins/Experimental/GeometryProcessing/Source/GeometricObjects/Private/GeometricObjectsModule.cpp
@@ -0,0 +1,20 @@
+// Copyright 1998-2019 Epic Games, Inc. All Rights Reserved.
+
+#include "GeometricObjectsModule.h"
+
+#define LOCTEXT_NAMESPACE "FGeometricObjectsModule"
+
+void FGeometricObjectsModule::StartupModule()
+{
+	// This code will execute after your module is loaded into memory; the exact timing is specified in the .uplugin file per-module
+}
+
+void FGeometricObjectsModule::ShutdownModule()
+{
+	// This function may be called during shutdown to clean up your module.  For modules that support dynamic reloading,
+	// we call this function before unloading the module.
+}
+
+#undef LOCTEXT_NAMESPACE
+	
+IMPLEMENT_MODULE(FGeometricObjectsModule, GeometricObjects)
\ No newline at end of file
diff --git a/Engine/Plugins/Experimental/GeometryProcessing/Source/GeometricObjects/Private/MathUtil.cpp b/Engine/Plugins/Experimental/GeometryProcessing/Source/GeometricObjects/Private/MathUtil.cpp
new file mode 100644
--- /dev/null
+++ b/Engine/Plugins/Experimental/GeometryProcessing/Source/GeometricObjects/Private/MathUtil.cpp
@@ -0,0 +1,50 @@
+// Copyright 1998-2019 Epic Games, Inc. All Rights Reserved.
+
+#include "MathUtil.h"
+
+#include <cfloat>
+#include <cmath>
+
+#include "Math/UnrealMath.h"
+
+template<> GEOMETRICOBJECTS_API const float TMathUtil<float>::Epsilon = FLT_EPSILON;
+template<> GEOMETRICOBJECTS_API const float TMathUtil<float>::ZeroTolerance = 1e-06f;
+template<> GEOMETRICOBJECTS_API const float TMathUtil<float>::MaxReal = FLT_MAX;
+template<> GEOMETRICOBJECTS_API const float TMathUtil<float>::Pi = (float)(4.0*atan(1.0));
+template<> GEOMETRICOBJECTS_API const float TMathUtil<float>::FourPi = 4.0f*TMathUtil<float>::Pi;
+template<> GEOMETRICOBJECTS_API const float TMathUtil<float>::TwoPi = 2.0f*TMathUtil<float>::Pi;
+template<> GEOMETRICOBJECTS_API const float TMathUtil<float>::HalfPi = 0.5f*TMathUtil<float>::Pi;
+template<> GEOMETRICOBJECTS_API const float TMathUtil<float>::InvPi = 1.0f / TMathUtil<float>::Pi;
+template<> GEOMETRICOBJECTS_API const float TMathUtil<float>::InvTwoPi = 1.0f / TMathUtil<float>::TwoPi;
+template<> GEOMETRICOBJECTS_API const float TMathUtil<float>::DegToRad = TMathUtil<float>::Pi / 180.0f;
+template<> GEOMETRICOBJECTS_API const float TMathUtil<float>::RadToDeg = 180.0f / TMathUtil<float>::Pi;
+//template<> const float TMathUtil<float>::LN_2 = TMathUtil<float>::Log(2.0f);
+//template<> const float TMathUtil<float>::LN_10 = TMathUtil<float>::Log(10.0f);
+//template<> const float TMathUtil<float>::INV_LN_2 = 1.0f / TMathUtil<float>::LN_2;
+//template<> const float TMathUtil<float>::INV_LN_10 = 1.0f / TMathUtil<float>::LN_10;
+template<> GEOMETRICOBJECTS_API const float TMathUtil<float>::Sqrt2 = (float)(sqrt(2.0));
+template<> GEOMETRICOBJECTS_API const float TMathUtil<float>::InvSqrt2 = 1.0f / TMathUtil<float>::Sqrt2;
+template<> GEOMETRICOBJECTS_API const float TMathUtil<float>::Sqrt3 = (float)(sqrt(3.0));
+template<> GEOMETRICOBJECTS_API const float TMathUtil<float>::InvSqrt3 = 1.0f / TMathUtil<float>::Sqrt3;
+
+template<> GEOMETRICOBJECTS_API const double TMathUtil<double>::Epsilon = DBL_EPSILON;
+template<> GEOMETRICOBJECTS_API const double TMathUtil<double>::ZeroTolerance = 1e-08;
+template<> GEOMETRICOBJECTS_API const double TMathUtil<double>::MaxReal = DBL_MAX;
+template<> GEOMETRICOBJECTS_API const double TMathUtil<double>::Pi = 4.0*atan(1.0);
+template<> GEOMETRICOBJECTS_API const double TMathUtil<double>::FourPi = 4.0*TMathUtil<double>::Pi;
+template<> GEOMETRICOBJECTS_API const double TMathUtil<double>::TwoPi = 2.0*TMathUtil<double>::Pi;
+template<> GEOMETRICOBJECTS_API const double TMathUtil<double>::HalfPi = 0.5*TMathUtil<double>::Pi;
+template<> GEOMETRICOBJECTS_API const double TMathUtil<double>::InvPi = 1.0 / TMathUtil<double>::Pi;
+template<> GEOMETRICOBJECTS_API const double TMathUtil<double>::InvTwoPi = 1.0 / TMathUtil<double>::TwoPi;
+template<> GEOMETRICOBJECTS_API const double TMathUtil<double>::DegToRad = TMathUtil<double>::Pi / 180.0;
+template<> GEOMETRICOBJECTS_API const double TMathUtil<double>::RadToDeg = 180.0 / TMathUtil<double>::Pi;
+//template<> const double TMathUtil<double>::LN_2 = TMathUtil<double>::Log(2.0);
+//template<> const double TMathUtil<double>::LN_10 = TMathUtil<double>::Log(10.0);
+//template<> const double TMathUtil<double>::INV_LN_2 = 1.0 / TMathUtil<double>::LN_2;
+//template<> const double TMathUtil<double>::INV_LN_10 = 1.0 / TMathUtil<double>::LN_10;
+template<> GEOMETRICOBJECTS_API const double TMathUtil<double>::Sqrt2 = sqrt(2.0);
+template<> GEOMETRICOBJECTS_API const double TMathUtil<double>::InvSqrt2 = 1.0f / TMathUtil<float>::Sqrt2;
+template<> GEOMETRICOBJECTS_API const double TMathUtil<double>::Sqrt3 = sqrt(3.0);
+template<> GEOMETRICOBJECTS_API const double TMathUtil<double>::InvSqrt3 = 1.0f / TMathUtil<float>::Sqrt3;
+
+
diff --git a/Engine/Plugins/Experimental/GeometryProcessing/Source/GeometricObjects/Public/BoxTypes.h b/Engine/Plugins/Experimental/GeometryProcessing/Source/GeometricObjects/Public/BoxTypes.h
new file mode 100644
--- /dev/null
+++ b/Engine/Plugins/Experimental/GeometryProcessing/Source/GeometricObjects/Public/BoxTypes.h
@@ -0,0 +1,488 @@
+// Copyright 1998-2019 Epic Games, Inc. All Rights Reserved.
+
+#pragma once
+
+#include "Math/Box.h"
+#include "Math/Box2D.h"
+#include "VectorTypes.h"
+
+template <typename RealType>
+struct FInterval1
+{
+	RealType Min;
+	RealType Max;
+
+	FInterval1()
+	{
+	}
+	FInterval1(const RealType& Min, const RealType& Max)
+	{
+		this->Min = Min;
+		this->Max = Max;
+	}
+
+	static FInterval1<RealType> Empty()
+	{
+		return FInterval1(TNumericLimits<RealType>::Max(), -TNumericLimits<RealType>::Max());
+	}
+
+	RealType Center() const
+	{
+		return (Min + Max) * (RealType)0.5;
+	}
+
+	RealType Extent() const
+	{
+		return Max - Min;
+	}
+	RealType Length() const
+	{
+		return Max - Min;
+	}
+
+	void Contain(const RealType& V)
+	{
+		if (V < Min)
+		{
+			Min = V;
+		}
+		if (V > Max)
+		{
+			Max = V;
+		}
+	}
+
+	bool Contains(RealType D) const
+	{
+		return D >= Min && D <= Max;
+	}
+
+	bool Overlaps(const FInterval1<RealType>& O) const
+	{
+		return !(O.Min > Max || O.Max < Min);
+	}
+
+	RealType SquaredDist(const FInterval1<RealType>& O) const
+	{
+		if (Max < O.Min)
+		{
+			return (O.Min - Max) * (O.Min - Max);
+		}
+		else if (Min > O.Max)
+		{
+			return (Min - O.Max) * (Min - O.Max);
+		}
+		else
+		{
+			return 0;
+		}
+	}
+	RealType Dist(const FInterval1<RealType>& O) const
+	{
+		if (Max < O.Min)
+		{
+			return O.Min - Max;
+		}
+		else if (Min > O.Max)
+		{
+			return Min - O.Max;
+		}
+		else
+		{
+			return 0;
+		}
+	}
+
+	FInterval1<RealType> IntersectionWith(const FInterval1<RealType>& O) const
+	{
+		if (O.Min > Max || O.Max < Min)
+		{
+			return FInterval1<RealType>::Empty();
+		}
+		return FInterval1<RealType>(TMathUtil<RealType>::Max(Min, O.Min), TMathUtil<RealType>::Min(Max, O.Max));
+	}
+
+	/**
+	 * clamp Value f to interval [Min,Max]
+	 */
+	RealType Clamp(RealType f) const
+	{
+		return (f < Min) ? Min : (f > Max) ? Max : f;
+	}
+
+	/**
+	 * interpolate between Min and Max using Value T in range [0,1]
+	 */
+	RealType Interpolate(RealType T) const
+	{
+		return (1 - T) * Min + (T)*Max;
+	}
+
+	/**
+	 * Convert Value into (clamped) T Value in range [0,1]
+	 */
+	RealType GetT(RealType Value) const
+	{
+		if (Value <= Min)
+		{
+			return 0;
+		}
+		else if (Value >= Max)
+		{
+			return 1;
+		}
+		else if (Min == Max)
+		{
+			return 0.5;
+		}
+		else
+		{
+			return (Value - Min) / (Max - Min);
+		}
+	}
+
+	void Set(FInterval1 O)
+	{
+		Min = O.Min;
+		Max = O.Max;
+	}
+
+	void Set(RealType A, RealType B)
+	{
+		Min = A;
+		Max = B;
+	}
+
+	FInterval1 operator-(FInterval1 V) const
+	{
+		return FInterval1(-V.Min, -V.Max);
+	}
+
+	FInterval1 operator+(RealType f) const
+	{
+		return FInterval1(Min + f, Max + f);
+	}
+
+	FInterval1 operator-(RealType f) const
+	{
+		return FInterval1(Min - f, Max - f);
+	}
+
+	FInterval1 operator*(RealType f) const
+	{
+		return FInterval1(Min * f, Max * f);
+	}
+};
+
+typedef FInterval1<float> FInterval1f;
+typedef FInterval1<double> FInterval1d;
+typedef FInterval1<int> FInterval1i;
+
+template <typename RealType>
+struct FAxisAlignedBox3
+{
+	FVector3<RealType> Min;
+	FVector3<RealType> Max;
+
+	FAxisAlignedBox3()
+	{
+	}
+	FAxisAlignedBox3(const FVector3<RealType>& Min, const FVector3<RealType>& Max)
+	{
+		this->Min = Min;
+		this->Max = Max;
+	}
+
+	operator FBox() const
+	{
+		FVector MinV((float)Min.X, (float)Min.Y, (float)Min.Z);
+		FVector MaxV((float)Max.X, (float)Max.Y, (float)Max.Z);
+		return FBox(MinV, MaxV);
+	}
+	FAxisAlignedBox3(const FBox& Box)
+	{
+		Min = Box.Min;
+		Max = Box.Max;
+	}
+
+	static FAxisAlignedBox3<RealType> Empty()
+	{
+		return FAxisAlignedBox3(
+			FVector3<RealType>(TNumericLimits<RealType>::Max(), TNumericLimits<RealType>::Max(), TNumericLimits<RealType>::Max()),
+			FVector3<RealType>(-TNumericLimits<RealType>::Max(), -TNumericLimits<RealType>::Max(), -TNumericLimits<RealType>::Max()));
+	}
+
+	FVector3<RealType> Center() const
+	{
+		return FVector3<RealType>(
+			(Min.X + Max.X) * (RealType)0.5,
+			(Min.Y + Max.Y) * (RealType)0.5,
+			(Min.Z + Max.Z) * (RealType)0.5);
+	}
+
+	FVector3<RealType> Extents() const
+	{
+		return FVector3<RealType>(Max.X - Min.X, Max.Y - Min.Y, Max.Z - Min.Z);
+	}
+
+	void Contain(const FVector3<RealType>& V)
+	{
+		if (V.X < Min.X)
+		{
+			Min.X = V.X;
+		}
+		if (V.X > Max.X)
+		{
+			Max.X = V.X;
+		}
+		if (V.Y < Min.Y)
+		{
+			Min.Y = V.Y;
+		}
+		if (V.Y > Max.Y)
+		{
+			Max.Y = V.Y;
+		}
+		if (V.Z < Min.Z)
+		{
+			Min.Z = V.Z;
+		}
+		if (V.Z > Max.Z)
+		{
+			Max.Z = V.Z;
+		}
+	}
+
+	void Contain(const FAxisAlignedBox3<RealType>& Other)
+	{
+		// todo: can be optimized
+		Contain(Other.Min);
+		Contain(Other.Max);
+	}
+
+	bool Contains(const FVector3<RealType>& V) const
+	{
+		return (Min.X <= V.X) && (Min.Y <= V.Y) && (Min.Z <= V.Z) && (Max.X >= V.X) && (Max.Y >= V.Y) && (Max.Z >= V.Z);
+	}
+
+	FAxisAlignedBox3<RealType> Intersect(const FAxisAlignedBox3<RealType>& Box) const
+	{
+		FAxisAlignedBox3<RealType> Intersection(
+			FVector3<RealType>(TMathUtil<RealType>::Max(Min.X, Box.Min.X), TMathUtil<RealType>::Max(Min.Y, Box.Min.Y), TMathUtil<RealType>::Max(Min.Z, Box.Min.Z)),
+			FVector3<RealType>(TMathUtil<RealType>::Min(Max.X, Box.Max.X), TMathUtil<RealType>::Min(Max.Y, Box.Max.Y), TMathUtil<RealType>::Min(Max.Z, Box.Max.Z)));
+		if (Intersection.Height() <= 0 || Intersection.Width() <= 0 || Intersection.Depth() <= 0)
+		{
+			return FAxisAlignedBox3<RealType>::Empty();
+		}
+		else
+		{
+			return Intersection;
+		}
+	}
+
+	bool Intersects(FAxisAlignedBox3 Box) const
+	{
+		return !((Box.Max.X <= Min.X) || (Box.Min.X >= Max.X) || (Box.Max.Y <= Min.Y) || (Box.Min.Y >= Max.Y) || (Box.Max.Z <= Min.Z) || (Box.Min.Z >= Max.Z));
+	}
+
+	RealType DistanceSquared(const FVector3<RealType>& V) const
+	{
+		RealType dx = (V.X < Min.X) ? Min.X - V.X : (V.X > Max.X ? V.X - Max.X : 0);
+		RealType dy = (V.Y < Min.Y) ? Min.Y - V.Y : (V.Y > Max.Y ? V.Y - Max.Y : 0);
+		RealType dz = (V.Z < Min.Z) ? Min.Z - V.Z : (V.Z > Max.Z ? V.Z - Max.Z : 0);
+		return dx * dx + dy * dy + dz * dz;
+	}
+
+	RealType Width() const
+	{
+		return TMathUtil<RealType>::Max(Max.X - Min.X, (RealType)0);
+	}
+
+	RealType Height() const
+	{
+		return TMathUtil<RealType>::Max(Max.Y - Min.Y, (RealType)0);
+	}
+
+	RealType Depth() const
+	{
+		return TMathUtil<RealType>::Max(Max.Z - Min.Z, (RealType)0);
+	}
+
+	RealType Volume() const
+	{
+		return Width() * Height() * Depth();
+	}
+
+	RealType DiagonalLength() const
+	{
+		return TMathUtil<RealType>::Sqrt((Max.X - Min.X) * (Max.X - Min.X) + (Max.Y - Min.Y) * (Max.Y - Min.Y) + (Max.Z - Min.Z) * (Max.Z - Min.Z));
+	}
+
+	RealType MaxDim() const
+	{
+		return TMathUtil<RealType>::Max(Width(), TMathUtil<RealType>::Max(Height(), Depth()));
+	}
+
+	FVector3<RealType> Diagonal() const
+	{
+		return FVector3<RealType>(Max.X - Min.X, Max.Y - Min.Y, Max.Z - Min.Z);
+	}
+};
+
+template <typename RealType>
+struct FAxisAlignedBox2
+{
+	FVector2<RealType> Min;
+	FVector2<RealType> Max;
+
+	FAxisAlignedBox2()
+	{
+	}
+	FAxisAlignedBox2(const FVector2<RealType>& Min, const FVector2<RealType>& Max)
+		: Min(Min), Max(Max)
+	{
+	}
+	FAxisAlignedBox2(RealType SquareSize)
+		: Min((RealType)0, (RealType)0), Max(SquareSize, SquareSize)
+	{
+	}
+	FAxisAlignedBox2(RealType Width, RealType Height)
+		: Min((RealType)0, (RealType)0), Max(Width, Height)
+	{
+	}
+
+	operator FBox2D() const
+	{
+		FVector2D MinV((float)Min.X, (float)Min.Y);
+		FVector2D MaxV((float)Max.X, (float)Max.Y);
+		return FBox2D(MinV, MaxV);
+	}
+	FAxisAlignedBox2(const FBox2D& Box)
+	{
+		Min = Box.Min;
+		Max = Box.Max;
+	}
+
+	static FAxisAlignedBox2<RealType> Empty()
+	{
+		return FAxisAlignedBox2(
+			FVector2<RealType>(TNumericLimits<RealType>::Max(), TNumericLimits<RealType>::Max()),
+			FVector2<RealType>(-TNumericLimits<RealType>::Max(), -TNumericLimits<RealType>::Max()));
+	}
+
+	FVector2<RealType> Center() const
+	{
+		return FVector2<RealType>(
+			(Min.X + Max.X) * (RealType)0.5,
+			(Min.Y + Max.Y) * (RealType)0.5);
+	}
+
+	FVector2<RealType> Extents() const
+	{
+		return FVector2<RealType>(Max.X - Min.X, Max.Y - Min.Y);
+	}
+
+	inline void Contain(const FVector2<RealType>& V)
+	{
+		if (V.X < Min.X)
+		{
+			Min.X = V.X;
+		}
+		if (V.X > Max.X)
+		{
+			Max.X = V.X;
+		}
+		if (V.Y < Min.Y)
+		{
+			Min.Y = V.Y;
+		}
+		if (V.Y > Max.Y)
+		{
+			Max.Y = V.Y;
+		}
+	}
+
+	inline void Contain(const FAxisAlignedBox2<RealType>& Other)
+	{
+		// todo: can be optimized
+		Contain(Other.Min);
+		Contain(Other.Max);
+	}
+
+	void Contain(const TArray<FVector2<RealType>>& Pts)
+	{
+		for (const FVector2<RealType>& Pt : Pts)
+		{
+			Contain(Pt);
+		}
+	}
+
+	bool Contains(const FVector2<RealType>& V) const
+	{
+		return (Min.X <= V.X) && (Min.Y <= V.Y) && (Max.X >= V.X) && (Max.Y >= V.Y);
+	}
+
+	bool Intersects(const FAxisAlignedBox2<RealType>& Box) const
+	{
+		return !((Box.Max.X < Min.X) || (Box.Min.X > Max.X) || (Box.Max.Y < Min.Y) || (Box.Min.Y > Max.Y));
+	}
+
+	FAxisAlignedBox2<RealType> Intersect(const FAxisAlignedBox2<RealType> &Box) const
+	{
+		FAxisAlignedBox2<RealType> Intersection(
+			FVector2<RealType>(TMathUtil<RealType>::Max(Min.X, Box.Min.X), TMathUtil<RealType>::Max(Min.Y, Box.Min.Y)),
+			FVector2<RealType>(TMathUtil<RealType>::Min(Max.X, Box.Max.X), TMathUtil<RealType>::Min(Max.Y, Box.Max.Y)));
+		if (Intersection.Height() <= 0 || Intersection.Width() <= 0)
+		{
+			return FAxisAlignedBox2<RealType>::Empty();
+		}
+		else
+		{
+			return Intersection;
+		}
+	}
+
+	RealType DistanceSquared(const FVector2<RealType>& V) const
+	{
+		RealType dx = (V.X < Min.X) ? Min.X - V.X : (V.X > Max.X ? V.X - Max.X : 0);
+		RealType dy = (V.Y < Min.Y) ? Min.Y - V.Y : (V.Y > Max.Y ? V.Y - Max.Y : 0);
+		return dx * dx + dy * dy;
+	}
+
+	inline RealType Width() const
+	{
+		return TMathUtil<RealType>::Max(Max.X - Min.X, (RealType)0);
+	}
+
+	inline RealType Height() const
+	{
+		return TMathUtil<RealType>::Max(Max.Y - Min.Y, (RealType)0);
+	}
+
+	inline RealType Area() const
+	{
+		return Width() * Height();
+	}
+
+	RealType DiagonalLength() const
+	{
+		return (RealType)TMathUtil<RealType>::Sqrt((Max.X - Min.X) * (Max.X - Min.X) + (Max.Y - Min.Y) * (Max.Y - Min.Y));
+	}
+
+	inline RealType MaxDim() const
+	{
+		return TMathUtil<RealType>::Max(Width(), Height());
+	}
+
+	inline RealType MinDim() const
+	{
+		return TMathUtil<RealType>::Min(Width(), Height());
+	}
+};
+
+typedef FAxisAlignedBox2<float> FAxisAlignedBox2f;
+typedef FAxisAlignedBox2<double> FAxisAlignedBox2d;
+typedef FAxisAlignedBox2<int> FAxisAlignedBox2i;
+typedef FAxisAlignedBox3<float> FAxisAlignedBox3f;
+typedef FAxisAlignedBox3<double> FAxisAlignedBox3d;
+typedef FAxisAlignedBox3<int> FAxisAlignedBox3i;
diff --git a/Engine/Plugins/Experimental/GeometryProcessing/Source/GeometricObjects/Public/CompGeom/PolygonTriangulation.h b/Engine/Plugins/Experimental/GeometryProcessing/Source/GeometricObjects/Public/CompGeom/PolygonTriangulation.h
new file mode 100644
--- /dev/null
+++ b/Engine/Plugins/Experimental/GeometryProcessing/Source/GeometricObjects/Public/CompGeom/PolygonTriangulation.h
@@ -0,0 +1,19 @@
+// Copyright 1998-2019 Epic Games, Inc. All Rights Reserved.
+
+#pragma once
+
+#include "VectorTypes.h"
+#include "IndexTypes.h"
+
+namespace PolygonTriangulation
+{
+
+	/**
+	 * Compute triangulation of simple polygon using ear-clipping
+	 * @param VertexPositions ordered vertices of polygon
+	 * @param OutTriangles computed triangulation. Each triangle is a tuple of indices into VertexPositions.
+	 */
+	template<typename RealType>
+	void GEOMETRICOBJECTS_API TriangulateSimplePolygon(const TArray<FVector2<RealType>>& VertexPositions, TArray<FIndex3i>& OutTriangles);
+
+}
\ No newline at end of file
diff --git a/Engine/Plugins/Experimental/GeometryProcessing/Source/GeometricObjects/Public/Distance/DistPoint3Triangle3.h b/Engine/Plugins/Experimental/GeometryProcessing/Source/GeometricObjects/Public/Distance/DistPoint3Triangle3.h
new file mode 100644
--- /dev/null
+++ b/Engine/Plugins/Experimental/GeometryProcessing/Source/GeometricObjects/Public/Distance/DistPoint3Triangle3.h
@@ -0,0 +1,275 @@
+// Copyright 1998-2019 Epic Games, Inc. All Rights Reserved.
+
+// Port of geometry3Sharp DistPoint3Triangle3
+
+#pragma once
+
+#include "VectorTypes.h"
+#include "TriangleTypes.h"
+
+
+/**
+ * Compute unsigned distance between 3D point and 3D triangle
+ */
+template <typename Real>
+class FDistPoint3Triangle3
+{
+public:
+	// Input
+	FVector3<Real> Point;
+	FTriangle3<Real> Triangle;
+
+	// Results
+	FVector3<Real> TriangleBaryCoords;
+	FVector3<Real> ClosestTrianglePoint;  // do we need this just use Triangle.BarycentricPoint
+
+
+	FDistPoint3Triangle3(const FVector3<Real>& PointIn, const FTriangle3<Real>& TriangleIn)
+	{
+		Point = PointIn;
+		Triangle = TriangleIn;
+	}
+
+	Real Get() 
+	{
+		return (Real)sqrt(ComputeResult());
+	}
+	Real GetSquared()
+	{
+		return ComputeResult();
+	}
+
+	Real ComputeResult()
+	{
+		FVector3<Real> diff = Triangle.V[0] - Point;
+		FVector3<Real> edge0 = Triangle.V[1] - Triangle.V[0];
+		FVector3<Real> edge1 = Triangle.V[2] - Triangle.V[0];
+		Real a00 = edge0.SquaredLength();
+		Real a01 = edge0.Dot(edge1);
+		Real a11 = edge1.SquaredLength();
+		Real b0 = diff.Dot(edge0);
+		Real b1 = diff.Dot(edge1);
+		Real c = diff.SquaredLength();
+		Real det = FMath::Abs(a00*a11 - a01 * a01);
+		Real s = a01 * b1 - a11 * b0;
+		Real t = a01 * b0 - a00 * b1;
+		Real sqrDistance;
+
+		if (s + t <= det)
+		{
+			if (s < (Real)0)
+			{
+				if (t < (Real)0)  // region 4
+				{
+					if (b0 < (Real)0)
+					{
+						t = (Real)0;
+						if (-b0 >= a00)
+						{
+							s = (Real)1;
+							sqrDistance = a00 + ((Real)2)*b0 + c;
+						}
+						else
+						{
+							s = -b0 / a00;
+							sqrDistance = b0 * s + c;
+						}
+					}
+					else
+					{
+						s = (Real)0;
+						if (b1 >= (Real)0)
+						{
+							t = (Real)0;
+							sqrDistance = c;
+						}
+						else if (-b1 >= a11)
+						{
+							t = (Real)1;
+							sqrDistance = a11 + ((Real)2)*b1 + c;
+						}
+						else
+						{
+							t = -b1 / a11;
+							sqrDistance = b1 * t + c;
+						}
+					}
+				}
+				else  // region 3
+				{
+					s = (Real)0;
+					if (b1 >= (Real)0)
+					{
+						t = (Real)0;
+						sqrDistance = c;
+					}
+					else if (-b1 >= a11)
+					{
+						t = (Real)1;
+						sqrDistance = a11 + ((Real)2)*b1 + c;
+					}
+					else
+					{
+						t = -b1 / a11;
+						sqrDistance = b1 * t + c;
+					}
+				}
+			}
+			else if (t < (Real)0)  // region 5
+			{
+				t = (Real)0;
+				if (b0 >= (Real)0)
+				{
+					s = (Real)0;
+					sqrDistance = c;
+				}
+				else if (-b0 >= a00)
+				{
+					s = (Real)1;
+					sqrDistance = a00 + ((Real)2)*b0 + c;
+				}
+				else
+				{
+					s = -b0 / a00;
+					sqrDistance = b0 * s + c;
+				}
+			}
+			else  // region 0
+			{
+				// minimum at interior point
+				Real invDet = ((Real)1) / det;
+				s *= invDet;
+				t *= invDet;
+				sqrDistance = s * (a00*s + a01 * t + ((Real)2)*b0) + t * (a01*s + a11 * t + ((Real)2)*b1) + c;
+			}
+		}
+		else
+		{
+			Real tmp0, tmp1, numer, denom;
+
+			if (s < (Real)0)  // region 2
+			{
+				tmp0 = a01 + b0;
+				tmp1 = a11 + b1;
+				if (tmp1 > tmp0)
+				{
+					numer = tmp1 - tmp0;
+					denom = a00 - ((Real)2)*a01 + a11;
+					if (numer >= denom)
+					{
+						s = (Real)1;
+						t = (Real)0;
+						sqrDistance = a00 + ((Real)2)*b0 + c;
+					}
+					else
+					{
+						s = numer / denom;
+						t = (Real)1 - s;
+						sqrDistance = s * (a00*s + a01 * t + ((Real)2)*b0) + t * (a01*s + a11 * t + ((Real)2)*b1) + c;
+					}
+				}
+				else
+				{
+					s = (Real)0;
+					if (tmp1 <= (Real)0)
+					{
+						t = (Real)1;
+						sqrDistance = a11 + ((Real)2)*b1 + c;
+					}
+					else if (b1 >= (Real)0)
+					{
+						t = (Real)0;
+						sqrDistance = c;
+					}
+					else
+					{
+						t = -b1 / a11;
+						sqrDistance = b1 * t + c;
+					}
+				}
+			}
+			else if (t < (Real)0)  // region 6
+			{
+				tmp0 = a01 + b1;
+				tmp1 = a00 + b0;
+				if (tmp1 > tmp0)
+				{
+					numer = tmp1 - tmp0;
+					denom = a00 - ((Real)2)*a01 + a11;
+					if (numer >= denom)
+					{
+						t = (Real)1;
+						s = (Real)0;
+						sqrDistance = a11 + ((Real)2)*b1 + c;
+					}
+					else
+					{
+						t = numer / denom;
+						s = (Real)1 - t;
+						sqrDistance = s * (a00*s + a01 * t + ((Real)2)*b0) + t * (a01*s + a11 * t + ((Real)2)*b1) + c;
+					}
+				}
+				else
+				{
+					t = (Real)0;
+					if (tmp1 <= (Real)0)
+					{
+						s = (Real)1;
+						sqrDistance = a00 + ((Real)2)*b0 + c;
+					}
+					else if (b0 >= (Real)0)
+					{
+						s = (Real)0;
+						sqrDistance = c;
+					}
+					else
+					{
+						s = -b0 / a00;
+						sqrDistance = b0 * s + c;
+					}
+				}
+			}
+			else  // region 1
+			{
+				numer = a11 + b1 - a01 - b0;
+				if (numer <= (Real)0)
+				{
+					s = (Real)0;
+					t = (Real)1;
+					sqrDistance = a11 + ((Real)2)*b1 + c;
+				}
+				else
+				{
+					denom = a00 - ((Real)2)*a01 + a11;
+					if (numer >= denom)
+					{
+						s = (Real)1;
+						t = (Real)0;
+						sqrDistance = a00 + ((Real)2)*b0 + c;
+					}
+					else
+					{
+						s = numer / denom;
+						t = (Real)1 - s;
+						sqrDistance = s * (a00*s + a01 * t + ((Real)2)*b0) + t * (a01*s + a11 * t + ((Real)2)*b1) + c;
+					}
+				}
+			}
+		}
+
+		// Account for numerical round-off error.
+		if (sqrDistance < (Real)0)
+		{
+			sqrDistance = (Real)0;
+		}
+
+		ClosestTrianglePoint = Triangle.V[0] + s * edge0 + t * edge1;
+		TriangleBaryCoords = FVector3<Real>((Real)1 - s - t, s, t);
+		return sqrDistance;
+	}
+
+};
+
+typedef FDistPoint3Triangle3<float> FDistPoint3Triangle3f;
+typedef FDistPoint3Triangle3<double> FDistPoint3Triangle3d;
+
diff --git a/Engine/Plugins/Experimental/GeometryProcessing/Source/GeometricObjects/Public/FrameTypes.h b/Engine/Plugins/Experimental/GeometryProcessing/Source/GeometricObjects/Public/FrameTypes.h
new file mode 100644
--- /dev/null
+++ b/Engine/Plugins/Experimental/GeometryProcessing/Source/GeometricObjects/Public/FrameTypes.h
@@ -0,0 +1,325 @@
+// Copyright 1998-2019 Epic Games, Inc. All Rights Reserved.
+
+
+#pragma once
+
+#include "VectorTypes.h"
+#include "VectorUtil.h"
+#include "RayTypes.h"
+#include "Quaternion.h"
+
+/**
+ * FFrame3 is an object that represents an oriented 3D coordinate frame, ie orthogonal X/Y/Z axes at a point in space.
+ * One can think of this Frame as a local coordinate space measured along these axes.
+ * Functions are provided to map geometric objects to/from the Frame coordinate space.
+ * 
+ * Internally the representation is the same as an FTransform, except a Frame has no Scale.
+ */
+template<typename RealType>
+struct FFrame3
+{
+	/**
+	 * Origin of the frame
+	 */
+	FVector3<RealType> Origin;
+
+	/**
+	 * Rotation of the frame. Think of this as the rotation of the unit X/Y/Z axes to the 3D frame axes.
+	 */
+	FQuaternion<RealType> Rotation;
+
+	/**
+	 * Construct a frame positioned at (0,0,0) aligned to the unit axes
+	 */
+	FFrame3()
+	{
+		Origin = FVector3<RealType>::Zero();
+		Rotation = FQuaternion<RealType>::Identity();
+	}
+
+	/**
+	 * Construct a frame at the given Origin aligned to the unit axes
+	 */
+	FFrame3(const FVector3<RealType>& OriginIn)
+	{
+		Origin = OriginIn;
+		Rotation = FQuaternion<RealType>::Identity();
+	}
+
+	/**
+	 * Construct a Frame from the given Origin and Rotation
+	 */
+	FFrame3(const FVector3<RealType>& OriginIn, const FQuaternion<RealType> RotationIn)
+	{
+		Origin = OriginIn;
+		Rotation = RotationIn;
+	}
+
+	/**
+	 * Construct a frame with the Z axis aligned to a target axis
+	 * @param OriginIn origin of frame
+	 * @param SetZ target Z axis
+	 */
+	FFrame3(const FVector3<RealType>& OriginIn, const FVector3<RealType>& SetZ)
+	{
+		Origin = OriginIn;
+		Rotation.SetFromTo(FVector3<RealType>::UnitZ(), SetZ);
+	}
+
+	/**
+	 * Construct Frame from X/Y/Z axis vectors. Vectors must be mutually orthogonal.
+	 * @param OriginIn origin of frame
+	 * @param X desired X axis of frame
+	 * @param Y desired Y axis of frame
+	 * @param Z desired Z axis of frame
+	 */
+	FFrame3(const FVector3<RealType>& OriginIn, const FVector3<RealType>& X, const FVector3<RealType>& Y, const FVector3<RealType>& Z)
+	{
+		Origin = OriginIn;
+		Rotation = FQuaternion<RealType>( FMatrix3<RealType>(X, Y, Z, false) );
+	}
+
+	/** Construct a Frame from an FTransform */
+	FFrame3(const FTransform& Transform)
+	{
+		Origin = Transform.GetTranslation();
+		Rotation = Transform.GetRotation();
+	}
+
+	
+	/**
+	 * @param AxisIndex index of axis of frame, either 0, 1, or 2
+	 * @return axis vector
+	 */
+	FVector3<RealType> GetAxis(int AxisIndex) const
+	{
+		switch (AxisIndex)
+		{
+		case 0:
+			return Rotation.AxisX();
+		case 1:
+			return Rotation.AxisY();
+		case 2:
+			return Rotation.AxisZ();
+		default:
+			checkNoEntry();
+			return FVector3<RealType>::Zero(); // compiler demands a return value
+		}
+	}
+
+	/** @return X axis of frame (axis 0) */
+	FVector3<RealType> X() const
+	{
+		return Rotation.AxisX();
+	}
+
+	/** @return Y axis of frame (axis 1) */
+	FVector3<RealType> Y() const
+	{
+		return Rotation.AxisY();
+	}
+
+	/** @return Z axis of frame (axis 2) */
+	FVector3<RealType> Z() const
+	{
+		return Rotation.AxisZ();
+	}
+
+	/** @return conversion of this Frame to FTransform */
+	FTransform ToFTransform() const
+	{
+		return FTransform(Rotation, Origin);
+	}
+
+
+	/** @return point at distances along frame axes */
+	FVector3<RealType> PointAt(RealType X, RealType Y, RealType Z) const
+	{
+		return Rotation * FVector3<RealType>(X,Y,Z) + Origin;
+	}
+
+	
+	/** @return input Point transformed into local coordinate system of Frame */
+	FVector3<RealType> ToFramePoint(const FVector3<RealType>& Point) const
+	{
+		return Rotation.InverseMultiply((Point-Origin));
+	}
+	/** @return input Point transformed from local coordinate system of Frame into "World" coordinate system */
+	FVector3<RealType> FromFramePoint(const FVector3<RealType>& Point) const
+	{
+		return Rotation * Point + Origin;
+	}
+
+
+	/** @return input Vector transformed into local coordinate system of Frame */
+	FVector3<RealType> ToFrameVector(const FVector3<RealType>& Vector) const
+	{
+		return Rotation.InverseMultiply(Vector);
+	}
+	/** @return input Vector transformed from local coordinate system of Frame into "World" coordinate system */
+	FVector3<RealType> FromFrameVector(const FVector3<RealType>& Vector) const
+	{
+		return Rotation * Vector;
+	}
+
+
+	/** @return input Quaternion transformed into local coordinate system of Frame */
+	FQuaternion<RealType> ToFrame(const FQuaternion<RealType>& Quat) const
+	{
+		return Rotation.Inverse() * Quat;
+	}
+	/** @return input Quaternion transformed from local coordinate system of Frame into "World" coordinate system */
+	FQuaternion<RealType> FromFrame(const FQuaternion<RealType>& Quat) const
+	{
+		return Rotation * Quat;
+	}
+
+
+	/** @return input Ray transformed into local coordinate system of Frame */
+	FRay3<RealType> ToFrame(const FRay3<RealType>& Ray) const
+	{
+		return FRay3<RealType>(ToFramePoint(Ray.Origin), ToFrameVector(Ray.Direction));
+	}
+	/** @return input Ray transformed from local coordinate system of Frame into "World" coordinate system */
+	FRay3<RealType> FromFrame(const FRay3<RealType>& Ray) const
+	{
+		return FRay3<RealType>(ToFramePoint(Ray.Origin), ToFrameVector(Ray.Direction));
+	}
+
+
+	/** @return input Frame transformed into local coordinate system of this Frame */
+	FFrame3<RealType> ToFrame(const FFrame3<RealType>& Frame) const
+	{
+		return FFrame3<RealType>(ToFramePoint(Frame.Origin), ToFrame(Frame.Rotation));
+	}
+	/** @return input Frame transformed from local coordinate system of this Frame into "World" coordinate system */
+	FFrame3<RealType> FromFrame(const FFrame3<RealType>& Frame) const
+	{
+		return FFrame3<RealType>(ToFramePoint(Frame.Origin), FromFrame(Frame.Rotation));
+	}
+
+
+
+
+	/**
+	 * Project 3D point into plane and convert to UV coordinates in that plane
+	 * @param Pos 3D position
+	 * @param nPlaneNormalAxis which plane to project onto, identified by perpendicular normal. Default is 2, ie normal is Z, plane is (X,Y)
+	 * @return 2D coordinates in UV plane, relative to origin
+	 */
+	FVector2<RealType> ToPlaneUV(const FVector3<RealType>& Pos, int nPlaneNormalAxis) const
+	{
+		int nAxis0 = 0, nAxis1 = 1;
+		if (nPlaneNormalAxis == 0)
+		{
+			nAxis0 = 2;
+		}
+		else if (nPlaneNormalAxis == 1)
+		{
+			nAxis1 = 2;
+		}
+		FVector3<RealType> LocalPos = Pos - Origin;
+		RealType U = LocalPos.Dot(GetAxis(nAxis0));
+		RealType V = LocalPos.Dot(GetAxis(nAxis1));
+		return FVector2<RealType>(U, V);
+	}
+
+
+
+	/**
+	 * Map a point from local UV plane coordinates to the corresponding 3D point in one of the planes of the frame
+	 * @param PosUV local UV plane coordinates
+	 * @param nPlaneNormalAxis which plane to map to, identified by perpendicular normal. Default is 2, ie normal is Z, plane is (X,Y)
+	 * @return 3D coordinates in frame's plane (including Origin translation)
+	 */
+	FVector3<RealType> FromPlaneUV(const FVector2<RealType>& PosUV, int PlaneNormalAxis) const
+	{
+		FVector3<RealType> PlanePos(PosUV[0], PosUV[1], 0);
+		if (PlaneNormalAxis == 0)
+		{
+			PlanePos[0] = 0; PlanePos[2] = PosUV[0];
+		}
+		else if (PlaneNormalAxis == 1)
+		{
+			PlanePos[1] = 0; PlanePos[2] = PosUV[1];
+		}
+		return Rotation*PlanePos + Origin;
+	}
+
+
+
+
+
+	/**
+	 * Rotate this frame by given quaternion
+	 */
+	void Rotate(const FQuaternion<RealType>& Quat)
+	{
+		Rotation = Quat * Rotation;
+	}
+
+
+	/**
+	 * transform this frame by the given transform
+	 */
+	void Transform(const FTransform& XForm)
+	{
+		Origin = (FVector3<RealType>)XForm.TransformPosition((FVector3f)Origin);
+		Rotation = FQuaternion<RealType>(XForm.GetRotation()) * Rotation;
+	}
+
+
+	/**
+	 * Align an axis of this frame with a target direction
+	 * @param AxisIndex which axis to align
+	 * @param ToDirection target direction
+	 */
+	void AlignAxis(int AxisIndex, const FVector3<RealType>& ToDirection)
+	{
+		FQuaternion<RealType> RelRotation(GetAxis(AxisIndex), ToDirection);
+		Rotate(RelRotation);
+	}
+
+
+	/**
+	 * Compute rotation around vector that best-aligns axis of frame with target direction
+	 * @param AxisIndex which axis to try to align
+	 * @param ToDirection target direction
+	 * @param AroundVector rotation is constrained to be around this vector (ie this direction in frame stays constant)
+	 */
+	void ConstrainedAlignAxis(int AxisIndex, const FVector3<RealType>& ToDirection, const FVector3<RealType>& AroundVector)
+	{
+		//@todo PlaneAngleSigned does acos() and then SetAxisAngleD() does cos/sin...can we optimize this?
+		FVector3<RealType> AxisVec = GetAxis(AxisIndex);
+		RealType AngleDeg = VectorUtil::PlaneAngleSignedD(AxisVec, ToDirection, AroundVector);
+		FQuaternion<RealType> RelRotation;
+		RelRotation.SetAxisAngleD(AroundVector, AngleDeg);
+		Rotate(RelRotation);
+	}
+
+
+
+
+	/**
+	 * Compute intersection of ray with plane defined by frame origin and axis as normal
+	 * @param RayOrigin origin of ray
+	 * @param RayDirection direction of ray
+	 * @param PlaneNormalAxis which axis of frame to use as plane normal
+	 * @return intersection point, or FVector3::Max() if ray is parallel to plane
+	 */
+	FVector3<RealType> RayPlaneIntersection(const FVector3<RealType>& RayOrigin, const FVector3<RealType>& RayDirection, int PlaneNormalAxis)
+	{
+		FVector3<RealType> Normal = GetAxis(PlaneNormalAxis);
+		RealType PlaneD = -Origin.Dot(Normal);
+		RealType NormalDot = RayDirection.Dot(Normal);
+		if (VectorUtil::EpsilonEqual(NormalDot, (RealType)0, TMathUtil<RealType>::ZeroTolerance))
+			return FVector3<RealType>::Max();
+		RealType t = -( RayOrigin.Dot(Normal) + PlaneD) / NormalDot;
+		return RayOrigin + t * RayDirection;
+	}
+
+};
+
+typedef FFrame3<float> FFrame3f;
+typedef FFrame3<double> FFrame3d;
+
diff --git a/Engine/Plugins/Experimental/GeometryProcessing/Source/GeometricObjects/Public/Intersection/IntrLine2Line2.h b/Engine/Plugins/Experimental/GeometryProcessing/Source/GeometricObjects/Public/Intersection/IntrLine2Line2.h
new file mode 100644
--- /dev/null
+++ b/Engine/Plugins/Experimental/GeometryProcessing/Source/GeometricObjects/Public/Intersection/IntrLine2Line2.h
@@ -0,0 +1,175 @@
+// Copyright 1998-2019 Epic Games, Inc. All Rights Reserved.
+
+// Port of geometry3Sharp IntrLine2Line2
+
+#pragma once
+
+#include "VectorTypes.h"
+#include "LineTypes.h"
+#include "VectorUtil.h"  // for EIntersectionType
+
+/**
+ * Compute intersection between two 2D lines
+ */
+template<typename RealType>
+class FIntrLine2Line2
+{
+protected:
+	// input data
+	FLine2<RealType> Line1;
+	FLine2<RealType> Line2;
+	RealType dotThresh = TMathUtil<RealType>::ZeroTolerance;
+	
+public:
+	// result data
+	FVector2<RealType> Point;
+	RealType Segment1Parameter;
+	RealType Segment2Parameter;
+	int Quantity = 0;
+	EIntersectionResult Result = EIntersectionResult::NotComputed;
+	EIntersectionType Type = EIntersectionType::Empty;
+
+
+
+	FIntrLine2Line2(const FLine2<RealType>& Line1In, const FLine2<RealType>& Line2In) 
+		: Line1(Line1In), Line2(Line2In)
+	{
+	}
+
+
+	const FLine2<RealType>& GetLine1() const
+	{
+		return Line1;
+	}
+
+	void SetLine1(const FLine2<RealType>& Value)
+	{
+		Line1 = Value;
+		Result = EIntersectionResult::NotComputed;
+	}
+
+	const FLine2<RealType>& GetLine2() const
+	{
+		return Line2;
+	}
+
+	void SetLine2(FLine2<RealType>& Value)
+	{
+		Line2 = Value;
+		Result = EIntersectionResult::NotComputed;
+	}
+
+
+	RealType GetDotThreshold() const
+	{
+		return dotThresh;
+	}
+
+	void SetDotThreshold(RealType Value)
+	{
+		dotThresh = FMath::Max(Value, (RealType)0);
+		Result = EIntersectionResult::NotComputed;
+	}
+
+	bool IsSimpleIntersection() const
+	{
+		return Result == EIntersectionResult::Intersects && Type == EIntersectionType::Point;
+	}
+
+
+	FIntrLine2Line2& Compute()
+	{
+		Find();
+		return *this;
+	}
+
+
+	bool Find()
+	{
+		if (Result != EIntersectionResult::NotComputed)
+		{
+			return (Result == EIntersectionResult::Intersects);
+		}
+
+		// [RMS] if either line direction is not a normalized vector, 
+		//   results are garbage, so fail query
+		if (Line1.Direction.IsNormalized() == false || Line2.Direction.IsNormalized() == false) 
+		{
+			Type = EIntersectionType::Empty;
+			Result = EIntersectionResult::InvalidQuery;
+			return false;
+		}
+
+		FVector2<RealType> s = FVector2<RealType>::Zero();
+		Type = Classify(Line1.Origin, Line1.Direction,
+			Line2.Origin, Line2.Direction, dotThresh, s);
+
+		if (Type == EIntersectionType::Point) 
+		{
+			Quantity = 1;
+			Point = Line1.Origin + s.X*Line1.Direction;
+			Segment1Parameter = s.X;
+			Segment2Parameter = s.Y;
+		}
+		else if (Type == EIntersectionType::Line) 
+		{
+			Quantity = std::numeric_limits<int>::max();
+		}
+		else 
+		{
+			Quantity = 0;
+		}
+
+		Result = (Type != EIntersectionType::Empty) ?
+			EIntersectionResult::Intersects : EIntersectionResult::NoIntersection;
+		return (Result == EIntersectionResult::Intersects);
+	}
+
+
+
+	static EIntersectionType Classify(
+		const FVector2<RealType>& P0, const FVector2<RealType>& D0, 
+		const FVector2<RealType>& P1, const FVector2<RealType>& D1, 
+		RealType dotThreshold, FVector2<RealType>& s)
+	{
+		// Ensure dotThreshold is nonnegative.
+		dotThreshold = FMath::Max(dotThreshold, (RealType)0);
+
+		// The intersection of two lines is a solution to P0+s0*D0 = P1+s1*D1.
+		// Rewrite this as s0*D0 - s1*D1 = P1 - P0 = Q.  If D0.Dot(Perp(D1)) = 0,
+		// the lines are parallel.  Additionally, if Q.Dot(Perp(D1)) = 0, the
+		// lines are the same.  If D0.Dot(Perp(D1)) is not zero, then
+		//   s0 = Q.Dot(Perp(D1))/D0.Dot(Perp(D1))
+		// produces the point of intersection.  Also,
+		//   s1 = Q.Dot(Perp(D0))/D0.Dot(Perp(D1))
+
+		FVector2<RealType> diff = P1 - P0;
+		RealType D0DotPerpD1 = D0.DotPerp(D1);
+		if (FMath::Abs(D0DotPerpD1) > dotThreshold) 
+		{
+			// Lines intersect in a single point.
+			RealType invD0DotPerpD1 = 1.0 / D0DotPerpD1;
+			RealType diffDotPerpD0 = diff.DotPerp(D0);
+			RealType diffDotPerpD1 = diff.DotPerp(D1);
+			s[0] = diffDotPerpD1 * invD0DotPerpD1;
+			s[1] = diffDotPerpD0 * invD0DotPerpD1;
+			return EIntersectionType::Point;
+		}
+
+		// Lines are parallel.
+		diff.Normalize();
+		RealType diffNDotPerpD1 = diff.DotPerp(D1);
+		if (FMath::Abs(diffNDotPerpD1) <= dotThreshold) 
+		{
+			// Lines are collinear.
+			return EIntersectionType::Line;
+		}
+
+		// Lines are parallel, but distinct.
+		return EIntersectionType::Empty;
+	}
+
+};
+
+typedef FIntrLine2Line2<double> FIntrLine2FLine2d;
+typedef FIntrLine2Line2<float> FIntrLine2Line2f;
diff --git a/Engine/Plugins/Experimental/GeometryProcessing/Source/GeometricObjects/Public/Intersection/IntrRay3AxisAlignedBox3.h b/Engine/Plugins/Experimental/GeometryProcessing/Source/GeometricObjects/Public/Intersection/IntrRay3AxisAlignedBox3.h
new file mode 100644
--- /dev/null
+++ b/Engine/Plugins/Experimental/GeometryProcessing/Source/GeometricObjects/Public/Intersection/IntrRay3AxisAlignedBox3.h
@@ -0,0 +1,220 @@
+// Copyright 1998-2019 Epic Games, Inc. All Rights Reserved.
+
+// Port of WildMagic IntrRay3Box3, simplified for Axis-Aligned Bounding Box
+
+#pragma once
+
+#include "RayTypes.h"
+#include "BoxTypes.h"
+#include "VectorUtil.h"
+
+
+/**
+ * Compute intersection between 3D ray and 3D axis-aligned box
+ */
+template <typename RealType>
+class FIntrRay3AxisAlignedBox3
+{
+public:
+
+	// @todo port object version
+
+
+
+	/**
+	 * Test if ray intersects box
+	 * @param Ray query ray
+	 * @param Box query box
+	 * @param ExpandExtents box is expanded by this amount in each direction, useful for dealing with float imprecision
+	 * @return true if ray intersects box
+	 */
+	static bool TestIntersection(const FRay3<RealType>& Ray, const FAxisAlignedBox3<RealType>& Box, RealType ExpandExtents = 0)
+	{
+		FVector3<RealType> WdU = FVector3<RealType>::Zero();
+		FVector3<RealType> AWdU = FVector3<RealType>::Zero();
+		FVector3<RealType> DdU = FVector3<RealType>::Zero();
+		FVector3<RealType> ADdU = FVector3<RealType>::Zero();
+		RealType RHS;
+
+		FVector3<RealType> diff = Ray.Origin - Box.Center();
+		FVector3<RealType> extent = Box.Extents() + ExpandExtents;
+
+		WdU.X = Ray.Direction.X;
+		AWdU.X = FMath::Abs(WdU.X);
+		DdU.X = diff.X;
+		ADdU.X = FMath::Abs(DdU.X);
+		if (ADdU.X > extent.X && DdU.X * WdU.X >= (RealType)0) 
+		{
+			return false;
+		}
+
+		WdU.Y = Ray.Direction.Y;
+		AWdU.Y = FMath::Abs(WdU.Y);
+		DdU.Y = diff.Y;
+		ADdU.Y = FMath::Abs(DdU.Y);
+		if (ADdU.Y > extent.Y && DdU.Y * WdU.Y >= (RealType)0) 
+		{
+			return false;
+		}
+
+		WdU.Z = Ray.Direction.Z;
+		AWdU.Z = FMath::Abs(WdU.Z);
+		DdU.Z = diff.Z;
+		ADdU.Z = FMath::Abs(DdU.Z);
+		if (ADdU.Z > extent.Z && DdU.Z * WdU.Z >= (RealType)0) 
+		{
+			return false;
+		}
+
+		FVector3<RealType> WxD = Ray.Direction.Cross(diff);
+		FVector3<RealType> AWxDdU = FVector3<RealType>::Zero();
+
+		AWxDdU.X = FMath::Abs(WxD.X);
+		RHS = extent.Y * AWdU.Z + extent.Z * AWdU.Y;
+		if (AWxDdU.X > RHS) 
+		{
+			return false;
+		}
+
+		AWxDdU.Y = FMath::Abs(WxD.Y);
+		RHS = extent.X * AWdU.Z + extent.Z * AWdU.X;
+		if (AWxDdU.Y > RHS) 
+		{
+			return false;
+		}
+
+		AWxDdU.Z = FMath::Abs(WxD.Z);
+		RHS = extent.X * AWdU.Y + extent.Y * AWdU.X;
+		if (AWxDdU.Z > RHS) 
+		{
+			return false;
+		}
+
+		return true;
+	}
+
+
+
+	/**
+	 * Find intersection of ray with AABB and returns ray T-value of intersection point (or TNumericLimits::Max() on miss)
+	 * @param Ray query ray
+	 * @param Box query box
+	 * @param RayParamOut ray intersect T-value, or TNumericLimits::Max()
+	 * @return true if ray intersects box
+	 */
+	static bool FindIntersection(const FRay3<RealType>& Ray, const FAxisAlignedBox3<RealType>& Box, RealType& RayParamOut)
+	{
+		RealType RayParam0 = 0.0;
+		RealType RayParam1 = TNumericLimits<RealType>::Max();
+		int Quantity = 0;
+		FVector3<RealType> Point0 = FVector3<RealType>::Zero();
+		FVector3<RealType> Point1 = FVector3<RealType>::Zero();
+		EIntersectionType Type = EIntersectionType::Empty;
+		DoClipping(RayParam0, RayParam1, Ray.Origin, Ray.Direction, Box,
+			true, Quantity, Point0, Point1, Type);
+
+		if (Type != EIntersectionType::Empty) 
+		{
+			RayParamOut = RayParam0;
+			return true;
+		}
+		else 
+		{
+			RayParamOut = TNumericLimits<RealType>::Max();
+			return false;
+		}
+	}
+
+
+
+
+protected:
+
+	// internal functions
+
+	static bool DoClipping(RealType t0, RealType t1,
+		const FVector3<RealType>& RayOrigin, const FVector3<RealType>& RayDirection,
+		const FAxisAlignedBox3<RealType>& Box, bool solid, 
+		int& quantity, FVector3<RealType>& Point0, FVector3<RealType>& Point1, EIntersectionType& intrType)
+	{
+		FVector3<RealType> BOrigin = RayOrigin - Box.Center();
+		FVector3<RealType> extent = Box.Extents();
+
+		RealType saveT0 = t0, saveT1 = t1;
+		bool notAllClipped =
+			Clip(+RayDirection.X, -BOrigin.X - extent.X, t0, t1) &&
+			Clip(-RayDirection.X, +BOrigin.X - extent.X, t0, t1) &&
+			Clip(+RayDirection.Y, -BOrigin.Y - extent.Y, t0, t1) &&
+			Clip(-RayDirection.Y, +BOrigin.Y - extent.Y, t0, t1) &&
+			Clip(+RayDirection.Z, -BOrigin.Z - extent.Z, t0, t1) &&
+			Clip(-RayDirection.Z, +BOrigin.Z - extent.Z, t0, t1);
+
+		if (notAllClipped && (solid || t0 != saveT0 || t1 != saveT1)) 
+		{
+			if (t1 > t0) 
+			{
+				intrType = EIntersectionType::Segment;
+				quantity = 2;
+				Point0 = RayOrigin + t0 * RayDirection;
+				Point1 = RayOrigin + t1 * RayDirection;
+			}
+			else 
+			{
+				intrType = EIntersectionType::Point;
+				quantity = 1;
+				Point0 = RayOrigin + t0 * RayDirection;
+			}
+		}
+		else 
+		{
+			quantity = 0;
+			intrType = EIntersectionType::Empty;
+		}
+
+		return intrType != EIntersectionType::Empty;
+	}
+
+
+
+
+	static bool Clip(RealType denom, RealType numer, RealType& t0, RealType& t1)
+	{
+		// Return value is 'true' if line segment intersects the current test
+		// plane.  Otherwise 'false' is returned in which case the line segment
+		// is entirely clipped.
+
+		if (denom > (RealType)0)
+		{
+			if (numer > denom*t1)
+			{
+				return false;
+			}
+			if (numer > denom*t0)
+			{
+				t0 = numer / denom;
+			}
+			return true;
+		}
+		else if (denom < (RealType)0)
+		{
+			if (numer > denom*t0)
+			{
+				return false;
+			}
+			if (numer > denom*t1)
+			{
+				t1 = numer / denom;
+			}
+			return true;
+		}
+		else
+		{
+			return numer <= (RealType)0;
+		}
+	}
+
+
+};
+
+typedef FIntrRay3AxisAlignedBox3<float> FIntrRay3AxisAlignedBox3f;
+typedef FIntrRay3AxisAlignedBox3<double> FIntrRay3AxisAlignedBox3d;
\ No newline at end of file
diff --git a/Engine/Plugins/Experimental/GeometryProcessing/Source/GeometricObjects/Public/Intersection/IntrRay3Triangle3.h b/Engine/Plugins/Experimental/GeometryProcessing/Source/GeometricObjects/Public/Intersection/IntrRay3Triangle3.h
new file mode 100644
--- /dev/null
+++ b/Engine/Plugins/Experimental/GeometryProcessing/Source/GeometricObjects/Public/Intersection/IntrRay3Triangle3.h
@@ -0,0 +1,172 @@
+// Copyright 1998-2019 Epic Games, Inc. All Rights Reserved.
+
+// Port of WildMagic IntrRay3Triangle3
+
+#pragma once
+
+#include "VectorTypes.h"
+#include "RayTypes.h"
+#include "TriangleTypes.h"
+#include "VectorUtil.h"
+
+/**
+ * Compute intersection between 3D ray and 3D triangle
+ */
+template <typename Real>
+class FIntrRay3Triangle3
+{
+public:
+	// Input
+	FRay3<Real> Ray;
+	FTriangle3<Real> Triangle;
+
+	// Output
+	Real RayParameter;
+	FVector3d TriangleBaryCoords;
+	EIntersectionType IntersectionType;
+
+
+	FIntrRay3Triangle3(const FRay3<Real>& RayIn, const FTriangle3<Real>& TriangleIn)
+	{
+		Ray = RayIn;
+		Triangle = TriangleIn;
+	}
+
+
+	/**
+	 * @return true if ray intersects triangle
+	 */
+	bool Test()
+	{
+		// Compute the offset origin, edges, and normal.
+		FVector3<Real> diff = Ray.Origin - Triangle.V[0];
+		FVector3<Real> edge1 = Triangle.V[1] - Triangle.V[0];
+		FVector3<Real> edge2 = Triangle.V[2] - Triangle.V[0];
+		FVector3<Real> normal = edge1.Cross(edge2);
+
+		// Solve Q + t*D = b1*E1 + b2*E2 (Q = kDiff, D = ray direction,
+		// E1 = kEdge1, E2 = kEdge2, N = Cross(E1,E2)) by
+		//   |Dot(D,N)|*b1 = sign(Dot(D,N))*Dot(D,Cross(Q,E2))
+		//   |Dot(D,N)|*b2 = sign(Dot(D,N))*Dot(D,Cross(E1,Q))
+		//   |Dot(D,N)|*t = -sign(Dot(D,N))*Dot(Q,N)
+		Real DdN = Ray.Direction.Dot(normal);
+		Real sign;
+		if (DdN > TMathUtil<Real>::ZeroTolerance)
+		{
+			sign = (Real)1;
+		}
+		else if (DdN < -TMathUtil<Real>::ZeroTolerance)
+		{
+			sign = (Real)-1;
+			DdN = -DdN;
+		}
+		else
+		{
+			// Ray and triangle are parallel, call it a "no intersection"
+			// even if the ray does intersect.
+			IntersectionType = EIntersectionType::Empty;
+			return false;
+		}
+
+		Real DdQxE2 = sign * Ray.Direction.Dot(diff.Cross(edge2));
+		if (DdQxE2 >= (Real)0)
+		{
+			Real DdE1xQ = sign * Ray.Direction.Dot(edge1.Cross(diff));
+			if (DdE1xQ >= (Real)0)
+			{
+				if (DdQxE2 + DdE1xQ <= DdN)
+				{
+					// Line intersects triangle, check if ray does.
+					Real QdN = -sign * diff.Dot(normal);
+					if (QdN >= (Real)0)
+					{
+						// Ray intersects triangle.
+						IntersectionType = EIntersectionType::Point;
+						return true;
+					}
+					// else: t < 0, no intersection
+				}
+				// else: b1+b2 > 1, no intersection
+			}
+			// else: b2 < 0, no intersection
+		}
+		// else: b1 < 0, no intersection
+
+		IntersectionType = EIntersectionType::Empty;
+		return false;
+	}
+
+
+	/**
+	 * Find intersection point
+	 * @return true if ray intersects triangle
+	 */
+	bool Find()
+	{
+		// Compute the offset origin, edges, and normal.
+		FVector3<Real> diff = Ray.Origin - Triangle.V[0];
+		FVector3<Real> edge1 = Triangle.V[1] - Triangle.V[0];
+		FVector3<Real> edge2 = Triangle.V[2] - Triangle.V[0];
+		FVector3<Real> normal = edge1.Cross(edge2);
+
+		// Solve Q + t*D = b1*E1 + b2*E2 (Q = kDiff, D = ray direction,
+		// E1 = kEdge1, E2 = kEdge2, N = Cross(E1,E2)) by
+		//   |Dot(D,N)|*b1 = sign(Dot(D,N))*Dot(D,Cross(Q,E2))
+		//   |Dot(D,N)|*b2 = sign(Dot(D,N))*Dot(D,Cross(E1,Q))
+		//   |Dot(D,N)|*t = -sign(Dot(D,N))*Dot(Q,N)
+		Real DdN = Ray.Direction.Dot(normal);
+		Real sign;
+		if (DdN > TMathUtil<Real>::ZeroTolerance)
+		{
+			sign = (Real)1;
+		}
+		else if (DdN < -TMathUtil<Real>::ZeroTolerance)
+		{
+			sign = (Real)-1;
+			DdN = -DdN;
+		}
+		else
+		{
+			// Ray and triangle are parallel, call it a "no intersection"
+			// even if the ray does intersect.
+			IntersectionType = EIntersectionType::Empty;
+			return false;
+		}
+
+		Real DdQxE2 = sign * Ray.Direction.Dot(diff.Cross(edge2));
+		if (DdQxE2 >= (Real)0)
+		{
+			Real DdE1xQ = sign * Ray.Direction.Dot(edge1.Cross(diff));
+			if (DdE1xQ >= (Real)0)
+			{
+				if (DdQxE2 + DdE1xQ <= DdN)
+				{
+					// Line intersects triangle, check if ray does.
+					Real QdN = -sign * diff.Dot(normal);
+					if (QdN >= (Real)0)
+					{
+						// Ray intersects triangle.
+						Real inv = ((Real)1) / DdN;
+						RayParameter = QdN * inv;
+						TriangleBaryCoords.Y = DdQxE2 * inv;
+						TriangleBaryCoords.Z = DdE1xQ * inv;
+						TriangleBaryCoords.X = (Real)1 - TriangleBaryCoords.Y - TriangleBaryCoords.Z;
+						IntersectionType = EIntersectionType::Point;
+						return true;
+					}
+					// else: t < 0, no intersection
+				}
+				// else: b1+b2 > 1, no intersection
+			}
+			// else: b2 < 0, no intersection
+		}
+		// else: b1 < 0, no intersection
+
+		IntersectionType = EIntersectionType::Empty;
+		return false;
+	}
+
+};
+
+typedef FIntrRay3Triangle3<float> FIntrRay3Triangle3f;
+typedef FIntrRay3Triangle3<double> FIntrRay3Triangle3d;
\ No newline at end of file
diff --git a/Engine/Plugins/Experimental/GeometryProcessing/Source/GeometricObjects/Public/LineTypes.h b/Engine/Plugins/Experimental/GeometryProcessing/Source/GeometricObjects/Public/LineTypes.h
new file mode 100644
--- /dev/null
+++ b/Engine/Plugins/Experimental/GeometryProcessing/Source/GeometricObjects/Public/LineTypes.h
@@ -0,0 +1,195 @@
+// Copyright 1998-2019 Epic Games, Inc. All Rights Reserved.
+
+// line types ported from WildMagic and geometry3Sharp
+
+#pragma once
+
+#include "Math/UnrealMath.h"
+#include "VectorTypes.h"
+
+/**
+ * FLine2 is a two-dimensional infinite line.
+ * The line is stored in (Center,Direction) form.
+ */
+template<typename T>
+struct FLine2
+{
+	/** Origin / Center Point of Line */
+	FVector2<T> Origin;
+
+	/** Direction of Line, Normalized */
+	FVector2<T> Direction;
+
+	/**
+	 * Construct default line along X axis
+	 */
+	FLine2()
+	{
+		Origin = FVector2<T>::Zero();
+		Direction = FVector2<T>::UnitX();
+	}
+
+	/**
+	 * Construct line with given Origin and Direction
+	 */
+	FLine2(const FVector2<T>& OriginIn, const FVector2<T>& DirectionIn)
+		: Origin(OriginIn), Direction(DirectionIn)
+	{
+	}
+
+
+	/**
+	 * @return line between two points
+	 */
+	static FLine2<T> FromPoints(const FVector2<T>& Point0, const FVector2<T>& Point1)
+	{
+		return FLine2<T>(Point0, (Point1 - Point0).Normalized());
+	}
+
+
+	/**
+	 * @return point on line at given line parameter value (distance along line from origin)
+	 */
+	inline FVector2<T> PointAt(double LineParameter) const
+	{
+		return Origin + LineParameter * Direction;
+	}
+
+
+	/**
+	 * @return line parameter (ie distance from Origin) at nearest point on line to QueryPoint
+	 */
+	inline double Project(const FVector2<T>& QueryPoint) const
+	{
+		return (QueryPoint - Origin).Dot(Direction);
+	}
+
+	/**
+	 * @return smallest squared distance from line to QueryPoint
+	 */
+	inline double DistanceSquared(const FVector2<T>& QueryPoint) const
+	{
+		double t = (QueryPoint - Origin).Dot(Direction);
+		FVector2<T> proj = Origin + t * Direction;
+		return (proj - QueryPoint).SquaredLength();
+	}
+
+	/**
+	 * @return +1 if QueryPoint is "right" of line, -1 if "left" or 0 if "on" line (up to given tolerance)
+	 */
+	inline int WhichSide(const FVector2<T>& QueryPoint, double OnLineTolerance = 0) const
+	{
+		double x0 = QueryPoint.X - Origin.X;
+		double y0 = QueryPoint.Y - Origin.Y;
+		double x1 = Direction.X;
+		double y1 = Direction.Y;
+		double det = x0 * y1 - x1 * y0;
+		return (det > OnLineTolerance ? +1 : (det < -OnLineTolerance ? -1 : 0));
+	}
+
+
+	/**
+	 * Calculate intersection point between this line and another one
+	 * @param OtherLine line to test against
+	 * @param IntersectionPointOut intersection point is stored here, if found
+	 * @param ParallelDotTolerance tolerance used to determine if lines are parallel (and hence cannot intersect)
+	 * @return true if lines intersect and IntersectionPointOut was computed
+	 */
+	bool IntersectionPoint(const FLine2& OtherLine, FVector2<T>& IntersectionPointOut, double ParallelDotTolerance = TMathUtil<T>::ZeroTolerance) const
+	{
+		// see IntrFLine2FLine2 for more detailed explanation
+		FVector2<T> diff = OtherLine.Origin - Origin;
+		double D0DotPerpD1 = Direction.DotPerp(OtherLine.Direction);
+		if (TMathUtil<T>::Abs(D0DotPerpD1) > ParallelDotTolerance)                     // FLines intersect in a single point.
+		{
+			double invD0DotPerpD1 = ((double)1) / D0DotPerpD1;
+			double diffDotPerpD1 = diff.DotPerp(OtherLine.Direction);
+			T s = diffDotPerpD1 * invD0DotPerpD1;
+			IntersectionPointOut = Origin + s * Direction;
+			return true;
+		}
+		// FLines are parallel.
+		return false;
+	}
+};
+
+
+typedef FLine2<double> FLine2d;
+typedef FLine2<float> FLine2f;
+
+
+
+
+
+
+/**
+ * FLine3 is a three-dimensional infinite line.
+ * The line is stored in (Center,Direction) form.
+ */
+template<typename T>
+struct FLine3
+{
+	/** Origin / Center Point of Line */
+	FVector3<T> Origin;
+
+	/** Direction of Line, Normalized */
+	FVector3<T> Direction;
+
+	/**
+	 * Construct default line along X axis
+	 */
+	FLine3()
+	{
+		Origin = FVector3<T>::Zero();
+		Direction = FVector3<T>::UnitX();
+	}
+
+	/**
+	 * Construct line with given Origin and Direction
+	 */
+	FLine3(const FVector3<T>& OriginIn, const FVector3<T>& DirectionIn)
+		: Origin(OriginIn), Direction(DirectionIn)
+	{
+	}
+
+
+	/**
+	 * @return line between two points
+	 */
+	static FLine3<T> FromPoints(const FVector3<T>& Point0, const FVector3<T>& Point1)
+	{
+		return FLine3<T>(Point0, (Point1 - Point0).Normalized());
+	}
+
+
+	/**
+	 * @return point on line at given line parameter value (distance along line from origin)
+	 */
+	inline FVector3<T> PointAt(double LineParameter) const
+	{
+		return Origin + LineParameter * Direction;
+	}
+
+
+	/**
+	 * @return line parameter (ie distance from Origin) at nearest point on line to QueryPoint
+	 */
+	inline double Project(const FVector3<T>& QueryPoint) const
+	{
+		return (QueryPoint - Origin).Dot(Direction);
+	}
+
+	/**
+	 * @return smallest squared distance from line to QueryPoint
+	 */
+	inline double DistanceSquared(const FVector3<T>& QueryPoint) const
+	{
+		double t = (QueryPoint - Origin).Dot(Direction);
+		FVector3<T> proj = Origin + t * Direction;
+		return (proj - QueryPoint).SquaredLength();
+	}
+
+};
+
+typedef FLine3<double> FLine3d;
+typedef FLine3<float> FLine3f;
\ No newline at end of file
