move GeometryProcessing plugin from NotForLicensees to Experimental

diff --git a/Engine/Plugins/Experimental/GeometryProcessing/Source/GeometricObjects/Public/MatrixTypes.h b/Engine/Plugins/Experimental/GeometryProcessing/Source/GeometricObjects/Public/MatrixTypes.h
new file mode 100644
--- /dev/null
+++ b/Engine/Plugins/Experimental/GeometryProcessing/Source/GeometricObjects/Public/MatrixTypes.h
@@ -0,0 +1,419 @@
+// Copyright 1998-2019 Epic Games, Inc. All Rights Reserved.
+
+#pragma once
+
+#include "VectorTypes.h"
+#include "VectorUtil.h"
+
+template <typename RealType>
+struct FMatrix3
+{
+	FVector3<RealType> Row0;
+	FVector3<RealType> Row1;
+	FVector3<RealType> Row2;
+
+	FMatrix3()
+	{
+	}
+
+	FMatrix3(RealType ConstantValue)
+	{
+		Row0 = FVector3<RealType>(ConstantValue, ConstantValue, ConstantValue);
+		Row1 = Row0;
+		Row2 = Row0;
+	}
+
+	FMatrix3(RealType Diag0, RealType Diag1, RealType Diag2)
+	{
+		Row0 = FVector3<RealType>(Diag0, 0, 0);
+		Row1 = FVector3<RealType>(0, Diag1, 0);
+		Row2 = FVector3<RealType>(0, 0, Diag2);
+	}
+
+	/**
+	 * Construct outer-product of U*transpose(V) of U and V
+	 * result is that Mij = u_i * v_j
+	 */
+	FMatrix3<RealType>(const FVector3<RealType>& U, const FVector3<RealType>& V)
+		: Row0(U.X * V.X, U.X * V.Y, U.X * V.Z),
+		  Row1(U.Y * V.X, U.Y * V.Y, U.Y * V.Z),
+		  Row2(U.Z * V.X, U.Z * V.Y, U.Z * V.Z)
+	{
+	}
+
+	FMatrix3(RealType M00, RealType M01, RealType M02, RealType M10, RealType M11, RealType M12, RealType M20, RealType M21, RealType M22)
+		: Row0(M00, M01, M02),
+		  Row1(M10, M11, M12),
+		  Row2(M20, M21, M22)
+	{
+	}
+
+	FMatrix3(const FVector3<RealType>& V1, const FVector3<RealType>& V2, const FVector3<RealType>& V3, bool bRows)
+	{
+		if (bRows)
+		{
+			Row0 = V1;
+			Row1 = V2;
+			Row2 = V3;
+		}
+		else
+		{
+			Row0 = FVector3<RealType>(V1.X, V2.X, V3.X);
+			Row1 = FVector3<RealType>(V1.Y, V2.Y, V3.Y);
+			Row2 = FVector3<RealType>(V1.Z, V2.Z, V3.Z);
+		}
+	}
+
+	static FMatrix3<RealType> Zero()
+	{
+		return FMatrix3<RealType>(0);
+	}
+	static FMatrix3<RealType> Identity()
+	{
+		return FMatrix3<RealType>(1, 1, 1);
+	}
+
+	RealType operator()(int Row, int Col) const
+	{
+		check(Row >= 0 && Row < 3 && Col >= 0 && Col < 3);
+		if (Row == 0)
+		{
+			return Row0[Col];
+		}
+		else if (Row == 1)
+		{
+			return Row1[Col];
+		}
+		else
+		{
+			return Row2[Col];
+		}
+	}
+
+	FMatrix3<RealType> operator*(RealType Scale) const
+	{
+		return FMatrix3<RealType>(
+			Row0.X * Scale, Row0.Y * Scale, Row0.Z * Scale,
+			Row1.X * Scale, Row1.Y * Scale, Row1.Z * Scale,
+			Row2.X * Scale, Row2.Y * Scale, Row2.Z * Scale);
+	}
+
+	FVector3<RealType> operator*(const FVector3<RealType>& V) const
+	{
+		return FVector3<RealType>(
+			Row0.X * V.X + Row0.Y * V.Y + Row0.Z * V.Z,
+			Row1.X * V.X + Row1.Y * V.Y + Row1.Z * V.Z,
+			Row2.X * V.X + Row2.Y * V.Y + Row2.Z * V.Z);
+	}
+
+	FMatrix3<RealType> operator*(const FMatrix3<RealType>& Mat2) const
+	{
+		RealType M00 = Row0.X * Mat2.Row0.X + Row0.Y * Mat2.Row1.X + Row0.Z * Mat2.Row2.X;
+		RealType M01 = Row0.X * Mat2.Row0.Y + Row0.Y * Mat2.Row1.Y + Row0.Z * Mat2.Row2.Y;
+		RealType M02 = Row0.X * Mat2.Row0.Z + Row0.Y * Mat2.Row1.Z + Row0.Z * Mat2.Row2.Z;
+
+		RealType M10 = Row1.X * Mat2.Row0.X + Row1.Y * Mat2.Row1.X + Row1.Z * Mat2.Row2.X;
+		RealType M11 = Row1.X * Mat2.Row0.Y + Row1.Y * Mat2.Row1.Y + Row1.Z * Mat2.Row2.Y;
+		RealType M12 = Row1.X * Mat2.Row0.Z + Row1.Y * Mat2.Row1.Z + Row1.Z * Mat2.Row2.Z;
+
+		RealType M20 = Row2.X * Mat2.Row0.X + Row2.Y * Mat2.Row1.X + Row2.Z * Mat2.Row2.X;
+		RealType M21 = Row2.X * Mat2.Row0.Y + Row2.Y * Mat2.Row1.Y + Row2.Z * Mat2.Row2.Y;
+		RealType M22 = Row2.X * Mat2.Row0.Z + Row2.Y * Mat2.Row1.Z + Row2.Z * Mat2.Row2.Z;
+
+		return FMatrix3<RealType>(M00, M01, M02, M10, M11, M12, M20, M21, M22);
+	}
+
+	FMatrix3<RealType> operator+(const FMatrix3<RealType>& Mat2)
+	{
+		return FMatrix3<RealType>(Row0 + Mat2.Row0, Row1 + Mat2.Row1, Row2 + Mat2.Row2, true);
+	}
+
+	FMatrix3<RealType> operator-(const FMatrix3<RealType>& Mat2)
+	{
+		return FMatrix3<RealType>(Row0 - Mat2.Row0, Row1 - Mat2.Row1, Row2 - Mat2.Row2, true);
+	}
+
+	inline FMatrix3<RealType>& operator*=(const RealType& Scalar)
+	{
+		Row0 *= Scalar;
+		Row1 *= Scalar;
+		Row2 *= Scalar;
+		return *this;
+	}
+
+	inline FMatrix3<RealType>& operator+=(const FMatrix3<RealType>& Mat2)
+	{
+		Row0 += Mat2.Row0;
+		Row1 += Mat2.Row1;
+		Row2 += Mat2.Row2;
+		return *this;
+	}
+
+	RealType InnerProduct(const FMatrix3<RealType>& Mat2) const
+	{
+		return Row0.Dot(Mat2.Row0) + Row1.Dot(Mat2.Row1) + Row2.Dot(Mat2.Row2);
+	}
+
+	RealType Determinant() const
+	{
+		RealType a11 = Row0.X, a12 = Row0.Y, a13 = Row0.Z, a21 = Row1.X, a22 = Row1.Y, a23 = Row1.Z, a31 = Row2.X, a32 = Row2.Y, a33 = Row2.Z;
+		RealType i00 = a33 * a22 - a32 * a23;
+		RealType i01 = -(a33 * a12 - a32 * a13);
+		RealType i02 = a23 * a12 - a22 * a13;
+		return a11 * i00 + a21 * i01 + a31 * i02;
+	}
+
+	FMatrix3<RealType> Inverse() const
+	{
+		RealType a11 = Row0.X, a12 = Row0.Y, a13 = Row0.Z, a21 = Row1.X, a22 = Row1.Y, a23 = Row1.Z, a31 = Row2.X, a32 = Row2.Y, a33 = Row2.Z;
+		RealType i00 = a33 * a22 - a32 * a23;
+		RealType i01 = -(a33 * a12 - a32 * a13);
+		RealType i02 = a23 * a12 - a22 * a13;
+
+		RealType i10 = -(a33 * a21 - a31 * a23);
+		RealType i11 = a33 * a11 - a31 * a13;
+		RealType i12 = -(a23 * a11 - a21 * a13);
+
+		RealType i20 = a32 * a21 - a31 * a22;
+		RealType i21 = -(a32 * a11 - a31 * a12);
+		RealType i22 = a22 * a11 - a21 * a12;
+
+		RealType det = a11 * i00 + a21 * i01 + a31 * i02;
+		ensure(TMathUtil<RealType>::Abs(det) >= TMathUtil<RealType>::Epsilon);
+		det = 1.0 / det;
+		return FMatrix3<RealType>(i00 * det, i01 * det, i02 * det, i10 * det, i11 * det, i12 * det, i20 * det, i21 * det, i22 * det);
+	}
+
+	FMatrix3<RealType> Transpose() const
+	{
+		return FMatrix3<RealType>(
+			Row0.X, Row1.X, Row2.X,
+			Row0.Y, Row1.Y, Row2.Y,
+			Row0.Z, Row1.Z, Row2.Z);
+	}
+
+	bool EpsilonEqual(const FMatrix3<RealType>& Mat2, RealType Epsilon) const
+	{
+		return VectorUtil::EpsilonEqual(Row0, Mat2.Row0, Epsilon) &&
+			   VectorUtil::EpsilonEqual(Row1, Mat2.Row1, Epsilon) &&
+			   VectorUtil::EpsilonEqual(Row2, Mat2.Row2, Epsilon);
+	}
+
+	static FMatrix3<RealType> AxisAngleR(const FVector3<RealType>& Axis, RealType AngleRad)
+	{
+		RealType cs = TMathUtil<RealType>::Cos(AngleRad);
+		RealType sn = TMathUtil<RealType>::Sin(AngleRad);
+		RealType oneMinusCos = 1.0 - cs;
+		RealType x2 = Axis[0] * Axis[0];
+		RealType y2 = Axis[1] * Axis[1];
+		RealType z2 = Axis[2] * Axis[2];
+		RealType xym = Axis[0] * Axis[1] * oneMinusCos;
+		RealType xzm = Axis[0] * Axis[2] * oneMinusCos;
+		RealType yzm = Axis[1] * Axis[2] * oneMinusCos;
+		RealType xSin = Axis[0] * sn;
+		RealType ySin = Axis[1] * sn;
+		RealType zSin = Axis[2] * sn;
+		return FMatrix3<RealType>(
+			x2 * oneMinusCos + cs, xym - zSin, xzm + ySin,
+			xym + zSin, y2 * oneMinusCos + cs, yzm - xSin,
+			xzm - ySin, yzm + xSin, z2 * oneMinusCos + cs);
+	}
+
+	static FMatrix3<RealType> AxisAngleD(const FVector3<RealType>& Axis, RealType AngleDeg)
+	{
+		return AxisAngleR(Axis, TMathUtil<RealType>::DegreesToRadians(AngleDeg));
+	}
+};
+
+
+template <typename RealType>
+struct FMatrix2
+{
+	FVector2<RealType> Row0;
+	FVector2<RealType> Row1;
+
+	FMatrix2()
+	{
+	}
+
+	FMatrix2(RealType ConstantValue)
+	{
+		Row0 = Row1 = FVector2<RealType>(ConstantValue, ConstantValue);
+	}
+
+	FMatrix2(RealType Diag0, RealType Diag1)
+	{
+		Row0 = FVector2<RealType>(Diag0, 0);
+		Row1 = FVector2<RealType>(0, Diag1);
+	}
+
+	/**
+	 * Construct outer-product of U*transpose(V) of U and V
+	 * result is that Mij = u_i * v_j
+	 */
+	FMatrix2<RealType>(const FVector2<RealType>& U, const FVector2<RealType>& V)
+		: Row0(U.X * V.X, U.X * V.Y),
+		  Row1(U.Y * V.X, U.Y * V.Y)
+	{
+	}
+
+	FMatrix2(RealType M00, RealType M01, RealType M10, RealType M11)
+		: Row0(M00, M01),
+		  Row1(M10, M11)
+	{
+	}
+
+	FMatrix2(const FVector2<RealType>& V1, const FVector2<RealType>& V2, bool bRows)
+	{
+		if (bRows)
+		{
+			Row0 = V1;
+			Row1 = V2;
+		}
+		else
+		{
+			Row0 = FVector2<RealType>(V1.X, V2.X);
+			Row1 = FVector2<RealType>(V1.Y, V2.Y);
+		}
+	}
+
+	static FMatrix2<RealType> Zero()
+	{
+		return FMatrix2<RealType>(0);
+	}
+	static FMatrix2<RealType> Identity()
+	{
+		return FMatrix2<RealType>(1, 1);
+	}
+
+	RealType operator()(int Row, int Col) const
+	{
+		check(Row >= 0 && Row < 2 && Col >= 0 && Col < 2);
+		if (Row == 0)
+		{
+			return Row0[Col];
+		}
+		else
+		{
+			return Row1[Col];
+		}
+	}
+
+	FMatrix2<RealType> operator*(RealType Scale) const
+	{
+		return FMatrix2<RealType>(
+			Row0.X * Scale, Row0.Y * Scale,
+			Row1.X * Scale, Row1.Y * Scale);
+	}
+
+	FVector2<RealType> operator*(const FVector2<RealType>& V) const
+	{
+		return FVector2<RealType>(
+			Row0.X * V.X + Row0.Y * V.Y,
+			Row1.X * V.X + Row1.Y * V.Y);
+	}
+
+	FMatrix2<RealType> operator*(const FMatrix2<RealType>& Mat2) const
+	{
+		RealType M00 = Row0.X * Mat2.Row0.X + Row0.Y * Mat2.Row1.X;
+		RealType M01 = Row0.X * Mat2.Row0.Y + Row0.Y * Mat2.Row1.Y;
+
+		RealType M10 = Row1.X * Mat2.Row0.X + Row1.Y * Mat2.Row1.X;
+		RealType M11 = Row1.X * Mat2.Row0.Y + Row1.Y * Mat2.Row1.Y;
+
+		return FMatrix2<RealType>(M00, M01, M10, M11);
+	}
+
+	FMatrix2<RealType> operator+(const FMatrix2<RealType>& Mat2)
+	{
+		return FMatrix2<RealType>(Row0 + Mat2.Row0, Row1 + Mat2.Row1, true);
+	}
+
+	FMatrix2<RealType> operator-(const FMatrix2<RealType>& Mat2)
+	{
+		return FMatrix2<RealType>(Row0 - Mat2.Row0, Row1 - Mat2.Row1, true);
+	}
+
+	inline FMatrix2<RealType>& operator*=(const RealType& Scalar)
+	{
+		Row0 *= Scalar;
+		Row1 *= Scalar;
+		return *this;
+	}
+
+	inline FMatrix2<RealType>& operator+=(const FMatrix2<RealType>& Mat2)
+	{
+		Row0 += Mat2.Row0;
+		Row1 += Mat2.Row1;
+		return *this;
+	}
+
+	RealType InnerProduct(const FMatrix2<RealType>& Mat2) const
+	{
+		return Row0.Dot(Mat2.Row0) + Row1.Dot(Mat2.Row1);
+	}
+
+	RealType Determinant() const
+	{
+		return Row0.X * Row1.Y - Row0.Y * Row1.X;
+	}
+
+	FMatrix2<RealType> Inverse() const
+	{
+		RealType Det = Determinant();
+		ensure(TMathUtil<RealType>::Abs(Det) >= TMathUtil<RealType>::Epsilon);
+		RealType DetInv = 1.0 / Det;
+		return FMatrix2<RealType>(Row1.Y * DetInv, -Row0.Y * DetInv, -Row1.X * DetInv, Row0.X * DetInv);
+	}
+
+	FMatrix2<RealType> Transpose() const
+	{
+		return FMatrix2<RealType>(
+			Row0.X, Row1.X,
+			Row0.Y, Row1.Y);
+	}
+
+	bool EpsilonEqual(const FMatrix2<RealType>& Mat2, RealType Epsilon) const
+	{
+		return VectorUtil::EpsilonEqual(Row0, Mat2.Row0, Epsilon) &&
+			   VectorUtil::EpsilonEqual(Row1, Mat2.Row1, Epsilon);
+	}
+
+	static FMatrix2<RealType> RotationRad(RealType AngleRad)
+	{
+		RealType cs = TMathUtil<RealType>::Cos(AngleRad);
+		RealType sn = TMathUtil<RealType>::Sin(AngleRad);
+		return FMatrix2<RealType>(cs, -sn, sn, cs);
+	}
+
+	/**
+	 * Assumes we have a rotation matrix (uniform scale ok)
+	 */
+	RealType GetAngleRad()
+	{
+		return TMathUtil<RealType>::Atan2(Row1.X, Row0.X);
+	}
+
+};
+
+template <typename RealType>
+inline FMatrix3<RealType> operator*(RealType Scale, const FMatrix3<RealType>& Mat)
+{
+	return FMatrix3<RealType>(
+		Mat.Row0.X * Scale, Mat.Row0.Y * Scale, Mat.Row0.Z * Scale,
+		Mat.Row1.X * Scale, Mat.Row1.Y * Scale, Mat.Row1.Z * Scale,
+		Mat.Row2.X * Scale, Mat.Row2.Y * Scale, Mat.Row2.Z * Scale);
+}
+
+template <typename RealType>
+inline FMatrix2<RealType> operator*(RealType Scale, const FMatrix2<RealType>& Mat)
+{
+	return FMatrix2<RealType>(
+		Mat.Row0.X * Scale, Mat.Row0.Y * Scale,
+		Mat.Row1.X * Scale, Mat.Row1.Y * Scale);
+}
+
+typedef FMatrix3<float> FMatrix3f;
+typedef FMatrix3<double> FMatrix3d;
+typedef FMatrix2<float> FMatrix2f;
+typedef FMatrix2<double> FMatrix2d;
\ No newline at end of file
diff --git a/Engine/Plugins/Experimental/GeometryProcessing/Source/GeometricObjects/Public/MeshAdapter.h b/Engine/Plugins/Experimental/GeometryProcessing/Source/GeometricObjects/Public/MeshAdapter.h
new file mode 100644
--- /dev/null
+++ b/Engine/Plugins/Experimental/GeometryProcessing/Source/GeometricObjects/Public/MeshAdapter.h
@@ -0,0 +1,48 @@
+// Copyright 1998-2019 Epic Games, Inc. All Rights Reserved.
+
+#pragma once
+
+#include "TriangleTypes.h"
+#include "VectorTypes.h"
+
+#include "Math/IntVector.h"
+#include "Math/Vector.h"
+
+template <typename RealType>
+struct TTriangleMeshAdapter
+{
+	TFunction<bool(int32 index)> IsTriangle;
+	TFunction<bool(int32 index)> IsVertex;
+	TFunction<int32()> MaxTriangleID;
+	TFunction<int32()> MaxVertexID;
+	TFunction<int32()> TriangleCount;
+	TFunction<int32()> VertexCount;
+	TFunction<int32()> GetShapeTimestamp;
+	TFunction<FIndex3i(int32)> GetTriangle;
+	TFunction<FVector3<RealType>(int32)> GetVertex;
+
+	inline void GetTriVertices(int TID, FVector3<RealType>& V0, FVector3<RealType>& V1, FVector3<RealType>& V2) const
+	{
+		FIndex3i TriIndices = GetTriangle(TID);
+		V0 = GetVertex(TriIndices.A);
+		V1 = GetVertex(TriIndices.B);
+		V2 = GetVertex(TriIndices.C);
+	}
+};
+
+typedef TTriangleMeshAdapter<double> TTriangleMeshAdapterd;
+typedef TTriangleMeshAdapter<float> TTriangleMeshAdapterf;
+
+TTriangleMeshAdapterd GetArrayMesh(TArray<FVector>& Vertices, TArray<FIntVector>& Triangles)
+{
+	return {
+		[&](int) { return true; },
+		[&](int) { return true; },
+		[&]() { return Triangles.Num(); },
+		[&]() { return Vertices.Num(); },
+		[&]() { return Triangles.Num(); },
+		[&]() { return Vertices.Num(); },
+		[&]() { return 0; },
+		[&](int Idx) { return FIndex3i(Triangles[Idx]); },
+		[&](int Idx) { return FVector3d(Vertices[Idx]); }};
+}
diff --git a/Engine/Plugins/Experimental/GeometryProcessing/Source/GeometricObjects/Public/OrientedBoxTypes.h b/Engine/Plugins/Experimental/GeometryProcessing/Source/GeometricObjects/Public/OrientedBoxTypes.h
new file mode 100644
--- /dev/null
+++ b/Engine/Plugins/Experimental/GeometryProcessing/Source/GeometricObjects/Public/OrientedBoxTypes.h
@@ -0,0 +1,222 @@
+// Copyright 1998-2019 Epic Games, Inc. All Rights Reserved.
+
+// ported from geometry3Sharp Box3
+
+#pragma once
+
+#include "VectorTypes.h"
+#include "BoxTypes.h"
+#include "FrameTypes.h"
+
+/**
+ * FOrientedBox3 is a non-axis-aligned 3D box defined by a 3D frame and extents along the axes of that frame
+ * The frame is at the center of the box.
+ */
+template<typename RealType>
+struct FOrientedBox3
+{
+	// available for porting: ContainPoint, MergeBoxes()
+
+
+	/** 3D position (center) and orientation (axes) of the box */
+	FFrame3<RealType> Frame;
+	/** Half-dimensions of box measured along the three axes */
+	FVector3<RealType> Extents;
+
+	FOrientedBox3() : Extents(1,1,1) {}
+
+	/**
+	 * Create axis-aligned box with given Origin and Extents
+	 */
+	FOrientedBox3(const FVector3<RealType>& Origin, const FVector3<RealType> & ExtentsIn)
+		: Frame(Origin), Extents(ExtentsIn)
+	{
+	}
+
+	/**
+	 * Create oriented box with given Frame and Extents
+	 */
+	FOrientedBox3(const FFrame3<RealType>& FrameIn, const FVector3<RealType> & ExtentsIn)
+		: Frame(FrameIn), Extents(ExtentsIn)
+	{
+	}
+
+	/**
+	 * Create oriented box from axis-aligned box
+	 */
+	FOrientedBox3(const FAxisAlignedBox3<RealType>& AxisBox)
+		: Frame(AxisBox.Center()), Extents((RealType)0.5 * AxisBox.Diagonal())
+	{
+	}
+
+
+	/** @return box with unit dimensions centered at origin */
+	static FOrientedBox3<RealType> UnitZeroCentered() { return FOrientedBox3<RealType>(FVector3<RealType>::Zero(), (RealType)0.5*FVector3<RealType>::One()); }
+
+	/** @return box with unit dimensions where minimum corner is at origin */
+	static FOrientedBox3<RealType> UnitPositive() { return FOrientedBox3<RealType>((RealType)0.5*FVector3<RealType>::One(), (RealType)0.5*FVector3<RealType>::One()); }
+
+
+	/** @return center of the box */
+	FVector3<RealType> Center() const { return Frame.Origin; }
+
+	/** @return X axis of the box */
+	FVector3<RealType> AxisX() const { return Frame.X(); }
+
+	/** @return Y axis of the box */
+	FVector3<RealType> AxisY() const { return Frame.Y(); }
+
+	/** @return Z axis of the box */
+	FVector3<RealType> AxisZ() const { return Frame.Z(); }
+
+	/** @return an axis of the box */
+	FVector3<RealType> GetAxis(int AxisIndex) const { return Frame.GetAxis(AxisIndex); }
+
+	/** @return maximum extent of box */
+	inline RealType MaxExtent() const
+	{
+		return Extents.MaxAbs();
+	}
+
+	/** @return minimum extent of box */
+	inline RealType MinExtent() const
+	{
+		return Extents.MinAbs();
+	}
+
+	/** @return vector from minimum-corner to maximum-corner of box */
+	inline FVector3<RealType> Diagonal() const
+	{
+		return Frame.PointAt(Extents.X, Extents.Y, Extents.Z) - Frame.PointAt(-Extents.X, -Extents.Y, -Extents.Z);
+	}
+
+	/** @return volume of box */
+	inline RealType Volume() const
+	{
+		return (RealType)8 * (Extents.X) * (Extents.Y) * (Extents.Z);
+	}
+
+	/** @return true if box contains point */
+	inline bool Contains(const FVector3<RealType>& Point) 
+	{
+		FVector3<RealType> InFramePoint = Frame.ToFramePoint(Point);
+		return (TMathUtil<RealType>::Abs(InFramePoint.X) <= Extents.X) &&
+			(TMathUtil<RealType>::Abs(InFramePoint.Y) <= Extents.Y) &&
+			(TMathUtil<RealType>::Abs(InFramePoint.Z) <= Extents.Z);
+	}
+
+
+	// corners [ (-x,-y), (x,-y), (x,y), (-x,y) ], -z, then +z
+	//
+	//   7---6     +z       or        3---2     -z
+	//   |\  |\                       |\  |\
+    //   4-\-5 \                      0-\-1 \
+    //    \ 3---2                      \ 7---6   
+	//     \|   |                       \|   |
+	//      0---1  -z                    4---5  +z
+	//
+	// @todo does this ordering make sense for UE? we are in LHS instead of RHS here
+	// if this is modified, likely need to update IndexUtil::BoxFaces and BoxFaceNormals
+
+	/**
+	 * @param Index corner index in range 0-7
+	 * @return Corner point on the box identified by the given index. See diagram in OrientedBoxTypes.h for index/corner mapping.
+	 */
+	FVector3<RealType> GetCorner(int Index) const
+	{
+		check(Index >= 0 && Index <= 7);
+		RealType dx = (((Index & 1) != 0) ^ ((Index & 2) != 0)) ? (Extents.X) : (-Extents.X);
+		RealType dy = ((Index / 2) % 2 == 0) ? (-Extents.Y) : (Extents.Y);
+		RealType dz = (Index < 4) ? (-Extents.Z) : (Extents.Z);
+		return Frame.PointAt(dx, dy, dz);
+	}
+
+
+
+
+
+	/**
+	 * Find squared distance to box.
+	 * @param Point input point
+	 * @return squared distance from point to box, or 0 if point is inside box
+	 */
+	RealType DistanceSquared(FVector3<RealType> Point)
+	{
+		 // Ported from WildMagic5 Wm5DistPoint3Box3.cpp
+			 
+		 // Work in the box's coordinate system.
+		 Point -= Frame.Origin;
+
+		// Compute squared distance and closest point on box.
+		RealType sqrDistance = 0;
+		RealType delta;
+		FVector3<RealType> closest(0, 0, 0);
+		for (int i = 0; i < 3; ++i) 
+		{
+			closest[i] = Point.Dot(GetAxis(i));
+			if (closest[i] < -Extents[i]) 
+			{
+				delta = closest[i] + Extents[i];
+				sqrDistance += delta * delta;
+				closest[i] = -Extents[i];
+			}
+			else if (closest[i] > Extents[i])
+			{
+				delta = closest[i] - Extents[i];
+				sqrDistance += delta * delta;
+				closest[i] = Extents[i];
+			}
+		}
+
+		return sqrDistance;
+	}
+
+
+
+
+	 /**
+	  * Find closest point on box
+	  * @param Point input point
+	  * @return closest point on box. Input point is returned if it is inside box.
+	  */
+	 FVector3<RealType> ClosestPoint(FVector3<RealType> Point)
+	 {
+		 // Work in the box's coordinate system.
+		 Point -= Frame.Origin;
+
+		 // Compute squared distance and closest point on box.
+		 RealType sqrDistance = 0;
+		 RealType delta;
+		 FVector3<RealType> closest;
+		 FVector3<RealType> Axes[3];
+		 for (int i = 0; i < 3; ++i) 
+		 {
+			 Axes[i] = GetAxis(i);
+			 closest[i] = Point.Dot(Axes[i]);
+			 RealType extent = Extents[i];
+			 if (closest[i] < -extent) 
+			 {
+				 delta = closest[i] + extent;
+				 sqrDistance += delta * delta;
+				 closest[i] = -extent;
+			 }
+			 else if (closest[i] > extent) 
+			 {
+				 delta = closest[i] - extent;
+				 sqrDistance += delta * delta;
+				 closest[i] = extent;
+			 }
+		 }
+
+		 return Frame.Origin + closest.X*Axes[0] + closest.Y*Axes[1] + closest.Z*Axes[2];
+	 }
+
+
+
+};
+
+
+
+
+typedef FOrientedBox3<float> FOrientedBox3f;
+typedef FOrientedBox3<double> FOrientedBox3d;
diff --git a/Engine/Plugins/Experimental/GeometryProcessing/Source/GeometricObjects/Public/Polygon2.h b/Engine/Plugins/Experimental/GeometryProcessing/Source/GeometricObjects/Public/Polygon2.h
new file mode 100644
--- /dev/null
+++ b/Engine/Plugins/Experimental/GeometryProcessing/Source/GeometricObjects/Public/Polygon2.h
@@ -0,0 +1,1045 @@
+// Copyright 1998-2019 Epic Games, Inc. All Rights Reserved.
+
+// port of geometry3Sharp Polygon
+
+#pragma once
+
+#include "Templates/UnrealTemplate.h"
+#include "Math/UnrealMath.h"
+#include "VectorTypes.h"
+#include "BoxTypes.h"
+#include "SegmentTypes.h"
+#include "LineTypes.h"
+#include "MathUtil.h"
+#include "Intersection/IntrSegment2Segment2.h"
+#include "Util/DynamicVector.h"
+
+/**
+ * FPolygon2 is a 2D polygon represented as a list of Vertices.
+ * 
+ * @todo move operators
+ */
+template<typename T>
+class FPolygon2
+{
+protected:
+	/** The list of vertices/corners of the polygon */
+	TArray<FVector2<T>> Vertices;
+
+	/** A counter that is incremented every time the polygon vertices are modified */
+	int Timestamp;
+
+public:
+
+	FPolygon2() : Timestamp(0)
+	{
+	}
+
+	/**
+	 * Construct polygon that is a copy of another polygon
+	 */
+	FPolygon2(const FPolygon2& Copy) : Vertices(Copy.Vertices), Timestamp(Copy.Timestamp)
+	{
+	}
+
+	/**
+	 * Construct polygon with given list of vertices
+	 */
+	FPolygon2(const TArray<FVector2<T>>& VertexList) : Vertices(VertexList), Timestamp(0)
+	{
+	}
+
+	/** @return the Timestamp for the polygon, which is updated every time the polygon is modified */
+	int GetTimestamp() const 
+	{
+		return Timestamp;
+	}
+
+	/**
+	 * Get the vertex at a given index
+	 */
+	const FVector2<T>& operator[](int Index) const
+	{
+		return Vertices[Index];
+	}
+
+	/**
+	 * Get the vertex at a given index
+	 * @warning changing the vertex via this operator does not update Timestamp!
+	 */
+	FVector2<T>& operator[](int Index)
+	{
+		return Vertices[Index];
+	}
+
+
+	/**
+	 * @return first vertex of Polygon
+	 */
+	const FVector2<T>& Start() const
+	{
+		return Vertices[0];
+	}
+
+	/**
+	 * @return list of Vertices of Polygon
+	 */
+	const TArray<FVector2<T>>& GetVertices() const
+	{
+		return Vertices;
+	}
+
+	/**
+	 * @return number of Vertices in Polygon
+	 */
+	int VertexCount() const
+	{
+		return Vertices.Num();
+	}
+
+	/**
+	 * Add a vertex to the Polygon
+	 */
+	void AppendVertex(const FVector2<T>& Position)
+	{
+		Vertices.Add(Position);
+		Timestamp++;
+	}
+
+	/**
+	 * Add a list of Vertices to the Polygon
+	 */
+	void AppendVertices(const TArray<FVector2<T>>& NewVertices)
+	{
+		Vertices.Append(NewVertices);
+		Timestamp++;
+	}
+
+	/**
+	 * Set vertex at given index to a new Position
+	 */
+	void Set(int VertexIndex, const FVector2<T>& Position)
+	{
+		Vertices[VertexIndex] = Position;
+		Timestamp++;
+	}
+
+	/**
+	 * Remove a vertex of the Polygon (existing Vertices are shifted)
+	 */
+	void RemoveVertex(int VertexIndex)
+	{
+		Vertices.RemoveAt(VertexIndex);
+		Timestamp++;
+	}
+
+	/**
+	 * Replace the list of Vertices with a new list
+	 */
+	void SetVertices(const TArray<FVector2<T>>& NewVertices)
+	{
+		Vertices = NewVertices;
+		Timestamp++;
+	}
+
+
+	/**
+	 * Reverse the order of the Vertices in the Polygon (ie switch between Clockwise and CounterClockwise)
+	 */
+	void Reverse()
+	{
+		int32 j = Vertices.Num()-1;
+		for (int32 VertexIndex = 0; VertexIndex < j; VertexIndex++, j--)
+		{
+			Swap(Vertices[VertexIndex], Vertices[j]);
+		}
+		Timestamp++;
+	}
+
+
+	/**
+	 * Get the tangent vector at a vertex of the polygon, which is the normalized
+	 * vector from the previous vertex to the next vertex
+	 */
+	FVector2<T> GetTangent(int VertexIndex) const
+	{
+		FVector2<T> next = Vertices[(VertexIndex + 1) % Vertices.Num()];
+		FVector2<T> prev = Vertices[VertexIndex == 0 ? Vertices.Num() - 1 : VertexIndex - 1];
+		return (next - prev).Normalized();
+	}
+
+
+	/**
+	 * Get the normal vector at a vertex of the polygon, which is perpendicular to GetTangent()
+	 * Points "inward" for a Clockwise Polygon, and outward for CounterClockwise
+	 */
+	FVector2<T> GetNormal(int VertexIndex) const
+	{
+		return GetTangent(VertexIndex).Perp();
+	}
+
+
+	/**
+	 * Construct a normal at a vertex of the Polygon by averaging the adjacent face normals.
+	 * This vector is independent of the lengths of the adjacent segments.
+	 * Points "inward" for a Clockwise Polygon, and outward for CounterClockwise
+	 */
+	FVector2<T> GetNormal_FaceAvg(int VertexIndex) const
+	{
+		FVector2<T> next = Vertices[(VertexIndex + 1) % Vertices.Num()];
+		FVector2<T> prev = Vertices[VertexIndex == 0 ? Vertices.Num() - 1 : VertexIndex - 1];
+		next -= Vertices[VertexIndex]; next.Normalize();
+		prev -= Vertices[VertexIndex]; prev.Normalize();
+
+		FVector2<T> n = (next.Perp() - prev.Perp());
+		double len = n.Normalize();
+		if (len == 0) 
+		{
+			return (next + prev).Normalized();   // this gives right direction for degenerate angle
+		}
+		else 
+		{
+			return n;
+		}
+	}
+
+
+	/**
+	 * @return the bounding box of the Polygon Vertices
+	 */
+	FAxisAlignedBox2<T> Bounds() const
+	{
+		FAxisAlignedBox2<T> box = FAxisAlignedBox2<T>::Empty();
+		box.Contain(Vertices);
+		return box;
+	}
+
+	
+	/**
+	 * SegmentIterator is used to iterate over the FSegment2<T> segments of the polygon
+	 */
+	class SegmentIterator 
+	{
+	public:
+		inline bool operator!()
+		{
+			return i < polygon->VertexCount();
+		}
+		inline FSegment2<T> operator*() const
+		{
+			check(polygon != nullptr && i < polygon->VertexCount());
+			return FSegment2<T>(polygon->Vertices[i], polygon->Vertices[i+1 % polygon->VertexCount()]);
+		}
+		//inline FSegment2<T> & operator*();
+		inline SegmentIterator & operator++() 		// prefix
+		{
+			i++;
+			return *this;
+		}
+		inline SegmentIterator operator++(int) 		// postfix
+		{
+			SegmentIterator copy(*this);
+			i++;
+			return copy;
+		}
+		inline bool operator==(const SegmentIterator & i2) { return i2.polygon == polygon && i2.i == i; }
+		inline bool operator!=(const SegmentIterator & i2) { return i2.polygon != polygon || i2.i != i; }
+	protected:
+		const FPolygon2 * polygon;
+		int i;
+		inline SegmentIterator(const FPolygon2 * p, int iCur) : polygon(p), i(iCur) {}
+		friend class FPolygon2;
+	};
+	friend class SegmentIterator;
+
+	SegmentIterator SegmentItr() const 
+	{
+		return SegmentIterator(this, 0);
+	}
+
+	/**
+	 * Wrapper around SegmentIterator that has begin() and end() suitable for range-based for loop
+	 */
+	class SegmentEnumerable
+	{
+	public:
+		const FPolygon2<T>* polygon;
+		SegmentEnumerable() : polygon(nullptr) {}
+		SegmentEnumerable(const FPolygon2<T> * p) : polygon(p) {}
+		SegmentIterator begin() { return polygon->SegmentItr(); }
+		SegmentIterator end() { return SegmentIterator(polygon, polygon->VertexCount()); }
+	};
+
+	/**
+	 * @return an object that can be used in a range-based for loop to iterate over the Segments of the Polygon
+	 */
+	SegmentEnumerable Segments() const
+	{
+		return SegmentEnumerable(this);
+	}
+
+
+	/**
+	 * @return true if the Polygon Vertices have Clockwise winding order / orientation  (signed area is negative)
+	 */
+	bool IsClockwise() const
+	{
+		return SignedArea() < 0;
+	}
+
+	/**
+	 * @return the signed area of the Polygon
+	 */
+	double SignedArea() const
+	{
+		double fArea = 0;
+		int N = Vertices.Num();
+		if (N == 0)
+		{
+			return 0;
+		}
+		for (int i = 0; i < N; ++i) 
+		{
+			const FVector2<T>& v1 = Vertices[i];
+			const FVector2<T>& v2 = Vertices[(i + 1) % N];
+			fArea += v1.X * v2.Y - v1.Y * v2.X;
+		}
+		return fArea * 0.5;
+	}
+
+	/**
+	 * @return the unsigned area of the Polygon
+	 */
+	double Area() const
+	{
+		return TMathUtil<T>::Abs(SignedArea());
+	}
+
+	/**
+	 * @return the total perimeter length of the Polygon
+	 */
+	double Perimeter() const
+	{
+		double fPerim = 0;
+		int N = Vertices.Num();
+		for (int i = 0; i < N; ++i)
+		{
+			fPerim += Vertices[i].Distance(Vertices[(i + 1) % N]);
+		}
+		return fPerim;
+	}
+
+
+	/**
+	 * Get the previous and next vertex positions for a given vertex of the Polygon
+	 */
+	void NeighbourPoints(int iVertex, FVector2<T> &PrevNbrOut, FVector2<T> &NextNbrOut) const
+	{
+		int N = Vertices.Num();
+		PrevNbrOut = Vertices[(iVertex == 0) ? N - 1 : iVertex - 1];
+		NextNbrOut = Vertices[(iVertex + 1) % N];
+	}
+
+
+	/**
+	 * Get the vectors from a given vertex to the previous and next Vertices, optionally normalized
+	 */
+	void NeighbourVectors(int iVertex, FVector2<T> &ToPrevOut, FVector2<T> &ToNextOut, bool bNormalize = false) const
+	{
+		int N = Vertices.Num();
+		ToPrevOut = Vertices[(iVertex == 0) ? N - 1 : iVertex - 1] - Vertices[iVertex];
+		ToNextOut = Vertices[(iVertex + 1) % N] - Vertices[iVertex];
+		if (bNormalize) 
+		{
+			ToPrevOut.Normalize();
+			ToNextOut.Normalize();
+		}
+	}
+
+
+	/**
+	 * @return the opening angle in degrees at a vertex of the Polygon
+	 */
+	double OpeningAngleDeg(int iVertex) const
+	{
+		FVector2<T> e0, e1;
+		NeighbourVectors(iVertex, e0, e1, true);
+		return e0.AngleD(e1);
+	}
+
+
+	/**
+	 * @return analytic winding integral for this Polygon at an arbitrary point
+	 */
+	double WindingIntegral(const FVector2<T>& QueryPoint) const
+	{
+		double sum = 0;
+		int N = Vertices.Num();
+		FVector2<T> a = Vertices[0] - QueryPoint, b = FVector2<T>::Zero();
+		for (int i = 0; i < N; ++i) 
+		{
+			b = Vertices[(i + 1) % N] - QueryPoint;
+			sum += TMathUtil<T>::Atan2(a.X * b.Y - a.Y * b.X, a.X * b.X + a.Y * b.Y);
+			a = b;
+		}
+		return sum / FMathd::TwoPi;
+	}
+
+
+	/**
+	 * @return true if the given query point is inside the Polygon, based on the winding integral
+	 */
+	bool Contains(const FVector2<T>& QueryPoint) const
+	{
+		int nWindingNumber = 0;
+
+		int N = Vertices.Num();
+		FVector2<T> a = Vertices[0], b = FVector2<T>::Zero();
+		for (int i = 0; i < N; ++i) 
+		{
+			b = Vertices[(i + 1) % N];
+
+			if (a.Y <= QueryPoint.Y)     // y <= P.Y (below)
+			{
+				if (b.Y > QueryPoint.Y)									// an upward crossing
+				{
+					if (FVector2<T>::Orient(a, b, QueryPoint) > 0)  // P left of edge
+						++nWindingNumber;                       // have a valid up intersect
+				}
+			}
+			else     // y > P.Y  (above)
+			{
+				if (b.Y <= QueryPoint.Y)									// a downward crossing
+				{
+					if (FVector2<T>::Orient(a, b, QueryPoint) < 0)  // P right of edge
+					{
+						--nWindingNumber;						// have a valid down intersect
+					}
+				}
+			}
+			a = b;
+		}
+		return nWindingNumber != 0;
+	}
+
+
+	/**
+	 * @return true if the Polygon fully contains the OtherPolygon
+	 */
+	bool Contains(const FPolygon2<T>& OtherPoly) const
+	{
+		// @todo fast bbox check?
+
+		int N = OtherPoly.VertexCount();
+		for (int i = 0; i < N; ++i) 
+		{
+			if (Contains(OtherPoly[i]) == false)
+			{
+				return false;
+			}
+		}
+
+		if (Intersects(OtherPoly))
+		{
+			return false;
+		}
+
+		return true;
+	}
+
+
+	/**
+	 * @return true if the Segment is fully contained inside the Polygon
+	 */
+	bool Contains(const FSegment2<T>& Segment) const
+	{
+		// [TODO] Add bbox check
+		if (Contains(Segment.StartPoint()) == false || Contains(Segment.EndPoint()) == false)
+		{
+			return false;
+		}
+
+		for (FSegment2<T> seg : Segments())
+		{
+			if (seg.Intersects(Segment))
+			{
+				return false;
+			}
+		}
+		return true;
+	}
+
+
+	/**
+	 * @return true if at least one edge of the OtherPolygon intersects the Polygon
+	 */
+	bool Intersects(const FPolygon2<T>& OtherPoly) const
+	{
+		if (!Bounds().Intersects(OtherPoly.Bounds()))
+		{
+			return false;
+		}
+
+		for (FSegment2<T> seg : Segments()) 
+		{
+			for (FSegment2<T> oseg : OtherPoly.Segments()) 
+			{
+				if (seg.Intersects(oseg))
+				{
+					return true;
+				}
+			}
+		}
+		return false;
+	}
+
+
+	/**
+	 * @return true if the Segment intersects an edge of the Polygon
+	 */
+	bool Intersects(const FSegment2<T>& Segment) const
+	{
+		// [TODO] Add bbox check
+		if (Contains(Segment.StartPoint()) == true || Contains(Segment.EndPoint()) == true)
+		{
+			return true;
+		}
+
+		// [TODO] Add bbox check
+		for (FSegment2<T> seg : Segments())
+		{
+			if (seg.Intersects(Segment))
+			{
+				return true;
+			}
+		}
+		return false;
+	}
+
+
+	/**
+	 * Find all the points where an edge of the Polygon intersects an edge of the OtherPolygon
+	 * @param OtherPoly polygon to test against
+	 * @param OutArray intersection points are stored here
+	 * @return true if any intersections were found
+	 */
+	bool FindIntersections(const FPolygon2<T>& OtherPoly, TArray<FVector2<T>>& OutArray) const
+	{
+		if (!Bounds().Intersects(OtherPoly.Bounds()))
+		{
+			return false;
+		}
+
+		bool bFoundIntersections = false;
+		for (FSegment2<T> seg : Segments()) 
+		{
+			for (FSegment2<T> oseg : OtherPoly.Segments())
+			{
+				// this computes test twice for intersections, but seg.intersects doesn't
+				// create any new objects so it should be much faster for majority of segments (should profile!)
+				if (seg.Intersects(oseg)) 
+				{
+					//@todo can we replace with something like seg.intersects?
+					FIntrSegment2Segment2<T> intr(seg, oseg);
+					if (intr.Find()) 
+					{
+						bFoundIntersections = true;
+						OutArray.Add(intr.Point0);
+						if (intr.Quantity == 2) 
+						{
+							OutArray.Add(intr.Point1);
+						}
+					}
+				}
+			}
+		}
+
+		return bFoundIntersections;
+	}
+
+
+	/**
+	 * @return edge of the polygon starting at vertex SegmentIndex
+	 */
+	FSegment2<T> Segment(int SegmentIndex) const
+	{
+		return FSegment2<T>(Vertices[SegmentIndex], Vertices[(SegmentIndex + 1) % Vertices.Num()]);
+	}
+
+	/**
+	 * @param SegmentIndex index of first vertex of the edge
+	 * @param SegmentParam parameter in range [0,1] along segment
+	 * @return point on the segment at the given parameter value
+	 */
+	FVector2<T> PointAt(int SegmentIndex, double SegmentParam) const
+	{
+		FSegment2<T> seg(Vertices[SegmentIndex], Vertices[(SegmentIndex + 1) % Vertices.Num()]);
+		return seg.PointAt(SegmentParam);
+	}
+
+
+	/**
+	 * @param SegmentIndex index of first vertex of the edge
+	 * @param SegmentParam parameter in range [0,1] along segment
+	 * @return interpolated normal to the segment at the given parameter value
+	 */
+	FVector2<T> GetNormal(int iSeg, double SegmentParam) const
+	{
+		FSegment2<T> seg(Vertices[iSeg], Vertices[(iSeg + 1) % Vertices.Num()]);
+		T t = ((SegmentParam / seg.Extent) + 1.0) / 2.0;
+
+		FVector2<T> n0 = GetNormal(iSeg);
+		FVector2<T> n1 = GetNormal((iSeg + 1) % Vertices.Num());
+		return ((T(1) - t) * n0 + t * n1).Normalized();
+	}
+
+
+
+	/**
+	 * Calculate the squared distance from a point to the polygon
+	 * @param QueryPoint the query point
+	 * @param NearestSegIndexOut The index of the nearest segment
+	 * @param NearestSegParamOut the parameter value of the nearest point on the segment
+	 * @return squared distance to the polygon
+	 */
+	double DistanceSquared(const FVector2<T>& QueryPoint, int& NearestSegIndexOut, double& NearestSegParamOut) const
+	{
+		NearestSegIndexOut = -1;
+		NearestSegParamOut = TNumericLimits<double>::Max();
+		double dist = TNumericLimits<double>::Max();
+		int N = Vertices.Num();
+		for (int vi = 0; vi < N; ++vi) 
+		{
+			// @todo can't we just use segment function here now?
+			FSegment2<T> seg = FSegment2<T>(Vertices[vi], Vertices[(vi + 1) % N]);
+			double t = (QueryPoint - seg.Center).Dot(seg.Direction);
+			double d = TNumericLimits<double>::Max();
+			if (t >= seg.Extent)
+			{
+				d = seg.EndPoint().DistanceSquared(QueryPoint);
+			}
+			else if (t <= -seg.Extent)
+			{
+				d = seg.StartPoint().DistanceSquared(QueryPoint);
+			}
+			else
+			{
+				d = (seg.PointAt(t) - QueryPoint).SquaredLength();
+			}
+			if (d < dist) 
+			{
+				dist = d;
+				NearestSegIndexOut = vi;
+				NearestSegParamOut = t;
+			}
+		}
+		return dist;
+	}
+
+
+	/**
+	 * Calculate the squared distance from a point to the polygon
+	 * @param QueryPoint the query point
+	 * @return squared distance to the polygon
+	 */
+	double DistanceSquared(const FVector2<T>& QueryPoint) const
+	{
+		int seg; double segt;
+		return DistanceSquared(QueryPoint, seg, segt);
+	}
+
+
+	/**
+	 * @return average edge length of all the edges of the Polygon
+	 */
+	double AverageEdgeLength() const
+	{
+		double avg = 0; int N = Vertices.Num();
+		for (int i = 1; i < N; ++i) {
+			avg += Vertices[i].Distance(Vertices[i - 1]);
+		}
+		avg += Vertices[N - 1].Distance(Vertices[0]);
+		return avg / N;
+	}
+
+
+	/**
+	 * Translate the polygon
+	 * @returns the Polygon, so that you can chain calls like Translate().Scale()
+	 */
+	FPolygon2<T>& Translate(const FVector2<T>& Translate) 
+	{
+		int N = Vertices.Num();
+		for (int i = 0; i < N; ++i)
+		{
+			Vertices[i] += Translate;
+		}
+		Timestamp++;
+		return *this;
+	}
+
+	/**
+	 * Scale the Polygon relative to a given point
+	 * @returns the Polygon, so that you can chain calls like Translate().Scale()
+	 */
+	FPolygon2<T>& Scale(const FVector2<T>& Scale, const FVector2<T>& Origin)
+	{
+		int N = Vertices.Num();
+		for (int i = 0; i < N; ++i)
+		{
+			Vertices[i] = Scale * (Vertices[i] - Origin) + Origin;
+		}
+		Timestamp++;
+		return *this;
+	}
+
+
+	/**
+	 * Apply an arbitrary transformation to the Polygon
+	 * @returns the Polygon, so that you can chain calls like Translate().Scale()
+	 */
+	FPolygon2<T>& Transform(const TFunction<FVector2<T> (const FVector2<T>&)>& TransformFunc)
+	{
+		int N = Vertices.Num();
+		for (int i = 0; i < N; ++i) 
+		{
+			Vertices[i] = TransformFunc(Vertices[i]);
+		}
+		Timestamp++;
+		return *this;
+	}
+
+
+
+
+	/**
+	 * Offset each point by the given Distance along vertex "normal" direction
+	 * @param OffsetDistance the distance to offset
+	 * @param bUseFaceAvg if true, we offset by the average-face normal instead of the perpendicular-tangent normal
+	 */
+	void VtxNormalOffset(T OffsetDistance, bool bUseFaceAvg = false)
+	{
+		TArray<FVector2<T>> NewVertices;
+		NewVertices.SetNumUninitialized(Vertices.Num());
+		if (bUseFaceAvg) 
+		{
+			for (int k = 0; k < Vertices.Num(); ++k)
+			{
+				NewVertices[k] = Vertices[k] + OffsetDistance * GetNormal_FaceAvg(k);
+			}
+		}
+		else 
+		{
+			for (int k = 0; k < Vertices.Num(); ++k)
+			{
+				NewVertices[k] = Vertices[k] + OffsetDistance * GetNormal(k);
+			}
+		}
+		for (int k = 0; k < Vertices.Num(); ++k)
+		{
+			Vertices[k] = NewVertices[k];
+		}
+
+		Timestamp++;
+	}
+
+
+	/**
+	 * Offset polygon by fixed distance, by offsetting and intersecting edges.
+	 * CounterClockWise Polygon offsets "outwards", ClockWise "inwards".
+	 */
+	void PolyOffset(T OffsetDistance)
+	{
+		// [TODO] possibly can do with half as many normalizes if we do w/ sequential edges,
+		//  rather than centering on each v?
+		TArray<FVector2<T>> NewVertices;
+		NewVertices.SetNumUninitialized(Vertices.Num());
+		for (int k = 0; k < Vertices.Num(); ++k) 
+		{
+			FVector2<T> v = Vertices[k];
+			FVector2<T> next = Vertices[(k + 1) % Vertices.Num()];
+			FVector2<T> prev = Vertices[k == 0 ? Vertices.Num() - 1 : k - 1];
+			FVector2<T> dn = (next - v).Normalized();
+			FVector2<T> dp = (prev - v).Normalized();
+			FLine2<T> ln(v + OffsetDistance * dn.Perp(), dn);
+			FLine2<T> lp(v - OffsetDistance * dp.Perp(), dp);
+
+			bool bIntersectsAtPoint = ln.IntersectionPoint(lp, NewVertices[k]);
+			if (!bIntersectsAtPoint)  // lines were parallel
+			{
+				NewVertices[k] = Vertices[k] + OffsetDistance * GetNormal_FaceAvg(k);
+			}
+		}
+		for (int k = 0; k < Vertices.Num(); ++k)
+		{
+			Vertices[k] = NewVertices[k];
+		}
+
+		Timestamp++;
+	}
+
+
+private:
+	// Polygon simplification
+	// code adapted from: http://softsurfer.com/Archive/algorithm_0205/algorithm_0205.htm
+	// simplifyDP():
+	//  This is the Douglas-Peucker recursive simplification routine
+	//  It just marks Vertices that are part of the simplified polyline
+	//  for approximating the polyline subchain v[j] to v[k].
+	//    Input:  tol = approximation tolerance
+	//            v[] = polyline array of vertex points
+	//            j,k = indices for the subchain v[j] to v[k]
+	//    Output: mk[] = array of markers matching vertex array v[]
+	static void SimplifyDouglasPeucker(double Tolerance, const TArray<FVector2<T>>& Vertices, int j, int k, TArray<bool>& Marked)
+	{
+		Marked.SetNum(Vertices.Num());
+		if (k <= j + 1) // there is nothing to simplify
+			return;
+
+		// check for adequate approximation by segment S from v[j] to v[k]
+		int maxi = j;          // index of vertex farthest from S
+		double maxd2 = 0;         // distance squared of farthest vertex
+		double tol2 = Tolerance * Tolerance;  // tolerance squared
+		FSegment2<T> S = FSegment2<T>(Vertices[j], Vertices[k]);    // segment from v[j] to v[k]
+
+		// test each vertex v[i] for max distance from S
+		// Note: this works in any dimension (2D, 3D, ...)
+		for (int i = j + 1; i < k; i++)
+		{
+			double dv2 = S.DistanceSquared(Vertices[i]);
+			if (dv2 <= maxd2)
+				continue;
+			// v[i] is a max vertex
+			maxi = i;
+			maxd2 = dv2;
+		}
+		if (maxd2 > tol2)       // error is worse than the tolerance
+		{
+			// split the polyline at the farthest vertex from S
+			Marked[maxi] = true;      // mark v[maxi] for the simplified polyline
+			// recursively simplify the two subpolylines at v[maxi]
+			SimplifyDouglasPeucker(Tolerance, Vertices, j, maxi, Marked);  // polyline v[j] to v[maxi]
+			SimplifyDouglasPeucker(Tolerance, Vertices, maxi, k, Marked);  // polyline v[maxi] to v[k]
+		}
+		// else the approximation is OK, so ignore intermediate Vertices
+		return;
+	}
+
+
+public:
+
+	/**
+	 * Simplify the Polygon to reduce the vertex count
+	 * @param ClusterTolerance Vertices closer than this distance will be merged into a single vertex
+	 * @param LineDeviationTolerance Vertices are allowed to deviate this much from the input polygon lines
+	 */
+	void Simplify(double ClusterTolerance = 0.0001, double LineDeviationTolerance = 0.01)
+	{
+		int n = Vertices.Num();
+		if (n < 3)
+		{
+			return;
+		}
+
+		int i, k, pv;            // misc counters
+		TArray<FVector2<T>> NewVertices;
+		NewVertices.SetNumUninitialized(n + 1);  // vertex buffer
+		TArray<bool> Marked;
+		Marked.SetNumUninitialized(n + 1);
+		for (i = 0; i < n + 1; ++i)		// marker buffer
+		{
+			Marked[i] = false;
+		}
+
+		// STAGE 1.  Vertex Reduction within tolerance of prior vertex cluster
+		double clusterTol2 = ClusterTolerance * ClusterTolerance;
+		NewVertices[0] = Vertices[0];              // start at the beginning
+		for (i = 1, k = 1, pv = 0; i < n; i++) 
+		{
+			if (Vertices[i].DistanceSquared(Vertices[pv]) < clusterTol2)
+			{
+				continue;
+			}
+			NewVertices[k++] = Vertices[i];
+			pv = i;
+		}
+		bool skip_dp = false;
+		if (k == 1) 
+		{
+			NewVertices[k++] = Vertices[1];
+			NewVertices[k++] = Vertices[2];
+			skip_dp = true;
+		}
+		else if (k == 2) 
+		{
+			NewVertices[k++] = Vertices[0];
+			skip_dp = true;
+		}
+
+		// push on start vertex again, because simplifyDP is for polylines, not polygons
+		NewVertices[k++] = Vertices[0];
+
+		// STAGE 2.  Douglas-Peucker polyline simplification
+		int nv = 0;
+		if (skip_dp == false && LineDeviationTolerance > 0) 
+		{
+			Marked[0] = Marked[k - 1] = true;       // mark the first and last Vertices
+			SimplifyDouglasPeucker(LineDeviationTolerance, NewVertices, 0, k - 1, Marked);
+			for (i = 0; i < k - 1; ++i) 
+			{
+				if (Marked[i])
+				{
+					nv++;
+				}
+			}
+		}
+		else 
+		{
+			for (i = 0; i < k; ++i)
+			{
+				Marked[i] = true;
+			}
+			nv = k - 1;
+		}
+
+		// polygon requires at least 3 Vertices
+		if (nv == 2) 
+		{
+			for (i = 1; i < k - 1; ++i) 
+			{
+				if (Marked[1] == false)
+				{
+					Marked[1] = true;
+				}
+				else if (Marked[k - 2] == false)
+				{
+					Marked[k - 2] = true;
+				}
+			}
+			nv++;
+		}
+		else if (nv == 1) 
+		{
+			Marked[1] = true;
+			Marked[2] = true;
+			nv += 2;
+		}
+
+		// copy marked Vertices back to this polygon
+		Vertices.Reset();
+		for (i = 0; i < k - 1; ++i)    // last vtx is copy of first, and definitely marked
+		{
+			if (Marked[i]) 
+			{
+				Vertices.Add(NewVertices[i]);
+			}
+		}
+
+		Timestamp++;
+		return;
+	}
+
+
+
+	/**
+	 * Chamfer each vertex corner of the Polygon
+	 * @param ChamferDist offset distance from corner that we cut at
+	 */
+	void Chamfer(double ChamferDist, double MinConvexAngleDeg = 30, double MinConcaveAngleDeg = 30)
+	{
+		check(IsClockwise());
+
+		TArray<FVector2<T>> OldV = Vertices;
+		int N = OldV.Num();
+		Vertices.Reset();
+
+		int iCur = 0;
+		do {
+			FVector2<T> center = OldV[iCur];
+
+			int iPrev = (iCur == 0) ? N - 1 : iCur - 1;
+			FVector2<T> prev = OldV[iPrev];
+			int iNext = (iCur + 1) % N;
+			FVector2<T> next = OldV[iNext];
+
+			FVector2<T> cp = prev - center;
+			double cpdist = cp.Normalize();
+			FVector2<T> cn = next - center;
+			double cndist = cn.Normalize();
+
+			// if degenerate, skip this vert
+			if (cpdist < TMathUtil<T>::ZeroTolerance || cndist < TMathUtil<T>::ZeroTolerance) 
+			{
+				iCur = iNext;
+				continue;
+			}
+
+			double angle = cp.AngleD(cn);
+			// TODO document what this means sign-wise
+			// TODO re-test post Unreal port that this DotPerp is doing the right thing
+			double sign = cn.DotPerp(cp);
+			bool bConcave = (sign > 0);
+
+			double thresh = (bConcave) ? MinConcaveAngleDeg : MinConvexAngleDeg;
+
+			// ok not too sharp
+			if (angle > thresh) 
+			{
+				Vertices.Add(center);
+				iCur = iNext;
+				continue;
+			}
+
+
+			double prev_cut_dist = TMathUtil<T>::Min(ChamferDist, cpdist*0.5);
+			FVector2<T> prev_cut = center + cp * prev_cut_dist;
+			double next_cut_dist = TMathUtil<T>::Min(ChamferDist, cndist * 0.5);
+			FVector2<T> next_cut = center + cn * next_cut_dist;
+
+			Vertices.Add(prev_cut);
+			Vertices.Add(next_cut);
+			iCur = iNext;
+		} while (iCur != 0);
+
+		Timestamp++;
+	}
+
+
+
+
+	/**
+	 * Construct a four-vertex rectangle Polygon
+	 */
+	static FPolygon2<T> MakeRectangle(const FVector2<T>& Center, double Width, double Height)
+	{
+		FPolygon2<T> Rectangle;
+		Rectangle.Vertices.SetNumUninitialized(4);
+		Rectangle.Set(0, FVector2<T>(Center.X - Width / 2, Center.Y - Height / 2));
+		Rectangle.Set(1, FVector2<T>(Center.X + Width / 2, Center.Y - Height / 2));
+		Rectangle.Set(2, FVector2<T>(Center.X + Width / 2, Center.Y + Height / 2));
+		Rectangle.Set(3, FVector2<T>(Center.X - Width / 2, Center.Y + Height / 2));
+		return Rectangle;
+	}
+
+
+	/**
+	 * Construct a circular Polygon
+	 */
+	static FPolygon2<T> MakeCircle(double Radius, int Steps, double AngleShiftRadians = 0)
+	{
+		FPolygon2<T> Circle;
+		Circle.Vertices.SetNumUninitialized(Steps);
+
+		for (int i = 0; i < Steps; ++i)
+		{
+			double t = (double)i / (double)Steps;
+			double a = TMathUtil<T>::TwoPi * t + AngleShiftRadians;
+			Circle.Set(i, FVector2<T>(Radius * TMathUtil<T>::Cos(a), Radius * TMathUtil<T>::Sin(a)));
+		}
+
+		return Circle;
+	}
+};
+
+typedef FPolygon2<double> FPolygon2d;
+typedef FPolygon2<float> FPolygon2f;
diff --git a/Engine/Plugins/Experimental/GeometryProcessing/Source/GeometricObjects/Public/QuadricError.h b/Engine/Plugins/Experimental/GeometryProcessing/Source/GeometricObjects/Public/QuadricError.h
new file mode 100644
--- /dev/null
+++ b/Engine/Plugins/Experimental/GeometryProcessing/Source/GeometricObjects/Public/QuadricError.h
@@ -0,0 +1,141 @@
+// Copyright 1998-2019 Epic Games, Inc. All Rights Reserved.
+
+#pragma once
+
+#include "MathUtil.h"
+#include "VectorTypes.h"
+
+
+/**
+ * QuadricError represents a quadratic function that evaluates distance to plane.
+ * Stores minimal 10-coefficient form, following http://mgarland.org/files/papers/qtheory.pdf
+ * (symmetric matrix A, vector b, constant c)
+ */
+template<typename RealType>
+struct TQuadricError 
+{
+	RealType Axx, Axy, Axz, Ayy, Ayz, Azz;
+	RealType bx, by, bz;
+	RealType c;
+
+	inline static TQuadricError Zero() { return TQuadricError(); };
+
+	TQuadricError()
+	{
+		Axx = Axy = Axz = Ayy = Ayz = Azz = bx = by = bz = c = 0;
+	}
+
+	/**
+	 * Construct TQuadricError a plane with the given normal and a point on plane
+	 */
+	TQuadricError(const FVector3<RealType>& Normal, const FVector3<RealType>& Point) 
+	{
+		Axx = Normal.X * Normal.X;
+		Axy = Normal.X * Normal.Y;
+		Axz = Normal.X * Normal.Z;
+		Ayy = Normal.Y * Normal.Y;
+		Ayz = Normal.Y * Normal.Z;
+		Azz = Normal.Z * Normal.Z;
+		bx = by = bz = c = 0;
+		FVector3<RealType> v = MultiplyA(Point);
+		bx = -v.X; by = -v.Y; bz = -v.Z;
+		c = Point.Dot(v);
+	}
+
+	/**
+	 * Construct TQuadricError that is the sum of two other TQuadricErrors
+	 */
+	TQuadricError(const TQuadricError& a, const TQuadricError& b) 
+	{
+		Axx = a.Axx + b.Axx;
+		Axy = a.Axy + b.Axy;
+		Axz = a.Axz + b.Axz;
+		Ayy = a.Ayy + b.Ayy;
+		Ayz = a.Ayz + b.Ayz;
+		Azz = a.Azz + b.Azz;
+		bx = a.bx + b.bx;
+		by = a.by + b.by;
+		bz = a.bz + b.bz;
+		c = a.c + b.c;
+	}
+
+
+	/**
+	 * Add scalar multiple of a TQuadricError to this TQuadricError
+	 */
+	void Add(RealType w, const TQuadricError& b) 
+	{
+		Axx += w * b.Axx;
+		Axy += w * b.Axy;
+		Axz += w * b.Axz;
+		Ayy += w * b.Ayy;
+		Ayz += w * b.Ayz;
+		Azz += w * b.Azz;
+		bx += w * b.bx;
+		by += w * b.by;
+		bz += w * b.bz;
+		c += w * b.c;
+	}
+
+
+	/**
+	 * Evaluates p*A*p + 2*dot(p,b) + c
+	 * @return 
+	 */
+	RealType Evaluate(const FVector3<RealType>& pt) const
+	{
+		RealType x = Axx * pt.X + Axy * pt.Y + Axz * pt.Z;
+		RealType y = Axy * pt.X + Ayy * pt.Y + Ayz * pt.Z;
+		RealType z = Axz * pt.X + Ayz * pt.Y + Azz * pt.Z;
+		return (pt.X * x + pt.Y * y + pt.Z * z) +
+			2.0 * (pt.X * bx + pt.Y * by + pt.Z * bz) + c;
+	}
+
+	/**
+	 * 
+	 */
+	FVector3<RealType> MultiplyA(const FVector3<RealType>& pt) const
+	{
+		RealType x = Axx * pt.X + Axy * pt.Y + Axz * pt.Z;
+		RealType y = Axy * pt.X + Ayy * pt.Y + Ayz * pt.Z;
+		RealType z = Axz * pt.X + Ayz * pt.Y + Azz * pt.Z;
+		return FVector3<RealType>(x, y, z);
+	}
+
+
+
+	bool OptimalPoint(FVector3<RealType>& OutResult, RealType minThresh = 1000.0*TMathUtil<RealType>::Epsilon ) const
+	{
+		RealType a11 = Azz * Ayy - Ayz * Ayz;
+		RealType a12 = Axz * Ayz - Azz * Axy;
+		RealType a13 = Axy * Ayz - Axz * Ayy;
+		RealType a22 = Azz * Axx - Axz * Axz;
+		RealType a23 = Axy * Axz - Axx * Ayz;
+		RealType a33 = Axx * Ayy - Axy * Axy;
+		RealType det = (Axx * a11) + (Axy * a12) + (Axz * a13);
+
+		// [RMS] not sure what we should be using for this threshold...have seen
+		//  det less than 10^-9 on "normal" meshes.
+		if ( (RealType)fabs(det) > minThresh)
+		{
+			det = 1.0 / det;
+			a11 *= det; a12 *= det; a13 *= det;
+			a22 *= det; a23 *= det; a33 *= det;
+			RealType x = a11 * bx + a12 * by + a13 * bz;
+			RealType y = a12 * bx + a22 * by + a23 * bz;
+			RealType z = a13 * bx + a23 * by + a33 * bz;
+			OutResult = FVector3<RealType>(-x, -y, -z);
+			return true;
+		}
+		else 
+		{
+			return false;
+		}
+
+	}
+
+};
+
+
+typedef TQuadricError<float> FQuadricErrorf;
+typedef TQuadricError<double> FQuadricErrord;
\ No newline at end of file
diff --git a/Engine/Plugins/Experimental/GeometryProcessing/Source/GeometricObjects/Public/Quaternion.h b/Engine/Plugins/Experimental/GeometryProcessing/Source/GeometricObjects/Public/Quaternion.h
new file mode 100644
--- /dev/null
+++ b/Engine/Plugins/Experimental/GeometryProcessing/Source/GeometricObjects/Public/Quaternion.h
@@ -0,0 +1,456 @@
+// Copyright 1998-2019 Epic Games, Inc. All Rights Reserved.
+
+#pragma once
+
+#include "VectorTypes.h"
+#include "MatrixTypes.h"
+#include "IndexTypes.h"
+
+// ported from geometry3Sharp Quaternion
+
+template<typename RealType>
+struct FQuaternion
+{
+	// note: in Wm5 version, this is a 4-element arraY stored in order (w,x,y,z).
+	RealType X;
+	RealType Y;
+	RealType Z;
+	RealType W;
+
+	FQuaternion();
+	FQuaternion(RealType X, RealType Y, RealType Z, RealType W);
+	FQuaternion(const RealType* Values);
+	FQuaternion(const FQuaternion& Copy);
+	FQuaternion(const FVector3<RealType>& Axis, RealType Angle, bool bAngleIsDegrees);
+	FQuaternion(const FVector3<RealType>& From, const FVector3<RealType>& To);
+	FQuaternion(const FQuaternion<RealType>& From, const FQuaternion<RealType>& To, RealType InterpT);
+	FQuaternion(const FMatrix3<RealType>& RotationMatrix);
+
+	void SetAxisAngleD(const FVector3<RealType>& Axis, RealType AngleDeg);
+	void SetAxisAngleR(const FVector3<RealType>& Axis, RealType AngleRad);
+	void SetFromTo(const FVector3<RealType>& From, const FVector3<RealType>& To);
+	void SetToSlerp(const FQuaternion<RealType>& From, const FQuaternion<RealType>& To, RealType InterpT);
+	void SetFromRotationMatrix(const FMatrix3<RealType>& RotationMatrix);
+
+	static FQuaternion<RealType> Zero() { return FQuaternion<RealType>((RealType)0, (RealType)0, (RealType)0, (RealType)0); }
+	static FQuaternion<RealType> Identity() { return FQuaternion<RealType>((RealType)0, (RealType)0, (RealType)0, (RealType)1); }
+
+	RealType & operator[](int i) { return (&X)[i]; }
+	const RealType & operator[](int i) const { return (&X)[i]; }
+
+	bool EpsilonEqual(const FQuaternion<RealType>& Other, RealType Epsilon) const;
+
+	RealType Length() const { return (RealType)sqrt(X*X + Y*Y + Z*Z + W*W); }
+	RealType SquaredLength() const { return X*X + Y*Y + Z*Z + W*W; }
+
+	FVector3<RealType> AxisX() const;
+	FVector3<RealType> AxisY() const;
+	FVector3<RealType> AxisZ() const;
+
+	RealType Normalize(const RealType epsilon = 0);
+	FQuaternion<RealType> Normalized(const RealType epsilon = 0) const;
+
+	RealType Dot(const FQuaternion<RealType>& Other) const;
+
+	FQuaternion<RealType> Inverse() const;
+	FVector3<RealType> InverseMultiply(const FVector3<RealType>& Other) const;
+
+	FMatrix3<RealType> ToRotationMatrix() const;
+
+
+	// available for porting:
+	//   SetFromRotationMatrix(FMatrix3f)
+
+
+	inline operator FQuat() const
+	{
+		FQuat Quat; 
+		Quat.X = (float)X;
+		Quat.Y = (float)Y;
+		Quat.Z = (float)Z;
+		Quat.W = (float)W;
+		return Quat;
+	}
+	inline FQuaternion(const FQuat& Quat)
+	{
+		X = (RealType)Quat.X;
+		Y = (RealType)Quat.Y;
+		Z = (RealType)Quat.Z;
+		W = (RealType)Quat.W;
+	}
+
+};
+
+
+
+typedef FQuaternion<float> FQuaternionf;
+typedef FQuaternion<double> FQuaterniond;
+
+
+
+
+template<typename RealType>
+FQuaternion<RealType>::FQuaternion()
+{
+	X = Y = Z = 0; W = 1;
+}
+
+template<typename RealType>
+FQuaternion<RealType>::FQuaternion(RealType x, RealType y, RealType z, RealType w)
+{
+	X = x;
+	Y = y;
+	Z = z;
+	W = w;
+}
+
+template<typename RealType>
+FQuaternion<RealType>::FQuaternion(const RealType* Values)
+{
+	X = Values[0];
+	Y = Values[1];
+	Z = Values[2];
+	W = Values[3];
+}
+
+template<typename RealType>
+FQuaternion<RealType>::FQuaternion(const FQuaternion& Copy)
+{
+	X = Copy.X;
+	Y = Copy.Y;
+	Z = Copy.Z;
+	W = Copy.W;
+}
+
+template<typename RealType>
+FQuaternion<RealType>::FQuaternion(const FVector3<RealType>& Axis, RealType Angle, bool bAngleIsDegrees)
+{
+	X = Y = Z = 0; W = 1;
+	SetAxisAngleR(Axis, Angle * (bAngleIsDegrees ? TMathUtil<RealType>::DegToRad : (RealType)1));
+}
+
+template<typename RealType>
+FQuaternion<RealType>::FQuaternion(const FVector3<RealType>& From, const FVector3<RealType>& To)
+{
+	X = Y = Z = 0; W = 1;
+	SetFromTo(From, To);
+}
+
+template<typename RealType>
+FQuaternion<RealType>::FQuaternion(const FQuaternion<RealType>& From, const FQuaternion<RealType>& To, RealType InterpT)
+{
+	X = Y = Z = 0; W = 1;
+	SetToSlerp(From, To, InterpT);
+}
+
+template<typename RealType>
+FQuaternion<RealType>::FQuaternion(const FMatrix3<RealType>& RotationMatrix)
+{
+	X = Y = Z = 0; W = 1;
+	SetFromRotationMatrix(RotationMatrix);
+}
+
+
+
+
+
+template<typename RealType>
+RealType FQuaternion<RealType>::Normalize(const RealType Epsilon)
+{
+	RealType length = Length();
+	if (length > Epsilon)
+	{
+		RealType invLength = ((RealType)1) / length;
+		X *= invLength;
+		Y *= invLength;
+		Z *= invLength;
+		W *= invLength;
+		return invLength;
+	}
+	X = Y = Z = W = (RealType)0;
+	return 0;
+}
+
+template<typename RealType>
+FQuaternion<RealType> FQuaternion<RealType>::Normalized(const RealType Epsilon) const
+{
+	RealType length = Length();
+	if (length > Epsilon)
+	{
+		RealType invLength = ((RealType)1) / length;
+		return FQuaternion<RealType>(X*invLength, Y*invLength, Z*invLength, W*invLength);
+	}
+	return Zero();
+}
+
+template<typename RealType>
+RealType FQuaternion<RealType>::Dot(const FQuaternion<RealType>& Other) const
+{
+	return X * Other.X + Y * Other.Y + Z * Other.Z + W * Other.W;
+}
+
+
+template<typename RealType>
+FQuaternion<RealType> operator*(const FQuaternion<RealType>& A, const FQuaternion<RealType>& B) 
+{
+	RealType W = A.W * B.W - A.X * B.X - A.Y * B.Y - A.Z * B.Z;
+	RealType X = A.W * B.X + A.X * B.W + A.Y * B.Z - A.Z * B.Y;
+	RealType Y = A.W * B.Y + A.Y * B.W + A.Z * B.X - A.X * B.Z;
+	RealType Z = A.W * B.Z + A.Z * B.W + A.X * B.Y - A.Y * B.X;
+	return FQuaternion<RealType>(X, Y, Z, W);
+}
+
+template<typename RealType>
+FQuaternion<RealType> operator -(const FQuaternion<RealType>& A, const FQuaternion<RealType>& B) 
+{
+	return FQuaternion<RealType>(A.X - B.X, A.Y - B.Y, A.Z - B.Z, A.W - B.W);
+}
+
+template<typename RealType>
+FVector3<RealType> operator*(const FQuaternion<RealType>& Q, const FVector3<RealType>& V) 
+{
+	//return q.ToRotationMatrix() * v;
+	// inline-expansion of above:
+	RealType twoX = (RealType)2*Q.X; RealType twoY = (RealType)2*Q.Y; RealType twoZ = (RealType)2*Q.Z;
+	RealType twoWX = twoX * Q.W; RealType twoWY = twoY * Q.W; RealType twoWZ = twoZ * Q.W;
+	RealType twoXX = twoX * Q.X; RealType twoXY = twoY * Q.X; RealType twoXZ = twoZ * Q.X;
+	RealType twoYY = twoY * Q.Y; RealType twoYZ = twoZ * Q.Y; RealType twoZZ = twoZ * Q.Z;
+	return FVector3<RealType>(
+		V.X * ((RealType)1 - (twoYY + twoZZ)) + V.Y * (twoXY - twoWZ) + V.Z * (twoXZ + twoWY),
+		V.X * (twoXY + twoWZ) + V.Y * ((RealType)1 - (twoXX + twoZZ)) + V.Z * (twoYZ - twoWX),
+		V.X * (twoXZ - twoWY) + V.Y * (twoYZ + twoWX) + V.Z * ((RealType)1 - (twoXX + twoYY))); ;
+}
+
+
+template<typename RealType>
+FVector3<RealType> FQuaternion<RealType>::InverseMultiply(const FVector3<RealType>& V) const
+{
+	RealType norm = SquaredLength();
+	if (norm > 0) 
+	{
+		RealType invNorm = (RealType)1 / norm;
+		RealType qX = -X * invNorm, qY = -Y * invNorm, qZ = -Z * invNorm, qW = W * invNorm;
+		RealType twoX = (RealType)2 * qX; RealType twoY = (RealType)2 * qY; RealType twoZ = (RealType)2 * qZ;
+		RealType twoWX = twoX * qW; RealType twoWY = twoY * qW; RealType twoWZ = twoZ * qW;
+		RealType twoXX = twoX * qX; RealType twoXY = twoY * qX; RealType twoXZ = twoZ * qX;
+		RealType twoYY = twoY * qY; RealType twoYZ = twoZ * qY; RealType twoZZ = twoZ * qZ;
+		return FVector3<RealType>(
+			V.X * ((RealType)1 - (twoYY + twoZZ)) + V.Y * (twoXY - twoWZ) + V.Z * (twoXZ + twoWY),
+			V.X * (twoXY + twoWZ) + V.Y * ((RealType)1 - (twoXX + twoZZ)) + V.Z * (twoYZ - twoWX),
+			V.X * (twoXZ - twoWY) + V.Y * (twoYZ + twoWX) + V.Z * ((RealType)1 - (twoXX + twoYY)));
+	}
+	return FVector3<RealType>::Zero();
+}
+
+
+template<typename RealType>
+FVector3<RealType> FQuaternion<RealType>::AxisX() const
+{
+	RealType twoY = (RealType)2 * Y; RealType twoZ = (RealType)2 * Z;
+	RealType twoWY = twoY * W; RealType twoWZ = twoZ * W;
+	RealType twoXY = twoY * X; RealType twoXZ = twoZ * X;
+	RealType twoYY = twoY * Y; RealType twoZZ = twoZ * Z;
+	return FVector3<RealType>((RealType)1 - (twoYY + twoZZ), twoXY + twoWZ, twoXZ - twoWY);
+}
+
+template<typename RealType>
+FVector3<RealType> FQuaternion<RealType>::AxisY() const
+{
+	RealType twoX = (RealType)2 * X; RealType twoY = (RealType)2 * Y; RealType twoZ = (RealType)2 * Z;
+	RealType twoWX = twoX * W; RealType twoWZ = twoZ * W; RealType twoXX = twoX * X;
+	RealType twoXY = twoY * X; RealType twoYZ = twoZ * Y; RealType twoZZ = twoZ * Z;
+	return FVector3<RealType>(twoXY - twoWZ, (RealType)1 - (twoXX + twoZZ), twoYZ + twoWX);
+}
+
+template<typename RealType>
+FVector3<RealType> FQuaternion<RealType>::AxisZ() const
+{
+	RealType twoX = (RealType)2 * X; RealType twoY = (RealType)2 * Y; RealType twoZ = (RealType)2 * Z;
+	RealType twoWX = twoX * W; RealType twoWY = twoY * W; RealType twoXX = twoX * X;
+	RealType twoXZ = twoZ * X; RealType twoYY = twoY * Y; RealType twoYZ = twoZ * Y;
+	return FVector3<RealType>(twoXZ + twoWY, twoYZ - twoWX, (RealType)1 - (twoXX + twoYY));
+}
+
+template<typename RealType>
+FQuaternion<RealType> FQuaternion<RealType>::Inverse() const
+{
+	RealType norm = SquaredLength();
+	if (norm > (RealType)0) 
+	{
+		RealType invNorm = (RealType)1 / norm;
+		return FQuaternion<RealType>(
+			-X * invNorm, -Y * invNorm, -Z * invNorm, W * invNorm);
+	}
+	return FQuaternion<RealType>::Zero();
+}
+
+
+template<typename RealType>
+void FQuaternion<RealType>::SetAxisAngleD(const FVector3<RealType>& Axis, RealType AngleDeg) 
+{
+	SetAxisAngleR(Axis, TMathUtil<RealType>::DegToRad * AngleDeg);
+}
+
+template<typename RealType>
+void FQuaternion<RealType>::SetAxisAngleR(const FVector3<RealType>& Axis, RealType AngleRad)
+{
+	RealType halfAngle = (RealType)0.5 * AngleRad;
+	RealType sn = (RealType)sin(halfAngle);
+	W = (RealType)cos(halfAngle);
+	X = (sn * Axis.X);
+	Y = (sn * Axis.Y);
+	Z = (sn * Axis.Z);
+}
+
+
+
+// this function can take non-normalized vectors vFrom and vTo (normalizes internally)
+template<typename RealType>
+void FQuaternion<RealType>::SetFromTo(const FVector3<RealType>& From, const FVector3<RealType>& To)
+{
+	// [TODO] this page seems to have optimized version:
+	//    http://lolengine.net/blog/2013/09/18/beautiful-maths-quaternion-from-vectors
+
+	// [RMS] not ideal to explicitly normalize here, but if we don't,
+	//   output FQuaternion is not normalized and this causes problems,
+	//   eg like drift if we do repeated SetFromTo()
+	FVector3<RealType> from = From.Normalized(), to = To.Normalized();
+	FVector3<RealType> bisector = (from + to).Normalized();
+	W = from.Dot(bisector);
+	if (W != 0) 
+	{
+		FVector3<RealType> cross = from.Cross(bisector);
+		X = cross.X;
+		Y = cross.Y;
+		Z = cross.Z;
+	}
+	else 
+	{
+		RealType invLength;
+		if ( (RealType)fabs(from.X) >= (RealType)fabs(from.Y)) 
+		{
+			// V1.X or V1.Z is the largest magnitude component.
+			invLength = ((RealType)1 / (RealType)sqrt(from.X * from.X + from.Z * from.Z));
+			X = -from.Z * invLength;
+			Y = 0;
+			Z = +from.X * invLength;
+		}
+		else 
+		{
+			// V1.Y or V1.Z is the largest magnitude component.
+			invLength = ((RealType)1 / (RealType)sqrt(from.Y * from.Y + from.Z * from.Z));
+			X = 0;
+			Y = +from.Z * invLength;
+			Z = -from.Y * invLength;
+		}
+	}
+	Normalize();   // just to be safe...
+}
+
+
+
+template<typename RealType>
+void FQuaternion<RealType>::SetToSlerp( const FQuaternion<RealType>& From, const FQuaternion<RealType>& To, RealType InterpT)
+{
+	RealType cs = From.Dot(To);
+	RealType angle = (RealType)acos(cs);
+	if (TMathUtil<RealType>::Abs(angle) >= TMathUtil<RealType>::ZeroTolerance)
+	{
+		RealType sn = (RealType)sin(angle);
+		RealType invSn = 1 / sn;
+		RealType tAngle = InterpT * angle;
+		RealType coeff0 = (RealType)sin(angle - tAngle) * invSn;
+		RealType coeff1 = (RealType)sin(tAngle) * invSn;
+		X = coeff0 * From.X + coeff1 * To.X;
+		Y = coeff0 * From.Y + coeff1 * To.Y;
+		Z = coeff0 * From.Z + coeff1 * To.Z;
+		W = coeff0 * From.W + coeff1 * To.W;
+	}
+	else 
+	{
+		X = From.X;
+		Y = From.Y;
+		Z = From.Z;
+		W = From.W;
+	}
+}
+
+
+
+template<typename RealType>
+void FQuaternion<RealType>::SetFromRotationMatrix(const FMatrix3<RealType>& RotationMatrix)
+{
+	// Algorithm in Ken Shoemake's article in 1987 SIGGRAPH course notes
+	// article "FQuaternion Calculus and Fast Animation".
+	FIndex3i next(1, 2, 0);
+
+	RealType trace = RotationMatrix(0,0) + RotationMatrix(1,1) + RotationMatrix(2, 2);
+	RealType root;
+
+	if (trace > 0) 
+	{
+		// |w| > 1/2, may as well choose w > 1/2
+		root = (RealType)sqrt(trace + (RealType)1);  // 2w
+		W = ((RealType)0.5) * root;
+		root = ((RealType)0.5) / root;  // 1/(4w)
+		X = (RotationMatrix(2, 1) - RotationMatrix(1, 2)) * root;
+		Y = (RotationMatrix(0, 2) - RotationMatrix(2, 0)) * root;
+		Z = (RotationMatrix(1, 0) - RotationMatrix(0, 1)) * root;
+	}
+	else 
+	{
+		// |w| <= 1/2
+		int i = 0;
+		if (RotationMatrix(1, 1) > RotationMatrix(0, 0))
+		{
+			i = 1;
+		}
+		if (RotationMatrix(2, 2) > RotationMatrix(i, i))
+		{
+			i = 2;
+		}
+		int j = next[i];
+		int k = next[j];
+
+		root = (RealType)sqrt(RotationMatrix(i, i) - RotationMatrix(j, j) - RotationMatrix(k, k) + (RealType)1);
+
+		FVector3<RealType> quat(X, Y, Z);
+		quat[i] = ((RealType)0.5) * root;
+		root = ((RealType)0.5) / root;
+		W = (RotationMatrix(k, j) - RotationMatrix(j, k)) * root;
+		quat[j] = (RotationMatrix(j, i) + RotationMatrix(i, j)) * root;
+		quat[k] = (RotationMatrix(k, i) + RotationMatrix(i, k)) * root;
+		X = quat.X; Y = quat.Y; Z = quat.Z;
+	}
+
+	Normalize();   // we prefer normalized FQuaternions...
+}
+
+
+
+
+template<typename RealType>
+FMatrix3<RealType> FQuaternion<RealType>::ToRotationMatrix() const
+{
+	RealType twoX = 2 * X; RealType twoY = 2 * Y; RealType twoZ = 2 * Z;
+	RealType twoWX = twoX * W; RealType twoWY = twoY * W; RealType twoWZ = twoZ * W;
+	RealType twoXX = twoX * X; RealType twoXY = twoY * X; RealType twoXZ = twoZ * X;
+	RealType twoYY = twoY * Y; RealType twoYZ = twoZ * Y; RealType twoZZ = twoZ * Z;
+	FMatrix3<RealType> m = FMatrix3<RealType>::Zero();
+	m.Row0 = FVector3<RealType>(1 - (twoYY + twoZZ), twoXY - twoWZ, twoXZ + twoWY);
+	m.Row1 = FVector3<RealType>(twoXY + twoWZ, 1 - (twoXX + twoZZ), twoYZ - twoWX);
+	m.Row2 = FVector3<RealType>(twoXZ - twoWY, twoYZ + twoWX, 1 - (twoXX + twoYY));
+	return m;
+}
+
+
+
+template<typename RealType>
+bool FQuaternion<RealType>::EpsilonEqual(const FQuaternion<RealType>& Other, RealType Epsilon) const
+{
+	return (RealType)fabs(X - Other.X) <= Epsilon &&
+		(RealType)fabs(Y - Other.Y) <= Epsilon &&
+		(RealType)fabs(Z - Other.Z) <= Epsilon &&
+		(RealType)fabs(W - Other.W) <= Epsilon;
+}
+
diff --git a/Engine/Plugins/Experimental/GeometryProcessing/Source/GeometricObjects/Public/RayTypes.h b/Engine/Plugins/Experimental/GeometryProcessing/Source/GeometricObjects/Public/RayTypes.h
new file mode 100644
--- /dev/null
+++ b/Engine/Plugins/Experimental/GeometryProcessing/Source/GeometricObjects/Public/RayTypes.h
@@ -0,0 +1,71 @@
+// Copyright 1998-2019 Epic Games, Inc. All Rights Reserved.
+
+// port of geometry3Sharp Ray3
+
+#pragma once
+
+#include "Math/UnrealMath.h"
+#include "VectorTypes.h"
+#include "Math/Ray.h"
+
+/*
+ * 3D Ray stored as Origin point and normalized Direction vector
+ */
+template<typename RealType>
+class FRay3
+{
+public:
+	/** Origin point */
+	FVector3<RealType> Origin;
+	/** Direction vector, always normalized */
+	FVector3<RealType> Direction;
+
+	/** Construct ray at origin pointed down Z axis */
+	FRay3()
+	{
+		Origin = FVector3<RealType>::Zero();
+		Direction = FVector3<RealType>::UnitZ();
+	}
+
+	/**
+	 * Construct ray from origin and direction
+	 * @param Origin origin point
+	 * @param Direction direction vector. Must be normalized if bIsNormalized=true
+	 * @param bIsNormalized if true, Direction will not be re-normalized
+	 */
+	FRay3(const FVector3<RealType>& Origin, const FVector3<RealType>& Direction, bool bIsNormalized = false)
+	{
+		this->Origin = Origin;
+		this->Direction = Direction;
+		if (bIsNormalized == false)
+		{
+			this->Direction.Normalize();
+		}
+	}
+
+	/**
+	 * @return point on ray at given (signed) Distance from the ray Origin
+	 */
+	FVector3<RealType> PointAt(RealType Distance) const
+	{
+		return Origin + Distance * Direction;
+	}
+
+
+	// conversion operators
+
+	inline operator FRay() const
+	{
+		return FRay((FVector)Origin, (FVector)Direction);
+	}
+	inline FRay3(const FRay & RayIn)
+	{
+		Origin = (FVector3<RealType>)RayIn.Origin;
+		Direction = (FVector3<RealType>)RayIn.Direction;
+	}
+
+
+};
+typedef FRay3<float> FRay3f;
+typedef FRay3<double> FRay3d;
+
diff --git a/Engine/Plugins/Experimental/GeometryProcessing/Source/GeometricObjects/Public/SegmentTypes.h b/Engine/Plugins/Experimental/GeometryProcessing/Source/GeometricObjects/Public/SegmentTypes.h
new file mode 100644
--- /dev/null
+++ b/Engine/Plugins/Experimental/GeometryProcessing/Source/GeometricObjects/Public/SegmentTypes.h
@@ -0,0 +1,302 @@
+// Copyright 1998-2019 Epic Games, Inc. All Rights Reserved.
+
+// ported from geometry3Sharp Segment2
+
+#pragma once
+
+#include "Math/UnrealMath.h"
+#include "VectorTypes.h"
+
+
+/*
+ * 2D Line Segmented stored as Center point, normalized Direction vector, and scalar Extent
+ */
+template<typename T>
+struct FSegment2
+{
+public:
+	/** Center point of segment */
+	FVector2<T> Center;
+	/** normalized Direction vector of segment */
+	FVector2<T> Direction;
+	/** Extent of segment, which is half the total length */
+	T Extent;
+
+	/**
+	 * Construct a Segment from two Points
+	 */
+	FSegment2(const FVector2<T>& Point0, const FVector2<T>& Point1)
+	{
+		// set from endpoints 
+		Center = T(.5) * (Point0 + Point1);
+		Direction = Point1 - Point0;
+		Extent = T(.5) * Direction.Normalize();
+	}
+
+	/**
+	 * Construct a segment from a Center Point, normalized Direction, and scalar Extent
+	 */
+	FSegment2(const FVector2<T>& CenterIn, const FVector2<T>& DirectionIn, double ExtentIn)
+	{
+		Center = CenterIn;
+		Direction = DirectionIn; 
+		Extent = ExtentIn;
+	}
+
+
+
+	/** Update the Segment with a new start point */
+	inline void SetStartPoint(const FVector2<T>& Point)
+	{
+		update_from_endpoints(Point, EndPoint());
+	}
+
+	/** Update the Segment with a new end point */
+	inline void SetEndPoint(const FVector2<T>& Point)
+	{
+		update_from_endpoints(StartPoint(), Point);
+	}
+
+	/** Reverse the segment */
+	void Reverse()
+	{
+		update_from_endpoints(EndPoint(), StartPoint());
+	}
+
+
+
+
+	/** @return start point of segment */
+	inline FVector2<T> StartPoint() const
+	{
+		return Center - Extent * Direction;
+	}
+
+	/** @return end point of segment */
+	inline FVector2<T> EndPoint() const
+	{
+		return Center + Extent * Direction;
+	}
+
+
+	/** @return the Length of the segment */
+	inline T Length() const
+	{
+		return (T)2 * Extent;
+	}
+
+	/** @return first (i == 0) or second (i == 1) endpoint of the Segment  */
+	inline FVector2<T> GetPointFromIndex(int i) const
+	{
+		return (i == 0) ? (Center - Extent * Direction) : (Center + Extent * Direction);
+	}
+
+	/**
+	 * @return point on segment at given (signed) Distance from the segment Origin
+	 */
+	inline FVector2<T> PointAt(T DistanceParameter) const
+	{
+		return Center + DistanceParameter * Direction;
+	}
+
+	/**
+	 * @param UnitParameter value in range [0,1]
+	 * @return point on segment at that linearly interpolates between start and end based on Unit Parameter
+	 */
+	inline FVector2<T> PointBetween(double UnitParameter) const
+	{
+		return Center + ((T)2 * UnitParameter - (T)1) * Extent * Direction;
+	}
+
+	/**
+	 * @return minimum squared distance from Point to segment
+	 */
+	inline double DistanceSquared(const FVector2<T>& Point) const
+	{
+		double DistParameter;
+		return DistanceSquared(Point, DistParameter);
+	}
+
+
+
+	/**
+	 * @param Point query point
+	 * @param DistParameterOut calculated distance parameter in range [-Extent,Extent]
+	 * @return minimum squared distance from Point to Segment
+	 */	
+	double DistanceSquared(const FVector2<T>& Point, double& DistParameterOut) const
+	{
+		DistParameterOut = (Point - Center).Dot(Direction);
+		if (DistParameterOut >= Extent)
+		{
+			DistParameterOut = Extent;
+			return Point.DistanceSquared(EndPoint());
+		}
+		else if (DistParameterOut <= -Extent)
+		{
+			DistParameterOut = -Extent;
+			return Point.DistanceSquared(StartPoint());
+		}
+		FVector2<T> ProjectedPt = Center + DistParameterOut * Direction;
+		return ProjectedPt.DistanceSquared(Point);
+	}
+
+
+	/**
+	 * @return nearest point on segment to QueryPoint
+	 */
+	inline FVector2<T> NearestPoint(const FVector2<T>& QueryPoint) const
+	{
+		double t = (QueryPoint - Center).Dot(Direction);
+		if (t >= Extent)
+		{
+			return EndPoint();
+		}
+		if (t <= -Extent)
+		{
+			return StartPoint();
+		}
+		return Center + t * Direction;
+	}
+
+
+	/**
+	 * @return scalar projection of QueryPoint onto line of Segment (not clamped to Extents)
+	 */
+	inline double Project(const FVector2<T>& QueryPoint) const
+	{
+		return (QueryPoint - Center).Dot(Direction);
+	}
+
+
+	/**
+	 * Determine which side of the segment the query point lies on
+	 * @param QueryPoint test point
+	 * @param Tolerance tolerance band in which we return 0
+	 * @return +1 if point is to right of line, -1 if left, and 0 if on line or within tolerance band
+	 */
+	int WhichSide(const FVector2<T>& QueryPoint, double Tolerance = 0)
+	{
+		// [TODO] subtract Center from test?
+		FVector2<T> EndPt = Center + Extent * Direction;
+		FVector2<T> StartPt = Center - Extent * Direction;
+		double det = -FVector2<T>::Orient(EndPt, StartPt, QueryPoint);
+		return (det > Tolerance ? +1 : (det < -Tolerance ? -1 : 0));
+	}
+
+
+	/**
+	 * Test if this segment intersects with OtherSegment. Returns true for parallel-line overlaps. Returns same result as IntrSegment2Segment2
+	 * @param OtherSegment segment to test against
+	 * @param DotThresh dot-product tolerance used to determine if segments are parallel
+	 * @param IntervalThresh distance tolerance used to allow slighly-not-touching segments to be considered overlapping
+	 * @return true if segments intersect
+	 */
+	bool Intersects(const FSegment2<T>& OtherSegment, T DotThresh = TMathUtil<T>::Epsilon, T IntervalThresh = 0) const
+	{
+		// see IntrLine2Line2 and IntrSegment2Segment2 for details on this code
+
+		FVector2<T> diff = OtherSegment.Center - Center;
+		T D0DotPerpD1 = Direction.DotPerp(OtherSegment.Direction);
+		if (TMathUtil<T>::Abs(D0DotPerpD1) > DotThresh)      // Lines intersect in a single point.
+		{
+			T invD0DotPerpD1 = ((T)1) / D0DotPerpD1;
+			T diffDotPerpD0 = diff.DotPerp(Direction);
+			T diffDotPerpD1 = diff.DotPerp(OtherSegment.Direction);
+			T s = diffDotPerpD1 * invD0DotPerpD1;
+			T s2 = diffDotPerpD0 * invD0DotPerpD1;
+			return TMathUtil<T>::Abs(s) <= (Extent + IntervalThresh)
+				&& TMathUtil<T>::Abs(s2) <= (OtherSegment.Extent + IntervalThresh);
+		}
+
+		// Lines are parallel.
+		diff.Normalize();
+		T diffNDotPerpD1 = diff.DotPerp(OtherSegment.Direction);
+		if (TMathUtil<T>::Abs(diffNDotPerpD1) <= DotThresh)
+		{
+			// Compute the location of OtherSegment endpoints relative to our Segment
+			diff = OtherSegment.Center - Center;
+			T t1 = Direction.Dot(diff);
+			T tmin = t1 - OtherSegment.Extent;
+			T tmax = t1 + OtherSegment.Extent;
+			FInterval1<T> extents(-Extent, Extent);
+			if (extents.Overlaps(FInterval1<T>(tmin, tmax)))
+			{
+				return true;
+			}
+			return false;
+		}
+
+		// lines are parallel but not collinear
+		return false;
+	}
+
+
+
+
+
+	// 2D segment utility functions
+
+
+	/**
+	 * Calculate distance from QueryPoint to segment (StartPt,EndPt)
+	 */
+	static T FastDistanceSquared(const FVector2<T>& StartPt, const FVector2<T>& EndPt, const FVector2<T>& QueryPt, T Tolerance = TMathUtil<T>::Epsilon)
+	{
+		T vx = EndPt.X - StartPt.X, vy = EndPt.Y - StartPt.Y;
+		T len2 = vx * vx + vy * vy;
+		T dx = QueryPt.X - StartPt.X, dy = QueryPt.Y - StartPt.Y;
+		//if (len2 < 1e-13) 
+		if (len2 < Tolerance)
+		{
+			return dx * dx + dy * dy;
+		}
+		T t = (dx*vx + dy * vy);
+		if (t <= 0) 
+		{
+			return dx * dx + dy * dy;
+		}
+		else if (t >= len2) 
+		{
+			dx = QueryPt.X - EndPt.X; 
+			dy = QueryPt.Y - EndPt.Y;
+			return dx * dx + dy * dy;
+		}
+		dx = QueryPt.X - (StartPt.X + ((t * vx) / len2));
+		dy = QueryPt.Y - (StartPt.Y + ((t * vy) / len2));
+		return dx * dx + dy * dy;
+	}
+
+
+	/**
+	 * Determine which side of the segment the query point lies on
+	 * @param StartPt first point of Segment
+	 * @param EndPt second point of Segment
+	 * @param QueryPoint test point
+	 * @param Tolerance tolerance band in which we return 0
+	 * @return +1 if point is to right of line, -1 if left, and 0 if on line or within tolerance band
+	 */
+	static int WhichSide(const FVector2<T>& StartPt, const FVector2<T>& EndPt, const FVector2<T>& QueryPt, T Tolerance = (T)0)
+	{
+		T det = -FVector2<T>::Orient(StartPt, EndPt, QueryPt);
+		return (det > Tolerance ? +1 : (det < -Tolerance ? -1 : 0));
+	}
+
+
+
+
+
+protected:
+
+	// update segment based on new endpoints
+	inline void update_from_endpoints(const FVector2<T>& p0, const FVector2<T>& p1)
+	{
+		Center = 0.5 * (p0 + p1);
+		Direction = p1 - p0;
+		Extent = 0.5 * Direction.Normalize();
+	}
+
+};
+typedef FSegment2<float> FSegment2f;
+typedef FSegment2<double> FSegment2d;
+
diff --git a/Engine/Plugins/Experimental/GeometryProcessing/Source/GeometricObjects/Public/Spatial/FastWinding.h b/Engine/Plugins/Experimental/GeometryProcessing/Source/GeometricObjects/Public/Spatial/FastWinding.h
new file mode 100644
--- /dev/null
+++ b/Engine/Plugins/Experimental/GeometryProcessing/Source/GeometricObjects/Public/Spatial/FastWinding.h
@@ -0,0 +1,461 @@
+// Copyright 1998-2019 Epic Games, Inc. All Rights Reserved.
+
+#include "MathUtil.h"
+#include "MeshQueries.h"
+#include "Spatial/MeshAABBTree3.h"
+#include "Util/MeshCaches.h"
+#include "MatrixTypes.h"
+
+/**
+ *  Formulas for triangle winding number approximation
+ */
+namespace FastTriWinding
+{
+	/**
+	 *  precompute constant coefficients of triangle winding number approximation
+	 *  P: 'Center' of expansion for Triangles (area-weighted centroid avg)
+	 *  R: max distance from P to Triangles
+	 *  Order1: first-order vector coeff
+	 *  Order2: second-order matrix coeff
+	 *  TriCache: precomputed triangle centroid/normal/area (todo @jimmy either support passing in as possibly-nullptr or remove commented-out null branches!)
+	 */
+	template <class TriangleMeshType, class IterableTriangleIndices>
+	void ComputeCoeffs(const TriangleMeshType& Mesh,
+					   const IterableTriangleIndices& Triangles,
+					   const FMeshTriInfoCache& TriCache,
+					   FVector3d& P, double& R,
+					   FVector3d& Order1, FMatrix3d& Order2)
+	{
+		P = FVector3d::Zero();
+		Order1 = FVector3d::Zero();
+		Order2 = FMatrix3d::Zero();
+		R = 0;
+
+		// compute area-weighted centroid of Triangles, we use this as the expansion point
+		FVector3d P0 = FVector3d::Zero(), P1 = FVector3d::Zero(), P2 = FVector3d::Zero();
+		double sum_area = 0;
+		for (int tid : Triangles)
+		{
+			//if (TriCache == null)
+			//{
+			//	double area = TriCache.Areas[tid];
+			//	sum_area += area;
+			//	P += area * TriCache.Centroids[tid];
+			//}
+			//else
+			{
+				Mesh.GetTriVertices(tid, P0, P1, P2);
+				double area = VectorUtil::Area(P0, P1, P2);
+				sum_area += area;
+				P += area * ((P0 + P1 + P2) / 3.0);
+			}
+		}
+		P /= sum_area;
+
+		// compute first and second-order coefficients of FWN taylor expansion, as well as
+		// 'radius' value R, which is max dist from any tri vertex to P
+		FVector3d n, c;
+		double a = 0;
+		for (int tid : Triangles)
+		{
+			Mesh.GetTriVertices(tid, P0, P1, P2);
+
+			//if (TriCache == null)
+			//{
+			//	c = (1.0 / 3.0) * (P0 + P1 + P2);
+			//	n = FMathd::FastNormalArea(P0, P1, P2, out a);
+			//}
+			//else
+			{
+				TriCache.GetTriInfo(tid, n, a, c);
+			}
+
+			Order1 += a * n;
+
+			FVector3d dcp = c - P;
+			Order2 += a * FMatrix3d(dcp, n);
+
+			// this is just for return value...
+			double maxdist = FMath::Max3(P0.DistanceSquared(P), P1.DistanceSquared(P), P2.DistanceSquared(P));
+			R = FMath::Max(R, sqrt(maxdist));
+		}
+	}
+
+	/**
+	 *  Evaluate first-order FWN approximation at point Q, relative to Center c
+	 */
+	double EvaluateOrder1Approx(const FVector3d& Center, const FVector3d& Order1Coeff, const FVector3d& Q)
+	{
+		FVector3d dpq = (Center - Q);
+		double len = dpq.Length();
+
+		return (1.0 / FMathd::FourPi) * Order1Coeff.Dot(dpq / (len * len * len));
+	}
+
+	/**
+	 *  Evaluate second-order FWN approximation at point Q, relative to Center c
+	 */
+	double EvaluateOrder2Approx(const FVector3d& Center, const FVector3d& Order1Coeff, const FMatrix3d& Order2Coeff, const FVector3d& Q)
+	{
+		FVector3d dpq = (Center - Q);
+		double len = dpq.Length();
+		double len3 = len * len * len;
+		double fourPi_len3 = 1.0 / (FMathd::FourPi * len3);
+
+		double Order1 = fourPi_len3 * Order1Coeff.Dot(dpq);
+
+		// second-order hessian \grad^2(G)
+		double c = -3.0 / (FMathd::FourPi * len3 * len * len);
+
+		// expanded-out version below avoids extra constructors
+		//FMatrix3d xqxq(dpq, dpq);
+		//FMatrix3d hessian(fourPi_len3, fourPi_len3, fourPi_len3) - c * xqxq;
+		FMatrix3d hessian(
+			fourPi_len3 + c * dpq.X * dpq.X, c * dpq.X * dpq.Y, c * dpq.X * dpq.Z,
+			c * dpq.Y * dpq.X, fourPi_len3 + c * dpq.Y * dpq.Y, c * dpq.Y * dpq.Z,
+			c * dpq.Z * dpq.X, c * dpq.Z * dpq.Y, fourPi_len3 + c * dpq.Z * dpq.Z);
+
+		double Order2 = Order2Coeff.InnerProduct(hessian);
+
+		return Order1 + Order2;
+	}
+
+	// triangle-winding-number first-order approximation.
+	// T is triangle, P is 'Center' of cluster of dipoles, Q is evaluation point
+	// (This is really just for testing)
+	double Order1Approx(const FTriangle3d& T, const FVector3d& P, const FVector3d& XN, double XA, const FVector3d& Q)
+	{
+		FVector3d at0 = XA * XN;
+
+		FVector3d dpq = (P - Q);
+		double len = dpq.Length();
+		double len3 = len * len * len;
+
+		return (1.0 / FMathd::FourPi) * at0.Dot(dpq / (len * len * len));
+	}
+
+	// triangle-winding-number second-order approximation
+	// T is triangle, P is 'Center' of cluster of dipoles, Q is evaluation point
+	// (This is really just for testing)
+	double Order2Approx(const FTriangle3d& T, const FVector3d& P, const FVector3d& XN, double XA, const FVector3d& Q)
+	{
+		FVector3d dpq = (P - Q);
+
+		double len = dpq.Length();
+		double len3 = len * len * len;
+
+		// first-order approximation - integrated_normal_area * \grad(G)
+		double Order1 = (XA / FMathd::FourPi) * XN.Dot(dpq / len3);
+
+		// second-order hessian \grad^2(G)
+		FMatrix3d xqxq(dpq, dpq);
+		xqxq *= 3.0 / (FMathd::FourPi * len3 * len * len);
+		double diag = 1 / (FMathd::FourPi * len3);
+		FMatrix3d hessian = FMatrix3d(diag, diag, diag) - xqxq;
+
+		// second-order LHS - integrated second-order area matrix (formula 26)
+		FVector3d centroid(
+			(T.V[0].X + T.V[1].X + T.V[2].X) / 3.0, (T.V[0].Y + T.V[1].Y + T.V[2].Y) / 3.0, (T.V[0].Z + T.V[1].Z + T.V[2].Z) / 3.0);
+		FVector3d dcp = centroid - P;
+		FMatrix3d o2_lhs(dcp, XN);
+		double Order2 = XA * o2_lhs.InnerProduct(hessian);
+
+		return Order1 + Order2;
+	}
+} // namespace FastTriWinding
+
+/**
+ *  Fast Mesh Winding Number computation
+ *
+ *  TODO: this is tightly coupled to the internal guts of TMeshAABBTree3 (because it is ported from code where it *was* the guts of the AABBTree)
+ *			we should probably at least try to decouple it some
+ */
+template <class TriangleMeshType>
+class TFastWindingTree
+{
+	TMeshAABBTree3<TriangleMeshType>* Tree;
+
+	struct FWNInfo
+	{
+		FVector3d Center;
+		double R;
+		FVector3d Order1Vec;
+		FMatrix3d Order2Mat;
+	};
+
+	TMap<int, FWNInfo> FastWindingCache;
+	int FastWindingCacheTimestamp = -1;
+
+public:
+	/**
+	 * FWN beta parameter - is 2.0 in paper
+	 */
+	double FWNBeta = 2.0;
+
+	/**
+	 * FWN approximation order. Must be 1 or 2. 2 is more accurate, obviously.
+	 */
+	int FWNApproxOrder = 2;
+
+	TFastWindingTree(TMeshAABBTree3<TriangleMeshType>* TreeToRef)
+	{
+		SetTree(TreeToRef);
+	}
+
+	void SetTree(TMeshAABBTree3<TriangleMeshType>* TreeToRef)
+	{
+		this->Tree = TreeToRef;
+		Build(true);
+	}
+
+	void Build(bool bAlwaysBuildRegardlessOfTimestamp = true)
+	{
+		check(Tree);
+		if (Tree->MeshTimestamp != Tree->Mesh->GetShapeTimestamp())
+		{
+			Tree->Build();
+		}
+		if (bAlwaysBuildRegardlessOfTimestamp || FastWindingCacheTimestamp != Tree->MeshTimestamp)
+		{
+			build_fast_winding_cache();
+			FastWindingCacheTimestamp = Tree->MeshTimestamp;
+		}
+	}
+
+	/**
+	 *  Fast approximation of winding number using far-field approximations
+	 */
+	double FastWindingNumber(const FVector3d& P)
+	{
+		Build(false);
+		double sum = branch_fast_winding_num(Tree->RootIndex, P);
+		return sum;
+	}
+
+private:
+	// evaluate winding number contribution for all Triangles below IBox
+	double branch_fast_winding_num(int IBox, FVector3d P)
+	{
+		double branch_sum = 0;
+
+		int idx = Tree->BoxToIndex[IBox];
+		if (idx < Tree->TrianglesEnd)
+		{ // triangle-list case, array is [N t1 t2 ... tN]
+			int num_tris = Tree->IndexList[idx];
+			for (int i = 1; i <= num_tris; ++i)
+			{
+				FVector3d a, b, c;
+				int ti = Tree->IndexList[idx + i];
+				Tree->Mesh->GetTriVertices(ti, a, b, c);
+				double angle = VectorUtil::TriSolidAngle(a, b, c, P);
+				branch_sum += angle / FMathd::FourPi;
+			}
+		}
+		else
+		{ // internal node, either 1 or 2 child boxes
+			int iChild1 = Tree->IndexList[idx];
+			if (iChild1 < 0)
+			{ // 1 child, descend if nearer than cur min-dist
+				iChild1 = (-iChild1) - 1;
+
+				// if we have winding cache, we can more efficiently compute contribution of all Triangles
+				// below this box. Otherwise, recursively descend Tree.
+				bool contained = Tree->box_contains(iChild1, P);
+				if (contained == false && can_use_fast_winding_cache(iChild1, P))
+				{
+					branch_sum += evaluate_box_fast_winding_cache(iChild1, P);
+				}
+				else
+				{
+					branch_sum += branch_fast_winding_num(iChild1, P);
+				}
+			}
+			else
+			{ // 2 children, descend closest first
+				iChild1 = iChild1 - 1;
+				int iChild2 = Tree->IndexList[idx + 1] - 1;
+
+				bool contained1 = Tree->box_contains(iChild1, P);
+				if (contained1 == false && can_use_fast_winding_cache(iChild1, P))
+				{
+					branch_sum += evaluate_box_fast_winding_cache(iChild1, P);
+				}
+				else
+				{
+					branch_sum += branch_fast_winding_num(iChild1, P);
+				}
+
+				bool contained2 = Tree->box_contains(iChild2, P);
+				if (contained2 == false && can_use_fast_winding_cache(iChild2, P))
+				{
+					branch_sum += evaluate_box_fast_winding_cache(iChild2, P);
+				}
+				else
+				{
+					branch_sum += branch_fast_winding_num(iChild2, P);
+				}
+			}
+		}
+
+		return branch_sum;
+	}
+
+	void build_fast_winding_cache()
+	{
+		// set this to a larger number to ignore caches if number of Triangles is too small.
+		// (seems to be no benefit to doing this...is holdover from Tree-decomposition FWN code)
+		int WINDING_CACHE_THRESH = 1;
+
+		//FMeshTriInfoCache TriCache = null;
+		FMeshTriInfoCache TriCache = FMeshTriInfoCache::BuildTriInfoCache(*Tree->Mesh);
+
+		FastWindingCache.Empty(); // = TMap<int, FWNInfo>();
+		TOptional<TSet<int>> root_hash;
+		build_fast_winding_cache(Tree->RootIndex, 0, WINDING_CACHE_THRESH, root_hash, TriCache);
+	}
+	int build_fast_winding_cache(int IBox, int Depth, int TriCountThresh, TOptional<TSet<int>>& TriHash, const FMeshTriInfoCache& TriCache)
+	{
+		TriHash.Reset();
+
+		int idx = Tree->BoxToIndex[IBox];
+		if (idx < Tree->TrianglesEnd)
+		{ // triangle-list case, array is [N t1 t2 ... tN]
+			int num_tris = Tree->IndexList[idx];
+			return num_tris;
+		}
+		else
+		{ // internal node, either 1 or 2 child boxes
+			int iChild1 = Tree->IndexList[idx];
+			if (iChild1 < 0)
+			{ // 1 child, descend if nearer than cur min-dist
+				iChild1 = (-iChild1) - 1;
+				int num_child_tris = build_fast_winding_cache(iChild1, Depth + 1, TriCountThresh, TriHash, TriCache);
+
+				// if count in child is large enough, we already built a cache at lower node
+				return num_child_tris;
+			}
+			else
+			{ // 2 children, descend closest first
+				iChild1 = iChild1 - 1;
+				int iChild2 = Tree->IndexList[idx + 1] - 1;
+
+				// let each child build its own cache if it wants. If so, it will return the
+				// list of its child tris
+				TOptional<TSet<int>> child2_hash;
+				int num_tris_1 = build_fast_winding_cache(iChild1, Depth + 1, TriCountThresh, TriHash, TriCache);
+				int num_tris_2 = build_fast_winding_cache(iChild2, Depth + 1, TriCountThresh, child2_hash, TriCache);
+				bool build_cache = (num_tris_1 + num_tris_2 > TriCountThresh);
+
+				if (Depth == 0)
+				{
+					return num_tris_1 + num_tris_2; // cannot build cache at level 0...
+				}
+
+				// collect up the Triangles we need. there are various cases depending on what children already did
+				if (TriHash.IsSet() || child2_hash.IsSet() || build_cache)
+				{
+					if (!TriHash.IsSet() && child2_hash.IsSet())
+					{
+						collect_triangles(iChild1, child2_hash.GetValue());
+						TriHash = child2_hash;
+					}
+					else
+					{
+						if (!TriHash.IsSet())
+						{
+							TriHash.Emplace();
+							collect_triangles(iChild1, TriHash.GetValue());
+						}
+						if (!child2_hash.IsSet())
+						{
+							collect_triangles(iChild2, TriHash.GetValue());
+						}
+						else
+						{
+							TriHash->Append(child2_hash.GetValue());
+						}
+					}
+				}
+				if (build_cache)
+				{
+					check(TriHash.IsSet());
+					make_box_fast_winding_cache(IBox, TriHash.GetValue(), TriCache);
+				}
+
+				return (num_tris_1 + num_tris_2);
+			}
+		}
+	}
+
+	// check if value is in cache and far enough away from Q that we can use cached value
+	bool can_use_fast_winding_cache(int IBox, const FVector3d& Q)
+	{
+		FWNInfo* cacheInfo = FastWindingCache.Find(IBox);
+		if (!cacheInfo)
+		{
+			return false;
+		}
+
+		double dist_qp = cacheInfo->Center.Distance(Q);
+		if (dist_qp > FWNBeta * cacheInfo->R)
+		{
+			return true;
+		}
+
+		return false;
+	}
+
+	// compute FWN cache for all Triangles underneath this box
+	void make_box_fast_winding_cache(int IBox, const TSet<int>& Triangles, const FMeshTriInfoCache& TriCache)
+	{
+		check(FastWindingCache.Find(IBox) == nullptr);
+
+		// construct cache
+		FWNInfo cacheInfo;
+		FastTriWinding::ComputeCoeffs(*Tree->Mesh, Triangles, TriCache, cacheInfo.Center, cacheInfo.R, cacheInfo.Order1Vec, cacheInfo.Order2Mat);
+
+		FastWindingCache.Add(IBox, cacheInfo);
+	}
+
+	// evaluate the FWN cache for IBox
+	double evaluate_box_fast_winding_cache(int IBox, const FVector3d& Q)
+	{
+		const FWNInfo& cacheInfo = FastWindingCache[IBox];
+
+		if (FWNApproxOrder == 2)
+		{
+			return FastTriWinding::EvaluateOrder2Approx(cacheInfo.Center, cacheInfo.Order1Vec, cacheInfo.Order2Mat, Q);
+		}
+		else
+		{
+			return FastTriWinding::EvaluateOrder1Approx(cacheInfo.Center, cacheInfo.Order1Vec, Q);
+		}
+	}
+
+	// collect all the Triangles below IBox in a hash
+	void collect_triangles(int IBox, TSet<int>& Triangles)
+	{
+		int idx = Tree->BoxToIndex[IBox];
+		if (idx < Tree->TrianglesEnd)
+		{ // triangle-list case, array is [N t1 t2 ... tN]
+			int num_tris = Tree->IndexList[idx];
+			for (int i = 1; i <= num_tris; ++i)
+			{
+				Triangles.Add(Tree->IndexList[idx + i]);
+			}
+		}
+		else
+		{
+			int iChild1 = Tree->IndexList[idx];
+			if (iChild1 < 0)
+			{ // 1 child, descend if nearer than cur min-dist
+				collect_triangles((-iChild1) - 1, Triangles);
+			}
+			else
+			{ // 2 children, descend closest first
+				collect_triangles(iChild1 - 1, Triangles);
+				collect_triangles(Tree->IndexList[idx + 1] - 1, Triangles);
+			}
+		}
+	}
+};
diff --git a/Engine/Plugins/Experimental/GeometryProcessing/Source/GeometricObjects/Public/Spatial/PointHashGrid2.h b/Engine/Plugins/Experimental/GeometryProcessing/Source/GeometricObjects/Public/Spatial/PointHashGrid2.h
new file mode 100644
--- /dev/null
+++ b/Engine/Plugins/Experimental/GeometryProcessing/Source/GeometricObjects/Public/Spatial/PointHashGrid2.h
@@ -0,0 +1,201 @@
+// Copyright 1998-2019 Epic Games, Inc. All Rights Reserved.
+
+// Port of geometry3Sharp PointHashGrid2
+
+#include "CoreMinimal.h"
+#include "Misc/ScopeLock.h"
+#include "Util/GridIndexing2.h"
+
+#pragma once
+
+
+/**
+ * Hash Grid for values associated with 2D points.
+ *
+ * This class addresses the situation where you have a list of (point, point_data) and you
+ * would like to be able to do efficient proximity queries, ie find the nearest point_data
+ * for a given query point.
+ *
+ * We don't store copies of the 2D points. You provide a point_data type. This could just be the
+ * integer index into your list for example, a pointer to something more complex, etc.
+ * Insert and Remove functions require you to pass in the 2D point for the point_data.
+ * To Update a point you need to know it's old and new 2D coordinates.
+ *
+ * @todo also template on point realtype?
+ */
+template<typename PointDataType>
+class TPointHashGrid2d
+{
+private:
+	TMultiMap<FVector2i, PointDataType> Hash;
+	FCriticalSection CriticalSection;
+	FScaleGridIndexer2d Indexer;
+	PointDataType invalidValue;
+
+public:
+
+	/**
+	 * Construct 2D hashgrid
+	 * @param cellSize size of grid cells
+	 * @param invalidValue this value will be returned by queries if no valid result is found (eg bounded-distance query)
+	 */
+	TPointHashGrid2d(double cellSize, PointDataType InvalidValue) : Indexer(cellSize), invalidValue(InvalidValue)
+	{
+	}
+
+	/** Invalid grid value */
+	PointDataType InvalidValue() const
+	{
+		return invalidValue;
+	}
+
+	/**
+	 * Insert at given position. This function is thread-safe.
+	 * @param Value the point/value to insert
+	 * @param Position the position associated with this value
+	 */
+	void InsertPoint(const PointDataType& Value, const FVector2d& Position)
+	{
+		FVector2i idx = Indexer.ToGrid(Position);
+		{
+			FScopeLock Lock(&CriticalSection);
+			Hash.Add(idx, Value);
+		}
+	}
+
+	/**
+	 * Insert at given position, without locking / thread-safety
+	 * @param Value the point/value to insert
+	 * @param Position the position associated with this value
+	 */
+	void InsertPointUnsafe(const PointDataType& Value, const FVector2d& Position)
+	{
+		FVector2i idx = Indexer.ToGrid(Position);
+		Hash.Add(idx, Value);
+	}
+
+
+	/**
+	 * Remove at given position. This function is thread-safe.
+	 * @param Value the point/value to remove
+	 * @param Position the position associated with this value
+	 * @return true if the value existed at this position
+	 */
+	bool RemovePoint(const PointDataType& Value, const FVector2d& Position)
+	{
+		FVector2i idx = Indexer.ToGrid(Position);
+		{
+			FScopeLock Lock(&CriticalSection);
+			return Hash.RemoveSingle(idx, Value) > 0;
+		}
+	}
+
+	/**
+	 * Remove at given position, without locking / thread-safety
+	 * @param Value the point/value to remove
+	 * @param Position the position associated with this value
+	 * @return true if the value existed at this position
+	 */
+	bool RemovePointUnsafe(const PointDataType& Value, const FVector2d& Position)
+	{
+		FVector2i idx = Indexer.ToGrid(Position);
+		return Hash.RemoveSingle(idx, Value) > 0;
+	}
+
+
+	/**
+	 * Move value from old to new position. This function is thread-safe.
+	 * @param Value the point/value to update
+	 * @param OldPosition the current position associated with this value
+	 * @param NewPosition the new position for this value
+	 */
+	void UpdatePoint(const PointDataType& Value, const FVector2d& OldPosition, const FVector2d& NewPosition)
+	{
+		FVector2i old_idx = Indexer.ToGrid(OldPosition);
+		FVector2i new_idx = Indexer.ToGrid(NewPosition);
+		if (old_idx == new_idx)
+		{
+			return;
+		}
+		bool bWasAtOldPos;
+		{
+			FScopeLock Lock(&CriticalSection);
+			bWasAtOldPos = Hash.RemoveSingle(old_idx, Value);
+		}
+		check(bWasAtOldPos);
+		{
+			FScopeLock Lock(&CriticalSection);
+			Hash.Add(new_idx, Value);
+		}
+		return;
+	}
+
+
+	/**
+	 * Move value from old to new position, without locking / thread-safety
+	 * @param Value the point/value to update
+	 * @param OldPosition the current position associated with this value
+	 * @param NewPosition the new position for this value
+	 */
+	void UpdatePointUnsafe(const PointDataType& Value, const FVector2d& OldPosition, const FVector2d& NewPosition)
+	{
+		FVector2i old_idx = Indexer.ToGrid(OldPosition);
+		FVector2i new_idx = Indexer.ToGrid(NewPosition);
+		if (old_idx == new_idx)
+		{
+			return;
+		}
+		bool bWasAtOldPos = Hash.RemoveSingle(old_idx, Value);
+		check(bWasAtOldPos);
+		Hash.Add(new_idx, Value);
+		return;
+	}
+
+
+	/**
+	 * Find nearest point in grid, within a given sphere, without locking / thread-safety.
+	 * @param QueryPoint the center of the query sphere
+	 * @param Radius the radius of the query sphere
+	 * @param DistanceFunc Function you provide which measures the distance between QueryPoint and a Value
+	 * @param IgnoreFunc optional Function you may provide which will result in a Value being ignored if IgnoreFunc(Value) returns true
+	 * @return the found pair (Value,DistanceFunc(Value)), or (InvalidValue,MaxDouble) if not found
+	 */
+	TPair<PointDataType, double> FindNearestInRadius(
+		const FVector2d& QueryPoint, double Radius, 
+		TFunction<double(const PointDataType&)> DistanceFunc,
+		TFunction<bool(const PointDataType&)> IgnoreFunc = [](const PointDataType& data) { return false; }) const
+	{
+		FVector2i min_idx = Indexer.ToGrid(QueryPoint - Radius * FVector2d::One());
+		FVector2i max_idx = Indexer.ToGrid(QueryPoint + Radius * FVector2d::One());
+
+		double min_dist = TNumericLimits<double>::Max();
+		PointDataType nearest = InvalidValue();
+
+		TArray<PointDataType> Values;
+		for (int yi = min_idx.Y; yi <= max_idx.Y; yi++) 
+		{
+			for (int xi = min_idx.X; xi <= max_idx.X; xi++) 
+			{
+				FVector2i idx(xi, yi);
+				Values.Reset();
+				Hash.MultiFind(idx, Values);
+				for (PointDataType Value : Values) 
+				{
+					if (IgnoreFunc(Value))
+					{
+						continue;
+					}
+					double dist = DistanceFunc(Value);
+					if (dist < Radius && dist < min_dist)
+					{
+						nearest = Value;
+						min_dist = dist;
+					}
+				}
+			}
+		}
+
+		return TPair<PointDataType, double>(nearest, min_dist);
+	}
+};
+
