move GeometryProcessing plugin from NotForLicensees to Experimental

diff --git a/Engine/Plugins/Experimental/GeometryProcessing/Source/GeometricObjects/Public/TriangleTypes.h b/Engine/Plugins/Experimental/GeometryProcessing/Source/GeometricObjects/Public/TriangleTypes.h
new file mode 100644
--- /dev/null
+++ b/Engine/Plugins/Experimental/GeometryProcessing/Source/GeometricObjects/Public/TriangleTypes.h
@@ -0,0 +1,165 @@
+// Copyright 1998-2019 Epic Games, Inc. All Rights Reserved.
+
+// Port of WildMagic DistPoint3FTriangle3
+
+#pragma once
+
+#include "VectorTypes.h"
+#include "VectorUtil.h"
+
+/**
+ * Triangle utility functions
+ */
+namespace TriangleUtil
+{
+	/**
+	 * @return the edge length of an equilateral/regular triangle with the given area
+	 */
+	template<typename RealType>
+	RealType EquilateralEdgeLengthForArea(RealType TriArea)
+	{
+		return TMathUtil<RealType>::Sqrt(((RealType)4 * TriArea) / TMathUtil<RealType>::Sqrt3);
+	}
+
+};
+
+
+
+
+template<typename RealType>
+struct FTriangle2
+{
+	FVector2<RealType> V[3];
+
+	FTriangle2() {}
+
+	FTriangle2(const FVector2<RealType>& V0, const FVector2<RealType>& V1, const FVector2<RealType>& V2)
+	{
+		V[0] = V0;
+		V[1] = V1;
+		V[2] = V2;
+	}
+
+	FTriangle2(const FVector2<RealType> VIn[3])
+	{
+		V[0] = VIn[0];
+		V[1] = VIn[1];
+		V[2] = VIn[2];
+	}
+
+	FVector2<RealType> BarycentricPoint(RealType Bary0, RealType Bary1, RealType Bary2) const
+	{
+		return Bary0 * V[0] + Bary1 * V[1] + Bary2 * V[2];
+	}
+
+	FVector2<RealType> BarycentricPoint(const FVector3<RealType>& BaryCoords) const
+	{
+		return BaryCoords[0] * V[0] + BaryCoords[1] * V[1] + BaryCoords[2] * V[2];
+	}
+
+
+	/**
+	 * @param A first vertex of triangle
+	 * @param B second vertex of triangle
+	 * @param C third vertex of triangle
+	 * @return signed area of triangle
+	 */
+	static RealType SignedArea(const FVector2<RealType>& A, const FVector2<RealType>& B, const FVector2<RealType>& C)
+	{
+		return ((RealType)0.5) * ((A.X*B.Y - A.Y*B.X) + (B.X*C.Y - B.Y*C.X) + (C.X*A.Y - C.Y*A.X));
+	}
+
+	/** @return signed area of triangle */
+	RealType SignedArea() const
+	{
+		return SignedArea(V[0], V[1], V[2]);
+	}
+
+
+	/**
+	 * @param A first vertex of triangle
+	 * @param B second vertex of triangle
+	 * @param C third vertex of triangle
+	 * @param QueryPoint test point
+	 * @return true if QueryPoint is inside triangle
+	 */
+	static bool IsInside(const FVector2<RealType>& A, const FVector2<RealType>& B, const FVector2<RealType>& C, const FVector2<RealType>& QueryPoint)
+	{
+		RealType Sign1 = FVector2<RealType>::Orient(A, B, QueryPoint);
+		RealType Sign2 = FVector2<RealType>::Orient(B, C, QueryPoint);
+		RealType Sign3 = FVector2<RealType>::Orient(C, A, QueryPoint);
+		return (Sign1*Sign2 > 0) && (Sign2*Sign3 > 0) && (Sign3*Sign1 > 0);
+	}
+
+	/** @return true if QueryPoint is inside triangle */
+	bool IsInside(const FVector2<RealType>& QueryPoint) const
+	{
+		return IsInside(V[0], V[1], V[2], QueryPoint);
+	}
+
+
+};
+
+typedef FTriangle2<float> FTriangle2f;
+typedef FTriangle2<double> FTriangle2d;
+typedef FTriangle2<int> FTriangle2i;
+
+
+
+
+
+
+template<typename RealType>
+struct FTriangle3
+{
+	FVector3<RealType> V[3];
+
+	FTriangle3() {}
+
+	FTriangle3(const FVector3<RealType>& V0, const FVector3<RealType>& V1, const FVector3<RealType>& V2)
+	{
+		V[0] = V0;
+		V[1] = V1;
+		V[2] = V2;
+	}
+
+	FTriangle3(const FVector3<RealType> VIn[3])
+	{
+		V[0] = VIn[0];
+		V[1] = VIn[1];
+		V[2] = VIn[2];
+	}
+
+	FVector3<RealType> BarycentricPoint(RealType Bary0, RealType Bary1, RealType Bary2) const
+	{
+		return Bary0*V[0] + Bary1*V[1] + Bary2*V[2];
+	}
+
+	FVector3<RealType> BarycentricPoint(const FVector3<RealType> & BaryCoords) const
+	{
+		return BaryCoords[0]*V[0] + BaryCoords[1]*V[1] + BaryCoords[2]*V[2];
+	}
+
+
+	/** @return vector that is perpendicular to the plane of this triangle */
+	FVector3<RealType> Normal() const
+	{
+		return VectorUtil::Normal(V[0], V[1], V[2]);
+	}
+
+	/** @return centroid of this triangle */
+	FVector3<RealType> Centroid() const
+	{
+		constexpr RealType f = 1.0 / 3.0;
+		return FVector3<RealType>(
+			(V[0].X + V[1].X + V[2].X) * f,
+			(V[0].Y + V[1].Y + V[2].Y) * f,
+			(V[0].Z + V[1].Z + V[2].Z) * f
+		);
+	}
+};
+
+typedef FTriangle3<float> FTriangle3f;
+typedef FTriangle3<double> FTriangle3d;
+typedef FTriangle3<int> FTriangle3i;
+
diff --git a/Engine/Plugins/Experimental/GeometryProcessing/Source/GeometricObjects/Public/Util/DynamicVector.h b/Engine/Plugins/Experimental/GeometryProcessing/Source/GeometricObjects/Public/Util/DynamicVector.h
new file mode 100644
--- /dev/null
+++ b/Engine/Plugins/Experimental/GeometryProcessing/Source/GeometricObjects/Public/Util/DynamicVector.h
@@ -0,0 +1,495 @@
+// Copyright 1998-2019 Epic Games, Inc. All Rights Reserved.
+
+// Port of geometry3cpp dvector
+
+#pragma once
+
+#include <CoreMinimal.h>
+#include "VectorTypes.h"
+#include "IndexTypes.h"
+
+#include <array>
+
+
+/*
+ * Blocked array with fixed, power-of-two sized blocks.
+ * 
+ * Iterator functions suitable for use with range-based for are provided
+ */
+template <class Type>
+class FDynamicVector
+{
+public:
+	FDynamicVector()
+	{
+		CurBlock = 0;
+		CurBlockUsed = 0;
+		Blocks.Emplace();
+	}
+
+	FDynamicVector(const FDynamicVector& Copy)
+	{
+		*this = Copy;
+	}
+
+	FDynamicVector(FDynamicVector&& Moved)
+	{
+		*this = MoveTemp(Moved);
+	}
+
+	~FDynamicVector() {}
+
+	const FDynamicVector& operator=(const FDynamicVector& Copy);
+	const FDynamicVector& operator=(FDynamicVector&& Moved);
+
+	inline void Clear();
+	inline void Fill(const Type& Value);
+	inline void Resize(size_t Count);
+	inline void Resize(size_t Count, const Type& InitValue);
+	inline void SetNum(size_t Count) { Resize(Count); }
+
+	inline bool IsEmpty() const { return CurBlock == 0 && CurBlockUsed == 0; }
+	inline size_t GetLength() const { return CurBlock * BlockSize + CurBlockUsed; }
+	inline size_t Num() const { return CurBlock * BlockSize + CurBlockUsed; }
+	inline int GetBlockSize() const { return BlockSize; }
+	inline size_t GetByteCount() const { return Blocks.Num() * BlockSize * sizeof(Type); }
+
+	inline void Add(const Type& Data);
+	inline void Add(const FDynamicVector<Type>& Data);
+	inline void PopBack();
+
+	inline void InsertAt(const Type& Data, unsigned int Index);
+
+	inline const Type& Front() const { return Blocks[0][0]; }
+	inline const Type& Back() const { return Blocks[CurBlock][CurBlockUsed - 1]; }
+
+	inline Type& operator[](unsigned int Index)
+	{
+		return Blocks[Index >> nShiftBits][Index & BlockIndexBitmask];
+	}
+	inline const Type& operator[](unsigned int Index) const
+	{
+		return Blocks[Index >> nShiftBits][Index & BlockIndexBitmask];
+	}
+
+	// apply f() to each member sequentially
+	template <typename Func>
+	void Apply(const Func& f);
+
+
+protected:
+	// [RMS] BlockSize must be a power-of-two, so we can use bit-shifts in operator[]
+	static constexpr int BlockSize = 2048; // (1 << 11)
+	static constexpr int nShiftBits = 11;
+	static constexpr int BlockIndexBitmask = 2047; // low 11 bits
+
+	unsigned int CurBlock;
+	unsigned int CurBlockUsed;
+
+	using BlockType = std::array<Type, BlockSize>;
+	TArray<BlockType> Blocks;
+
+	friend class FIterator;
+
+
+public:
+	/*
+	 * FIterator class iterates over values of vector
+	 */
+	class FIterator
+	{
+	public:
+		inline const Type& operator*() const
+		{
+			return (*DVector)[Idx];
+		}
+		inline Type& operator*()
+		{
+			return (*DVector)[Idx];
+		}
+		inline FIterator& operator++()   // prefix
+		{
+			Idx++;
+			return *this;
+		}
+		inline FIterator operator++(int) // postfix
+		{
+			FIterator Copy(*this);
+			Idx++;
+			return Copy;
+		}
+		inline bool operator==(const FIterator& Itr2)
+		{
+			return DVector == Itr2.DVector && Idx == Itr2.Idx;
+		}
+		inline bool operator!=(const FIterator& Itr2)
+		{
+			return DVector != Itr2.DVector || Idx != Itr2.Idx;
+		}
+
+	protected:
+		FDynamicVector<Type>* DVector;
+		int Idx;
+		inline FIterator(FDynamicVector<Type>* Parent, int ICur)
+		{
+			DVector = Parent;
+			Idx = ICur;
+		}
+		friend class FDynamicVector<Type>;
+	};
+
+	/** @return iterator at beginning of vector */
+	FIterator begin()
+	{
+		return IsEmpty() ? end() : FIterator(this, 0);
+	}
+	/** @return iterator at end of vector */
+	FIterator end()
+	{
+		return FIterator(this, (int)GetLength());
+	}
+};
+
+
+
+
+
+
+template <class Type, int N>
+class FDynamicVectorN
+{
+public:
+	FDynamicVectorN()
+	{
+	}
+	FDynamicVectorN(const FDynamicVectorN& Copy)
+		: Data(Copy.Data)
+	{
+	}
+	FDynamicVectorN(FDynamicVectorN&& Moved)
+		: Data(MoveTemp(Moved.Data))
+	{
+	}
+	~FDynamicVectorN()
+	{
+	}
+
+	const FDynamicVectorN& operator=(const FDynamicVectorN& Copy)
+	{
+		Data = Copy.Data;
+		return *this;
+	}
+	const FDynamicVectorN& operator=(FDynamicVectorN&& Moved)
+	{
+		Data = MoveTemp(Moved.Data);
+		return *this;
+	}
+
+	inline void Clear()
+	{
+		Data.Clear();
+	}
+	inline void Fill(const Type& Value)
+	{
+		Data.Fill(Value);
+	}
+	inline void Resize(size_t Count)
+	{
+		Data.Resize(Count * N);
+	}
+	inline void Resize(size_t Count, const Type& InitValue)
+	{
+		Data.Resize(Count * N, InitValue);
+	}
+
+	inline bool IsEmpty() const
+	{
+		return Data.IsEmpty();
+	}
+	inline size_t GetLength() const
+	{
+		return Data.GetLength() / N;
+	}
+	inline int GetBlockSize() const
+	{
+		return Data.GetBlockSize();
+	}
+	inline size_t GetByteCount() const
+	{
+		return Data.GetByteCount();
+	}
+
+	inline void Add(const std::array<Type, N>& AddData)
+	{
+		// todo specialize for N=2,3,4
+		for (int i = 0; i < N; i++)
+		{
+			Data.Add(AddData[i]);
+		}
+	}
+
+	inline void PopBack()
+	{
+		for (int i = 0; i < N; i++)
+		{
+			PopBack();
+		}
+	} // TODO specialize
+
+	inline void InsertAt(const std::array<Type, N>& AddData, unsigned int Index)
+	{
+		for (int i = 1; i <= N; i++)
+		{
+			Data.InsertAt(AddData[N - i], N * (Index + 1) - i);
+		}
+	}
+
+	inline Type& operator()(unsigned int TopIndex, unsigned int SubIndex)
+	{
+		return Data[TopIndex * N + SubIndex];
+	}
+	inline const Type& operator()(unsigned int TopIndex, unsigned int SubIndex) const
+	{
+		return Data[TopIndex * N + SubIndex];
+	}
+	inline void SetVector2(unsigned int TopIndex, const FVector2<Type>& V)
+	{
+		check(N >= 2);
+		unsigned int i = TopIndex * N;
+		Data[i] = V.X;
+		Data[i + 1] = V.Y;
+	}
+	inline void SetVector3(unsigned int TopIndex, const FVector3<Type>& V)
+	{
+		check(N >= 3);
+		unsigned int i = TopIndex * N;
+		Data[i] = V.X;
+		Data[i + 1] = V.Y;
+		Data[i + 2] = V.Z;
+	}
+	inline FVector2<Type> AsVector2(unsigned int TopIndex) const
+	{
+		check(N >= 2);
+		return FVector2<Type>(
+			Data[TopIndex * N + 0],
+			Data[TopIndex * N + 1]);
+	}
+	inline FVector3<Type> AsVector3(unsigned int TopIndex) const
+	{
+		check(N >= 3);
+		return FVector3<Type>(
+			Data[TopIndex * N + 0],
+			Data[TopIndex * N + 1],
+			Data[TopIndex * N + 2]);
+	}
+	inline FIndex2i AsIndex2(unsigned int TopIndex) const
+	{
+		check(N >= 2);
+		return FIndex2i(
+			Data[TopIndex * N + 0],
+			Data[TopIndex * N + 1]);
+	}
+	inline FIndex3i AsIndex3(unsigned int TopIndex) const
+	{
+		check(N >= 3);
+		return FIndex3i(
+			Data[TopIndex * N + 0],
+			Data[TopIndex * N + 1],
+			Data[TopIndex * N + 2]);
+	}
+	inline FIndex4i AsIndex4(unsigned int TopIndex) const
+	{
+		check(N >= 4);
+		return FIndex4i(
+			Data[TopIndex * N + 0],
+			Data[TopIndex * N + 1],
+			Data[TopIndex * N + 2],
+			Data[TopIndex * N + 3]);
+	}
+
+protected:
+	FDynamicVector<Type> Data;
+
+	friend class FIterator;
+};
+
+template class FDynamicVectorN<double, 2>;
+
+typedef FDynamicVectorN<float, 3> FDynamicVector3f;
+typedef FDynamicVectorN<float, 2> FDynamicVector2f;
+typedef FDynamicVectorN<double, 3> FDynamicVector3d;
+typedef FDynamicVectorN<double, 2> FDynamicVector2d;
+typedef FDynamicVectorN<int, 3> FDynamicVector3i;
+typedef FDynamicVectorN<int, 2> FDynamicVector2i;
+
+
+
+
+
+template <class Type>
+const FDynamicVector<Type>& FDynamicVector<Type>::operator=(const FDynamicVector& Copy)
+{
+	Blocks = Copy.Blocks;
+	CurBlock = Copy.CurBlock;
+	CurBlockUsed = Copy.CurBlockUsed;
+	return *this;
+}
+
+template <class Type>
+const FDynamicVector<Type>& FDynamicVector<Type>::operator=(FDynamicVector&& Moved)
+{
+	Blocks = MoveTemp(Moved.Blocks);
+	CurBlock = Moved.CurBlock;
+	CurBlockUsed = Moved.CurBlockUsed;
+	return *this;
+}
+
+template <class Type>
+void FDynamicVector<Type>::Clear()
+{
+	Blocks.Empty();
+	CurBlock = 0;
+	CurBlockUsed = 0;
+	Blocks.Add(BlockType());
+}
+
+template <class Type>
+void FDynamicVector<Type>::Fill(const Type& Value)
+{
+	size_t Count = Blocks.Num();
+	for (unsigned int i = 0; i < Count; ++i)
+	{
+		Blocks[i].fill(Value);
+	}
+}
+
+template <class Type>
+void FDynamicVector<Type>::Resize(size_t Count)
+{
+	if (GetLength() == Count)
+	{
+		return;
+	}
+
+	// figure out how many segments we need
+	int nNumSegs = 1 + (int)Count / BlockSize;
+
+	// figure out how many are currently allocated...
+	size_t nCurCount = Blocks.Num();
+
+	// erase extra segments memory
+	// [RMS] not necessary right? std::array will deallocate?
+	//for (int i = nNumSegs; i < nCurCount; ++i)
+	//	Blocks[i] = null;
+
+	// resize to right number of segments
+	if (nNumSegs >= Blocks.Num())
+	{
+		// allocate new segments
+		for (int i = (int)nCurCount; i < nNumSegs; ++i)
+		{
+			Blocks.Emplace();
+		}
+	}
+	else
+	{
+		//Blocks.RemoveRange(nNumSegs, Blocks.Count - nNumSegs);
+		Blocks.SetNum(nNumSegs);
+	}
+
+	// mark last segment
+	CurBlockUsed = (unsigned int)(Count - (nNumSegs - 1) * BlockSize);
+	CurBlock = nNumSegs - 1;
+}
+
+template <class Type>
+void FDynamicVector<Type>::Resize(size_t Count, const Type& InitValue)
+{
+	size_t nCurSize = GetLength();
+	Resize(Count);
+	for (size_t Index = nCurSize; Index < Count; ++Index)
+	{
+		this->operator[](Index) = InitValue;
+	}
+}
+
+template <class Type>
+void FDynamicVector<Type>::Add(const Type& Value)
+{
+	if (CurBlockUsed == BlockSize)
+	{
+		if (CurBlock == Blocks.Num() - 1)
+		{
+			Blocks.Emplace();
+		}
+		CurBlock++;
+		CurBlockUsed = 0;
+	}
+	Blocks[CurBlock][CurBlockUsed] = Value;
+	CurBlockUsed++;
+}
+
+
+template <class Type>
+void FDynamicVector<Type>::Add(const FDynamicVector<Type>& AddData)
+{
+	// @todo it could be more efficient to use memcopies here...
+	size_t nSize = AddData.Num();
+	for (unsigned int k = 0; k < nSize; ++k)
+	{
+		Add(AddData[k]);
+	}
+}
+
+template <class Type>
+void FDynamicVector<Type>::PopBack()
+{
+	if (CurBlockUsed > 0)
+	{
+		CurBlockUsed--;
+	}
+	if (CurBlockUsed == 0 && CurBlock > 0)
+	{
+		CurBlock--;
+		CurBlockUsed = BlockSize;
+		// remove block ??
+	}
+}
+
+template <class Type>
+void FDynamicVector<Type>::InsertAt(const Type& AddData, unsigned int Index)
+{
+	size_t s = GetLength();
+	if (Index == s)
+	{
+		Add(AddData);
+	}
+	else if (Index > s)
+	{
+		Resize(Index);
+		Add(AddData);
+	}
+	else
+	{
+		(*this)[Index] = AddData;
+	}
+}
+
+template <typename Type>
+template <typename Func>
+void FDynamicVector<Type>::Apply(const Func& f)
+{
+	for (int bi = 0; bi < CurBlock; ++bi)
+	{
+		auto block = Blocks[bi];
+		for (int k = 0; k < BlockSize; ++k)
+		{
+			applyF(block[k], k);
+		}
+	}
+	auto lastblock = Blocks[CurBlock];
+	for (int k = 0; k < CurBlockUsed; ++k)
+	{
+		applyF(lastblock[k], k);
+	}
+}
diff --git a/Engine/Plugins/Experimental/GeometryProcessing/Source/GeometricObjects/Public/Util/GridIndexing2.h b/Engine/Plugins/Experimental/GeometryProcessing/Source/GeometricObjects/Public/Util/GridIndexing2.h
new file mode 100644
--- /dev/null
+++ b/Engine/Plugins/Experimental/GeometryProcessing/Source/GeometricObjects/Public/Util/GridIndexing2.h
@@ -0,0 +1,82 @@
+// Copyright 1998-2019 Epic Games, Inc. All Rights Reserved.
+
+// Port of geometry3Sharp GridIndexing2
+
+#pragma once
+
+#include "VectorTypes.h"
+
+
+/**
+ * Convert between integer grid coordinates and scaled real-valued coordinates (ie assumes integer grid origin == real origin)
+ */
+template<typename RealType>
+struct TScaleGridIndexer2
+{
+	/** Real-valued size of an integer grid cell */
+	RealType CellSize;
+
+	TScaleGridIndexer2(RealType CellSize) : CellSize(CellSize) {}
+
+	/** Convert real-valued point to integer grid coordinates */
+	inline FVector2i ToGrid(const FVector2<RealType>& P) const
+	{
+		return FVector2i((int)P.X / CellSize, (int)P.Y / CellSize);
+	}
+
+	/** Convert integer grid coordinates to real-valued point */
+	inline FVector2<RealType> FromGrid(const FVector2i& GridPoint) const
+	{
+		return FVector2<RealType>(GridPoint.X*CellSize, GridPoint.Y*CellSize);
+	}
+};
+typedef TScaleGridIndexer2<float> FScaleGridIndexer2f;
+typedef TScaleGridIndexer2<double> FScaleGridIndexer2d;
+
+
+/**
+ * Convert between integer grid coordinates and scaled+translated real-valued coordinates
+ */
+template<typename RealType>
+struct TShiftGridIndexer2
+{
+	/** Real-valued size of an integer grid cell */
+	RealType CellSize;
+	/** Real-valued origin of grid, position of integer grid origin */
+	FVector2<RealType> Origin;
+
+	TShiftGridIndexer2(const FVector2<RealType>& origin, RealType cellSize)
+	{
+		Origin = origin;
+		CellSize = cellSize;
+	}
+
+	/** Convert real-valued point to integer grid coordinates */
+	inline FVector2i ToGrid(const FVector2<RealType>& point) const
+	{
+		return FVector2i(
+			(int)((point.X - Origin.X) / CellSize),
+			(int)((point.Y - Origin.Y) / CellSize));
+	}
+
+	/** Convert integer grid coordinates to real-valued point */
+	inline FVector2<RealType> FromGrid(const FVector2i& gridpoint) const
+	{
+		return FVector2<RealType>(
+			((RealType)gridpoint.X * CellSize) + Origin.X,
+			((RealType)gridpoint.Y * CellSize) + Origin.Y);
+	}
+
+	/** Convert real-valued grid coordinates to real-valued point */
+	inline FVector2<RealType> FromGrid(const FVector2<RealType>& gridpointf) const
+	{
+		return FVector2<RealType>(
+			((RealType)gridpointf.X * CellSize) + Origin.X,
+			((RealType)gridpointf.Y * CellSize) + Origin.Y);
+	}
+};
+typedef TShiftGridIndexer2<float> FShiftGridIndexer2f;
+typedef TShiftGridIndexer2<double> FShiftGridIndexer2d;
+
+
+
diff --git a/Engine/Plugins/Experimental/GeometryProcessing/Source/GeometricObjects/Public/Util/GridIndexing3.h b/Engine/Plugins/Experimental/GeometryProcessing/Source/GeometricObjects/Public/Util/GridIndexing3.h
new file mode 100644
--- /dev/null
+++ b/Engine/Plugins/Experimental/GeometryProcessing/Source/GeometricObjects/Public/Util/GridIndexing3.h
@@ -0,0 +1,96 @@
+// Copyright 1998-2019 Epic Games, Inc. All Rights Reserved.
+
+// Port of geometry3Sharp GridIndexing3
+
+#pragma once
+
+#include "VectorTypes.h"
+
+
+/**
+ * Convert between integer grid coordinates and scaled real-valued coordinates (ie assumes integer grid origin == real origin)
+ */
+template<typename RealType>
+struct TScaleGridIndexer3
+{
+	/** Real-valued size of an integer grid cell */
+	RealType CellSize;
+
+	TScaleGridIndexer3() : CellSize((RealType)1) 
+	{
+	}
+	
+	TScaleGridIndexer3(RealType CellSize) : CellSize(CellSize) 
+	{
+	}
+
+	/** Convert real-valued point to integer grid coordinates */
+	inline FVector3i ToGrid(const FVector3<RealType>& P) const
+	{
+		return FVector3i((int)(P.X / CellSize), (int)(P.Y / CellSize), (int)(P.Z / CellSize));
+	}
+
+	/** Convert integer grid coordinates to real-valued point */
+	inline FVector3<RealType> FromGrid(const FVector3i& GridPoint) const
+	{
+		return FVector3<RealType>(GridPoint.X*CellSize, GridPoint.Y*CellSize, GridPoint.Z*CellSize);
+	}
+};
+typedef TScaleGridIndexer3<float> FScaleGridIndexer3f;
+typedef TScaleGridIndexer3<double> FScaleGridIndexer3d;
+
+
+
+/**
+ * Convert between integer grid coordinates and scaled+translated real-valued coordinates
+ */
+template<typename RealType>
+struct TShiftGridIndexer3
+{
+	/** Real-valued size of an integer grid cell */
+	RealType CellSize;
+	/** Real-valued origin of grid, position of integer grid origin */
+	FVector3<RealType> Origin;
+
+	TShiftGridIndexer3() 
+		: CellSize((RealType)1), Origin(FVector3<RealType>::Zero())
+	{
+	}
+
+	TShiftGridIndexer3(const FVector3<RealType>& origin, RealType cellSize)
+		: CellSize(cellSize), Origin(origin)
+	{
+	}
+
+	/** Convert real-valued point to integer grid coordinates */
+	inline FVector3i ToGrid(const FVector3<RealType>& point) const
+	{
+		return FVector3i(
+			(int)((point.X - Origin.X) / CellSize),
+			(int)((point.Y - Origin.Y) / CellSize),
+			(int)((point.Z - Origin.Z) / CellSize));
+	}
+
+	/** Convert integer grid coordinates to real-valued point */
+	inline FVector3<RealType> FromGrid(const FVector3i& gridpoint) const
+	{
+		return FVector3<RealType>(
+			((RealType)gridpoint.X * CellSize) + Origin.X,
+			((RealType)gridpoint.Y * CellSize) + Origin.Y,
+			((RealType)gridpoint.Z * CellSize) + Origin.Z);
+	}
+
+	/** Convert real-valued grid coordinates to real-valued point */
+	inline FVector3<RealType> FromGrid(const FVector3<RealType>& gridpointf)  const
+	{
+		return FVector3d(
+			((RealType)gridpointf.X * CellSize) + Origin.X,
+			((RealType)gridpointf.Y * CellSize) + Origin.Y,
+			((RealType)gridpointf.Z * CellSize) + Origin.Z);
+	}
+};
+typedef TShiftGridIndexer3<float> FShiftGridIndexer3f;
+typedef TShiftGridIndexer3<double> FShiftGridIndexer3d;
+
+
+
diff --git a/Engine/Plugins/Experimental/GeometryProcessing/Source/GeometricObjects/Public/Util/MeshCaches.h b/Engine/Plugins/Experimental/GeometryProcessing/Source/GeometricObjects/Public/Util/MeshCaches.h
new file mode 100644
--- /dev/null
+++ b/Engine/Plugins/Experimental/GeometryProcessing/Source/GeometricObjects/Public/Util/MeshCaches.h
@@ -0,0 +1,40 @@
+// Copyright 1998-2019 Epic Games, Inc. All Rights Reserved.
+
+#include "DynamicVector.h"
+
+#include "Async/ParallelFor.h"
+
+/*
+ * Basic cache of per-triangle information for a mesh
+ */
+struct FMeshTriInfoCache
+{
+	FDynamicVector<FVector3d> Centroids;
+	FDynamicVector<FVector3d> Normals;
+	FDynamicVector<double> Areas;
+
+	void GetTriInfo(int TriangleID, FVector3d& NormalOut, double& AreaOut, FVector3d& CentroidOut) const
+	{
+		NormalOut = Normals[TriangleID];
+		AreaOut = Areas[TriangleID];
+		CentroidOut = Centroids[TriangleID];
+	}
+
+	template<class TriangleMeshType>
+	static FMeshTriInfoCache BuildTriInfoCache(const TriangleMeshType& Mesh)
+	{
+		FMeshTriInfoCache Cache;
+		int NT = Mesh.TriangleCount();
+		Cache.Centroids.Resize(NT);
+		Cache.Normals.Resize(NT);
+		Cache.Areas.Resize(NT);
+
+		ParallelFor(NT, [&](int TriIdx)
+		{
+			TMeshQueries<TriangleMeshType>::GetTriNormalAreaCentroid(Mesh, TriIdx, Cache.Normals[TriIdx], Cache.Areas[TriIdx], Cache.Centroids[TriIdx]);
+		});
+
+		return Cache;
+	}
+};
+
diff --git a/Engine/Plugins/Experimental/GeometryProcessing/Source/GeometricObjects/Public/Util/ProgressCancel.h b/Engine/Plugins/Experimental/GeometryProcessing/Source/GeometricObjects/Public/Util/ProgressCancel.h
new file mode 100644
--- /dev/null
+++ b/Engine/Plugins/Experimental/GeometryProcessing/Source/GeometricObjects/Public/Util/ProgressCancel.h
@@ -0,0 +1,76 @@
+// Copyright 1998-2019 Epic Games, Inc. All Rights Reserved.
+
+// Port of geometry3cpp ProgressCancel
+
+#pragma once
+
+
+/**
+ * ICancelSource is an object that provides a cancel function
+ */
+class ICancelSource
+{
+public:
+	virtual ~ICancelSource() {}
+
+	/**
+	 * @return true if object wishes to cancel expensive operation
+	 */
+	virtual bool Cancelled() = 0;
+};
+
+
+/**
+ * FCancelFunction uses a TFunction to implement the ICancelSource interface
+ */
+class FCancelFunction : public ICancelSource
+{
+public:
+	virtual ~FCancelFunction() {}
+
+	/** function that returns true if object wishes to cancel operation */
+	TFunction<bool()> CancelF;
+
+	FCancelFunction(const TFunction<bool()> & cancelF) : CancelF(cancelF)
+	{
+	}
+
+	/**
+	 * @return true if object wishes to cancel expensive operation
+	 */
+	bool Cancelled() { return CancelF(); }
+};
+
+
+
+/**
+ * FProgressCancel is an obejct that is intended to be passed to long-running
+ * computes to do two things:
+ * 1) provide progress info back to caller (not ported yet)
+ * 2) allow caller to cancel the computation
+ */
+class FProgressCancel
+{
+public:
+	TSharedPtr<ICancelSource> Source;
+
+	bool WasCancelled = false;  // will be set to true if CancelF() ever returns true
+
+	FProgressCancel(TSharedPtr<ICancelSource> source)
+	{
+		Source = source;
+	}
+
+	/**
+	 * @return true if client would like to cancel operation
+	 */
+	bool Cancelled()
+	{
+		if (WasCancelled)
+		{
+			return true;
+		}
+		WasCancelled = Source->Cancelled();
+		return WasCancelled;
+	}
+};
diff --git a/Engine/Plugins/Experimental/GeometryProcessing/Source/GeometricObjects/Public/Util/SmallListSet.h b/Engine/Plugins/Experimental/GeometryProcessing/Source/GeometricObjects/Public/Util/SmallListSet.h
new file mode 100644
--- /dev/null
+++ b/Engine/Plugins/Experimental/GeometryProcessing/Source/GeometricObjects/Public/Util/SmallListSet.h
@@ -0,0 +1,384 @@
+// Copyright 1998-2019 Epic Games, Inc. All Rights Reserved.
+
+// Port of geometry3cpp small_list_set
+
+#pragma once
+
+#include "Util/DynamicVector.h"
+
+/**
+ * FSmallListSet stores a set of short integer-valued variable-size lists.
+ * The lists are encoded into a few large FDynamicVector buffers, with internal pooling,
+ * so adding/removing lists usually does not involve any or delete ops.
+ * 
+ * The lists are stored in two parts. The first N elements are stored in a linear
+ * subset of a FDynamicVector. If the list spills past these N elements, the extra elements
+ * are stored in a linked list (which is also stored in a flat array).
+ * 
+ * Each list stores its count, so list-size operations are constant time.
+ * All the internal "pointers" are 32-bit.
+ * 
+ * @todo look at usage of TFunction, are we making unnecessary copies?
+ */
+class GEOMETRICOBJECTS_API FSmallListSet
+{
+protected:
+	/** This value is used to indicate Null in internal pointers */
+	static const int32 NullValue; // = -1;		// note: cannot be constexpr because we pass as reference to several functions, requires C++17
+
+	/** size of initial linear-memory portion of lists */
+	static constexpr int32 BLOCKSIZE = 8;
+	/** offset from start of linear-memory portion of list that contains pointer to head of variable-length linked list */
+	static constexpr int32 BLOCK_LIST_OFFSET = BLOCKSIZE + 1;
+
+	/** mapping from list index to offset into block_store that contains list data */
+	FDynamicVector<int32> ListHeads;
+
+	/** 
+	 * flat buffer used to store per-list linear-memory blocks. 
+	 * blocks are BLOCKSIZE+2 long, elements are [CurrentCount, item0...itemN, LinkedListPtr]
+	 */
+	FDynamicVector<int32> ListBlocks; 
+
+	/** list of free blocks as indices/offsets into block_store */
+	FDynamicVector<int32> FreeBlocks;
+
+	/** number of allocated lists */
+	int32 AllocatedCount = 0;
+
+	/**
+	 * flat buffer used to store linked-list "spill" elements
+	 * each element is [value, next_ptr]
+	 */
+	FDynamicVector<int32> LinkedListElements;
+
+	/** index of first free element in linked_store */
+	int32 FreeHeadIndex;
+
+
+public:
+	FSmallListSet();
+
+	FSmallListSet(const FSmallListSet& copy);
+
+
+	/**
+	 * @return largest current list index
+	 */
+	size_t Size() const 
+	{
+		return ListHeads.GetLength();
+	}
+
+	/**
+	 * set new number of lists
+	 */
+	void Resize(int32 NewSize);
+
+
+	/**
+	 * @return true if a list has been allocated at the given ListIndex
+	 */
+	bool IsAllocated(int32 ListIndex) const
+	{
+		return (ListIndex >= 0 && ListIndex < (int32)ListHeads.GetLength() && ListHeads[ListIndex] != NullValue);
+	}
+
+	/**
+	 * Create a list at the given ListIndex
+	 */
+	void AllocateAt(int32 ListIndex);
+
+
+	/**
+	 * Insert Value into list at ListIndex
+	 */
+	void Insert(int32 ListIndex, int32 Value);
+
+
+
+	/**
+	 * remove Value from the list at ListIndex
+	 * @return false if Value was not in this list
+	 */
+	bool Remove(int32 ListIndex, int32 Value);
+
+
+
+	/**
+	 * Move list at FromIndex to ToIndex
+	 */
+	void Move(int32 FromIndex, int32 ToIndex);
+
+
+
+	/**
+	 * Remove all elements from the list at ListIndex
+	 */
+	void Clear(int32 ListIndex);
+
+
+	/**
+	 * @return the size of the list at ListIndex
+	 */
+	inline int32 GetCount(int32 ListIndex) const
+	{
+		check(ListIndex >= 0);
+		int32 block_ptr = ListHeads[ListIndex];
+		return (block_ptr == NullValue) ? 0 : ListBlocks[block_ptr];
+	}
+
+
+	/**
+	 * @return the first item in the list at ListIndex
+	 * @warning does not check for zero-size-list!
+	 */
+	inline int32 First(int32 ListIndex) const
+	{
+		check(ListIndex >= 0);
+		int32 block_ptr = ListHeads[ListIndex];
+		return ListBlocks[block_ptr + 1];
+	}
+
+
+	/**
+	 * Search for the given Value in list at ListIndex
+	 * @return true if found
+	 */
+	bool Contains(int32 ListIndex, int32 Value) const;
+
+
+	/**
+	 * Search the list at ListIndex for a value where PredicateFunc(value) returns true
+	 * @return the found value, or the InvalidValue argument if not found
+	 */
+	int32 Find(int32 ListIndex, const TFunction<bool(int32)>& PredicateFunc, int32 InvalidValue = -1) const;
+
+
+
+	/**
+	 * Search the list at ListIndex for a value where PredicateFunc(value) returns true, and replace it with NewValue
+	 * @return true if the value was found and replaced
+	 */
+	bool Replace(int32 ListIndex, const TFunction<bool(int32)>& PredicateFunc, int32 NewValue);
+
+
+
+	//
+	// iterator support
+	// 
+
+
+	friend class ValueIterator;
+
+	/**
+	 * ValueIterator iterates over the values of a small list
+	 * An optional mapping function can be provided which will then be applied to the values returned by the * operator
+	 */
+	class ValueIterator
+	{
+	public:
+		inline ValueIterator() 
+		{ 
+			ListSet = nullptr;
+			MapFunc = nullptr; 
+			ListIndex = 0; 
+		}
+
+		inline bool operator==(const ValueIterator& Other) const 
+		{
+			return ListSet == Other.ListSet && ListIndex == Other.ListIndex;
+		}
+		inline bool operator!=(const ValueIterator& Other) const 
+		{
+			return ListSet != Other.ListSet || ListIndex != Other.ListIndex || iCur != Other.iCur || cur_ptr != Other.cur_ptr;
+		}
+
+		inline int32 operator*() const 
+		{
+			return (MapFunc == nullptr) ? cur_value : MapFunc(cur_value);
+		}
+
+		inline const ValueIterator& operator++() 		// prefix
+		{
+			this->GotoNext();
+			return *this;
+		}
+		//inline ValueIterator operator++(int32) {		// postfix
+		//	index_iterator copy(*this);
+		//	this->GotoNext();
+		//	return copy;
+		//}
+
+
+	protected:
+		inline void GotoNext() 
+		{
+			if (N == 0)
+			{
+				SetToEnd();
+				return;
+			}
+			GotoNextOverflow();
+		}
+
+		inline void GotoNextOverflow() 
+		{
+			if (iCur <= iEnd) 
+			{
+				cur_value = ListSet->ListBlocks[iCur];
+				iCur++;
+			}
+			else if (cur_ptr != NullValue) 
+			{
+				cur_value = ListSet->LinkedListElements[cur_ptr];
+				cur_ptr = ListSet->LinkedListElements[cur_ptr + 1];
+			}
+			else
+			{
+				SetToEnd();
+			}
+		}
+
+		inline ValueIterator(
+			const FSmallListSet* ListSetIn,
+			int32 ListIndex, 
+			bool is_end,
+			const TFunction<int32(int32)>& MapFuncIn = nullptr)
+		{
+			this->ListSet = ListSetIn;
+			this->MapFunc = MapFuncIn;
+			this->ListIndex = ListIndex;
+			if (is_end) 
+			{
+				SetToEnd();
+			}
+			else 
+			{
+				block_ptr = ListSet->ListHeads[ListIndex];
+				if (block_ptr != ListSet->NullValue)
+				{
+					N = ListSet->ListBlocks[block_ptr];
+					iEnd = (N < BLOCKSIZE) ? (block_ptr + N) : (block_ptr + BLOCKSIZE);
+					iCur = block_ptr + 1;
+					cur_ptr = (N < BLOCKSIZE) ? NullValue : ListSet->ListBlocks[block_ptr + BLOCK_LIST_OFFSET];
+					GotoNext();
+				}
+				else
+				{
+					SetToEnd();
+				}
+			}
+		}
+
+		inline void SetToEnd() 
+		{
+			block_ptr = ListSet->NullValue;
+			N = 0;
+			iCur = -1;
+			cur_ptr = -1;
+		}
+
+		const FSmallListSet * ListSet;
+		TFunction<int32(int32)> MapFunc;
+		int32 ListIndex;
+		int32 block_ptr;
+		int32 N;
+		int32 iEnd;
+		int32 iCur;
+		int32 cur_ptr;
+		int32 cur_value;
+		friend class FSmallListSet;
+	};
+
+	/**
+	 * @return iterator for start of list at ListIndex
+	 */
+	inline ValueIterator BeginValues(int32 ListIndex) const 
+	{
+		return ValueIterator(this, ListIndex, false);
+	}
+
+	/**
+	 * @return iterator for start of list at ListIndex, with given value mapping function
+	 */
+	inline ValueIterator BeginValues(int32 ListIndex, const TFunction<int32(int32)>& MapFunc) const 
+	{
+		return ValueIterator(this, ListIndex, false, MapFunc);
+	}
+
+	/**
+	 * @return iterator for end of list at ListIndex
+	 */
+	inline ValueIterator EndValues(int32 ListIndex) const 
+	{
+		return ValueIterator(this, ListIndex, true);
+	}
+
+	/**
+	 * ValueEnumerable is an object that provides begin/end semantics for a small list, suitable for use with a range-based for loop
+	 */
+	class ValueEnumerable
+	{
+	public:
+		const FSmallListSet* ListSet;
+		int32 ListIndex;
+		TFunction<int32(int32)> MapFunc;
+		ValueEnumerable() {}
+		ValueEnumerable(const FSmallListSet* ListSetIn, int32 ListIndex, TFunction<int32(int32)> MapFunc = nullptr)
+		{
+			this->ListSet = ListSetIn;
+			this->ListIndex = ListIndex;
+			this->MapFunc = MapFunc;
+		}
+		typename FSmallListSet::ValueIterator begin() const { return ListSet->BeginValues(ListIndex, MapFunc); }
+		typename FSmallListSet::ValueIterator end() const { return ListSet->EndValues(ListIndex); }
+	};
+
+	/**
+	 * @return a value enumerable for the given ListIndex
+	 */
+	inline ValueEnumerable Values(int32 ListIndex) const
+	{
+		return ValueEnumerable(this, ListIndex);
+	}
+
+	/**
+	 * @return a value enumerable for the given ListIndex, with the given value mapping function
+	 */
+	inline ValueEnumerable Values(int32 ListIndex, const TFunction<int32(int32)>& MapFunc) const 
+	{
+		return ValueEnumerable(this, ListIndex, MapFunc);
+	}
+
+
+
+
+protected:
+
+
+	// grab a block from the free list, or allocate a one
+	int32 AllocateBlock();
+
+	// push a link-node onto the free list
+	inline void AddFreeLink(int32 ptr)
+	{
+		LinkedListElements[ptr + 1] = FreeHeadIndex;
+		FreeHeadIndex = ptr;
+	}
+
+
+	// remove val from the linked-list attached to block_ptr
+	bool RemoveFromLinkedList(int32 block_ptr, int32 val);
+
+
+public:
+	inline FString MemoryUsage()
+	{
+		return FString::Printf(TEXT("ListSize %d  Blocks Count %d  Free %d  Mem %dkb   Linked Mem %dkb"),
+			ListHeads.GetLength(), AllocatedCount, (FreeBlocks.GetLength() * sizeof(int32) / 1024),
+			ListBlocks.GetLength(), (LinkedListElements.GetLength() * sizeof(int32) / 1024));
+	}
+
+
+};
\ No newline at end of file
diff --git a/Engine/Plugins/Experimental/GeometryProcessing/Source/GeometricObjects/Public/VectorUtil.h b/Engine/Plugins/Experimental/GeometryProcessing/Source/GeometricObjects/Public/VectorUtil.h
new file mode 100644
--- /dev/null
+++ b/Engine/Plugins/Experimental/GeometryProcessing/Source/GeometricObjects/Public/VectorUtil.h
@@ -0,0 +1,207 @@
+// Copyright 1998-2019 Epic Games, Inc. All Rights Reserved.
+
+#pragma once
+
+#include "MathUtil.h"
+#include "VectorTypes.h"
+
+// these should go somewhere else?
+
+enum class EIntersectionResult
+{
+	NotComputed,
+	Intersects,
+	NoIntersection,
+	InvalidQuery
+};
+
+enum class EIntersectionType
+{
+	Empty,
+	Point,
+	Segment,
+	Line,
+	Polygon,
+	Plane,
+	Unknown
+};
+
+namespace VectorUtil
+{
+	template <typename RealType>
+	inline bool IsFinite(const FVector2<RealType>& V)
+	{
+		return TMathUtil<RealType>::IsFinite(V.X) && TMathUtil<RealType>::IsFinite(V.Y);
+	}
+	template <typename RealType>
+	inline bool IsFinite(const FVector3<RealType>& V)
+	{
+		return TMathUtil<RealType>::IsFinite(V.X) && TMathUtil<RealType>::IsFinite(V.Y) && TMathUtil<RealType>::IsFinite(V.Z);
+	}
+
+
+	template <typename RealType>
+	inline RealType Clamp(RealType Value, RealType MinValue, RealType MaxValue)
+	{
+		return (Value < MinValue) ? MinValue : ((Value > MaxValue) ? MaxValue : Value);
+	}
+
+	template <typename RealType>
+	inline FVector3<RealType> Normal(const FVector3<RealType>& V0, const FVector3<RealType>& V1, const FVector3<RealType>& V2)
+	{
+		FVector3<RealType> edge1(V1 - V0);
+		FVector3<RealType> edge2(V2 - V0);
+		edge1.Normalize();
+		edge2.Normalize();
+		// Unreal has Left-Hand Coordinate System so we need to reverse this cross-product to get proper triangle normal
+		FVector3<RealType> vCross(edge2.Cross(edge1));
+		//FVector3<RealType> vCross(edge1.Cross(edge2));
+		return vCross.Normalized();
+	}
+
+	template <typename RealType>
+	inline FVector3<RealType> FastNormalDirection(const FVector3<RealType>& V0, const FVector3<RealType>& V1, const FVector3<RealType>& V2)
+	{
+		// Unreal has Left-Hand Coordinate System so we need to reverse this cross-product to get proper triangle normal
+		return (V2 - V0).Cross(V1 - V0);
+		//return (V1 - V0).Cross(V2 - V0);
+	}
+
+	template <typename RealType>
+	inline RealType Area(const FVector3<RealType>& V0, const FVector3<RealType>& V1, const FVector3<RealType>& V2)
+	{
+		FVector3<RealType> edge1(V1 - V0);
+		FVector3<RealType> edge2(V2 - V0);
+		RealType Dot = edge1.Dot(edge2);
+		return (RealType)(0.5 * sqrt(edge1.SquaredLength() * edge2.SquaredLength() - Dot * Dot));
+	}
+
+	template <typename RealType>
+	inline bool IsObtuse(const FVector3<RealType>& V1, const FVector3<RealType>& V2, const FVector3<RealType>& V3)
+	{
+		RealType a2 = V1.DistanceSquared(V2);
+		RealType b2 = V1.DistanceSquared(V3);
+		RealType c2 = V2.DistanceSquared(V3);
+		return (a2 + b2 < c2) || (b2 + c2 < a2) || (c2 + a2 < b2);
+	}
+
+	template <typename RealType>
+	inline FVector3<RealType> FastNormalArea(const FVector3<RealType>& V0, const FVector3<RealType>& V1, const FVector3<RealType>& V2, double& Area)
+	{
+		FVector3<RealType> edge1(V1 - V0);
+		FVector3<RealType> edge2(V2 - V0);
+		// Unreal has Left-Hand Coordinate System so we need to reverse this cross-product to get proper triangle normal
+		FVector3d vCross = edge2.Cross(edge1);
+		//FVector3d vCross = edge1.Cross(edge2);
+		Area = RealType(0.5) * vCross.Normalize();
+		return vCross;
+	}
+
+	template <typename RealType>
+	inline bool EpsilonEqual(RealType A, RealType B, RealType Epsilon)
+	{
+		return TMathUtil<RealType>::Abs(A - B) < Epsilon;
+	}
+
+	template <typename RealType>
+	inline bool EpsilonEqual(const FVector2<RealType>& V0, const FVector2<RealType>& V1, RealType Epsilon)
+	{
+		return EpsilonEqual(V0.X, V1.X, Epsilon) && EpsilonEqual(V0.Y, V1.Y, Epsilon);
+	}
+	template <typename RealType>
+	inline bool EpsilonEqual(const FVector3<RealType>& V0, const FVector3<RealType>& V1, RealType Epsilon)
+	{
+		return EpsilonEqual(V0.X, V1.X, Epsilon) && EpsilonEqual(V0.Y, V1.Y, Epsilon) && EpsilonEqual(V0.Z, V1.Z, Epsilon);
+	}
+
+	/**
+	 * Returns two vectors perpendicular to n, as efficiently as possible.
+	 * Duff et all method, from https://graphics.pixar.com/library/OrthonormalB/paper.pdf
+	 */
+	template <typename RealType>
+	inline void MakePerpVectors(const FVector3<RealType>& Normal, FVector3<RealType>& OutPerp1, FVector3<RealType>& OutPerp2)
+	{
+		if (Normal.Z < (RealType)0)
+		{
+			RealType A = (RealType)1 / ((RealType)1 - Normal.Z);
+			RealType B = Normal.X * Normal.Y * A;
+			OutPerp1.X = (RealType)1 - Normal.X * Normal.X * A;
+			OutPerp1.Y = -B;
+			OutPerp1.Z = Normal.X;
+			OutPerp2.X = B;
+			OutPerp2.Y = Normal.Y * Normal.Y * A - (RealType)1;
+			OutPerp2.Z = -Normal.Y;
+		}
+		else
+		{
+			RealType A = (RealType)1 / ((RealType)1 + Normal.Z);
+			RealType B = -Normal.X * Normal.Y * A;
+			OutPerp1.X = (RealType)1 - Normal.X * Normal.X * A;
+			OutPerp1.Y = B;
+			OutPerp1.Z = -Normal.X;
+			OutPerp2.X = B;
+			OutPerp2.Y = (RealType)1 - Normal.Y * Normal.Y * A;
+			OutPerp2.Z = -Normal.Y;
+		}
+	}
+
+	template <typename RealType>
+	inline double PlaneAngleSignedD(const FVector3<RealType>& VFromIn, const FVector3<RealType>& VToIn, const FVector3<RealType>& PlaneN)
+	{
+		FVector3<RealType> vFrom = VFromIn - VFromIn.Dot(PlaneN) * PlaneN;
+		FVector3<RealType> vTo = VToIn - VToIn.Dot(PlaneN) * PlaneN;
+		vFrom.Normalize();
+		vTo.Normalize();
+		FVector3<RealType> C = vFrom.Cross(vTo);
+		if (C.SquaredLength() < TMathUtil<RealType>::ZeroTolerance)
+		{ // vectors are parallel
+			return vFrom.Dot(vTo) < 0 ? 180.0 : 0;
+		}
+		RealType fSign = C.Dot(PlaneN) < 0 ? (RealType)-1 : (RealType)1;
+		return (RealType)(fSign * vFrom.AngleD(vTo));
+	}
+
+	/**
+	 * tan(theta/2) = +/- sqrt( (1-cos(theta)) / (1+cos(theta)) )
+	 * This function returns positive Value
+	 */
+	template <typename RealType>
+	RealType VectorTanHalfAngle(const FVector3<RealType>& A, const FVector3<RealType>& B)
+	{
+		RealType cosAngle = A.Dot(B);
+		RealType sqr = ((RealType)1 - cosAngle) / ((RealType)1 + cosAngle);
+		sqr = Clamp(sqr, (RealType)0, TMathUtil<RealType>::MaxReal);
+		return (RealType)sqrt(sqr);
+	}
+
+	//! fast cotangent between two normalized vectors
+	//! cot = cos/sin, both of which can be computed from vector identities
+	//! returns zero if result would be unstable (eg infinity)
+	template <typename RealType>
+	RealType VectorCot(const FVector3<RealType>& V1, const FVector3<RealType>& V2)
+	{
+		// formula from http://www.geometry.caltech.edu/pubs/DMSB_III.pdf
+		RealType fDot = V1.Dot(V2);
+		RealType lensqr1 = V1.SquaredLength();
+		RealType lensqr2 = V2.SquaredLength();
+		RealType d = Clamp(lensqr1 * lensqr2 - fDot * fDot, (RealType)0.0, TMathUtil<RealType>::MaxReal);
+		if (d < TMathUtil<RealType>::ZeroTolerance)
+			return (RealType)0;
+		else
+			return fDot / (RealType)sqrt(d);
+	}
+
+	template <typename RealType>
+	inline double TriSolidAngle(FVector3<RealType> A, FVector3<RealType> B, FVector3<RealType> C, const FVector3<RealType>& P)
+	{
+		A -= P;
+		B -= P;
+		C -= P;
+		RealType la = A.Length(), lb = B.Length(), lc = C.Length();
+		RealType top = (la * lb * lc) + A.Dot(B) * lc + B.Dot(C) * la + C.Dot(A) * lb;
+		RealType bottom = A.X * (B.Y * C.Z - C.Y * B.Z) - A.Y * (B.X * C.Z - C.X * B.Z) + A.Z * (B.X * C.Y - C.X * B.Y);
+		// -2 instead of 2 to account for UE winding
+		return RealType(-2.0) * atan2(bottom, top);
+	}
+
+}; // namespace VectorUtil
diff --git a/Engine/Plugins/Experimental/GeometryProcessing/Source/GeometryAlgorithms/Private/Arrangement2d.cpp b/Engine/Plugins/Experimental/GeometryProcessing/Source/GeometryAlgorithms/Private/Arrangement2d.cpp
new file mode 100644
--- /dev/null
+++ b/Engine/Plugins/Experimental/GeometryProcessing/Source/GeometryAlgorithms/Private/Arrangement2d.cpp
@@ -0,0 +1,67 @@
+// Copyright 1998-2019 Epic Games, Inc. All Rights Reserved.
+
+// Port of geometry3Sharp Arrangement2d
+
+#include "Arrangement2d.h"
+#include "GTEngine/Mathematics/GteConstrainedDelaunay2.h"
+#include "GTEngine/Mathematics/GteBSNumber.h"
+#include "GTEngine/Mathematics/GteUIntegerFP32.h"
+#include <vector>
+
+bool FArrangement2d::AttemptTriangulate(TArray<FIntVector> &Triangles, TArray<int32> &SkippedEdges)
+{
+	Triangles.Empty();
+
+	// non-robust version is very non-robust
+	//gte::ConstrainedDelaunay2<double, double> Delaunay;
+
+	// Here are the types you could use for robust math (doesn't make the algorithm 100% robust, but improves it a lot)
+	gte::ConstrainedDelaunay2<double, gte::BSNumber<gte::UIntegerFP32<263>>> Delaunay; // Value of 263 is from comment in GTEngine/Mathematics/GteDelaunay2.h
+	//gte::ConstrainedDelaunay2<double, gte::BSNumber<gte::UIntegerAP32>> Delaunay; // Full arbitrary precision (slowest method, creates a std::vector per number)
+	
+	// Copy vertices to a compact buffer, and create a mapping in and out of the compact space (TODO: maybe extract this for general use?  or is it already somewhere else?)
+	std::vector<gte::Vector2<double>> InputVertices;
+	TArray<int> InputIndices; InputIndices.SetNumZeroed(Graph.MaxVertexID());
+	TArray<int> OutputIndices; OutputIndices.SetNumZeroed(Graph.VertexCount());
+	int TotalOffset = 0;
+	for (int i = 0; i < Graph.MaxVertexID(); i++)
+	{
+		if (Graph.IsVertex(i))
+		{
+			OutputIndices[i - TotalOffset] = i;
+			InputIndices[i] = i - TotalOffset;
+			FVector2d Vertex = Graph.GetVertex(i);
+			InputVertices.push_back(gte::Vector2<double> {Vertex.X, Vertex.Y});
+		}
+		else
+		{
+			InputIndices[i] = -1;
+			TotalOffset++;
+		}
+	}
+	if (!Delaunay(Graph.VertexCount(), &InputVertices[0], 0.001f))
+	{
+		return false;
+	}
+	std::vector<int> OutEdges; // (needed for Delaunay.Insert call, but otherwise unused)
+	bool InsertConstraintFailure = false;
+	for (int EdgeIdx : Graph.EdgeIndices())
+	{
+		FDynamicGraph::FEdge Edge = Graph.GetEdge(EdgeIdx);
+		if (!Delaunay.Insert({ InputIndices[Edge.A], InputIndices[Edge.B] }, OutEdges))
+		{
+			// Note the failed edge; we will try to proceed anyway, just without this edge
+			ensureMsgf(false, TEXT("CDT edge insertion failed -- possibly bad data?"));
+			InsertConstraintFailure = true;
+			SkippedEdges.Add(EdgeIdx);
+		}
+	}
+	
+	std::vector<int> const& Indices = Delaunay.GetIndices();
+	for (size_t i = 0, n = Indices.size() / 3; i < n; i++)
+	{
+		Triangles.Add(FIntVector(OutputIndices[Indices[i * 3]], OutputIndices[Indices[i * 3 + 1]], OutputIndices[Indices[i * 3 + 2]]));
+	}
+	return !InsertConstraintFailure;
+}
+
diff --git a/Engine/Plugins/Experimental/GeometryProcessing/Source/GeometryAlgorithms/Private/GTEngine/GTEngine.h b/Engine/Plugins/Experimental/GeometryProcessing/Source/GeometryAlgorithms/Private/GTEngine/GTEngine.h
new file mode 100644
--- /dev/null
+++ b/Engine/Plugins/Experimental/GeometryProcessing/Source/GeometryAlgorithms/Private/GTEngine/GTEngine.h
@@ -0,0 +1,11 @@
+// David Eberly, Geometric Tools, Redmond WA 98052
+// Copyright (c) 1998-2019
+// Distributed under the Boost Software License, Version 1.0.
+// http://www.boost.org/LICENSE_1_0.txt
+// http://www.geometrictools.com/License/Boost/LICENSE_1_0.txt
+// File Version: 3.0.4 (2019/02/12)
+
+#pragma once
+
+#include <GTEngine/GTLowLevel.h>
+#include <GTEngine/GTMathematics.h>
diff --git a/Engine/Plugins/Experimental/GeometryProcessing/Source/GeometryAlgorithms/Private/GTEngine/GTEngineDEF.h b/Engine/Plugins/Experimental/GeometryProcessing/Source/GeometryAlgorithms/Private/GTEngine/GTEngineDEF.h
new file mode 100644
--- /dev/null
+++ b/Engine/Plugins/Experimental/GeometryProcessing/Source/GeometryAlgorithms/Private/GTEngine/GTEngineDEF.h
@@ -0,0 +1,80 @@
+// David Eberly, Geometric Tools, Redmond WA 98052
+// Copyright (c) 1998-2019
+// Distributed under the Boost Software License, Version 1.0.
+// http://www.boost.org/LICENSE_1_0.txt
+// http://www.geometrictools.com/License/Boost/LICENSE_1_0.txt
+// File Version: 3.0.0 (2016/06/19)
+
+#pragma once
+
+//----------------------------------------------------------------------------
+// The platform specification.
+//
+// __MSWINDOWS__            :  Microsoft Windows (WIN32 or WIN64)
+// __APPLE__                :  Macintosh OS X
+// __LINUX__                :  Linux or Cygwin
+//----------------------------------------------------------------------------
+
+#if !defined(__LINUX__) && (defined(WIN32) || defined(_WIN64))
+#define __MSWINDOWS__
+
+#if !defined(_MSC_VER)
+#error Microsoft Visual Studio 2013 or later is required.
+#endif
+
+//  MSVC  6   is version 12.0
+//  MSVC  7.0 is version 13.0 (MSVS 2002)
+//  MSVC  7.1 is version 13.1 (MSVS 2003)
+//  MSVC  8.0 is version 14.0 (MSVS 2005)
+//  MSVC  9.0 is version 15.0 (MSVS 2008)
+//  MSVC 10.0 is version 16.0 (MSVS 2010)
+//  MSVC 11.0 is version 17.0 (MSVS 2012)
+//  MSVC 12.0 is version 18.0 (MSVS 2013)
+//  MSVC 14.0 is version 19.0 (MSVS 2015)
+//  Currently, projects are provided only for MSVC 12.0 and 14.0.
+#if _MSC_VER < 1800
+#error Microsoft Visual Studio 2013 or later is required.
+#endif
+
+// Debug build values (choose_your_value is 0, 1, or 2)
+// 0:  Disables checked iterators and disables iterator debugging.
+// 1:  Enables checked iterators and disables iterator debugging.
+// 2:  (default) Enables iterator debugging; checked iterators are not relevant.
+//
+// Release build values (choose_your_value is 0 or 1)
+// 0:  (default) Disables checked iterators.
+// 1:  Enables checked iterators; iterator debugging is not relevant.
+//
+// #define _ITERATOR_DEBUG_LEVEL choose_your_value
+
+#endif  // WIN32 or _WIN64
+
+// TODO: Windows DLL configurations have not yet been added to the project,
+// but these defines are required to support them (when we do add them).
+//
+// Add GTE_EXPORT to project preprocessor options for dynamic library
+// configurations to export their symbols.
+#if defined(GTE_EXPORT)
+    // For the dynamic library configurations.
+    #define GTE_IMPEXP __declspec(dllexport)
+#else
+    // For a client of the dynamic library or for the static library
+    // configurations.
+    #define GTE_IMPEXP
+#endif
+
+// Expose exactly one of these.
+#define GTE_USE_ROW_MAJOR
+//#define GTE_USE_COL_MAJOR
+
+// Expose exactly one of these.
+#define GTE_USE_MAT_VEC
+//#define GTE_USE_VEC_MAT
+
+#if (defined(GTE_USE_ROW_MAJOR) && defined(GTE_USE_COL_MAJOR)) || (!defined(GTE_USE_ROW_MAJOR) && !defined(GTE_USE_COL_MAJOR))
+#error Exactly one storage order must be specified.
+#endif
+
+#if (defined(GTE_USE_MAT_VEC) && defined(GTE_USE_VEC_MAT)) || (!defined(GTE_USE_MAT_VEC) && !defined(GTE_USE_VEC_MAT))
+#error Exactly one multiplication convention must be specified.
+#endif
diff --git a/Engine/Plugins/Experimental/GeometryProcessing/Source/GeometryAlgorithms/Private/GTEngine/GTEnginePCH.cpp b/Engine/Plugins/Experimental/GeometryProcessing/Source/GeometryAlgorithms/Private/GTEngine/GTEnginePCH.cpp
new file mode 100644
--- /dev/null
+++ b/Engine/Plugins/Experimental/GeometryProcessing/Source/GeometryAlgorithms/Private/GTEngine/GTEnginePCH.cpp
@@ -0,0 +1,8 @@
+// David Eberly, Geometric Tools, Redmond WA 98052
+// Copyright (c) 1998-2019
+// Distributed under the Boost Software License, Version 1.0.
+// http://www.boost.org/LICENSE_1_0.txt
+// http://www.geometrictools.com/License/Boost/LICENSE_1_0.txt
+// File Version: 3.0.0 (2016/06/19)
+
+#include <GTEngine/GTEnginePCH.h>
diff --git a/Engine/Plugins/Experimental/GeometryProcessing/Source/GeometryAlgorithms/Private/GTEngine/GTEnginePCH.h b/Engine/Plugins/Experimental/GeometryProcessing/Source/GeometryAlgorithms/Private/GTEngine/GTEnginePCH.h
new file mode 100644
--- /dev/null
+++ b/Engine/Plugins/Experimental/GeometryProcessing/Source/GeometryAlgorithms/Private/GTEngine/GTEnginePCH.h
@@ -0,0 +1,13 @@
+// David Eberly, Geometric Tools, Redmond WA 98052
+// Copyright (c) 1998-2019
+// Distributed under the Boost Software License, Version 1.0.
+// http://www.boost.org/LICENSE_1_0.txt
+// http://www.geometrictools.com/License/Boost/LICENSE_1_0.txt
+// File Version: 3.0.0 (2016/06/19)
+
+// Define GTE_DISABLE_PCH to turn off the precompiled header system.
+#ifndef GTE_DISABLE_PCH
+#include <GTEngine/GTEngine.h>
+#endif
+
+
diff --git a/Engine/Plugins/Experimental/GeometryProcessing/Source/GeometryAlgorithms/Private/GTEngine/GTLowLevel.h b/Engine/Plugins/Experimental/GeometryProcessing/Source/GeometryAlgorithms/Private/GTEngine/GTLowLevel.h
new file mode 100644
--- /dev/null
+++ b/Engine/Plugins/Experimental/GeometryProcessing/Source/GeometryAlgorithms/Private/GTEngine/GTLowLevel.h
@@ -0,0 +1,33 @@
+// David Eberly, Geometric Tools, Redmond WA 98052
+// Copyright (c) 1998-2019
+// Distributed under the Boost Software License, Version 1.0.
+// http://www.boost.org/LICENSE_1_0.txt
+// http://www.geometrictools.com/License/Boost/LICENSE_1_0.txt
+// File Version: 3.0.1 (2019/01/09)
+
+#pragma once
+
+// DataTypes
+#include <GTEngine/LowLevel/GteArray2.h>
+#include <GTEngine/LowLevel/GteArray3.h>
+#include <GTEngine/LowLevel/GteArray4.h>
+#include <GTEngine/LowLevel/GteAtomicMinMax.h>
+#include <GTEngine/LowLevel/GteComputeModel.h>
+#include <GTEngine/LowLevel/GteLexicoArray2.h>
+#include <GTEngine/LowLevel/GteMinHeap.h>
+#include <GTEngine/LowLevel/GteRangeIteration.h>
+#include <GTEngine/LowLevel/GteSharedPtrCompare.h>
+#include <GTEngine/LowLevel/GteThreadSafeMap.h>
+#include <GTEngine/LowLevel/GteThreadSafeQueue.h>
+#include <GTEngine/LowLevel/GteWeakPtrCompare.h>
+#include <GTEngine/LowLevel/GteWrapper.h>
+
+// Logger
+#include <GTEngine/LowLevel/GteLogger.h>
+#include <GTEngine/LowLevel/GteLogReporter.h>
+#include <GTEngine/LowLevel/GteLogToFile.h>
+#include <GTEngine/LowLevel/GteLogToStdout.h>
+#include <GTEngine/LowLevel/GteLogToStringArray.h>
+
+// Timer
+#include <GTEngine/LowLevel/GteTimer.h>
diff --git a/Engine/Plugins/Experimental/GeometryProcessing/Source/GeometryAlgorithms/Private/GTEngine/GTMathematics.h b/Engine/Plugins/Experimental/GeometryProcessing/Source/GeometryAlgorithms/Private/GTEngine/GTMathematics.h
new file mode 100644
--- /dev/null
+++ b/Engine/Plugins/Experimental/GeometryProcessing/Source/GeometryAlgorithms/Private/GTEngine/GTMathematics.h
@@ -0,0 +1,398 @@
+// David Eberly, Geometric Tools, Redmond WA 98052
+// Copyright (c) 1998-2019
+// Distributed under the Boost Software License, Version 1.0.
+// http://www.boost.org/LICENSE_1_0.txt
+// http://www.geometrictools.com/License/Boost/LICENSE_1_0.txt
+// File Version: 3.0.29 (2019/02/13)
+
+#pragma once
+
+// I used to have these sections ordered alphabetically by the names in
+// the comments.  Two phase name lookups for template matching by the
+// MSVS 2017 compiler had no problems with the ordering, but the Linux
+// g++ compiler does.  This occurred when trying to match std::sqrt(...)
+// and other math functions when the inputs are based on BSNumber or
+// BSRational.  The "Arithmetic" section has been moved before all other
+// headers, and the UInteger* files have been moved before the BS* files.
+
+// Arithmetic
+#include <GTEngine/Mathematics/GteBitHacks.h>
+#include <GTEngine/Mathematics/GteIEEEBinary.h>
+#include <GTEngine/Mathematics/GteIEEEBinary16.h>
+#include <GTEngine/Mathematics/GteMath.h>
+#include <GTEngine/Mathematics/GteArbitraryPrecision.h>
+#include <GTEngine/Mathematics/GteQuadraticField.h>
+
+// Algebra
+#include <GTEngine/Mathematics/GteAxisAngle.h>
+#include <GTEngine/Mathematics/GteBandedMatrix.h>
+#include <GTEngine/Mathematics/GteConvertCoordinates.h>
+#include <GTEngine/Mathematics/GteEulerAngles.h>
+#include <GTEngine/Mathematics/GteGMatrix.h>
+#include <GTEngine/Mathematics/GteGVector.h>
+#include <GTEngine/Mathematics/GteMatrix.h>
+#include <GTEngine/Mathematics/GteMatrix2x2.h>
+#include <GTEngine/Mathematics/GteMatrix3x3.h>
+#include <GTEngine/Mathematics/GteMatrix4x4.h>
+#include <GTEngine/Mathematics/GtePolynomial1.h>
+#include <GTEngine/Mathematics/GteQuaternion.h>
+#include <GTEngine/Mathematics/GteRotation.h>
+#include <GTEngine/Mathematics/GteVector.h>
+#include <GTEngine/Mathematics/GteVector2.h>
+#include <GTEngine/Mathematics/GteVector3.h>
+#include <GTEngine/Mathematics/GteVector4.h>
+
+// Approximation
+#include <GTEngine/Mathematics/GteApprCircle2.h>
+#include <GTEngine/Mathematics/GteApprCone3.h>
+#include <GTEngine/Mathematics/GteApprCylinder3.h>
+#include <GTEngine/Mathematics/GteApprEllipseByArcs.h>
+#include <GTEngine/Mathematics/GteApprEllipse2.h>
+#include <GTEngine/Mathematics/GteApprEllipsoid3.h>
+#include <GTEngine/Mathematics/GteApprGaussian2.h>
+#include <GTEngine/Mathematics/GteApprGaussian3.h>
+#include <GTEngine/Mathematics/GteApprGreatCircle3.h>
+#include <GTEngine/Mathematics/GteApprHeightLine2.h>
+#include <GTEngine/Mathematics/GteApprHeightPlane3.h>
+#include <GTEngine/Mathematics/GteApprOrthogonalLine2.h>
+#include <GTEngine/Mathematics/GteApprOrthogonalLine3.h>
+#include <GTEngine/Mathematics/GteApprOrthogonalPlane3.h>
+#include <GTEngine/Mathematics/GteApprParaboloid3.h>
+#include <GTEngine/Mathematics/GteApprPolynomial2.h>
+#include <GTEngine/Mathematics/GteApprPolynomial3.h>
+#include <GTEngine/Mathematics/GteApprPolynomial4.h>
+#include <GTEngine/Mathematics/GteApprPolynomialSpecial2.h>
+#include <GTEngine/Mathematics/GteApprPolynomialSpecial3.h>
+#include <GTEngine/Mathematics/GteApprPolynomialSpecial4.h>
+#include <GTEngine/Mathematics/GteApprQuadratic2.h>
+#include <GTEngine/Mathematics/GteApprQuadratic3.h>
+#include <GTEngine/Mathematics/GteApprQuery.h>
+#include <GTEngine/Mathematics/GteApprSphere3.h>
+#include <GTEngine/Mathematics/GteApprTorus3.h>
+
+// ComputationalGeometry
+#include <GTEngine/Mathematics/GteConstrainedDelaunay2.h>
+#include <GTEngine/Mathematics/GteConvexHull2.h>
+#include <GTEngine/Mathematics/GteConvexHull3.h>
+#include <GTEngine/Mathematics/GteDelaunay2.h>
+#include <GTEngine/Mathematics/GteDelaunay2Mesh.h>
+#include <GTEngine/Mathematics/GteDelaunay3.h>
+#include <GTEngine/Mathematics/GteDelaunay3Mesh.h>
+#include <GTEngine/Mathematics/GteEdgeKey.h>
+#include <GTEngine/Mathematics/GteETManifoldMesh.h>
+#include <GTEngine/Mathematics/GteETNonmanifoldMesh.h>
+#include <GTEngine/Mathematics/GteFeatureKey.h>
+#include <GTEngine/Mathematics/GteGenerateMeshUV.h>
+#include <GTEngine/Mathematics/GteIsPlanarGraph.h>
+#include <GTEngine/Mathematics/GteMinimalCycleBasis.h>
+#include <GTEngine/Mathematics/GteMinimumAreaBox2.h>
+#include <GTEngine/Mathematics/GteMinimumAreaCircle2.h>
+#include <GTEngine/Mathematics/GteMinimumVolumeBox3.h>
+#include <GTEngine/Mathematics/GteMinimumVolumeSphere3.h>
+#include <GTEngine/Mathematics/GteNearestNeighborQuery.h>
+#include <GTEngine/Mathematics/GteOBBTreeOfPoints.h>
+#include <GTEngine/Mathematics/GtePlanarMesh.h>
+#include <GTEngine/Mathematics/GtePrimalQuery2.h>
+#include <GTEngine/Mathematics/GtePrimalQuery3.h>
+#include <GTEngine/Mathematics/GteSeparatePoints2.h>
+#include <GTEngine/Mathematics/GteSeparatePoints3.h>
+#include <GTEngine/Mathematics/GteTetrahedronKey.h>
+#include <GTEngine/Mathematics/GteTriangleKey.h>
+#include <GTEngine/Mathematics/GteTriangulateCDT.h>
+#include <GTEngine/Mathematics/GteTriangulateEC.h>
+#include <GTEngine/Mathematics/GteTSManifoldMesh.h>
+#include <GTEngine/Mathematics/GteUniqueVerticesTriangles.h>
+#include <GTEngine/Mathematics/GteVEManifoldMesh.h>
+#include <GTEngine/Mathematics/GteVETManifoldMesh.h>
+#include <GTEngine/Mathematics/GteVETNonmanifoldMesh.h>
+#include <GTEngine/Mathematics/GteVertexCollapseMesh.h>
+
+// Containment
+#include <GTEngine/Mathematics/GteContAlignedBox.h>
+#include <GTEngine/Mathematics/GteContCapsule3.h>
+#include <GTEngine/Mathematics/GteContCircle2.h>
+#include <GTEngine/Mathematics/GteContCone.h>
+#include <GTEngine/Mathematics/GteContCylinder3.h>
+#include <GTEngine/Mathematics/GteContEllipse2.h>
+#include <GTEngine/Mathematics/GteContEllipse2MinCR.h>
+#include <GTEngine/Mathematics/GteContEllipsoid3.h>
+#include <GTEngine/Mathematics/GteContEllipsoid3MinCR.h>
+#include <GTEngine/Mathematics/GteContOrientedBox2.h>
+#include <GTEngine/Mathematics/GteContOrientedBox3.h>
+#include <GTEngine/Mathematics/GteContPointInPolygon2.h>
+#include <GTEngine/Mathematics/GteContPointInPolyhedron3.h>
+#include <GTEngine/Mathematics/GteContScribeCircle2.h>
+#include <GTEngine/Mathematics/GteContScribeCircle3Sphere3.h>
+#include <GTEngine/Mathematics/GteContSphere3.h>
+
+// CurvesSurfacesVolumes
+#include <GTEngine/Mathematics/GteBasisFunction.h>
+#include <GTEngine/Mathematics/GteBezierCurve.h>
+#include <GTEngine/Mathematics/GteBSplineCurve.h>
+#include <GTEngine/Mathematics/GteBSplineCurveFit.h>
+#include <GTEngine/Mathematics/GteBSplineSurface.h>
+#include <GTEngine/Mathematics/GteBSplineSurfaceFit.h>
+#include <GTEngine/Mathematics/GteBSplineVolume.h>
+#include <GTEngine/Mathematics/GteDarbouxFrame.h>
+#include <GTEngine/Mathematics/GteFrenetFrame.h>
+#include <GTEngine/Mathematics/GteIndexAttribute.h>
+#include <GTEngine/Mathematics/GteMesh.h>
+#include <GTEngine/Mathematics/GteNaturalSplineCurve.h>
+#include <GTEngine/Mathematics/GteNURBSCircle.h>
+#include <GTEngine/Mathematics/GteNURBSCurve.h>
+#include <GTEngine/Mathematics/GteNURBSSurface.h>
+#include <GTEngine/Mathematics/GteNURBSSphere.h>
+#include <GTEngine/Mathematics/GteNURBSVolume.h>
+#include <GTEngine/Mathematics/GteParametricCurve.h>
+#include <GTEngine/Mathematics/GteParametricSurface.h>
+#include <GTEngine/Mathematics/GteRectangleMesh.h>
+#include <GTEngine/Mathematics/GteRectanglePatchMesh.h>
+#include <GTEngine/Mathematics/GteRevolutionMesh.h>
+#include <GTEngine/Mathematics/GteTCBSplineCurve.h>
+#include <GTEngine/Mathematics/GteTubeMesh.h>
+#include <GTEngine/Mathematics/GteVertexAttribute.h>
+
+// Distance
+#include <GTEngine/Mathematics/GteDCPQuery.h>
+#include <GTEngine/Mathematics/GteDistAlignedBoxAlignedBox.h>
+#include <GTEngine/Mathematics/GteDistAlignedBox3OrientedBox3.h>
+#include <GTEngine/Mathematics/GteDistCircle3Circle3.h>
+#include <GTEngine/Mathematics/GteDistLine3AlignedBox3.h>
+#include <GTEngine/Mathematics/GteDistLine3Circle3.h>
+#include <GTEngine/Mathematics/GteDistLine3OrientedBox3.h>
+#include <GTEngine/Mathematics/GteDistLine3Rectangle3.h>
+#include <GTEngine/Mathematics/GteDistLine3Triangle3.h>
+#include <GTEngine/Mathematics/GteDistLineLine.h>
+#include <GTEngine/Mathematics/GteDistLineRay.h>
+#include <GTEngine/Mathematics/GteDistLineSegment.h>
+#include <GTEngine/Mathematics/GteDistOrientedBox3OrientedBox3.h>
+#include <GTEngine/Mathematics/GteDistPoint3Circle3.h>
+#include <GTEngine/Mathematics/GteDistPoint3ConvexPolyhedron3.h>
+#include <GTEngine/Mathematics/GteDistPoint3Cylinder3.h>
+#include <GTEngine/Mathematics/GteDistPoint3Frustum3.h>
+#include <GTEngine/Mathematics/GteDistPoint3Plane3.h>
+#include <GTEngine/Mathematics/GteDistPoint3Rectangle3.h>
+#include <GTEngine/Mathematics/GteDistPoint3Tetrahedron3.h>
+#include <GTEngine/Mathematics/GteDistPointAlignedBox.h>
+#include <GTEngine/Mathematics/GteDistPointHyperellipsoid.h>
+#include <GTEngine/Mathematics/GteDistPointLine.h>
+#include <GTEngine/Mathematics/GteDistPointOrientedBox.h>
+#include <GTEngine/Mathematics/GteDistPointRay.h>
+#include <GTEngine/Mathematics/GteDistPointSegment.h>
+#include <GTEngine/Mathematics/GteDistPointTriangle.h>
+#include <GTEngine/Mathematics/GteDistPointTriangleExact.h>
+#include <GTEngine/Mathematics/GteDistRay3AlignedBox3.h>
+#include <GTEngine/Mathematics/GteDistRay3OrientedBox3.h>
+#include <GTEngine/Mathematics/GteDistRay3Rectangle3.h>
+#include <GTEngine/Mathematics/GteDistRay3Triangle3.h>
+#include <GTEngine/Mathematics/GteDistRayRay.h>
+#include <GTEngine/Mathematics/GteDistRaySegment.h>
+#include <GTEngine/Mathematics/GteDistRectangle3Rectangle3.h>
+#include <GTEngine/Mathematics/GteDistRectangle3AlignedBox3.h>
+#include <GTEngine/Mathematics/GteDistRectangle3OrientedBox3.h>
+#include <GTEngine/Mathematics/GteDistSegment3AlignedBox3.h>
+#include <GTEngine/Mathematics/GteDistSegment3OrientedBox3.h>
+#include <GTEngine/Mathematics/GteDistSegment3Rectangle3.h>
+#include <GTEngine/Mathematics/GteDistSegment3Triangle3.h>
+#include <GTEngine/Mathematics/GteDistSegmentSegment.h>
+#include <GTEngine/Mathematics/GteDistSegmentSegmentExact.h>
+#include <GTEngine/Mathematics/GteDistTriangle3AlignedBox3.h>
+#include <GTEngine/Mathematics/GteDistTriangle3OrientedBox3.h>
+#include <GTEngine/Mathematics/GteDistTriangle3Rectangle3.h>
+#include <GTEngine/Mathematics/GteDistTriangle3Triangle3.h>
+
+// Functions
+#include <GTEngine/Mathematics/GteACosEstimate.h>
+#include <GTEngine/Mathematics/GteASinEstimate.h>
+#include <GTEngine/Mathematics/GteATanEstimate.h>
+#include <GTEngine/Mathematics/GteChebyshevRatio.h>
+#include <GTEngine/Mathematics/GteCosEstimate.h>
+#include <GTEngine/Mathematics/GteExp2Estimate.h>
+#include <GTEngine/Mathematics/GteExpEstimate.h>
+#include <GTEngine/Mathematics/GteInvSqrtEstimate.h>
+#include <GTEngine/Mathematics/GteLog2Estimate.h>
+#include <GTEngine/Mathematics/GteLogEstimate.h>
+#include <GTEngine/Mathematics/GteSinEstimate.h>
+#include <GTEngine/Mathematics/GteSlerpEstimate.h>
+#include <GTEngine/Mathematics/GteSqrtEstimate.h>
+#include <GTEngine/Mathematics/GteTanEstimate.h>
+
+// GeometricPrimitives
+#include <GTEngine/Mathematics/GteAlignedBox.h>
+#include <GTEngine/Mathematics/GteArc2.h>
+#include <GTEngine/Mathematics/GteCapsule.h>
+#include <GTEngine/Mathematics/GteCircle3.h>
+#include <GTEngine/Mathematics/GteCone.h>
+#include <GTEngine/Mathematics/GteConvexPolyhedron3.h>
+#include <GTEngine/Mathematics/GteCylinder3.h>
+#include <GTEngine/Mathematics/GteEllipse3.h>
+#include <GTEngine/Mathematics/GteFrustum3.h>
+#include <GTEngine/Mathematics/GteHalfspace.h>
+#include <GTEngine/Mathematics/GteHyperellipsoid.h>
+#include <GTEngine/Mathematics/GteHyperplane.h>
+#include <GTEngine/Mathematics/GteHypersphere.h>
+#include <GTEngine/Mathematics/GteLine.h>
+#include <GTEngine/Mathematics/GteOrientedBox.h>
+#include <GTEngine/Mathematics/GtePolygon2.h>
+#include <GTEngine/Mathematics/GtePolyhedron3.h>
+#include <GTEngine/Mathematics/GteRay.h>
+#include <GTEngine/Mathematics/GteRectangle.h>
+#include <GTEngine/Mathematics/GteSector2.h>
+#include <GTEngine/Mathematics/GteSegment.h>
+#include <GTEngine/Mathematics/GteTetrahedron3.h>
+#include <GTEngine/Mathematics/GteTorus3.h>
+#include <GTEngine/Mathematics/GteTriangle.h>
+
+// Interpolation
+#include <GTEngine/Mathematics/GteIntpAkima1.h>
+#include <GTEngine/Mathematics/GteIntpAkimaNonuniform1.h>
+#include <GTEngine/Mathematics/GteIntpAkimaUniform1.h>
+#include <GTEngine/Mathematics/GteIntpAkimaUniform2.h>
+#include <GTEngine/Mathematics/GteIntpAkimaUniform3.h>
+#include <GTEngine/Mathematics/GteIntpBicubic2.h>
+#include <GTEngine/Mathematics/GteIntpBilinear2.h>
+#include <GTEngine/Mathematics/GteIntpBSplineUniform.h>
+#include <GTEngine/Mathematics/GteIntpLinearNonuniform2.h>
+#include <GTEngine/Mathematics/GteIntpLinearNonuniform3.h>
+#include <GTEngine/Mathematics/GteIntpQuadraticNonuniform2.h>
+#include <GTEngine/Mathematics/GteIntpSphere2.h>
+#include <GTEngine/Mathematics/GteIntpThinPlateSpline2.h>
+#include <GTEngine/Mathematics/GteIntpThinPlateSpline3.h>
+#include <GTEngine/Mathematics/GteIntpTricubic3.h>
+#include <GTEngine/Mathematics/GteIntpTrilinear3.h>
+#include <GTEngine/Mathematics/GteIntpVectorField2.h>
+
+// Intersection
+#include <GTEngine/Mathematics/GteFIQuery.h>
+#include <GTEngine/Mathematics/GteIntrAlignedBox2AlignedBox2.h>
+#include <GTEngine/Mathematics/GteIntrAlignedBox2Circle2.h>
+#include <GTEngine/Mathematics/GteIntrAlignedBox2OrientedBox2.h>
+#include <GTEngine/Mathematics/GteIntrAlignedBox3AlignedBox3.h>
+#include <GTEngine/Mathematics/GteIntrAlignedBox3Cone3.h>
+#include <GTEngine/Mathematics/GteIntrAlignedBox3Cylinder3.h>
+#include <GTEngine/Mathematics/GteIntrAlignedBox3OrientedBox3.h>
+#include <GTEngine/Mathematics/GteIntrAlignedBox3Sphere3.h>
+#include <GTEngine/Mathematics/GteIntrArc2Arc2.h>
+#include <GTEngine/Mathematics/GteIntrCapsule3Capsule3.h>
+#include <GTEngine/Mathematics/GteIntrCircle2Arc2.h>
+#include <GTEngine/Mathematics/GteIntrCircle2Circle2.h>
+#include <GTEngine/Mathematics/GteIntrConvexPolygonHyperplane.h>
+#include <GTEngine/Mathematics/GteIntrDisk2Sector2.h>
+#include <GTEngine/Mathematics/GteIntrEllipse2Ellipse2.h>
+#include <GTEngine/Mathematics/GteIntrEllipsoid3Ellipsoid3.h>
+#include <GTEngine/Mathematics/GteIntrHalfspace2Polygon2.h>
+#include <GTEngine/Mathematics/GteIntrHalfspace3Capsule3.h>
+#include <GTEngine/Mathematics/GteIntrHalfspace3Cylinder3.h>
+#include <GTEngine/Mathematics/GteIntrHalfspace3Ellipsoid3.h>
+#include <GTEngine/Mathematics/GteIntrHalfspace3OrientedBox3.h>
+#include <GTEngine/Mathematics/GteIntrHalfspace3Segment3.h>
+#include <GTEngine/Mathematics/GteIntrHalfspace3Sphere3.h>
+#include <GTEngine/Mathematics/GteIntrHalfspace3Triangle3.h>
+#include <GTEngine/Mathematics/GteIntrIntervals.h>
+#include <GTEngine/Mathematics/GteIntrLine2AlignedBox2.h>
+#include <GTEngine/Mathematics/GteIntrLine2Arc2.h>
+#include <GTEngine/Mathematics/GteIntrLine2Circle2.h>
+#include <GTEngine/Mathematics/GteIntrLine2Line2.h>
+#include <GTEngine/Mathematics/GteIntrLine2OrientedBox2.h>
+#include <GTEngine/Mathematics/GteIntrLine2Ray2.h>
+#include <GTEngine/Mathematics/GteIntrLine2Segment2.h>
+#include <GTEngine/Mathematics/GteIntrLine2Triangle2.h>
+#include <GTEngine/Mathematics/GteIntrLine3AlignedBox3.h>
+#include <GTEngine/Mathematics/GteIntrLine3Capsule3.h>
+#include <GTEngine/Mathematics/GteIntrLine3Cone3.h>
+#include <GTEngine/Mathematics/GteIntrLine3Cylinder3.h>
+#include <GTEngine/Mathematics/GteIntrLine3Ellipsoid3.h>
+#include <GTEngine/Mathematics/GteIntrLine3OrientedBox3.h>
+#include <GTEngine/Mathematics/GteIntrLine3Plane3.h>
+#include <GTEngine/Mathematics/GteIntrLine3Sphere3.h>
+#include <GTEngine/Mathematics/GteIntrLine3Triangle3.h>
+#include <GTEngine/Mathematics/GteIntrOrientedBox2Circle2.h>
+#include <GTEngine/Mathematics/GteIntrOrientedBox2Cone2.h>
+#include <GTEngine/Mathematics/GteIntrOrientedBox2OrientedBox2.h>
+#include <GTEngine/Mathematics/GteIntrOrientedBox2Sector2.h>
+#include <GTEngine/Mathematics/GteIntrOrientedBox3Cone3.h>
+#include <GTEngine/Mathematics/GteIntrOrientedBox3Cylinder3.h>
+#include <GTEngine/Mathematics/GteIntrOrientedBox3Frustum3.h>
+#include <GTEngine/Mathematics/GteIntrOrientedBox3OrientedBox3.h>
+#include <GTEngine/Mathematics/GteIntrOrientedBox3Sphere3.h>
+#include <GTEngine/Mathematics/GteIntrPlane3Capsule3.h>
+#include <GTEngine/Mathematics/GteIntrPlane3Circle3.h>
+#include <GTEngine/Mathematics/GteIntrPlane3Cylinder3.h>
+#include <GTEngine/Mathematics/GteIntrPlane3Ellipsoid3.h>
+#include <GTEngine/Mathematics/GteIntrPlane3OrientedBox3.h>
+#include <GTEngine/Mathematics/GteIntrPlane3Plane3.h>
+#include <GTEngine/Mathematics/GteIntrPlane3Sphere3.h>
+#include <GTEngine/Mathematics/GteIntrPlane3Triangle3.h>
+#include <GTEngine/Mathematics/GteIntrRay2AlignedBox2.h>
+#include <GTEngine/Mathematics/GteIntrRay2Arc2.h>
+#include <GTEngine/Mathematics/GteIntrRay2Circle2.h>
+#include <GTEngine/Mathematics/GteIntrRay2OrientedBox2.h>
+#include <GTEngine/Mathematics/GteIntrRay2Ray2.h>
+#include <GTEngine/Mathematics/GteIntrRay2Segment2.h>
+#include <GTEngine/Mathematics/GteIntrRay2Triangle2.h>
+#include <GTEngine/Mathematics/GteIntrRay3AlignedBox3.h>
+#include <GTEngine/Mathematics/GteIntrRay3Capsule3.h>
+#include <GTEngine/Mathematics/GteIntrRay3Cone3.h>
+#include <GTEngine/Mathematics/GteIntrRay3Cylinder3.h>
+#include <GTEngine/Mathematics/GteIntrRay3Ellipsoid3.h>
+#include <GTEngine/Mathematics/GteIntrRay3OrientedBox3.h>
+#include <GTEngine/Mathematics/GteIntrRay3Plane3.h>
+#include <GTEngine/Mathematics/GteIntrRay3Sphere3.h>
+#include <GTEngine/Mathematics/GteIntrRay3Triangle3.h>
+#include <GTEngine/Mathematics/GteIntrSegment2AlignedBox2.h>
+#include <GTEngine/Mathematics/GteIntrSegment2Arc2.h>
+#include <GTEngine/Mathematics/GteIntrSegment2Circle2.h>
+#include <GTEngine/Mathematics/GteIntrSegment2OrientedBox2.h>
+#include <GTEngine/Mathematics/GteIntrSegment2Segment2.h>
+#include <GTEngine/Mathematics/GteIntrSegment2Triangle2.h>
+#include <GTEngine/Mathematics/GteIntrSegment3AlignedBox3.h>
+#include <GTEngine/Mathematics/GteIntrSegment3Capsule3.h>
+#include <GTEngine/Mathematics/GteIntrSegment3Cone3.h>
+#include <GTEngine/Mathematics/GteIntrSegment3Cylinder3.h>
+#include <GTEngine/Mathematics/GteIntrSegment3Ellipsoid3.h>
+#include <GTEngine/Mathematics/GteIntrSegment3OrientedBox3.h>
+#include <GTEngine/Mathematics/GteIntrSegment3Plane3.h>
+#include <GTEngine/Mathematics/GteIntrSegment3Sphere3.h>
+#include <GTEngine/Mathematics/GteIntrSegment3Triangle3.h>
+#include <GTEngine/Mathematics/GteIntrSphere3Cone3.h>
+#include <GTEngine/Mathematics/GteIntrSphere3Frustum3.h>
+#include <GTEngine/Mathematics/GteIntrSphere3Sphere3.h>
+#include <GTEngine/Mathematics/GteIntrSphere3Triangle3.h>
+#include <GTEngine/Mathematics/GteIntrTriangle3OrientedBox3.h>
+#include <GTEngine/Mathematics/GteTIQuery.h>
+
+// NumericalMethods
+#include <GTEngine/Mathematics/GteCholeskyDecomposition.h>
+#include <GTEngine/Mathematics/GteCubicRootsQR.h>
+#include <GTEngine/Mathematics/GteGaussNewtonMinimizer.h>
+#include <GTEngine/Mathematics/GteGaussianElimination.h>
+#include <GTEngine/Mathematics/GteIntegration.h>
+#include <GTEngine/Mathematics/GteLCPSolver.h>
+#include <GTEngine/Mathematics/GteLevenbergMarquardtMinimizer.h>
+#include <GTEngine/Mathematics/GteLinearSystem.h>
+#include <GTEngine/Mathematics/GteMinimize1.h>
+#include <GTEngine/Mathematics/GteMinimizeN.h>
+#include <GTEngine/Mathematics/GteOdeEuler.h>
+#include <GTEngine/Mathematics/GteOdeImplicitEuler.h>
+#include <GTEngine/Mathematics/GteOdeMidpoint.h>
+#include <GTEngine/Mathematics/GteOdeRungeKutta4.h>
+#include <GTEngine/Mathematics/GteOdeSolver.h>
+#include <GTEngine/Mathematics/GteQuarticRootsQR.h>
+#include <GTEngine/Mathematics/GteRootsBisection.h>
+#include <GTEngine/Mathematics/GteRootsBrentsMethod.h>
+#include <GTEngine/Mathematics/GteRootsPolynomial.h>
+#include <GTEngine/Mathematics/GteSingularValueDecomposition.h>
+#include <GTEngine/Mathematics/GteSymmetricEigensolver.h>
+#include <GTEngine/Mathematics/GteSymmetricEigensolver2x2.h>
+#include <GTEngine/Mathematics/GteSymmetricEigensolver3x3.h>
+#include <GTEngine/Mathematics/GteUnsymmetricEigenvalues.h>
+
+// Projection
+#include <GTEngine/Mathematics/GteProjection.h>
+
+// SIMD (removed due to compile issues + not needed currently)
+// #if defined(__MSWINDOWS__) && !defined(MINGW)
+// #include <GTEngine/Mathematics/MSW/GteIntelSSE.h>
+// #include <GTEngine/Mathematics/MSW/GteCPUQueryInstructions.h>
+// #endif
+
diff --git a/Engine/Plugins/Experimental/GeometryProcessing/Source/GeometryAlgorithms/Private/GTEngine/LowLevel/GteArray2.h b/Engine/Plugins/Experimental/GeometryProcessing/Source/GeometryAlgorithms/Private/GTEngine/LowLevel/GteArray2.h
new file mode 100644
--- /dev/null
+++ b/Engine/Plugins/Experimental/GeometryProcessing/Source/GeometryAlgorithms/Private/GTEngine/LowLevel/GteArray2.h
@@ -0,0 +1,171 @@
+// David Eberly, Geometric Tools, Redmond WA 98052
+// Copyright (c) 1998-2019
+// Distributed under the Boost Software License, Version 1.0.
+// http://www.boost.org/LICENSE_1_0.txt
+// http://www.geometrictools.com/License/Boost/LICENSE_1_0.txt
+// File Version: 3.0.0 (2016/06/19)
+
+#pragma once
+
+#include <GTEngine/GTEngineDEF.h>
+#include <cstddef>
+#include <vector>
+
+// The Array2 class represents a 2-dimensional array that minimizes the number
+// of new and delete calls.  The T objects are stored in a contiguous array.
+
+namespace gte
+{
+
+template <typename T>
+class Array2
+{
+public:
+    // Construction.  The first constructor generates an array of objects that
+    // are owned by Array2.  The second constructor is given an array of
+    // objects that are owned by the caller.  The array has bound0 columns and
+    // bound1 rows.
+    Array2(size_t bound0, size_t bound1);
+    Array2(size_t bound0, size_t bound1, T* objects);
+
+    // Support for dynamic resizing, copying, or moving.  If 'other' does
+    // not own the original 'objects', they are not copied by the assignment
+    // operator.
+    Array2();
+    Array2(Array2 const& other);
+    Array2& operator=(Array2 const& other);
+    Array2(Array2&& other);
+    Array2& operator=(Array2&& other);
+
+    // Access to the array.  Sample usage is
+    //   Array2<T> myArray(3, 2);
+    //   T* row1 = myArray[1];
+    //   T row1Col2 = myArray[1][2];
+    inline size_t GetBound0() const;
+    inline size_t GetBound1() const;
+    inline T const* operator[] (int row) const;
+    inline T* operator[] (int row);
+
+private:
+    void SetPointers(T* objects);
+    void SetPointers(Array2 const& other);
+
+    size_t mBound0, mBound1;
+    std::vector<T> mObjects;
+    std::vector<T*> mIndirect1;
+};
+
+template <typename T>
+Array2<T>::Array2(size_t bound0, size_t bound1)
+    :
+    mBound0(bound0),
+    mBound1(bound1),
+    mObjects(bound0 * bound1),
+    mIndirect1(bound1)
+{
+    SetPointers(mObjects.data());
+}
+
+template <typename T>
+Array2<T>::Array2(size_t bound0, size_t bound1, T* objects)
+    :
+    mBound0(bound0),
+    mBound1(bound1),
+    mIndirect1(bound1)
+{
+    SetPointers(objects);
+}
+
+template <typename T>
+Array2<T>::Array2()
+    :
+    mBound0(0),
+    mBound1(0)
+{
+}
+
+template <typename T>
+Array2<T>::Array2(Array2 const& other)
+{
+    *this = other;
+}
+
+template <typename T>
+Array2<T>& Array2<T>::operator=(Array2 const& other)
+{
+    // The copy is valid whether or not other.mObjects has elements.
+    mObjects = other.mObjects;
+    SetPointers(other);
+    return *this;
+}
+
+template <typename T>
+Array2<T>::Array2(Array2&& other)
+{
+    *this = std::move(other);
+}
+
+template <typename T>
+Array2<T>& Array2<T>::operator=(Array2&& other)
+{
+    // The move is valid whether or not other.mObjects has elements.
+    mObjects = std::move(other.mObjects);
+    SetPointers(other);
+    return *this;
+}
+
+template <typename T> inline
+size_t Array2<T>::GetBound0() const
+{
+    return mBound0;
+}
+
+template <typename T> inline
+size_t Array2<T>::GetBound1() const
+{
+    return mBound1;
+}
+
+template <typename T> inline
+T const* Array2<T>::operator[] (int row) const
+{
+    return mIndirect1[row];
+}
+
+template <typename T> inline
+T* Array2<T>::operator[] (int row)
+{
+    return mIndirect1[row];
+}
+
+template <typename T>
+void Array2<T>::SetPointers(T* objects)
+{
+    for (size_t i1 = 0; i1 < mBound1; ++i1)
+    {
+        size_t j0 = mBound0 * i1;  // = bound0 * (i1 + j1) where j1 = 0
+        mIndirect1[i1] = &objects[j0];
+    }
+}
+
+template <typename T>
+void Array2<T>::SetPointers(Array2 const& other)
+{
+    mBound0 = other.mBound0;
+    mBound1 = other.mBound1;
+    mIndirect1.resize(mBound1);
+
+    if (mBound0 > 0)
+    {
+        // The objects are owned.
+        SetPointers(mObjects.data());
+    }
+    else if (mIndirect1.size() > 0)
+    {
+        // The objects are not owned.
+        SetPointers(other.mIndirect1[0]);
+    }
+    // else 'other' is an empty Array2.
+}
+
+}
diff --git a/Engine/Plugins/Experimental/GeometryProcessing/Source/GeometryAlgorithms/Private/GTEngine/LowLevel/GteArray3.h b/Engine/Plugins/Experimental/GeometryProcessing/Source/GeometryAlgorithms/Private/GTEngine/LowLevel/GteArray3.h
new file mode 100644
--- /dev/null
+++ b/Engine/Plugins/Experimental/GeometryProcessing/Source/GeometryAlgorithms/Private/GTEngine/LowLevel/GteArray3.h
@@ -0,0 +1,192 @@
+// David Eberly, Geometric Tools, Redmond WA 98052
+// Copyright (c) 1998-2019
+// Distributed under the Boost Software License, Version 1.0.
+// http://www.boost.org/LICENSE_1_0.txt
+// http://www.geometrictools.com/License/Boost/LICENSE_1_0.txt
+// File Version: 3.0.0 (2016/06/19)
+
+#pragma once
+
+#include <GTEngine/GTEngineDEF.h>
+#include <cstddef>
+#include <vector>
+
+// The Array3 class represents a 3-dimensional array that minimizes the number
+// of new and delete calls.  The T objects are stored in a contiguous array.
+
+namespace gte
+{
+
+template <typename T>
+class Array3
+{
+public:
+    // Construction.  The first constructor generates an array of objects that
+    // are owned by Array3.  The second constructor is given an array of
+    // objects that are owned by the caller.  The array has bound0 columns,
+    // bound1 rows, and bound2 slices.
+    Array3(size_t bound0, size_t bound1, size_t bound2);
+    Array3(size_t bound0, size_t bound1, size_t bound2, T* objects);
+
+    // Support for dynamic resizing, copying, or moving.  If 'other' does
+    // not own the original 'objects', they are not copied by the assignment
+    // operator.
+    Array3();
+    Array3(Array3 const&);
+    Array3& operator=(Array3 const&);
+    Array3(Array3&&);
+    Array3& operator=(Array3&&);
+
+    // Access to the array.  Sample usage is
+    //   Array3<T> myArray(4, 3, 2);
+    //   T** slice1 = myArray[1];
+    //   T* slice1row2 = myArray[1][2];
+    //   T slice1Row2Col3 = myArray[1][2][3];
+    inline size_t GetBound0() const;
+    inline size_t GetBound1() const;
+    inline size_t GetBound2() const;
+    inline T* const* operator[] (int slice) const;
+    inline T** operator[] (int slice);
+
+private:
+    void SetPointers(T* objects);
+    void SetPointers(Array3 const& other);
+
+    size_t mBound0, mBound1, mBound2;
+    std::vector<T> mObjects;
+    std::vector<T*> mIndirect1;
+    std::vector<T**> mIndirect2;
+};
+
+template <typename T>
+Array3<T>::Array3(size_t bound0, size_t bound1, size_t bound2)
+    :
+    mBound0(bound0),
+    mBound1(bound1),
+    mBound2(bound2),
+    mObjects(bound0 * bound1 * bound2),
+    mIndirect1(bound1 * bound2),
+    mIndirect2(bound2)
+{
+    SetPointers(mObjects.data());
+}
+
+template <typename T>
+Array3<T>::Array3(size_t bound0, size_t bound1, size_t bound2, T* objects)
+    :
+    mBound0(bound0),
+    mBound1(bound1),
+    mBound2(bound2),
+    mIndirect1(bound1 * bound2),
+    mIndirect2(bound2)
+{
+    SetPointers(objects);
+}
+
+template <typename T>
+Array3<T>::Array3()
+    :
+    mBound0(0),
+    mBound1(0),
+    mBound2(0)
+{
+}
+
+template <typename T>
+Array3<T>::Array3(Array3 const& other)
+{
+    *this = other;
+}
+
+template <typename T>
+Array3<T>& Array3<T>::operator=(Array3 const& other)
+{
+    // The copy is valid whether or not other.mObjects has elements.
+    mObjects = other.mObjects;
+    SetPointers(other);
+    return *this;
+}
+
+template <typename T>
+Array3<T>::Array3(Array3&& other)
+{
+    *this = std::move(other);
+}
+
+template <typename T>
+Array3<T>& Array3<T>::operator=(Array3&& other)
+{
+    // The move is valid whether or not other.mObjects has elements.
+    mObjects = std::move(other.mObjects);
+    SetPointers(other);
+    return *this;
+}
+
+template <typename T> inline
+size_t Array3<T>::GetBound0() const
+{
+    return mBound0;
+}
+
+template <typename T> inline
+size_t Array3<T>::GetBound1() const
+{
+    return mBound1;
+}
+
+template <typename T> inline
+size_t Array3<T>::GetBound2() const
+{
+    return mBound2;
+}
+
+template <typename T> inline
+T* const* Array3<T>::operator[] (int slice) const
+{
+    return mIndirect2[slice];
+}
+
+template <typename T> inline
+T** Array3<T>::operator[] (int slice)
+{
+    return mIndirect2[slice];
+}
+
+template <typename T>
+void Array3<T>::SetPointers(T* objects)
+{
+    for (size_t i2 = 0; i2 < mBound2; ++i2)
+    {
+        size_t j1 = mBound1 * i2;  // = bound1 * (i2 + j2) where j2 = 0
+        mIndirect2[i2] = &mIndirect1[j1];
+        for (size_t i1 = 0; i1 < mBound1; ++i1)
+        {
+            size_t j0 = mBound0 * (i1 + j1);
+            mIndirect2[i2][i1] = &objects[j0];
+        }
+    }
+}
+
+template <typename T>
+void Array3<T>::SetPointers(Array3 const& other)
+{
+    mBound0 = other.mBound0;
+    mBound1 = other.mBound1;
+    mBound2 = other.mBound2;
+    mIndirect1.resize(mBound1 * mBound2);
+    mIndirect2.resize(mBound2);
+
+    if (mBound0 > 0)
+    {
+        // The objects are owned.
+        SetPointers(mObjects.data());
+    }
+    else if (mIndirect1.size() > 0)
+    {
+        // The objects are not owned.
+        SetPointers(other.mIndirect2[0][0]);
+    }
+    // else 'other' is an empty Array3.
+}
+
+}
diff --git a/Engine/Plugins/Experimental/GeometryProcessing/Source/GeometryAlgorithms/Private/GTEngine/LowLevel/GteArray4.h b/Engine/Plugins/Experimental/GeometryProcessing/Source/GeometryAlgorithms/Private/GTEngine/LowLevel/GteArray4.h
new file mode 100644
--- /dev/null
+++ b/Engine/Plugins/Experimental/GeometryProcessing/Source/GeometryAlgorithms/Private/GTEngine/LowLevel/GteArray4.h
@@ -0,0 +1,213 @@
+// David Eberly, Geometric Tools, Redmond WA 98052
+// Copyright (c) 1998-2019
+// Distributed under the Boost Software License, Version 1.0.
+// http://www.boost.org/LICENSE_1_0.txt
+// http://www.geometrictools.com/License/Boost/LICENSE_1_0.txt
+// File Version: 3.0.0 (2016/06/19)
+
+#pragma once
+
+#include <GTEngine/GTEngineDEF.h>
+#include <cstddef>
+#include <vector>
+
+// The Array4 class represents a 4-dimensional array that minimizes the number
+// of new and delete calls.  The T objects are stored in a contiguous array.
+
+namespace gte
+{
+
+template <typename T>
+class Array4
+{
+public:
+    // Construction.  The first constructor generates an array of objects that
+    // are owned by Array4.  The second constructor is given an array of
+    // objects that are owned by the caller.  The array has bound0 columns,
+    // bound1 rows, bound2 slices, and bound3 cuboids.
+    Array4(size_t bound0, size_t bound1, size_t bound2, size_t bound3);
+    Array4(size_t bound0, size_t bound1, size_t bound2, size_t bound3, T* objects);
+
+    // Support for dynamic resizing, copying, or moving.  If 'other' does
+    // not own the original 'objects', they are not copied by the assignment
+    // operator.
+    Array4();
+    Array4(Array4 const&);
+    Array4& operator=(Array4 const&);
+    Array4(Array4&&);
+    Array4& operator=(Array4&&);
+
+    // Access to the array.  Sample usage is
+    //   Array4<T> myArray(5, 4, 3, 2);
+    //   T*** cuboid1 = myArray[1];
+    //   T** cuboid1Slice2 = myArray[1][2];
+    //   T* cuboid1Slice2Row3 = myArray[1][2][3];
+    //   T cuboid1Slice2Row3Col4 = myArray[1][2][3][4];
+    inline size_t GetBound0() const;
+    inline size_t GetBound1() const;
+    inline size_t GetBound2() const;
+    inline size_t GetBound3() const;
+    inline T** const* operator[] (int cuboid) const;
+    inline T*** operator[] (int cuboid);
+
+private:
+    void SetPointers(T* objects);
+    void SetPointers(Array4 const& other);
+
+    size_t mBound0, mBound1, mBound2, mBound3;
+    std::vector<T> mObjects;
+    std::vector<T*> mIndirect1;
+    std::vector<T**> mIndirect2;
+    std::vector<T***> mIndirect3;
+};
+
+template <typename T>
+Array4<T>::Array4(size_t bound0, size_t bound1, size_t bound2, size_t bound3)
+    :
+    mBound0(bound0),
+    mBound1(bound1),
+    mBound2(bound2),
+    mBound3(bound3),
+    mObjects(bound0 * bound1 * bound2 * bound3),
+    mIndirect1(bound1 * bound2 * bound3),
+    mIndirect2(bound2 * bound3),
+    mIndirect3(bound3)
+{
+    SetPointers(mObjects.data());
+}
+
+template <typename T>
+Array4<T>::Array4(size_t bound0, size_t bound1, size_t bound2, size_t bound3, T* objects)
+    :
+    mBound0(bound0),
+    mBound1(bound1),
+    mBound2(bound2),
+    mBound3(bound3),
+    mIndirect1(bound1 * bound2 * bound3),
+    mIndirect2(bound2 * bound3),
+    mIndirect3(bound3)
+{
+    SetPointers(objects);
+}
+
+template <typename T>
+Array4<T>::Array4()
+    :
+    mBound0(0),
+    mBound1(0),
+    mBound2(0),
+    mBound3(0)
+{
+}
+
+template <typename T>
+Array4<T>::Array4(Array4 const& other)
+{
+    *this = other;
+}
+
+template <typename T>
+Array4<T>& Array4<T>::operator=(Array4 const& other)
+{
+    // The copy is valid whether or not other.mObjects has elements.
+    mObjects = other.mObjects;
+    SetPointers(other);
+    return *this;
+}
+
+template <typename T>
+Array4<T>::Array4(Array4&& other)
+{
+    *this = std::move(other);
+}
+
+template <typename T>
+Array4<T>& Array4<T>::operator=(Array4&& other)
+{
+    // The move is valid whether or not other.mObjects has elements.
+    mObjects = std::move(other.mObjects);
+    SetPointers(other);
+    return *this;
+}
+
+template <typename T> inline
+size_t Array4<T>::GetBound0() const
+{
+    return mBound0;
+}
+
+template <typename T> inline
+size_t Array4<T>::GetBound1() const
+{
+    return mBound1;
+}
+
+template <typename T> inline
+size_t Array4<T>::GetBound2() const
+{
+    return mBound2;
+}
+
+template <typename T> inline
+size_t Array4<T>::GetBound3() const
+{
+    return mBound3;
+}
+
+template <typename T> inline
+T** const* Array4<T>::operator[] (int cuboid) const
+{
+    return mIndirect3[cuboid];
+}
+
+template <typename T> inline
+T*** Array4<T>::operator[] (int cuboid)
+{
+    return mIndirect3[cuboid];
+}
+
+template <typename T>
+void Array4<T>::SetPointers(T* objects)
+{
+    for (size_t i3 = 0; i3 < mBound3; ++i3)
+    {
+        size_t j2 = mBound2 * i3;  // = bound2 * (i3 + j3) where j3 = 0
+        mIndirect3[i3] = &mIndirect2[j2];
+        for (size_t i2 = 0; i2 < mBound2; ++i2)
+        {
+            size_t j1 = mBound1 * (i2 + j2);
+            mIndirect3[i3][i2] = &mIndirect1[j1];
+            for (size_t i1 = 0; i1 < mBound1; ++i1)
+            {
+                size_t j0 = mBound0 * (i1 + j1);
+                mIndirect3[i3][i2][i1] = &objects[j0];
+            }
+        }
+    }
+}
+
+template <typename T>
+void Array4<T>::SetPointers(Array4 const& other)
+{
+    mBound0 = other.mBound0;
+    mBound1 = other.mBound1;
+    mBound2 = other.mBound2;
+    mBound3 = other.mBound3;
+    mIndirect1.resize(mBound1 * mBound2 * mBound3);
+    mIndirect2.resize(mBound2 * mBound3);
+    mIndirect3.resize(mBound3);
+
+    if (mBound0 > 0)
+    {
+        // The objects are owned.
+        SetPointers(mObjects.data());
+    }
+    else if (mIndirect1.size() > 0)
+    {
+        // The objects are not owned.
+        SetPointers(other.mIndirect3[0][0][0]);
+    }
+    // else 'other' is an empty Array3.
+}
+
+}
diff --git a/Engine/Plugins/Experimental/GeometryProcessing/Source/GeometryAlgorithms/Private/GTEngine/LowLevel/GteAtomicMinMax.h b/Engine/Plugins/Experimental/GeometryProcessing/Source/GeometryAlgorithms/Private/GTEngine/LowLevel/GteAtomicMinMax.h
new file mode 100644
--- /dev/null
+++ b/Engine/Plugins/Experimental/GeometryProcessing/Source/GeometryAlgorithms/Private/GTEngine/LowLevel/GteAtomicMinMax.h
@@ -0,0 +1,52 @@
+// David Eberly, Geometric Tools, Redmond WA 98052
+// Copyright (c) 1998-2019
+// Distributed under the Boost Software License, Version 1.0.
+// http://www.boost.org/LICENSE_1_0.txt
+// http://www.geometrictools.com/License/Boost/LICENSE_1_0.txt
+// File Version: 3.0.0 (2016/06/19)
+
+#pragma once
+
+#include <GTEngine/GTEngineDEF.h>
+#include <algorithm>
+#include <atomic>
+
+// Implementations of atomic minimum and atomic maximum computations.  These
+// are based on std::atomic_compare_exchange_strong.
+
+namespace gte
+{
+
+template <typename T>
+T AtomicMin(std::atomic<T>& v0, T const& v1);
+
+template <typename T>
+T AtomicMax(std::atomic<T>& v0, T const& v1);
+
+
+template <typename T>
+T AtomicMin(std::atomic<T>& v0, T const& v1)
+{
+    T vInitial, vMin;
+    do
+    {
+        vInitial = v0;
+        vMin = std::min(vInitial, v1);
+    } while (!std::atomic_compare_exchange_strong(&v0, &vInitial, vMin));
+    return vInitial;
+}
+
+template <typename T>
+T AtomicMax(std::atomic<T>& v0, T const& v1)
+{
+    T vInitial, vMax;
+    do
+    {
+        vInitial = v0;
+        vMax = std::max(vInitial, v1);
+    } while (!std::atomic_compare_exchange_strong(&v0, &vInitial, vMax));
+    return vInitial;
+}
+
+
+}
diff --git a/Engine/Plugins/Experimental/GeometryProcessing/Source/GeometryAlgorithms/Private/GTEngine/LowLevel/GteComputeModel.h b/Engine/Plugins/Experimental/GeometryProcessing/Source/GeometryAlgorithms/Private/GTEngine/LowLevel/GteComputeModel.h
new file mode 100644
--- /dev/null
+++ b/Engine/Plugins/Experimental/GeometryProcessing/Source/GeometryAlgorithms/Private/GTEngine/LowLevel/GteComputeModel.h
@@ -0,0 +1,102 @@
+// David Eberly, Geometric Tools, Redmond WA 98052
+// Copyright (c) 1998-2019
+// Distributed under the Boost Software License, Version 1.0.
+// http://www.boost.org/LICENSE_1_0.txt
+// http://www.geometrictools.com/License/Boost/LICENSE_1_0.txt
+// File Version: 3.0.0 (2016/06/19)
+
+#pragma once
+
+#include <GTEngine/GTEngineDEF.h>
+
+// Expose this define if you want GPGPU support in computing any algorithms
+// that have a GPU implemetnation.  Alternatively, your application can
+// define this in the project settings so that you do not have to edit this
+// source file.
+//
+//#define GTE_COMPUTE_MODEL_ALLOW_GPGPU
+
+// The ComputeModel class allows you to select the type of hardware to use
+// in your computational algorithms.
+//
+// If your computational algorithm requires the GPU, set 'inEngine' to the
+// object that will execute the compute shaders.  If your algorithm
+// requires CPU multithreading, set the 'inNumThreads' to the desired
+// number of threads, presumably 2 or larger.  You can query for the number
+// of concurrent hardware threads using std::thread::hardware_concurrency().
+// If you want single-threaded computations (on the main thread), set
+// inNumThreads to 1.  An example of using this class is
+//
+//  ComputeModel cmodel(...);
+//  if (cmodel.engine)
+//  {
+//      ComputeUsingGPU(...);
+//  }
+//  else if (cmodel.numThreads > 1)
+//  {
+//      ComputeUsingCPUMultipleThreads();
+//  }
+//  else
+//  {
+//      ComputeUsingCPUSingleThread();
+//  }
+// See GenerateMeshUV<Real>::SolveSystem(...) for a concrete example.
+//
+// Of course, your algorithm can interpret cmodel anyway it likes.  For
+// example, you might ignore cmodel.engine if all you care about is
+// multithreading on the CPU.
+
+#if defined(GTE_COMPUTE_MODEL_ALLOW_GPGPU)
+#include <memory>
+#endif
+
+namespace gte
+{
+
+#if defined(GTE_COMPUTE_MODEL_ALLOW_GPGPU)
+class GraphicsEngine;
+class ProgramFactory;
+#endif
+
+class GTE_IMPEXP ComputeModel
+{
+public:
+    // Construction and destruction.  You may derive from this class to make
+    // additional behavior available to your algorithms.  For example, you
+    // might add a callback function to report progress of the algorithm.
+    virtual ~ComputeModel()
+    {
+    }
+
+    ComputeModel()
+        :
+        numThreads(1)
+    {
+    }
+
+    ComputeModel(unsigned int inNumThreads)
+        :
+        numThreads(inNumThreads > 0 ? inNumThreads : 1)
+    {
+    }
+
+#if defined(GTE_COMPUTE_MODEL_ALLOW_GPGPU)
+    ComputeModel(unsigned int inNumThreads,
+        std::shared_ptr<GraphicsEngine> const& inEngine,
+        std::shared_ptr<ProgramFactory> const& inFactory)
+        :
+        numThreads(inNumThreads > 0 ? inNumThreads : 1),
+        engine(inEngine),
+        factory(inFactory)
+    {
+    }
+#endif
+
+    unsigned int numThreads;
+#if defined(GTE_COMPUTE_MODEL_ALLOW_GPGPU)
+    std::shared_ptr<GraphicsEngine> engine;
+    std::shared_ptr<ProgramFactory> factory;
+#endif
+};
+
+}
diff --git a/Engine/Plugins/Experimental/GeometryProcessing/Source/GeometryAlgorithms/Private/GTEngine/LowLevel/GteLexicoArray2.h b/Engine/Plugins/Experimental/GeometryProcessing/Source/GeometryAlgorithms/Private/GTEngine/LowLevel/GteLexicoArray2.h
new file mode 100644
--- /dev/null
+++ b/Engine/Plugins/Experimental/GeometryProcessing/Source/GeometryAlgorithms/Private/GTEngine/LowLevel/GteLexicoArray2.h
@@ -0,0 +1,219 @@
+// David Eberly, Geometric Tools, Redmond WA 98052
+// Copyright (c) 1998-2019
+// Distributed under the Boost Software License, Version 1.0.
+// http://www.boost.org/LICENSE_1_0.txt
+// http://www.geometrictools.com/License/Boost/LICENSE_1_0.txt
+// File Version: 3.0.0 (2016/06/19)
+
+#pragma once
+
+#include <GTEngine/GTEngineDEF.h>
+
+namespace gte
+{
+
+// A template class to provide 2D array access that conforms to row-major
+// order (RowMajor = true) or column-major order (RowMajor = false).  The
+template <bool RowMajor, typename Real, int... Dimensions>
+class LexicoArray2 {};
+
+// The array dimensions are known only at run time.
+template <typename Real>
+class LexicoArray2<true, Real>
+{
+public:
+    inline LexicoArray2(int numRows, int numCols, Real* matrix);
+
+    inline int GetNumRows() const;
+    inline int GetNumCols() const;
+    inline Real& operator()(int r, int c);
+    inline Real const& operator()(int r, int c) const;
+
+private:
+    int mNumRows, mNumCols;
+    Real* mMatrix;
+};
+
+template <typename Real>
+class LexicoArray2<false, Real>
+{
+public:
+    inline LexicoArray2(int numRows, int numCols, Real* matrix);
+
+    inline int GetNumRows() const;
+    inline int GetNumCols() const;
+    inline Real& operator()(int r, int c);
+    inline Real const& operator()(int r, int c) const;
+
+private:
+    int mNumRows, mNumCols;
+    Real* mMatrix;
+};
+
+// The array dimensions are known at compile time.
+template <typename Real, int NumRows, int NumCols>
+class LexicoArray2<true, Real, NumRows, NumCols>
+{
+public:
+    inline LexicoArray2(Real* matrix);
+
+    inline int GetNumRows() const;
+    inline int GetNumCols() const;
+    inline Real& operator()(int r, int c);
+    inline Real const& operator()(int r, int c) const;
+
+private:
+    Real* mMatrix;
+};
+
+template <typename Real, int NumRows, int NumCols>
+class LexicoArray2<false, Real, NumRows, NumCols>
+{
+public:
+    inline LexicoArray2(Real* matrix);
+
+    inline int GetNumRows() const;
+    inline int GetNumCols() const;
+    inline Real& operator()(int r, int c);
+    inline Real const& operator()(int r, int c) const;
+
+private:
+    Real* mMatrix;
+};
+
+
+template <typename Real> inline
+LexicoArray2<true, Real>::LexicoArray2(int numRows, int numCols, Real* matrix)
+    :
+    mNumRows(numRows),
+    mNumCols(numCols),
+    mMatrix(matrix)
+{
+}
+
+template <typename Real> inline
+int LexicoArray2<true, Real>::GetNumRows() const
+{
+    return mNumRows;
+}
+
+template <typename Real> inline
+int LexicoArray2<true, Real>::GetNumCols() const
+{
+    return mNumCols;
+}
+
+template <typename Real> inline
+Real& LexicoArray2<true, Real>::operator()(int r, int c)
+{
+    return mMatrix[c + mNumCols*r];
+}
+
+template <typename Real> inline
+Real const& LexicoArray2<true, Real>::operator()(int r, int c) const
+{
+    return mMatrix[c + mNumCols*r];
+}
+
+
+
+template <typename Real> inline
+LexicoArray2<false, Real>::LexicoArray2(int numRows, int numCols, Real* matrix)
+    :
+    mNumRows(numRows),
+    mNumCols(numCols),
+    mMatrix(matrix)
+{
+}
+
+template <typename Real> inline
+int LexicoArray2<false, Real>::GetNumRows() const
+{
+    return mNumRows;
+}
+
+template <typename Real> inline
+int LexicoArray2<false, Real>::GetNumCols() const
+{
+    return mNumCols;
+}
+
+template <typename Real> inline
+Real& LexicoArray2<false, Real>::operator()(int r, int c)
+{
+    return mMatrix[r + mNumRows*c];
+}
+
+template <typename Real> inline
+Real const& LexicoArray2<false, Real>::operator()(int r, int c) const
+{
+    return mMatrix[r + mNumRows*c];
+}
+
+
+
+template <typename Real, int NumRows, int NumCols> inline
+LexicoArray2<true, Real, NumRows, NumCols>::LexicoArray2(Real* matrix)
+    :
+    mMatrix(matrix)
+{
+}
+
+template <typename Real, int NumRows, int NumCols> inline
+int LexicoArray2<true, Real, NumRows, NumCols>::GetNumRows() const
+{
+    return NumRows;
+}
+
+template <typename Real, int NumRows, int NumCols> inline
+int LexicoArray2<true, Real, NumRows, NumCols>::GetNumCols() const
+{
+    return NumCols;
+}
+
+template <typename Real, int NumRows, int NumCols> inline
+Real& LexicoArray2<true, Real, NumRows, NumCols>::operator()(int r, int c)
+{
+    return mMatrix[c + NumCols*r];
+}
+
+template <typename Real, int NumRows, int NumCols> inline
+Real const& LexicoArray2<true, Real, NumRows, NumCols>::operator()(int r, int c) const
+{
+    return mMatrix[c + NumCols*r];
+}
+
+
+
+template <typename Real, int NumRows, int NumCols> inline
+LexicoArray2<false, Real, NumRows, NumCols>::LexicoArray2(Real* matrix)
+    :
+    mMatrix(matrix)
+{
+}
+
+template <typename Real, int NumRows, int NumCols> inline
+int LexicoArray2<false, Real, NumRows, NumCols>::GetNumRows() const
+{
+    return NumRows;
+}
+
+template <typename Real, int NumRows, int NumCols> inline
+int LexicoArray2<false, Real, NumRows, NumCols>::GetNumCols() const
+{
+    return NumCols;
+}
+
+template <typename Real, int NumRows, int NumCols> inline
+Real& LexicoArray2<false, Real, NumRows, NumCols>::operator()(int r, int c)
+{
+    return mMatrix[r + NumRows*c];
+}
+
+template <typename Real, int NumRows, int NumCols> inline
+Real const& LexicoArray2<false, Real, NumRows, NumCols>::operator()(int r, int c) const
+{
+    return mMatrix[r + NumRows*c];
+}
+
+}
