move GeometryProcessing plugin from NotForLicensees to Experimental

diff --git a/Engine/Plugins/Experimental/GeometryProcessing/Source/GeometryAlgorithms/Private/GTEngine/LowLevel/GteLogReporter.cpp b/Engine/Plugins/Experimental/GeometryProcessing/Source/GeometryAlgorithms/Private/GTEngine/LowLevel/GteLogReporter.cpp
new file mode 100644
--- /dev/null
+++ b/Engine/Plugins/Experimental/GeometryProcessing/Source/GeometryAlgorithms/Private/GTEngine/LowLevel/GteLogReporter.cpp
@@ -0,0 +1,47 @@
+// David Eberly, Geometric Tools, Redmond WA 98052
+// Copyright (c) 1998-2019
+// Distributed under the Boost Software License, Version 1.0.
+// http://www.boost.org/LICENSE_1_0.txt
+// http://www.geometrictools.com/License/Boost/LICENSE_1_0.txt
+// File Version: 3.0.0 (2016/06/19)
+
+#include <GTEngine/LowLevel/GteLogReporter.h>
+#include <GTEngine/GTEnginePCH.h>
+using namespace gte;
+
+LogReporter::~LogReporter()
+{
+    if (mLogToStdout)
+    {
+        Logger::Unsubscribe(mLogToStdout.get());
+    }
+
+    if (mLogToFile)
+    {
+        Logger::Unsubscribe(mLogToFile.get());
+    }
+
+
+}
+
+LogReporter::LogReporter(std::string const& logFile, int logFileFlags,
+	int logStdoutFlags)
+    :
+    mLogToFile(nullptr),
+    mLogToStdout(nullptr)
+
+{
+    if (logFileFlags != Logger::Listener::LISTEN_FOR_NOTHING)
+    {
+        mLogToFile = std::make_unique<LogToFile>(logFile, logFileFlags);
+        Logger::Subscribe(mLogToFile.get());
+    }
+
+    if (logStdoutFlags != Logger::Listener::LISTEN_FOR_NOTHING)
+    {
+        mLogToStdout = std::make_unique<LogToStdout>(logStdoutFlags);
+        Logger::Subscribe(mLogToStdout.get());
+    }
+
+
+}
diff --git a/Engine/Plugins/Experimental/GeometryProcessing/Source/GeometryAlgorithms/Private/GTEngine/LowLevel/GteLogReporter.h b/Engine/Plugins/Experimental/GeometryProcessing/Source/GeometryAlgorithms/Private/GTEngine/LowLevel/GteLogReporter.h
new file mode 100644
--- /dev/null
+++ b/Engine/Plugins/Experimental/GeometryProcessing/Source/GeometryAlgorithms/Private/GTEngine/LowLevel/GteLogReporter.h
@@ -0,0 +1,34 @@
+// David Eberly, Geometric Tools, Redmond WA 98052
+// Copyright (c) 1998-2019
+// Distributed under the Boost Software License, Version 1.0.
+// http://www.boost.org/LICENSE_1_0.txt
+// http://www.geometrictools.com/License/Boost/LICENSE_1_0.txt
+// File Version: 3.0.0 (2016/06/19)
+
+#pragma once
+
+#include <GTEngine/LowLevel/GteLogToFile.h>
+#include <GTEngine/LowLevel/GteLogToStdout.h>
+#include <memory>
+
+namespace gte
+{
+
+class GTE_IMPEXP LogReporter
+{
+public:
+    // Construction and destruction.  Create one of these objects in an
+    // application for logging.  The GenerateProject tool creates such code.
+    // If you do not want a particular logger, set the flags to
+    // LISTEN_FOR_NOTHING and set logFile to "" if you do not want a file.
+    ~LogReporter();
+
+    LogReporter(std::string const& logFile, int logFileFlags, int logStdoutFlags);
+
+private:
+    std::unique_ptr<LogToFile> mLogToFile;
+    std::unique_ptr<LogToStdout> mLogToStdout;
+
+};
+
+}
diff --git a/Engine/Plugins/Experimental/GeometryProcessing/Source/GeometryAlgorithms/Private/GTEngine/LowLevel/GteLogToFile.cpp b/Engine/Plugins/Experimental/GeometryProcessing/Source/GeometryAlgorithms/Private/GTEngine/LowLevel/GteLogToFile.cpp
new file mode 100644
--- /dev/null
+++ b/Engine/Plugins/Experimental/GeometryProcessing/Source/GeometryAlgorithms/Private/GTEngine/LowLevel/GteLogToFile.cpp
@@ -0,0 +1,53 @@
+// David Eberly, Geometric Tools, Redmond WA 98052
+// Copyright (c) 1998-2019
+// Distributed under the Boost Software License, Version 1.0.
+// http://www.boost.org/LICENSE_1_0.txt
+// http://www.geometrictools.com/License/Boost/LICENSE_1_0.txt
+// File Version: 3.0.0 (2016/06/19)
+
+#include <GTEngine/LowLevel/GteLogToFile.h>
+#include <GTEngine/GTEnginePCH.h>
+#include <fstream>
+using namespace gte;
+
+
+LogToFile::LogToFile(std::string const& filename, int flags)
+    :
+    Logger::Listener(flags),
+    mFilename(filename)
+{
+    std::ofstream logFile(filename);
+    if (logFile)
+    {
+        // This clears the file contents from any previous runs.
+        logFile.close();
+    }
+    else
+    {
+        // The file cannot be opened.  Use a null string for Report to know
+        // not to attempt opening the file for append.
+        mFilename = "";
+    }
+}
+
+void LogToFile::Report(std::string const& message)
+{
+    if (mFilename != "")
+    {
+        // Open for append.
+        std::ofstream logFile(mFilename,
+            std::ios_base::out | std::ios_base::app);
+        if (logFile)
+        {
+            logFile << message.c_str();
+            logFile.close();
+        }
+        else
+        {
+            // The file cannot be opened.  Use a null string for Report not
+            // to attempt opening the file for append on the next call.
+            mFilename = "";
+        }
+    }
+}
+
diff --git a/Engine/Plugins/Experimental/GeometryProcessing/Source/GeometryAlgorithms/Private/GTEngine/LowLevel/GteLogToFile.h b/Engine/Plugins/Experimental/GeometryProcessing/Source/GeometryAlgorithms/Private/GTEngine/LowLevel/GteLogToFile.h
new file mode 100644
--- /dev/null
+++ b/Engine/Plugins/Experimental/GeometryProcessing/Source/GeometryAlgorithms/Private/GTEngine/LowLevel/GteLogToFile.h
@@ -0,0 +1,26 @@
+// David Eberly, Geometric Tools, Redmond WA 98052
+// Copyright (c) 1998-2019
+// Distributed under the Boost Software License, Version 1.0.
+// http://www.boost.org/LICENSE_1_0.txt
+// http://www.geometrictools.com/License/Boost/LICENSE_1_0.txt
+// File Version: 3.0.0 (2016/06/19)
+
+#pragma once
+
+#include <GTEngine/LowLevel/GteLogger.h>
+
+namespace gte
+{
+
+class GTE_IMPEXP LogToFile : public Logger::Listener
+{
+public:
+    LogToFile(std::string const& filename, int flags);
+
+private:
+    virtual void Report(std::string const& message);
+
+    std::string mFilename;
+};
+
+}
diff --git a/Engine/Plugins/Experimental/GeometryProcessing/Source/GeometryAlgorithms/Private/GTEngine/LowLevel/GteLogToStdout.cpp b/Engine/Plugins/Experimental/GeometryProcessing/Source/GeometryAlgorithms/Private/GTEngine/LowLevel/GteLogToStdout.cpp
new file mode 100644
--- /dev/null
+++ b/Engine/Plugins/Experimental/GeometryProcessing/Source/GeometryAlgorithms/Private/GTEngine/LowLevel/GteLogToStdout.cpp
@@ -0,0 +1,24 @@
+// David Eberly, Geometric Tools, Redmond WA 98052
+// Copyright (c) 1998-2019
+// Distributed under the Boost Software License, Version 1.0.
+// http://www.boost.org/LICENSE_1_0.txt
+// http://www.geometrictools.com/License/Boost/LICENSE_1_0.txt
+// File Version: 3.0.0 (2016/06/19)
+
+#include <GTEngine/LowLevel/GteLogToStdout.h>
+#include <GTEngine/GTEnginePCH.h>
+#include <iostream>
+using namespace gte;
+
+
+LogToStdout::LogToStdout(int flags)
+    :
+    Logger::Listener(flags)
+{
+}
+
+void LogToStdout::Report(std::string const& message)
+{
+    std::cout << message.c_str() << std::flush;
+}
+
diff --git a/Engine/Plugins/Experimental/GeometryProcessing/Source/GeometryAlgorithms/Private/GTEngine/LowLevel/GteLogToStdout.h b/Engine/Plugins/Experimental/GeometryProcessing/Source/GeometryAlgorithms/Private/GTEngine/LowLevel/GteLogToStdout.h
new file mode 100644
--- /dev/null
+++ b/Engine/Plugins/Experimental/GeometryProcessing/Source/GeometryAlgorithms/Private/GTEngine/LowLevel/GteLogToStdout.h
@@ -0,0 +1,24 @@
+// David Eberly, Geometric Tools, Redmond WA 98052
+// Copyright (c) 1998-2019
+// Distributed under the Boost Software License, Version 1.0.
+// http://www.boost.org/LICENSE_1_0.txt
+// http://www.geometrictools.com/License/Boost/LICENSE_1_0.txt
+// File Version: 3.0.0 (2016/06/19)
+
+#pragma once
+
+#include <GTEngine/LowLevel/GteLogger.h>
+
+namespace gte
+{
+
+class GTE_IMPEXP LogToStdout : public Logger::Listener
+{
+public:
+    LogToStdout(int flags);
+
+private:
+    virtual void Report(std::string const& message);
+};
+
+}
diff --git a/Engine/Plugins/Experimental/GeometryProcessing/Source/GeometryAlgorithms/Private/GTEngine/LowLevel/GteLogToStringArray.cpp b/Engine/Plugins/Experimental/GeometryProcessing/Source/GeometryAlgorithms/Private/GTEngine/LowLevel/GteLogToStringArray.cpp
new file mode 100644
--- /dev/null
+++ b/Engine/Plugins/Experimental/GeometryProcessing/Source/GeometryAlgorithms/Private/GTEngine/LowLevel/GteLogToStringArray.cpp
@@ -0,0 +1,39 @@
+// David Eberly, Geometric Tools, Redmond WA 98052
+// Copyright (c) 1998-2019
+// Distributed under the Boost Software License, Version 1.0.
+// http://www.boost.org/LICENSE_1_0.txt
+// http://www.geometrictools.com/License/Boost/LICENSE_1_0.txt
+// File Version: 3.0.0 (2016/06/19)
+
+#include <GTEngine/LowLevel/GteLogToStringArray.h>
+#include <GTEngine/GTEnginePCH.h>
+using namespace gte;
+
+
+LogToStringArray::LogToStringArray(std::string const& name, int flags)
+    :
+    Logger::Listener(flags),
+    mName(name)
+{
+}
+
+std::string const& LogToStringArray::GetName() const
+{
+    return mName;
+}
+
+std::vector<std::string> const& LogToStringArray::GetMessages() const
+{
+    return mMessages;
+}
+
+std::vector<std::string>& LogToStringArray::GetMessages()
+{
+    return mMessages;
+}
+
+void LogToStringArray::Report(std::string const& message)
+{
+    mMessages.push_back(message);
+}
+
diff --git a/Engine/Plugins/Experimental/GeometryProcessing/Source/GeometryAlgorithms/Private/GTEngine/LowLevel/GteLogToStringArray.h b/Engine/Plugins/Experimental/GeometryProcessing/Source/GeometryAlgorithms/Private/GTEngine/LowLevel/GteLogToStringArray.h
new file mode 100644
--- /dev/null
+++ b/Engine/Plugins/Experimental/GeometryProcessing/Source/GeometryAlgorithms/Private/GTEngine/LowLevel/GteLogToStringArray.h
@@ -0,0 +1,32 @@
+// David Eberly, Geometric Tools, Redmond WA 98052
+// Copyright (c) 1998-2019
+// Distributed under the Boost Software License, Version 1.0.
+// http://www.boost.org/LICENSE_1_0.txt
+// http://www.geometrictools.com/License/Boost/LICENSE_1_0.txt
+// File Version: 3.0.0 (2016/06/19)
+
+#pragma once
+
+#include <GTEngine/LowLevel/GteLogger.h>
+#include <vector>
+
+namespace gte
+{
+
+class GTE_IMPEXP LogToStringArray : public Logger::Listener
+{
+public:
+    LogToStringArray(std::string const& name, int flags);
+
+    std::string const& GetName() const;
+    std::vector<std::string> const& GetMessages() const;
+    std::vector<std::string>& GetMessages();
+
+private:
+    virtual void Report(std::string const& message);
+
+    std::string mName;
+    std::vector<std::string> mMessages;
+};
+
+}
diff --git a/Engine/Plugins/Experimental/GeometryProcessing/Source/GeometryAlgorithms/Private/GTEngine/LowLevel/GteLogger.cpp b/Engine/Plugins/Experimental/GeometryProcessing/Source/GeometryAlgorithms/Private/GTEngine/LowLevel/GteLogger.cpp
new file mode 100644
--- /dev/null
+++ b/Engine/Plugins/Experimental/GeometryProcessing/Source/GeometryAlgorithms/Private/GTEngine/LowLevel/GteLogger.cpp
@@ -0,0 +1,135 @@
+// David Eberly, Geometric Tools, Redmond WA 98052
+// Copyright (c) 1998-2019
+// Distributed under the Boost Software License, Version 1.0.
+// http://www.boost.org/LICENSE_1_0.txt
+// http://www.geometrictools.com/License/Boost/LICENSE_1_0.txt
+// File Version: 3.0.0 (2016/06/19)
+
+#include <GTEngine/LowLevel/GteLogger.h>
+#include <GTEngine/GTEnginePCH.h>
+using namespace gte;
+
+
+Logger::Logger(char const* file, char const* function, int line,
+    std::string const& message)
+{
+    mMessage =
+        "File: " + std::string(file) + "\n" +
+        "Func: " + std::string(function) + "\n" +
+        "Line: " + std::to_string(line) + "\n" +
+        message + "\n\n";
+}
+
+void Logger::Assertion()
+{
+    msMutex.lock();
+    for (auto listener : msListeners)
+    {
+        if (listener->GetFlags() & Listener::LISTEN_FOR_ASSERTION)
+        {
+            listener->Assertion(mMessage);
+        }
+    }
+    msMutex.unlock();
+}
+
+void Logger::Error()
+{
+    msMutex.lock();
+    for (auto listener : msListeners)
+    {
+        if (listener->GetFlags() & Listener::LISTEN_FOR_ERROR)
+        {
+            listener->Error(mMessage);
+        }
+    }
+    msMutex.unlock();
+}
+
+void Logger::Warning()
+{
+    msMutex.lock();
+    for (auto listener : msListeners)
+    {
+        if (listener->GetFlags() & Listener::LISTEN_FOR_WARNING)
+        {
+            listener->Warning(mMessage);
+        }
+    }
+    msMutex.unlock();
+}
+
+void Logger::Information()
+{
+    msMutex.lock();
+    for (auto listener : msListeners)
+    {
+        if (listener->GetFlags() & Listener::LISTEN_FOR_INFORMATION)
+        {
+            listener->Information(mMessage);
+        }
+    }
+    msMutex.unlock();
+}
+
+void Logger::Subscribe(Listener* listener)
+{
+    msMutex.lock();
+    msListeners.insert(listener);
+    msMutex.unlock();
+}
+
+void Logger::Unsubscribe(Listener* listener)
+{
+    msMutex.lock();
+    msListeners.erase(listener);
+    msMutex.unlock();
+}
+
+
+
+// Logger::Listener
+
+Logger::Listener::~Listener()
+{
+}
+
+Logger::Listener::Listener(int flags)
+    :
+    mFlags(flags)
+{
+}
+
+int Logger::Listener::GetFlags() const
+{
+    return mFlags;
+}
+
+void Logger::Listener::Assertion(std::string const& message)
+{
+    Report("\nGTE ASSERTION:\n" + message);
+}
+
+void Logger::Listener::Error(std::string const& message)
+{
+    Report("\nGTE ERROR:\n" + message);
+}
+
+void Logger::Listener::Warning(std::string const& message)
+{
+    Report("\nGTE WARNING:\n" + message);
+}
+
+void Logger::Listener::Information(std::string const& message)
+{
+    Report("\nGTE INFORMATION:\n" + message);
+}
+
+void Logger::Listener::Report(std::string const&)
+{
+    // Stub for derived classes.
+}
+
+
+std::mutex Logger::msMutex;
+std::set<Logger::Listener*> Logger::msListeners;
diff --git a/Engine/Plugins/Experimental/GeometryProcessing/Source/GeometryAlgorithms/Private/GTEngine/LowLevel/GteLogger.h b/Engine/Plugins/Experimental/GeometryProcessing/Source/GeometryAlgorithms/Private/GTEngine/LowLevel/GteLogger.h
new file mode 100644
--- /dev/null
+++ b/Engine/Plugins/Experimental/GeometryProcessing/Source/GeometryAlgorithms/Private/GTEngine/LowLevel/GteLogger.h
@@ -0,0 +1,110 @@
+// David Eberly, Geometric Tools, Redmond WA 98052
+// Copyright (c) 1998-2019
+// Distributed under the Boost Software License, Version 1.0.
+// http://www.boost.org/LICENSE_1_0.txt
+// http://www.geometrictools.com/License/Boost/LICENSE_1_0.txt
+// File Version: 3.0.0 (2016/06/19)
+
+#pragma once
+
+#include <GTEngine/GTEngineDEF.h>
+#include <mutex>
+#include <set>
+#include <string>
+
+// Uncomment this to turn off the logging system.  The macros LogAssert,
+// LogError, LogWarning, and LogInformation expand to nothing.  (Do this for
+// optimal performance.)
+//#define GTE_NO_LOGGER
+
+namespace gte
+{
+
+class GTE_IMPEXP Logger
+{
+public:
+    // Construction.  The Logger object is designed to exist only for a
+    // single-line call.  A string is generated from the input parameters and
+    // is used for reporting.
+    Logger(char const* file, char const* function, int line,
+        std::string const& message);
+
+    // Notify current listeners about the logged information.
+    void Assertion();
+    void Error();
+    void Warning();
+    void Information();
+
+    // Listeners subscribe to Logger to receive message strings.
+    class Listener
+    {
+    public:
+        enum
+        {
+            LISTEN_FOR_NOTHING      = 0x00000000,
+            LISTEN_FOR_ASSERTION    = 0x00000001,
+            LISTEN_FOR_ERROR        = 0x00000002,
+            LISTEN_FOR_WARNING      = 0x00000004,
+            LISTEN_FOR_INFORMATION  = 0x00000008,
+            LISTEN_FOR_ALL          = 0xFFFFFFFF
+        };
+
+        // Construction and destruction.
+        virtual ~Listener();
+        Listener(int flags = LISTEN_FOR_NOTHING);
+
+        // What the listener wants to hear.
+        int GetFlags() const;
+
+        // Handlers for the messages received from the logger.
+        void Assertion(std::string const& message);
+        void Error(std::string const& message);
+        void Warning(std::string const& message);
+        void Information(std::string const& message);
+
+    private:
+        virtual void Report(std::string const& message);
+
+        int mFlags;
+    };
+
+    static void Subscribe(Listener* listener);
+    static void Unsubscribe(Listener* listener);
+
+private:
+    std::string mMessage;
+
+    static std::mutex msMutex;
+    static std::set<Listener*> msListeners;
+};
+
+}
+
+
+#if !defined(GTE_NO_LOGGER)
+
+#define LogAssert(condition, message) \
+    if (!(condition)) \
+    { \
+        gte::Logger(__FILE__, __FUNCTION__, __LINE__, message).Assertion(); \
+    }
+
+#define LogError(message) \
+    gte::Logger(__FILE__, __FUNCTION__, __LINE__, message).Error()
+
+#define LogWarning(message) \
+    gte::Logger(__FILE__, __FUNCTION__, __LINE__, message).Warning()
+
+#define LogInformation(message) \
+    gte::Logger(__FILE__, __FUNCTION__, __LINE__, message).Information()
+
+#else
+
+// No logging of assertions, warnings, errors, or information.
+#define LogAssert(condition, message)
+#define LogError(message)
+#define LogWarning(message)
+#define LogInformation(message)
+
+#endif
+
diff --git a/Engine/Plugins/Experimental/GeometryProcessing/Source/GeometryAlgorithms/Private/GTEngine/LowLevel/GteMinHeap.h b/Engine/Plugins/Experimental/GeometryProcessing/Source/GeometryAlgorithms/Private/GTEngine/LowLevel/GteMinHeap.h
new file mode 100644
--- /dev/null
+++ b/Engine/Plugins/Experimental/GeometryProcessing/Source/GeometryAlgorithms/Private/GTEngine/LowLevel/GteMinHeap.h
@@ -0,0 +1,442 @@
+// David Eberly, Geometric Tools, Redmond WA 98052
+// Copyright (c) 1998-2019
+// Distributed under the Boost Software License, Version 1.0.
+// http://www.boost.org/LICENSE_1_0.txt
+// http://www.geometrictools.com/License/Boost/LICENSE_1_0.txt
+// File Version: 3.0.0 (2016/06/19)
+
+#pragma once
+
+#include <GTEngine/GTEngineDEF.h>
+#include <vector>
+
+// A min-heap is a binary tree whose nodes have weights and with the
+// constraint that the weight of a parent node is less than or equal to the
+// weights of its children.  This data structure may be used as a priority
+// queue.  If the std::priority_queue interface suffices for your needs, use
+// that instead.  However, for some geometric algorithms, that interface is
+// insufficient for optimal performance.  For example, if you have a polyline
+// vertices that you want to decimate, each vertex's weight depends on its
+// neighbors' locations.  If the minimum-weight vertex is removed from the
+// min-heap, the neighboring vertex weights must be updated--something that
+// is O(1) time when you store the vertices as a doubly linked list.  The
+// neighbors are already in the min-heap, so modifying their weights without
+// removing then from--and then reinserting into--the min-heap requires they
+// must be moved to their proper places to restore the invariant of the
+// min-heap.  With std::priority_queue, you have no direct access to the
+// modified vertices, forcing you to search for those vertices, remove them,
+// update their weights, and re-insert them.  The min-heap implementation here
+// does support the update without removal and reinsertion.
+//
+// The ValueType represents the weight and it must support comparisons
+// "<" and "<=".  Additional information can be stored in the min-heap for
+// convenient access; this is stored as the KeyType.  In the (open) polyline
+// decimation example, the KeyType is a structure that stores indices to
+// a vertex and its neighbors.  The following code illustrates the creation
+// and use of the min-heap.  The Weight() function is whatever you choose to
+// guide which vertices are removed first from the polyline.
+//
+//    struct Vertex { int previous, current, next; };
+//    int numVertices = <number of polyline vertices>;
+//    std::vector<Vector<N, Real>> positions(numVertices);
+//    <assign all positions[*]>;
+//    MinHeap<Vertex, Real> minHeap(numVertices);
+//    std::vector<MinHeap<Vertex, Real>::Record*> records(numVertices);
+//    for (int i = 0; i < numVertices; ++i)
+//    {
+//        Vertex vertex;
+//        vertex.previous = (i + numVertices - 1) % numVertices;
+//        vertex.current = i;
+//        vertex.next = (i + 1) % numVertices;
+//        records[i] = minHeap.Insert(vertex, Weight(positions, vertex));
+//    }
+//
+//    while (minHeap.GetNumElements() >= 2)
+//    {
+//        Vertex vertex;
+//        Real weight;
+//        minHeap.Remove(vertex, weight);
+//        <consume the 'vertex' according to your application's needs>;
+//
+//        // Remove 'vertex' from the doubly linked list.
+//        Vertex& vp = records[vertex.previous]->key;
+//        Vertex& vc = records[vertex.current]->key;
+//        Vertex& vn = records[vertex.next]->key;
+//        vp.next = vc.next;
+//        vn.previous = vc.previous;
+//
+//        // Update the neighbors' weights in the min-heap.
+//        minHeap.Update(records[vertex.previous], Weight(positions, vp));
+//        minHeap.Update(records[vertex.next], Weight(positions, vn));
+//    }
+
+namespace gte
+{
+
+template <typename KeyType, typename ValueType>
+class MinHeap
+{
+public:
+    struct Record
+    {
+        KeyType key;
+        ValueType value;
+        int index;
+    };
+
+    // Construction.  The record 'value' members are uninitialized for native
+    // types chosen for ValueType.  If ValueType is of class type, then the
+    // default constructor is used to set the 'value' members.
+    MinHeap(MinHeap const& minHeap);
+    MinHeap(int maxElements = 0);
+
+    // Assignment.
+    MinHeap& operator=(MinHeap const& minHeap);
+
+    // Clear the min-heap so that it has the specified max elements,
+    // mNumElements is zero, and mPointers are set to the natural ordering
+    // of mRecords.
+    void Reset(int maxElements);
+
+    // Get the remaining number of elements in the min-heap.  This number is
+    // in the range {0..maxElements}.
+    inline int GetNumElements() const;
+
+    // Get the root of the min-heap.  The return value is 'true' whenever the
+    // min-heap is not empty.  This function reads the root but does not
+    // remove the element from the min-heap.
+    inline bool GetMinimum(KeyType& key, ValueType& value) const;
+
+    // Insert into the min-heap the 'value' that corresponds to the 'key'.
+    // The return value is a pointer to the heap record that stores a copy of
+    // 'value', and the pointer value is constant for the life of the min-heap.
+    // If you must update a member of the min-heap, say, as illustrated in the
+    // polyline decimation example, pass the pointer to Update:
+    //    auto* valueRecord = minHeap.Insert(key, value);
+    //    <do whatever>;
+    //    minHeap.Update(valueRecord, newValue).
+    Record* Insert(KeyType const& key, ValueType const& value);
+
+    // Remove the root of the heap and return its 'key' and 'value members.
+    // The root contains the minimum value of all heap elements.  The return
+    // value is 'true' whenever the min-heap was not empty before the Remove
+    // call.
+    bool Remove(KeyType& key, ValueType& value);
+
+    // The value of a heap record must be modified through this function call.
+    // The side effect is that the heap is updated accordingly to restore the
+    // data structure to a min-heap.  The input 'record' should be a pointer
+    // returned by Insert(value); see the comments for the Insert() function.
+    void Update(Record* record, ValueType const& value);
+
+    // Support for debugging.  The functions test whether the data structure
+    // is a valid min-heap.
+    bool IsValid() const;
+
+private:
+    // A 2-level storage system is used.  The pointers have two roles.
+    // Firstly, they are unique to each inserted value in order to support
+    // the Update() capability of the min-heap.  Secondly, they avoid
+    // potentially expensive copying of Record objects as sorting occurs in
+    // the heap.
+    int mNumElements;
+    std::vector<Record> mRecords;
+    std::vector<Record*> mPointers;
+};
+
+
+template <typename KeyType, typename ValueType>
+MinHeap<KeyType, ValueType>::MinHeap(int maxElements)
+{
+    Reset(maxElements);
+}
+
+template <typename KeyType, typename ValueType>
+MinHeap<KeyType, ValueType>::MinHeap(MinHeap const& minHeap)
+{
+    *this = minHeap;
+}
+
+template <typename KeyType, typename ValueType>
+MinHeap<KeyType, ValueType>& MinHeap<KeyType, ValueType>::operator=(MinHeap const& minHeap)
+{
+    mNumElements = minHeap.mNumElements;
+    mRecords = minHeap.mRecords;
+    mPointers.resize(minHeap.mPointers.size());
+    for (auto& record : mRecords)
+    {
+        mPointers[record.index] = &record;
+    }
+    return *this;
+}
+
+template <typename KeyType, typename ValueType>
+void MinHeap<KeyType, ValueType>::Reset(int maxElements)
+{
+    mNumElements = 0;
+    if (maxElements > 0)
+    {
+        mRecords.resize(maxElements);
+        mPointers.resize(maxElements);
+        for (int i = 0; i < maxElements; ++i)
+        {
+            mPointers[i] = &mRecords[i];
+            mPointers[i]->index = i;
+        }
+    }
+    else
+    {
+        mRecords.clear();
+        mPointers.clear();
+    }
+}
+
+template <typename KeyType, typename ValueType>
+inline int MinHeap<KeyType, ValueType>::GetNumElements() const
+{
+    return mNumElements;
+}
+
+template <typename KeyType, typename ValueType>
+inline bool MinHeap<KeyType, ValueType>::GetMinimum(KeyType& key, ValueType& value) const
+{
+    if (mNumElements > 0)
+    {
+        key = mPointers[0]->key;
+        value = mPointers[0]->value;
+        return true;
+    }
+    else
+    {
+        return false;
+    }
+}
+
+template <typename KeyType, typename ValueType>
+typename MinHeap<KeyType, ValueType>::Record*
+MinHeap<KeyType, ValueType>::Insert(KeyType const& key, ValueType const& value)
+{
+    // Return immediately when the heap is full.
+    if (mNumElements == static_cast<int>(mRecords.size()))
+    {
+        return nullptr;
+    }
+
+    // Store the input information in the last heap record, which is the last
+    // leaf in the tree.
+    int child = mNumElements++;
+    Record * record = mPointers[child];
+    record->key = key;
+    record->value = value;
+
+    // Propagate the information toward the root of the tree until it reaches
+    // its correct position, thus restoring the tree to a valid heap.
+    while (child > 0)
+    {
+        int parent = (child - 1) / 2;
+        if (mPointers[parent]->value <= value)
+        {
+            // The parent has a value smaller than or equal to the child's
+            // value, so we now have a valid heap.
+            break;
+        }
+
+        // The parent has a larger value than the child's value.  Swap the
+        // parent and child:
+
+        // Move the parent into the child's slot.
+        mPointers[child] = mPointers[parent];
+        mPointers[child]->index = child;
+
+        // Move the child into the parent's slot.
+        mPointers[parent] = record;
+        mPointers[parent]->index = parent;
+
+        child = parent;
+    }
+
+    return mPointers[child];
+}
+
+template <typename KeyType, typename ValueType>
+bool MinHeap<KeyType, ValueType>::Remove(KeyType& key, ValueType& value)
+{
+    // Return immediately when the heap is empty.
+    if (mNumElements == 0)
+    {
+        return false;
+    }
+
+    // Get the information from the root of the heap.
+    Record* root = mPointers[0];
+    key = root->key;
+    value = root->value;
+
+    // Restore the tree to a heap.  Abstractly, record is the new root of
+    // the heap.  It is moved down the tree via parent-child swaps until it
+    // is in a location that restores the tree to a heap.
+    int last = --mNumElements;
+    Record* record = mPointers[last];
+    int parent = 0, child = 1;
+    while (child <= last)
+    {
+        if (child < last)
+        {
+            // Select the child with smallest value to be the one that is
+            // swapped with the parent, if necessary.
+            int childP1 = child + 1;
+            if (mPointers[childP1]->value < mPointers[child]->value)
+            {
+                child = childP1;
+            }
+        }
+
+        if (record->value <= mPointers[child]->value)
+        {
+            // The tree is now a heap.
+            break;
+        }
+
+        // Move the child into the parent's slot.
+        mPointers[parent] = mPointers[child];
+        mPointers[parent]->index = parent;
+
+        parent = child;
+        child = 2 * child + 1;
+    }
+
+    // The previous 'last' record was moved to the root and propagated down
+    // the tree to its final resting place, restoring the tree to a heap.
+    // The slot mPointers[parent] is that resting place.
+    mPointers[parent] = record;
+    mPointers[parent]->index = parent;
+
+    // The old root record must not be lost.  Attach it to the slot that
+    // contained the old last record.
+    mPointers[last] = root;
+    mPointers[last]->index = last;
+    return true;
+}
+
+template <typename KeyType, typename ValueType>
+void MinHeap<KeyType, ValueType>::Update(Record* record, ValueType const& value)
+{
+    // Return immediately on invalid record.
+    if (!record)
+    {
+        return;
+    }
+
+    int parent, child, childP1, maxChild;
+
+    if (record->value < value)
+    {
+        record->value = value;
+
+        // The new value is larger than the old value.  Propagate it toward
+        // the leaves.
+        parent = record->index;
+        child = 2 * parent + 1;
+        while (child < mNumElements)
+        {
+            // At least one child exists.  Locate the one of maximum value.
+            childP1 = child + 1;
+            if (childP1 < mNumElements)
+            {
+                // Two children exist.
+                if (mPointers[child]->value <= mPointers[childP1]->value)
+                {
+                    maxChild = child;
+                }
+                else
+                {
+                    maxChild = childP1;
+                }
+            }
+            else
+            {
+                // One child exists.
+                maxChild = child;
+            }
+
+            if (value <= mPointers[maxChild]->value)
+            {
+                // The new value is in the correct place to restore the tree
+                // to a heap.
+                break;
+            }
+
+            // The child has a larger value than the parent's value.  Swap
+            // the parent and child:
+
+            // Move the child into the parent's slot.
+            mPointers[parent] = mPointers[maxChild];
+            mPointers[parent]->index = parent;
+
+            // Move the parent into the child's slot.
+            mPointers[maxChild] = record;
+            mPointers[maxChild]->index = maxChild;
+
+            parent = maxChild;
+            child = 2 * parent + 1;
+        }
+    }
+    else if (value < record->value)
+    {
+        record->value = value;
+
+        // The new weight is smaller than the old weight.  Propagate it
+        // toward the root.
+        child = record->index;
+        while (child > 0)
+        {
+            // A parent exists.
+            parent = (child - 1) / 2;
+
+            if (mPointers[parent]->value <= value)
+            {
+                // The new value is in the correct place to restore the tree
+                // to a heap.
+                break;
+            }
+
+            // The parent has a smaller value than the child's value.  Swap
+            // the child and parent:
+
+            // Move the parent into the child's slot.
+            mPointers[child] = mPointers[parent];
+            mPointers[child]->index = child;
+
+            // Move the child into the parent's slot.
+            mPointers[parent] = record;
+            mPointers[parent]->index = parent;
+
+            child = parent;
+        }
+    }
+}
+
+template <typename KeyType, typename ValueType>
+bool MinHeap<KeyType, ValueType>::IsValid() const
+{
+    for (int child = 0; child < mNumElements; ++child)
+    {
+        int parent = (child - 1) / 2;
+        if (parent > 0)
+        {
+            if (mPointers[child]->value < mPointers[parent]->value)
+            {
+                return false;
+            }
+
+            if (mPointers[parent]->index != parent)
+            {
+                return false;
+            }
+        }
+    }
+
+    return true;
+}
+
+}
diff --git a/Engine/Plugins/Experimental/GeometryProcessing/Source/GeometryAlgorithms/Private/GTEngine/LowLevel/GteRangeIteration.h b/Engine/Plugins/Experimental/GeometryProcessing/Source/GeometryAlgorithms/Private/GTEngine/LowLevel/GteRangeIteration.h
new file mode 100644
--- /dev/null
+++ b/Engine/Plugins/Experimental/GeometryProcessing/Source/GeometryAlgorithms/Private/GTEngine/LowLevel/GteRangeIteration.h
@@ -0,0 +1,69 @@
+// David Eberly, Geometric Tools, Redmond WA 98052
+// Copyright (c) 1998-2019
+// Distributed under the Boost Software License, Version 1.0.
+// http://www.boost.org/LICENSE_1_0.txt
+// http://www.geometrictools.com/License/Boost/LICENSE_1_0.txt
+// File Version: 3.0.0 (2016/06/19)
+
+#pragma once
+
+#include <GTEngine/GTEngineDEF.h>
+#include <iterator>
+#include <type_traits>
+
+// For information on range-based for-loops, see
+// http://en.cppreference.com/w/cpp/language/range-for
+
+namespace gte
+{
+
+// The function gte::reverse supports reverse iteration in range-based
+// for-loops using the auto keyword.  For example,
+//
+//   std::vector<int> numbers(4);
+//   int i = 0;
+//   for (auto& number : numbers)
+//   {
+//       number = i++;
+//       std::cout << number << ' ';
+//   }
+//   // Output:  0 1 2 3
+//
+//   for (auto& number : gte::reverse(numbers))
+//   {
+//       std::cout << number << ' ';
+//   }
+//   // Output:  3 2 1 0
+
+template <typename Iterator>
+class ReversalObject
+{
+public:
+    ReversalObject(Iterator begin, Iterator end)
+        :
+        mBegin(begin),
+        mEnd(end)
+    {
+    }
+
+    Iterator begin() const { return mBegin; }
+    Iterator end() const { return mEnd; }
+
+private:
+    Iterator mBegin, mEnd;
+};
+
+template
+<
+    typename Iterable,
+    typename Iterator = decltype(std::begin(std::declval<Iterable>())),
+    typename ReverseIterator = std::reverse_iterator<Iterator>
+>
+ReversalObject<ReverseIterator> reverse(Iterable&& range)
+{
+    return ReversalObject<ReverseIterator>(
+        ReverseIterator(std::end(range)),
+        ReverseIterator(std::begin(range)));
+}
+
+}
diff --git a/Engine/Plugins/Experimental/GeometryProcessing/Source/GeometryAlgorithms/Private/GTEngine/LowLevel/GteSharedPtrCompare.h b/Engine/Plugins/Experimental/GeometryProcessing/Source/GeometryAlgorithms/Private/GTEngine/LowLevel/GteSharedPtrCompare.h
new file mode 100644
--- /dev/null
+++ b/Engine/Plugins/Experimental/GeometryProcessing/Source/GeometryAlgorithms/Private/GTEngine/LowLevel/GteSharedPtrCompare.h
@@ -0,0 +1,86 @@
+// David Eberly, Geometric Tools, Redmond WA 98052
+// Copyright (c) 1998-2018
+// Distributed under the Boost Software License, Version 1.0.
+// http://www.boost.org/LICENSE_1_0.txt
+// http://www.geometrictools.com/License/Boost/LICENSE_1_0.txt
+// File Version: 3.20.0 (2019/01/09)
+
+#pragma once
+
+#include <memory>
+
+// Comparison operators for std::shared_ptr objects. The type T must implement
+// comparison operators. You must be careful when managing containers whose
+// ordering is based on std::shared_ptr comparisons.  The underlying objects
+// can change, which invalidates the container ordering.  If objects do not
+// change while the container persists, these are safe to use.
+//
+// NOTE: std::shared_ptr<T> already has comparison operators, but these
+// compare pointer values instead of comparing the objects referenced by the
+// pointers.  If a container sorted using std::shared_ptr<T> is created for
+// two different executions of a program, the object ordering implied by the
+// pointer ordering can differ.  This might be undesirable for reproducibility
+// of results between executions.
+
+namespace gte
+{
+    // sp0 == sp1
+    template <typename T>
+    struct SharedPtrEQ
+    {
+        bool operator()(std::shared_ptr<T> const& sp0, std::shared_ptr<T> const& sp1) const
+        {
+            return (sp0 ? (sp1 ? *sp0 == *sp1 : false) : !sp1);
+        }
+    };
+
+    // sp0 != sp1
+    template <typename T>
+    struct SharedPtrNEQ
+    {
+        bool operator()(std::shared_ptr<T> const& sp0, std::shared_ptr<T> const& sp1) const
+        {
+            return !SharedPtrEQ<T>()(sp0, sp1);
+        }
+    };
+
+    // sp0 < sp1
+    template <typename T>
+    struct SharedPtrLT
+    {
+        bool operator()(std::shared_ptr<T> const& sp0, std::shared_ptr<T> const& sp1) const
+        {
+            return (sp1 ? (!sp0 || *sp0 < *sp1) : false);
+        }
+    };
+
+    // sp0 <= sp1
+    template <typename T>
+    struct SharedPtrLTE
+    {
+        bool operator()(std::shared_ptr<T> const& sp0, std::shared_ptr<T> const& sp1) const
+        {
+            return !SharedPtrLT<T>()(sp1, sp0);
+        }
+    };
+
+    // sp0 > sp1
+    template <typename T>
+    struct SharedPtrGT
+    {
+        bool operator()(std::shared_ptr<T> const& sp0, std::shared_ptr<T> const& sp1) const
+        {
+            return SharedPtrLT<T>()(sp1, sp0);
+        }
+    };
+
+    // sp0 >= sp1
+    template <typename T>
+    struct SharedPtrGTE
+    {
+        bool operator()(std::shared_ptr<T> const& sp0, std::shared_ptr<T> const& sp1) const
+        {
+            return !SharedPtrLT<T>()(sp0, sp1);
+        }
+    };
+}
diff --git a/Engine/Plugins/Experimental/GeometryProcessing/Source/GeometryAlgorithms/Private/GTEngine/LowLevel/GteThreadSafeMap.h b/Engine/Plugins/Experimental/GeometryProcessing/Source/GeometryAlgorithms/Private/GTEngine/LowLevel/GteThreadSafeMap.h
new file mode 100644
--- /dev/null
+++ b/Engine/Plugins/Experimental/GeometryProcessing/Source/GeometryAlgorithms/Private/GTEngine/LowLevel/GteThreadSafeMap.h
@@ -0,0 +1,161 @@
+// David Eberly, Geometric Tools, Redmond WA 98052
+// Copyright (c) 1998-2019
+// Distributed under the Boost Software License, Version 1.0.
+// http://www.boost.org/LICENSE_1_0.txt
+// http://www.geometrictools.com/License/Boost/LICENSE_1_0.txt
+// File Version: 3.0.0 (2016/06/19)
+
+#pragma once
+
+#include <GTEngine/GTEngineDEF.h>
+#include <map>
+#include <mutex>
+#include <vector>
+
+namespace gte
+{
+
+template <typename Key, typename Value>
+class ThreadSafeMap
+{
+public:
+    // Construction and destruction.
+    virtual ~ThreadSafeMap();
+    ThreadSafeMap();
+
+    // All the operations are thread-safe.
+    bool HasElements() const;
+    bool Exists(Key key) const;
+    void Insert(Key key, Value value);
+    bool Remove(Key key, Value& value);
+    void RemoveAll();
+    bool Get(Key key, Value& value) const;
+    void GatherAll(std::vector<Value>& values) const;
+
+protected:
+    std::map<Key,Value> mMap;
+    mutable std::mutex mMutex;
+};
+
+
+template <typename Key, typename Value>
+ThreadSafeMap<Key, Value>::~ThreadSafeMap()
+{
+}
+
+template <typename Key, typename Value>
+ThreadSafeMap<Key, Value>::ThreadSafeMap()
+{
+}
+
+template <typename Key, typename Value>
+bool ThreadSafeMap<Key, Value>::HasElements() const
+{
+    bool hasElements;
+    mMutex.lock();
+    {
+        hasElements = (mMap.size() > 0);
+    }
+    mMutex.unlock();
+    return hasElements;
+}
+
+template <typename Key, typename Value>
+bool ThreadSafeMap<Key, Value>::Exists(Key key) const
+{
+    bool exists;
+    mMutex.lock();
+    {
+        exists = (mMap.find(key) != mMap.end());
+    }
+    mMutex.unlock();
+    return exists;
+}
+
+template <typename Key, typename Value>
+void ThreadSafeMap<Key, Value>::Insert(Key key, Value value)
+{
+    mMutex.lock();
+    {
+        mMap[key] = value;
+    }
+    mMutex.unlock();
+}
+
+template <typename Key, typename Value>
+bool ThreadSafeMap<Key, Value>::Remove(Key key, Value& value)
+{
+    bool exists;
+    mMutex.lock();
+    {
+        auto iter = mMap.find(key);
+        if (iter != mMap.end())
+        {
+            value = iter->second;
+            mMap.erase(iter);
+            exists = true;
+        }
+        else
+        {
+            exists = false;
+        }
+    }
+    mMutex.unlock();
+    return exists;
+}
+
+template <typename Key, typename Value>
+void ThreadSafeMap<Key, Value>::RemoveAll()
+{
+    mMutex.lock();
+    {
+        mMap.clear();
+    }
+    mMutex.unlock();
+}
+
+template <typename Key, typename Value>
+bool ThreadSafeMap<Key, Value>::Get(Key key, Value& value) const
+{
+    bool exists;
+    mMutex.lock();
+    {
+        auto iter = mMap.find(key);
+        if (iter != mMap.end())
+        {
+            value = iter->second;
+            exists = true;
+        }
+        else
+        {
+            exists = false;
+        }
+    }
+    mMutex.unlock();
+    return exists;
+}
+
+template <typename Key, typename Value>
+void ThreadSafeMap<Key, Value>::GatherAll(std::vector<Value>& values) const
+{
+    mMutex.lock();
+    {
+        if (mMap.size() > 0)
+        {
+            values.resize(mMap.size());
+            auto viter = values.begin();
+            for (auto const& m : mMap)
+            {
+                *viter++ = m.second;
+            }
+        }
+        else
+        {
+            values.clear();
+        }
+    }
+    mMutex.unlock();
+}
+
+
+}
diff --git a/Engine/Plugins/Experimental/GeometryProcessing/Source/GeometryAlgorithms/Private/GTEngine/LowLevel/GteThreadSafeQueue.h b/Engine/Plugins/Experimental/GeometryProcessing/Source/GeometryAlgorithms/Private/GTEngine/LowLevel/GteThreadSafeQueue.h
new file mode 100644
--- /dev/null
+++ b/Engine/Plugins/Experimental/GeometryProcessing/Source/GeometryAlgorithms/Private/GTEngine/LowLevel/GteThreadSafeQueue.h
@@ -0,0 +1,116 @@
+// David Eberly, Geometric Tools, Redmond WA 98052
+// Copyright (c) 1998-2019
+// Distributed under the Boost Software License, Version 1.0.
+// http://www.boost.org/LICENSE_1_0.txt
+// http://www.geometrictools.com/License/Boost/LICENSE_1_0.txt
+// File Version: 3.0.0 (2016/06/19)
+
+#pragma once
+
+#include <GTEngine/GTEngineDEF.h>
+#include <mutex>
+#include <queue>
+
+namespace gte
+{
+
+template <typename Element>
+class ThreadSafeQueue
+{
+public:
+    // Construction and destruction.
+    virtual ~ThreadSafeQueue();
+    ThreadSafeQueue(size_t maxNumElements = 0);
+
+    // All the operations are thread-safe.
+    size_t GetMaxNumElements() const;
+    size_t GetNumElements() const;
+    bool Push(Element const& element);
+    bool Pop(Element& element);
+
+protected:
+    size_t mMaxNumElements;
+    std::queue<Element> mQueue;
+    mutable std::mutex mMutex;
+};
+
+
+template <typename Element>
+ThreadSafeQueue<Element>::~ThreadSafeQueue()
+{
+}
+
+template <typename Element>
+ThreadSafeQueue<Element>::ThreadSafeQueue(size_t maxNumElements)
+    :
+    mMaxNumElements(maxNumElements)
+{
+}
+
+template <typename Element>
+size_t ThreadSafeQueue<Element>::GetMaxNumElements() const
+{
+    size_t maxNumElements;
+    mMutex.lock();
+    {
+        maxNumElements = mMaxNumElements;
+    }
+    mMutex.unlock();
+    return maxNumElements;
+}
+
+template <typename Element>
+size_t ThreadSafeQueue<Element>::GetNumElements() const
+{
+    size_t numElements;
+    mMutex.lock();
+    {
+        numElements = mQueue.size();
+    }
+    mMutex.unlock();
+    return numElements;
+}
+
+template <typename Element>
+bool ThreadSafeQueue<Element>::Push(Element const& element)
+{
+    bool pushed;
+    mMutex.lock();
+    {
+        if (mQueue.size() < mMaxNumElements)
+        {
+            mQueue.push(element);
+            pushed = true;
+        }
+        else
+        {
+            pushed = false;
+        }
+    }
+    mMutex.unlock();
+    return pushed;
+}
+
+template <typename Element>
+bool ThreadSafeQueue<Element>::Pop(Element& element)
+{
+    bool popped;
+    mMutex.lock();
+    {
+        if (mQueue.size() > 0)
+        {
+            element = mQueue.front();
+            mQueue.pop();
+            popped = true;
+        }
+        else
+        {
+            popped = false;
+        }
+    }
+    mMutex.unlock();
+    return popped;
+}
+
+
+}
diff --git a/Engine/Plugins/Experimental/GeometryProcessing/Source/GeometryAlgorithms/Private/GTEngine/LowLevel/GteTimer.cpp b/Engine/Plugins/Experimental/GeometryProcessing/Source/GeometryAlgorithms/Private/GTEngine/LowLevel/GteTimer.cpp
new file mode 100644
--- /dev/null
+++ b/Engine/Plugins/Experimental/GeometryProcessing/Source/GeometryAlgorithms/Private/GTEngine/LowLevel/GteTimer.cpp
@@ -0,0 +1,115 @@
+// David Eberly, Geometric Tools, Redmond WA 98052
+// Copyright (c) 1998-2019
+// Distributed under the Boost Software License, Version 1.0.
+// http://www.boost.org/LICENSE_1_0.txt
+// http://www.geometrictools.com/License/Boost/LICENSE_1_0.txt
+// File Version: 3.0.0 (2016/06/19)
+
+#include <GTEngine/LowLevel/GteTimer.h>
+#include <GTEngine/GTEnginePCH.h>
+using namespace gte;
+
+#if defined(GTE_NEEDS_64_BIT_CLOCK)
+
+#include <cmath>
+#include <windows.h>
+
+Timer::Timer()
+    :
+    mFrequency(0),
+    mInitialTicks(0),
+    mInvFrequency(0.0)
+{
+    LARGE_INTEGER frequency = { 1 }, counter = { 0 };
+    QueryPerformanceFrequency(&frequency);
+    QueryPerformanceCounter(&counter);
+    mFrequency = static_cast<int64_t>(frequency.QuadPart);
+    mInitialTicks = static_cast<int64_t>(counter.QuadPart);
+    mInvFrequency = 1.0 / static_cast<double>(mFrequency);
+}
+
+int64_t Timer::GetNanoseconds() const
+{
+    int64_t ticks = GetTicks();
+    double seconds = mInvFrequency * static_cast<double>(ticks);
+    int64_t nanoseconds = static_cast<int64_t>(std::ceil(1000000000.0 * seconds));
+    return nanoseconds;
+}
+
+int64_t Timer::GetMicroseconds() const
+{
+    int64_t ticks = GetTicks();
+    double seconds = mInvFrequency * static_cast<double>(ticks);
+    int64_t microseconds = static_cast<int64_t>(std::ceil(1000000.0 * seconds));
+    return microseconds;
+}
+
+int64_t Timer::GetMilliseconds() const
+{
+    int64_t ticks = GetTicks();
+    double seconds = mInvFrequency * static_cast<double>(ticks);
+    int64_t milliseconds = static_cast<int64_t>(std::ceil(1000.0 * seconds));
+    return milliseconds;
+}
+
+double Timer::GetSeconds() const
+{
+    int64_t ticks = GetTicks();
+    double seconds = mInvFrequency * static_cast<double>(ticks);
+    return seconds;
+}
+
+void Timer::Reset()
+{
+    LARGE_INTEGER counter = { 0 };
+    QueryPerformanceCounter(&counter);
+    mInitialTicks = static_cast<int64_t>(counter.QuadPart);
+}
+
+int64_t Timer::GetTicks() const
+{
+    LARGE_INTEGER counter = { 0 };
+    QueryPerformanceCounter(&counter);
+    return static_cast<int64_t>(counter.QuadPart) - mInitialTicks;
+}
+
+#else
+
+Timer::Timer()
+{
+    Reset();
+}
+
+int64_t Timer::GetNanoseconds() const
+{
+    auto currentTime = std::chrono::high_resolution_clock::now();
+    return std::chrono::duration_cast<std::chrono::nanoseconds>(
+        currentTime - mInitialTime).count();
+}
+
+int64_t Timer::GetMicroseconds() const
+{
+    auto currentTime = std::chrono::high_resolution_clock::now();
+    return std::chrono::duration_cast<std::chrono::microseconds>(
+        currentTime - mInitialTime).count();
+}
+
+int64_t Timer::GetMilliseconds() const
+{
+    auto currentTime = std::chrono::high_resolution_clock::now();
+    return std::chrono::duration_cast<std::chrono::milliseconds>(
+        currentTime - mInitialTime).count();
+}
+
+double Timer::GetSeconds() const
+{
+    int64_t msecs = GetMilliseconds();
+    return static_cast<double>(msecs) / 1000.0;
+}
+
+void Timer::Reset()
+{
+    mInitialTime = std::chrono::high_resolution_clock::now();
+}
+
+#endif
diff --git a/Engine/Plugins/Experimental/GeometryProcessing/Source/GeometryAlgorithms/Private/GTEngine/LowLevel/GteTimer.h b/Engine/Plugins/Experimental/GeometryProcessing/Source/GeometryAlgorithms/Private/GTEngine/LowLevel/GteTimer.h
new file mode 100644
--- /dev/null
+++ b/Engine/Plugins/Experimental/GeometryProcessing/Source/GeometryAlgorithms/Private/GTEngine/LowLevel/GteTimer.h
@@ -0,0 +1,56 @@
+// David Eberly, Geometric Tools, Redmond WA 98052
+// Copyright (c) 1998-2019
+// Distributed under the Boost Software License, Version 1.0.
+// http://www.boost.org/LICENSE_1_0.txt
+// http://www.geometrictools.com/License/Boost/LICENSE_1_0.txt
+// File Version: 3.0.0 (2016/06/19)
+
+#pragma once
+
+#include <GTEngine/GTEngineDEF.h>
+#include <cstdint>
+
+#if defined(__MSWINDOWS__) && _MSC_VER < 1900
+// MSVS 2013 has an implementation of std::chrono::high_resolution_clock
+// that does not use a 64-bit clock (QueryPerformanceCounter).  Instead,
+// it appears to use a file-system clock that is 24 bits.  To obtain
+// accurate measurements, we need the 64-bit clock.  However, MSVS 2015
+// does use a 64-bit clock.
+#define GTE_NEEDS_64_BIT_CLOCK
+#endif
+
+#if !defined(GTE_NEEDS_64_BIT_CLOCK)
+#include <chrono>
+#endif
+
+namespace gte
+{
+
+class GTE_IMPEXP Timer
+{
+public:
+    // Construction of a high-resolution timer (64-bit).
+    Timer();
+
+    // Get the current time relative to the initial time.
+    int64_t GetNanoseconds() const;
+    int64_t GetMicroseconds() const;
+    int64_t GetMilliseconds() const;
+    double GetSeconds() const;
+
+    // Reset so that the current time is the initial time.
+    void Reset();
+
+private:
+#if defined(GTE_NEEDS_64_BIT_CLOCK)
+    // Internally use QueryPerformanceCounter.
+    int64_t GetTicks() const;
+
+    int64_t mFrequency, mInitialTicks;
+    double mInvFrequency;
+#else
+    std::chrono::high_resolution_clock::time_point mInitialTime;
+#endif
+};
+
+}
diff --git a/Engine/Plugins/Experimental/GeometryProcessing/Source/GeometryAlgorithms/Private/GTEngine/LowLevel/GteWeakPtrCompare.h b/Engine/Plugins/Experimental/GeometryProcessing/Source/GeometryAlgorithms/Private/GTEngine/LowLevel/GteWeakPtrCompare.h
new file mode 100644
--- /dev/null
+++ b/Engine/Plugins/Experimental/GeometryProcessing/Source/GeometryAlgorithms/Private/GTEngine/LowLevel/GteWeakPtrCompare.h
@@ -0,0 +1,81 @@
+// David Eberly, Geometric Tools, Redmond WA 98052
+// Copyright (c) 1998-2018
+// Distributed under the Boost Software License, Version 1.0.
+// http://www.boost.org/LICENSE_1_0.txt
+// http://www.geometrictools.com/License/Boost/LICENSE_1_0.txt
+// File Version: 3.20.0 (2019/01/09)
+
+#pragma once
+
+#include <memory>
+
+// Comparison operators for std::weak_ptr objects.  The type T must implement
+// comparison operators. You must be careful when managing containers whose
+// ordering is based on std::weak_ptr comparisons.  The underlying objects
+// can change, which invalidates the container ordering.  If objects do not
+// change while the container persists, these are safe to use.
+
+namespace gte
+{
+    // wp0 == wp1
+    template <typename T>
+    struct WeakPtrEQ
+    {
+        bool operator()(std::weak_ptr<T> const& wp0, std::weak_ptr<T> const& wp1) const
+        {
+            auto sp0 = wp0.lock(), sp1 = wp1.lock();
+            return (sp0 ? (sp1 ? *sp0 == *sp1 : false) : !sp1);
+        }
+    };
+
+    // wp0 != wp1
+    template <typename T>
+    struct WeakPtrNEQ
+    {
+        bool operator()(std::weak_ptr<T> const& wp0, std::weak_ptr<T> const& wp1) const
+        {
+            return !WeakPtrEQ<T>()(wp0, wp1);
+        }
+    };
+
+    // wp0 < wp1
+    template <typename T>
+    struct WeakPtrLT
+    {
+        bool operator()(std::weak_ptr<T> const& wp0, std::weak_ptr<T> const& wp1) const
+        {
+            auto sp0 = wp0.lock(), sp1 = wp1.lock();
+            return (sp1 ? (!sp0 || *sp0 < *sp1) : false);
+        }
+    };
+
+    // wp0 <= wp1
+    template <typename T>
+    struct WeakPtrLTE
+    {
+        bool operator()(std::weak_ptr<T> const& wp0, std::weak_ptr<T> const& wp1) const
+        {
+            return !WeakPtrLT<T>()(wp1, wp0);
+        }
+    };
+
+    // wp0 > wp1
+    template <typename T>
+    struct WeakPtrGT
+    {
+        bool operator()(std::weak_ptr<T> const& wp0, std::weak_ptr<T> const& wp1) const
+        {
+            return WeakPtrLT<T>()(wp1, wp0);
+        }
+    };
+
+    // wp0 >= wp1
+    template <typename T>
+    struct WeakPtrGTE
+    {
+        bool operator()(std::weak_ptr<T> const& wp0, std::weak_ptr<T> const& wp1) const
+        {
+            return !WeakPtrLT<T>()(wp0, wp1);
+        }
+    };
+}
diff --git a/Engine/Plugins/Experimental/GeometryProcessing/Source/GeometryAlgorithms/Private/GTEngine/LowLevel/GteWrapper.cpp b/Engine/Plugins/Experimental/GeometryProcessing/Source/GeometryAlgorithms/Private/GTEngine/LowLevel/GteWrapper.cpp
new file mode 100644
--- /dev/null
+++ b/Engine/Plugins/Experimental/GeometryProcessing/Source/GeometryAlgorithms/Private/GTEngine/LowLevel/GteWrapper.cpp
@@ -0,0 +1,36 @@
+// David Eberly, Geometric Tools, Redmond WA 98052
+// Copyright (c) 1998-2019
+// Distributed under the Boost Software License, Version 1.0.
+// http://www.boost.org/LICENSE_1_0.txt
+// http://www.geometrictools.com/License/Boost/LICENSE_1_0.txt
+// File Version: 3.0.0 (2016/06/19)
+
+#include <GTEngine/LowLevel/GteWrapper.h>
+#include <GTEngine/GTEnginePCH.h>
+#include <cstring>
+#include <cwchar>
+
+namespace gte
+{
+
+void Memcpy(void* target, void const* source, size_t count)
+{
+#if defined(__MSWINDOWS__)
+    errno_t result = memcpy_s(target, count, source, count);
+    (void)result;  // 0 on success
+#else
+    memcpy(target, source, count);
+#endif
+}
+
+void Memcpy(wchar_t* target, wchar_t const* source, size_t count)
+{
+#if defined(__MSWINDOWS__)
+    errno_t result = wmemcpy_s(target, count, source, count);
+    (void)result;  // 0 on success
+#else
+    wmemcpy(target, source, count);
+#endif
+}
+
+}
diff --git a/Engine/Plugins/Experimental/GeometryProcessing/Source/GeometryAlgorithms/Private/GTEngine/LowLevel/GteWrapper.h b/Engine/Plugins/Experimental/GeometryProcessing/Source/GeometryAlgorithms/Private/GTEngine/LowLevel/GteWrapper.h
new file mode 100644
--- /dev/null
+++ b/Engine/Plugins/Experimental/GeometryProcessing/Source/GeometryAlgorithms/Private/GTEngine/LowLevel/GteWrapper.h
@@ -0,0 +1,21 @@
+// David Eberly, Geometric Tools, Redmond WA 98052
+// Copyright (c) 1998-2019
+// Distributed under the Boost Software License, Version 1.0.
+// http://www.boost.org/LICENSE_1_0.txt
+// http://www.geometrictools.com/License/Boost/LICENSE_1_0.txt
+// File Version: 3.0.0 (2016/06/19)
+
+#pragma once
+
+#include <GTEngine/GTEngineDEF.h>
+#include <cstddef>
+
+// Wrappers around platform-specific low-level library calls.
+
+namespace gte
+{
+
+void Memcpy(void* target, void const* source, size_t count);
+void Memcpy(wchar_t* target, wchar_t const* source, size_t count);
+
+}
diff --git a/Engine/Plugins/Experimental/GeometryProcessing/Source/GeometryAlgorithms/Private/GTEngine/Mathematics/GteACosEstimate.h b/Engine/Plugins/Experimental/GeometryProcessing/Source/GeometryAlgorithms/Private/GTEngine/Mathematics/GteACosEstimate.h
new file mode 100644
--- /dev/null
+++ b/Engine/Plugins/Experimental/GeometryProcessing/Source/GeometryAlgorithms/Private/GTEngine/Mathematics/GteACosEstimate.h
@@ -0,0 +1,161 @@
+// David Eberly, Geometric Tools, Redmond WA 98052
+// Copyright (c) 1998-2019
+// Distributed under the Boost Software License, Version 1.0.
+// http://www.boost.org/LICENSE_1_0.txt
+// http://www.geometrictools.com/License/Boost/LICENSE_1_0.txt
+// File Version: 3.0.1 (2018/10/05)
+
+#pragma once
+
+#include <GTEngine/Mathematics/GteMath.h>
+
+// Approximations to acos(x) of the form f(x) = sqrt(1-x)*p(x)
+// where the polynomial p(x) of degree D minimizes the quantity
+// maximum{|acos(x)/sqrt(1-x) - p(x)| : x in [0,1]} over all
+// polynomials of degree D.
+
+namespace gte
+{
+
+template <typename Real>
+class ACosEstimate
+{
+public:
+    // The input constraint is x in [0,1].  For example,
+    //   float x; // in [0,1]
+    //   float result = ACosEstimate<float>::Degree<3>(x);
+    template <int D>
+    inline static Real Degree(Real x);
+
+private:
+    // Metaprogramming and private implementation to allow specialization of
+    // a template member function.
+    template <int D> struct degree {};
+    inline static Real Evaluate(degree<1>, Real x);
+    inline static Real Evaluate(degree<2>, Real x);
+    inline static Real Evaluate(degree<3>, Real x);
+    inline static Real Evaluate(degree<4>, Real x);
+    inline static Real Evaluate(degree<5>, Real x);
+    inline static Real Evaluate(degree<6>, Real x);
+    inline static Real Evaluate(degree<7>, Real x);
+    inline static Real Evaluate(degree<8>, Real x);
+};
+
+
+template <typename Real>
+template <int D>
+inline Real ACosEstimate<Real>::Degree(Real x)
+{
+    return Evaluate(degree<D>(), x);
+}
+
+template <typename Real>
+inline Real ACosEstimate<Real>::Evaluate(degree<1>, Real x)
+{
+    Real poly;
+    poly = (Real)GTE_C_ACOS_DEG1_C1;
+    poly = (Real)GTE_C_ACOS_DEG1_C0 + poly * x;
+    poly = poly * std::sqrt((Real)1 - x);
+    return poly;
+}
+
+template <typename Real>
+inline Real ACosEstimate<Real>::Evaluate(degree<2>, Real x)
+{
+    Real poly;
+    poly = (Real)GTE_C_ACOS_DEG2_C2;
+    poly = (Real)GTE_C_ACOS_DEG2_C1 + poly * x;
+    poly = (Real)GTE_C_ACOS_DEG2_C0 + poly * x;
+    poly = poly * std::sqrt((Real)1 - x);
+    return poly;
+}
+
+template <typename Real>
+inline Real ACosEstimate<Real>::Evaluate(degree<3>, Real x)
+{
+    Real poly;
+    poly = (Real)GTE_C_ACOS_DEG3_C3;
+    poly = (Real)GTE_C_ACOS_DEG3_C2 + poly * x;
+    poly = (Real)GTE_C_ACOS_DEG3_C1 + poly * x;
+    poly = (Real)GTE_C_ACOS_DEG3_C0 + poly * x;
+    poly = poly * std::sqrt((Real)1 - x);
+    return poly;
+}
+
+template <typename Real>
+inline Real ACosEstimate<Real>::Evaluate(degree<4>, Real x)
+{
+    Real poly;
+    poly = (Real)GTE_C_ACOS_DEG4_C4;
+    poly = (Real)GTE_C_ACOS_DEG4_C3 + poly * x;
+    poly = (Real)GTE_C_ACOS_DEG4_C2 + poly * x;
+    poly = (Real)GTE_C_ACOS_DEG4_C1 + poly * x;
+    poly = (Real)GTE_C_ACOS_DEG4_C0 + poly * x;
+    poly = poly * std::sqrt((Real)1 - x);
+    return poly;
+}
+
+template <typename Real>
+inline Real ACosEstimate<Real>::Evaluate(degree<5>, Real x)
+{
+    Real poly;
+    poly = (Real)GTE_C_ACOS_DEG5_C5;
+    poly = (Real)GTE_C_ACOS_DEG5_C4 + poly * x;
+    poly = (Real)GTE_C_ACOS_DEG5_C3 + poly * x;
+    poly = (Real)GTE_C_ACOS_DEG5_C2 + poly * x;
+    poly = (Real)GTE_C_ACOS_DEG5_C1 + poly * x;
+    poly = (Real)GTE_C_ACOS_DEG5_C0 + poly * x;
+    poly = poly * std::sqrt((Real)1 - x);
+    return poly;
+}
+
+template <typename Real>
+inline Real ACosEstimate<Real>::Evaluate(degree<6>, Real x)
+{
+    Real poly;
+    poly = (Real)GTE_C_ACOS_DEG6_C6;
+    poly = (Real)GTE_C_ACOS_DEG6_C5 + poly * x;
+    poly = (Real)GTE_C_ACOS_DEG6_C4 + poly * x;
+    poly = (Real)GTE_C_ACOS_DEG6_C3 + poly * x;
+    poly = (Real)GTE_C_ACOS_DEG6_C2 + poly * x;
+    poly = (Real)GTE_C_ACOS_DEG6_C1 + poly * x;
+    poly = (Real)GTE_C_ACOS_DEG6_C0 + poly * x;
+    poly = poly * std::sqrt((Real)1 - x);
+    return poly;
+}
+
+template <typename Real>
+inline Real ACosEstimate<Real>::Evaluate(degree<7>, Real x)
+{
+    Real poly;
+    poly = (Real)GTE_C_ACOS_DEG7_C7;
+    poly = (Real)GTE_C_ACOS_DEG7_C6 + poly * x;
+    poly = (Real)GTE_C_ACOS_DEG7_C5 + poly * x;
+    poly = (Real)GTE_C_ACOS_DEG7_C4 + poly * x;
+    poly = (Real)GTE_C_ACOS_DEG7_C3 + poly * x;
+    poly = (Real)GTE_C_ACOS_DEG7_C2 + poly * x;
+    poly = (Real)GTE_C_ACOS_DEG7_C1 + poly * x;
+    poly = (Real)GTE_C_ACOS_DEG7_C0 + poly * x;
+    poly = poly * std::sqrt((Real)1 - x);
+    return poly;
+}
+
+template <typename Real>
+inline Real ACosEstimate<Real>::Evaluate(degree<8>, Real x)
+{
+    Real poly;
+    poly = (Real)GTE_C_ACOS_DEG8_C8;
+    poly = (Real)GTE_C_ACOS_DEG8_C7 + poly * x;
+    poly = (Real)GTE_C_ACOS_DEG8_C6 + poly * x;
+    poly = (Real)GTE_C_ACOS_DEG8_C5 + poly * x;
+    poly = (Real)GTE_C_ACOS_DEG8_C4 + poly * x;
+    poly = (Real)GTE_C_ACOS_DEG8_C3 + poly * x;
+    poly = (Real)GTE_C_ACOS_DEG8_C2 + poly * x;
+    poly = (Real)GTE_C_ACOS_DEG8_C1 + poly * x;
+    poly = (Real)GTE_C_ACOS_DEG8_C0 + poly * x;
+    poly = poly * std::sqrt((Real)1 - x);
+    return poly;
+}
+
+
+}
diff --git a/Engine/Plugins/Experimental/GeometryProcessing/Source/GeometryAlgorithms/Private/GTEngine/Mathematics/GteASinEstimate.h b/Engine/Plugins/Experimental/GeometryProcessing/Source/GeometryAlgorithms/Private/GTEngine/Mathematics/GteASinEstimate.h
new file mode 100644
--- /dev/null
+++ b/Engine/Plugins/Experimental/GeometryProcessing/Source/GeometryAlgorithms/Private/GTEngine/Mathematics/GteASinEstimate.h
@@ -0,0 +1,40 @@
+// David Eberly, Geometric Tools, Redmond WA 98052
+// Copyright (c) 1998-2019
+// Distributed under the Boost Software License, Version 1.0.
+// http://www.boost.org/LICENSE_1_0.txt
+// http://www.geometrictools.com/License/Boost/LICENSE_1_0.txt
+// File Version: 3.0.0 (2016/06/19)
+
+#pragma once
+
+#include <GTEngine/Mathematics/GteACosEstimate.h>
+
+// Approximations to asin(x) of the form f(x) = pi/2 - sqrt(1-x)*p(x)
+// where the polynomial p(x) of degree D minimizes the quantity
+// maximum{|acos(x)/sqrt(1-x) - p(x)| : x in [0,1]} over all
+// polynomials of degree D.  We use the identity asin(x) = pi/2 - acos(x).
+
+namespace gte
+{
+
+template <typename Real>
+class ASinEstimate
+{
+public:
+    // The input constraint is x in [0,1].  For example,
+    //   float x; // in [0,1]
+    //   float result = ASinEstimate<float>::Degree<3>(x);
+    template <int D>
+    inline static Real Degree(Real x);
+};
+
+
+template <typename Real>
+template <int D>
+inline Real ASinEstimate<Real>::Degree(Real x)
+{
+    return (Real)GTE_C_HALF_PI - ACosEstimate<Real>::Degree<D>(x);
+}
+
+
+}
diff --git a/Engine/Plugins/Experimental/GeometryProcessing/Source/GeometryAlgorithms/Private/GTEngine/Mathematics/GteATanEstimate.h b/Engine/Plugins/Experimental/GeometryProcessing/Source/GeometryAlgorithms/Private/GTEngine/Mathematics/GteATanEstimate.h
new file mode 100644
--- /dev/null
+++ b/Engine/Plugins/Experimental/GeometryProcessing/Source/GeometryAlgorithms/Private/GTEngine/Mathematics/GteATanEstimate.h
@@ -0,0 +1,159 @@
+// David Eberly, Geometric Tools, Redmond WA 98052
+// Copyright (c) 1998-2019
+// Distributed under the Boost Software License, Version 1.0.
+// http://www.boost.org/LICENSE_1_0.txt
+// http://www.geometrictools.com/License/Boost/LICENSE_1_0.txt
+// File Version: 3.0.1 (2018/10/05)
+
+#pragma once
+
+#include <GTEngine/Mathematics/GteMath.h>
+
+// Minimax polynomial approximations to atan(x).  The polynomial p(x) of
+// degree D has only odd-power terms, is required to have linear term x,
+// and p(1) = atan(1) = pi/4.  It minimizes the quantity
+// maximum{|atan(x) - p(x)| : x in [-1,1]} over all polynomials of
+// degree D subject to the constraints mentioned.
+
+namespace gte
+{
+
+template <typename Real>
+class ATanEstimate
+{
+public:
+    // The input constraint is x in [-1,1].  For example,
+    //   float x; // in [-1,1]
+    //   float result = ATanEstimate<float>::Degree<3>(x);
+    template <int D>
+    inline static Real Degree(Real x);
+
+    // The input x can be any real number.  Range reduction is used via
+    // the identities atan(x) = pi/2 - atan(1/x) for x > 0, and
+    // atan(x) = -pi/2 - atan(1/x) for x < 0.  For example,
+    //   float x;  // x any real number
+    //   float result = ATanEstimate<float>::DegreeRR<3>(x);
+    template <int D>
+    inline static Real DegreeRR(Real x);
+
+private:
+    // Metaprogramming and private implementation to allow specialization of
+    // a template member function.
+    template <int D> struct degree {};
+    inline static Real Evaluate(degree<3>, Real x);
+    inline static Real Evaluate(degree<5>, Real x);
+    inline static Real Evaluate(degree<7>, Real x);
+    inline static Real Evaluate(degree<9>, Real x);
+    inline static Real Evaluate(degree<11>, Real x);
+    inline static Real Evaluate(degree<13>, Real x);
+};
+
+
+template <typename Real>
+template <int D>
+inline Real ATanEstimate<Real>::Degree(Real x)
+{
+    return Evaluate(degree<D>(), x);
+}
+
+template <typename Real>
+template <int D>
+inline Real ATanEstimate<Real>::DegreeRR(Real x)
+{
+    if (std::abs(x) <= (Real)1)
+    {
+        return Degree<D>(x);
+    }
+    else if (x > (Real)1)
+    {
+        return (Real)GTE_C_HALF_PI - Degree<D>((Real)1 / x);
+    }
+    else
+    {
+        return (Real)-GTE_C_HALF_PI - Degree<D>((Real)1 / x);
+    }
+}
+
+template <typename Real>
+inline Real ATanEstimate<Real>::Evaluate(degree<3>, Real x)
+{
+    Real xsqr = x * x;
+    Real poly;
+    poly = (Real)GTE_C_ATAN_DEG3_C1;
+    poly = (Real)GTE_C_ATAN_DEG3_C0 + poly * xsqr;
+    poly = poly * x;
+    return poly;
+}
+
+template <typename Real>
+inline Real ATanEstimate<Real>::Evaluate(degree<5>, Real x)
+{
+    Real xsqr = x * x;
+    Real poly;
+    poly = (Real)GTE_C_ATAN_DEG5_C2;
+    poly = (Real)GTE_C_ATAN_DEG5_C1 + poly * xsqr;
+    poly = (Real)GTE_C_ATAN_DEG5_C0 + poly * xsqr;
+    poly = poly * x;
+    return poly;
+}
+
+template <typename Real>
+inline Real ATanEstimate<Real>::Evaluate(degree<7>, Real x)
+{
+    Real xsqr = x * x;
+    Real poly;
+    poly = (Real)GTE_C_ATAN_DEG7_C3;
+    poly = (Real)GTE_C_ATAN_DEG7_C2 + poly * xsqr;
+    poly = (Real)GTE_C_ATAN_DEG7_C1 + poly * xsqr;
+    poly = (Real)GTE_C_ATAN_DEG7_C0 + poly * xsqr;
+    poly = poly * x;
+    return poly;
+}
+
+template <typename Real>
+inline Real ATanEstimate<Real>::Evaluate(degree<9>, Real x)
+{
+    Real xsqr = x * x;
+    Real poly;
+    poly = (Real)GTE_C_ATAN_DEG9_C4;
+    poly = (Real)GTE_C_ATAN_DEG9_C3 + poly * xsqr;
+    poly = (Real)GTE_C_ATAN_DEG9_C2 + poly * xsqr;
+    poly = (Real)GTE_C_ATAN_DEG9_C1 + poly * xsqr;
+    poly = (Real)GTE_C_ATAN_DEG9_C0 + poly * xsqr;
+    poly = poly * x;
+    return poly;
+}
+
+template <typename Real>
+inline Real ATanEstimate<Real>::Evaluate(degree<11>, Real x)
+{
+    Real xsqr = x * x;
+    Real poly;
+    poly = (Real)GTE_C_ATAN_DEG11_C5;
+    poly = (Real)GTE_C_ATAN_DEG11_C4 + poly * xsqr;
+    poly = (Real)GTE_C_ATAN_DEG11_C3 + poly * xsqr;
+    poly = (Real)GTE_C_ATAN_DEG11_C2 + poly * xsqr;
+    poly = (Real)GTE_C_ATAN_DEG11_C1 + poly * xsqr;
+    poly = (Real)GTE_C_ATAN_DEG11_C0 + poly * xsqr;
+    poly = poly * x;
+    return poly;
+}
+
+template <typename Real>
+inline Real ATanEstimate<Real>::Evaluate(degree<13>, Real x)
+{
+    Real xsqr = x * x;
+    Real poly;
+    poly = (Real)GTE_C_ATAN_DEG13_C6;
+    poly = (Real)GTE_C_ATAN_DEG13_C5 + poly * xsqr;
+    poly = (Real)GTE_C_ATAN_DEG13_C4 + poly * xsqr;
+    poly = (Real)GTE_C_ATAN_DEG13_C3 + poly * xsqr;
+    poly = (Real)GTE_C_ATAN_DEG13_C2 + poly * xsqr;
+    poly = (Real)GTE_C_ATAN_DEG13_C1 + poly * xsqr;
+    poly = (Real)GTE_C_ATAN_DEG13_C0 + poly * xsqr;
+    poly = poly * x;
+    return poly;
+}
+
+
+}
diff --git a/Engine/Plugins/Experimental/GeometryProcessing/Source/GeometryAlgorithms/Private/GTEngine/Mathematics/GteAlignedBox.h b/Engine/Plugins/Experimental/GeometryProcessing/Source/GeometryAlgorithms/Private/GTEngine/Mathematics/GteAlignedBox.h
new file mode 100644
--- /dev/null
+++ b/Engine/Plugins/Experimental/GeometryProcessing/Source/GeometryAlgorithms/Private/GTEngine/Mathematics/GteAlignedBox.h
@@ -0,0 +1,136 @@
+// David Eberly, Geometric Tools, Redmond WA 98052
+// Copyright (c) 1998-2019
+// Distributed under the Boost Software License, Version 1.0.
+// http://www.boost.org/LICENSE_1_0.txt
+// http://www.geometrictools.com/License/Boost/LICENSE_1_0.txt
+// File Version: 3.0.0 (2016/06/19)
+
+#pragma once
+
+#include <GTEngine/Mathematics/GteVector.h>
+
+// The box is aligned with the standard coordinate axes, which allows us to
+// represent it using minimum and maximum values along each axis.  Some
+// algorithms prefer the centered representation that is used for oriented
+// boxes.  The center is C and the extents are the half-lengths in each
+// coordinate-axis direction.
+
+namespace gte
+{
+
+template <int N, typename Real>
+class AlignedBox
+{
+public:
+    // Construction and destruction.  The default constructor sets the
+    // minimum values to -1 and the maximum values to +1.
+    AlignedBox();
+
+    // Please ensure that inMin[i] <= inMax[i] for all i.
+    AlignedBox(Vector<N, Real> const& inMin, Vector<N, Real> const& inMax);
+
+    // Compute the centered representation.  NOTE:  If you set the minimum
+    // and maximum values, compute C and extents, and then recompute the
+    // minimum and maximum values, the numerical round-off errors can lead to
+    // results different from what you started with.
+    void GetCenteredForm(Vector<N, Real>& center, Vector<N, Real>& extent)
+        const;
+
+    // Public member access.  It is required that min[i] <= max[i].
+    Vector<N, Real> min, max;
+
+public:
+    // Comparisons to support sorted containers.
+    bool operator==(AlignedBox const& box) const;
+    bool operator!=(AlignedBox const& box) const;
+    bool operator< (AlignedBox const& box) const;
+    bool operator<=(AlignedBox const& box) const;
+    bool operator> (AlignedBox const& box) const;
+    bool operator>=(AlignedBox const& box) const;
+};
+
+// Template aliases for convenience.
+template <typename Real>
+using AlignedBox2 = AlignedBox<2, Real>;
+
+template <typename Real>
+using AlignedBox3 = AlignedBox<3, Real>;
+
+
+template <int N, typename Real>
+AlignedBox<N, Real>::AlignedBox()
+{
+    for (int i = 0; i < N; ++i)
+    {
+        min[i] = (Real)-1;
+        max[i] = (Real)+1;
+    }
+}
+
+template <int N, typename Real>
+AlignedBox<N, Real>::AlignedBox(Vector<N, Real> const& inMin,
+    Vector<N, Real> const& inMax)
+{
+    for (int i = 0; i < N; ++i)
+    {
+        min[i] = inMin[i];
+        max[i] = inMax[i];
+    }
+}
+
+template <int N, typename Real>
+void AlignedBox<N, Real>::GetCenteredForm(Vector<N, Real>& center,
+    Vector<N, Real>& extent) const
+{
+    center = (max + min) * (Real)0.5;
+    extent = (max - min) * (Real)0.5;
+}
+
+template <int N, typename Real>
+bool AlignedBox<N, Real>::operator==(AlignedBox const& box) const
+{
+    return min == box.min && max == box.max;
+}
+
+template <int N, typename Real>
+bool AlignedBox<N, Real>::operator!=(AlignedBox const& box) const
+{
+    return !operator==(box);
+}
+
+template <int N, typename Real>
+bool AlignedBox<N, Real>::operator<(AlignedBox const& box) const
+{
+    if (min < box.min)
+    {
+        return true;
+    }
+
+    if (min > box.min)
+    {
+        return false;
+    }
+
+    return max < box.max;
+}
+
+template <int N, typename Real>
+bool AlignedBox<N, Real>::operator<=(AlignedBox const& box) const
+{
+    return operator<(box) || operator==(box);
+}
+
+template <int N, typename Real>
+bool AlignedBox<N, Real>::operator>(AlignedBox const& box) const
+{
+    return !operator<=(box);
+}
+
+template <int N, typename Real>
+bool AlignedBox<N, Real>::operator>=(AlignedBox const& box) const
+{
+    return !operator<(box);
+}
+
+
+}
diff --git a/Engine/Plugins/Experimental/GeometryProcessing/Source/GeometryAlgorithms/Private/GTEngine/Mathematics/GteApprCircle2.h b/Engine/Plugins/Experimental/GeometryProcessing/Source/GeometryAlgorithms/Private/GTEngine/Mathematics/GteApprCircle2.h
new file mode 100644
--- /dev/null
+++ b/Engine/Plugins/Experimental/GeometryProcessing/Source/GeometryAlgorithms/Private/GTEngine/Mathematics/GteApprCircle2.h
@@ -0,0 +1,101 @@
+// David Eberly, Geometric Tools, Redmond WA 98052
+// Copyright (c) 1998-2019
+// Distributed under the Boost Software License, Version 1.0.
+// http://www.boost.org/LICENSE_1_0.txt
+// http://www.geometrictools.com/License/Boost/LICENSE_1_0.txt
+// File Version: 3.0.0 (2016/06/19)
+
+#pragma once
+
+#include <GTEngine/Mathematics/GteApprQuadratic2.h>
+
+// Least-squares fit of a circle to a set of points.  A successful fit is
+// indicated by the return value of 'true'.  The return value is the number
+// of iterations used.  If the number is the maximum, you can either accept
+// the result or try increasing the maximum number and calling the function
+// again.  TODO: Currently, the test for terminating the algorithm is exact
+// (diff == (0,0)).  Expose an epsilon for this?
+//
+// If initialCenterIsAverage is set to 'true', the initial guess for the
+// circle center is the average of the data points.  If the data points are
+// clustered along a small arc, ApprCircle2 is very slow to converge.  If
+// initialCenterIsAverage is set to 'false', the initial guess for the
+// circle center is computed using a least-squares estimate of the
+// coefficients for a quadratic equation that represents a circle.  This
+// approach tends to converge rapidly.
+
+namespace gte
+{
+
+template <typename Real>
+class ApprCircle2
+{
+public:
+    unsigned int operator()(int numPoints, Vector2<Real> const* points,
+        unsigned int maxIterations, bool initialCenterIsAverage,
+        Circle2<Real>& circle);
+};
+
+
+template <typename Real>
+unsigned int ApprCircle2<Real>::operator()(int numPoints,
+    Vector2<Real> const* points, unsigned int maxIterations,
+    bool initialCenterIsAverage, Circle2<Real>& circle)
+{
+    // Compute the average of the data points.
+    Vector2<Real> average = points[0];
+    for (int i = 1; i < numPoints; ++i)
+    {
+        average += points[i];
+    }
+    Real invNumPoints = ((Real)1) / static_cast<Real>(numPoints);
+    average *= invNumPoints;
+
+    // The initial guess for the center.
+    if (initialCenterIsAverage)
+    {
+        circle.center = average;
+    }
+    else
+    {
+        ApprQuadraticCircle2<Real>()(numPoints, points, circle);
+    }
+
+    unsigned int iteration;
+    for (iteration = 0; iteration < maxIterations; ++iteration)
+    {
+        // Update the iterates.
+        Vector2<Real> current = circle.center;
+
+        // Compute average L, dL/da, dL/db.
+        Real lenAverage = (Real)0;
+        Vector2<Real> derLenAverage = Vector2<Real>::Zero();
+        for (int i = 0; i < numPoints; ++i)
+        {
+            Vector2<Real> diff = points[i] - circle.center;
+            Real length = Length(diff);
+            if (length >(Real)0)
+            {
+                lenAverage += length;
+                Real invLength = ((Real)1) / length;
+                derLenAverage -= invLength * diff;
+            }
+        }
+        lenAverage *= invNumPoints;
+        derLenAverage *= invNumPoints;
+
+        circle.center = average + lenAverage * derLenAverage;
+        circle.radius = lenAverage;
+
+        Vector2<Real> diff = circle.center - current;
+        if (diff == Vector2<Real>::Zero())
+        {
+            break;
+        }
+    }
+
+    return ++iteration;
+}
+
+
+}
diff --git a/Engine/Plugins/Experimental/GeometryProcessing/Source/GeometryAlgorithms/Private/GTEngine/Mathematics/GteApprCone3.h b/Engine/Plugins/Experimental/GeometryProcessing/Source/GeometryAlgorithms/Private/GTEngine/Mathematics/GteApprCone3.h
new file mode 100644
--- /dev/null
+++ b/Engine/Plugins/Experimental/GeometryProcessing/Source/GeometryAlgorithms/Private/GTEngine/Mathematics/GteApprCone3.h
@@ -0,0 +1,340 @@
+// David Eberly, Geometric Tools, Redmond WA 98052
+// Copyright (c) 1998-2019
+// Distributed under the Boost Software License, Version 1.0.
+// http://www.boost.org/LICENSE_1_0.txt
+// http://www.geometrictools.com/License/Boost/LICENSE_1_0.txt
+// File Version: 3.14.2 (2018/10/30)
+
+#pragma once
+
+#include <GTEngine/Mathematics/GteGaussNewtonMinimizer.h>
+#include <GTEngine/Mathematics/GteLevenbergMarquardtMinimizer.h>
+#include <GTEngine/Mathematics/GteRootsPolynomial.h>
+#include <GTEngine/Mathematics/GteMath.h>
+
+// The cone vertex is V, the unit-length axis direction is U and the
+// cone angle is A in (0,pi/2).  The cone is defined algebraically by
+// those points X for which
+//     Dot(U,X-V)/Length(X-V) = cos(A)
+// This can be written as a quadratic equation
+//     (V-X)^T * (cos(A)^2 - U * U^T) * (V-X) = 0
+// with the implicit constraint that Dot(U, X-V) > 0 (X is on the
+// "positive" cone).  Define W = U/cos(A), so Length(W) > 1 and
+//     F(X;V,W) = (V-X)^T * (I - W * W^T) * (V-X) = 0
+// The nonlinear least squares fitting of points {X[i]}_{i=0}^{n-1}
+// computes V and W to minimize the error function
+//     E(V,W) = sum_{i=0}^{n-1} F(X[i];V,W)^2
+// I recommend using the Gauss-Newton minimizer when your cone points
+// are truly nearly a cone; otherwise, try the Levenberg-Marquardt
+// minimizer.
+//
+// The mathematics used in this implementation are found in
+//   http://www.geometrictools.com/Documentation/LeastSquaresFitting.pdf
+// In particular, the details for choosing an initial cone for fitting
+// are somewhat complicated.
+
+namespace gte
+{
+    template <typename Real>
+    class ApprCone3
+    {
+    public:
+        ApprCone3()
+            :
+            mNumPoints(0),
+            mPoints(nullptr)
+        {
+            // F[i](V,W) = D^T * (I - W * W^T) * D, D = V - X[i], P = (V,W)
+            mFFunction = [this](GVector<Real> const& P, GVector<Real>& F)
+            {
+                Vector<3, Real> V = { P[0], P[1], P[2] };
+                Vector<3, Real> W = { P[3], P[4], P[5] };
+                for (int i = 0; i < mNumPoints; ++i)
+                {
+                    Vector<3, Real> delta = V - mPoints[i];
+                    Real deltaDotW = Dot(delta, W);
+                    F[i] = Dot(delta, delta) - deltaDotW * deltaDotW;
+                }
+            };
+
+            // dF[i]/dV = 2 * (D - Dot(W, D) * W)
+            // dF[i]/dW = -2 * Dot(W, D) * D
+            mJFunction = [this](GVector<Real> const& P, GMatrix<Real>& J)
+            {
+                Vector<3, Real> V = { P[0], P[1], P[2] };
+                Vector<3, Real> W = { P[3], P[4], P[5] };
+                for (int row = 0; row < mNumPoints; ++row)
+                {
+                    Vector<3, Real> delta = V - mPoints[row];
+                    Real deltaDotW = Dot(delta, W);
+                    Vector<3, Real> temp0 = delta - deltaDotW * W;
+                    Vector<3, Real> temp1 = deltaDotW * delta;
+                    for (int col = 0; col < 3; ++col)
+                    {
+                        J(row, col) = (Real)2 * temp0[col];
+                        J(row, col + 3) = (Real)-2 * temp1[col];
+                    }
+                }
+            };
+        }
+
+        // If you want to specify that coneVertex, coneAxis and coneAngle
+        // are the initial guesses for the minimizer, set the parameter
+        // useConeInputAsInitialGuess to 'true'.  If you want the function
+        // to compute initial guesses, set that parameter to 'false'.
+        // A Gauss-Newton minimizer is used to fit a cone using nonlinear
+        // least-squares.  The fitted cone is returned in coneVertex,
+        // coneAxis and coneAngle.  See GteGaussNewtonMinimizer.h for a
+        // description of the least-squares algorithm and the parameters
+        // that it requires.
+        typename GaussNewtonMinimizer<Real>::Result
+        operator()(int numPoints, Vector<3, Real> const* points,
+            size_t maxIterations, Real updateLengthTolerance, Real errorDifferenceTolerance,
+            bool useConeInputAsInitialGuess,
+            Vector<3, Real>& coneVertex, Vector<3, Real>& coneAxis, Real& coneAngle)
+        {
+            mNumPoints = numPoints;
+            mPoints = points;
+            GaussNewtonMinimizer<Real> minimizer(6, mNumPoints, mFFunction, mJFunction);
+
+            Real coneCosAngle;
+            if (useConeInputAsInitialGuess)
+            {
+                Normalize(coneAxis);
+                coneCosAngle = std::cos(coneAngle);
+            }
+            else
+            {
+                ComputeInitialCone(coneVertex, coneAxis, coneCosAngle);
+            }
+
+            // The initial guess for the cone vertex.
+            GVector<Real> initial(6);
+            initial[0] = coneVertex[0];
+            initial[1] = coneVertex[1];
+            initial[2] = coneVertex[2];
+
+            // The initial guess for the weighted cone axis.
+            initial[3] = coneAxis[0] / coneCosAngle;
+            initial[4] = coneAxis[1] / coneCosAngle;
+            initial[5] = coneAxis[2] / coneCosAngle;
+
+            auto result = minimizer(initial, maxIterations, updateLengthTolerance,
+                errorDifferenceTolerance);
+
+            // No test is made for result.converged so that we return some
+            // estimates of the cone.  The caller can decide how to respond
+            // when result.converged is false.
+            for (int i = 0; i < 3; ++i)
+            {
+                coneVertex[i] = result.minLocation[i];
+                coneAxis[i] = result.minLocation[i + 3];
+            }
+
+            // We know that coneCosAngle will be nonnegative.  The std::min
+            // call guards against rounding errors leading to a number
+            // slightly larger than 1.  The clamping ensures std::acos will
+            // not return a NaN.
+            coneCosAngle = std::min((Real)1 / Normalize(coneAxis), (Real)1);
+            coneAngle = std::acos(coneCosAngle);
+
+            mNumPoints = 0;
+            mPoints = nullptr;
+            return result;
+        }
+
+        // The parameters coneVertex, coneAxis and coneAngle are in/out
+        // variables.  The caller must provide initial guesses for these.
+        // The function estimates the cone parameters and returns them.  See
+        // GteGaussNewtonMinimizer.h for a description of the least-squares
+        // algorithm and the parameters that it requires.
+        typename LevenbergMarquardtMinimizer<Real>::Result
+        operator()(int numPoints, Vector<3, Real> const* points,
+            size_t maxIterations, Real updateLengthTolerance, Real errorDifferenceTolerance,
+            Real lambdaFactor, Real lambdaAdjust, size_t maxAdjustments,
+            bool useConeInputAsInitialGuess,
+            Vector<3, Real>& coneVertex, Vector<3, Real>& coneAxis, Real& coneAngle)
+        {
+            mNumPoints = numPoints;
+            mPoints = points;
+            LevenbergMarquardtMinimizer<Real> minimizer(6, mNumPoints, mFFunction, mJFunction);
+
+            Real coneCosAngle;
+            if (useConeInputAsInitialGuess)
+            {
+                Normalize(coneAxis);
+                coneCosAngle = std::cos(coneAngle);
+            }
+            else
+            {
+                ComputeInitialCone(coneVertex, coneAxis, coneCosAngle);
+            }
+
+            // The initial guess for the cone vertex.
+            GVector<Real> initial(6);
+            initial[0] = coneVertex[0];
+            initial[1] = coneVertex[1];
+            initial[2] = coneVertex[2];
+
+            // The initial guess for the weighted cone axis.
+            initial[3] = coneAxis[0] / coneCosAngle;
+            initial[4] = coneAxis[1] / coneCosAngle;
+            initial[5] = coneAxis[2] / coneCosAngle;
+
+            auto result = minimizer(initial, maxIterations, updateLengthTolerance,
+                errorDifferenceTolerance, lambdaFactor, lambdaAdjust, maxAdjustments);
+
+            // No test is made for result.converged so that we return some
+            // estimates of the cone.  The caller can decide how to respond
+            // when result.converged is false.
+            for (int i = 0; i < 3; ++i)
+            {
+                coneVertex[i] = result.minLocation[i];
+                coneAxis[i] = result.minLocation[i + 3];
+            }
+
+            // We know that coneCosAngle will be nonnegative.  The std::min
+            // call guards against rounding errors leading to a number
+            // slightly larger than 1.  The clamping ensures std::acos will
+            // not return a NaN.
+            coneCosAngle = std::min((Real)1 / Normalize(coneAxis), (Real)1);
+            coneAngle = std::acos(coneCosAngle);
+
+            mNumPoints = 0;
+            mPoints = nullptr;
+            return result;
+        }
+
+    private:
+        void ComputeInitialCone(Vector<3, Real>& coneVertex, Vector<3, Real>& coneAxis, Real& coneCosAngle)
+        {
+            // Compute the average of the sample points.
+            Vector<3, Real> center{ (Real)0, (Real)0, (Real)0 };
+            Real const invNumPoints = (Real)1 / (Real)mNumPoints;
+            for (int i = 0; i < mNumPoints; ++i)
+            {
+                center += mPoints[i];
+            }
+            center *= invNumPoints;
+
+            // The cone axis is estimated from ZZTZ (see the PDF).
+            coneAxis = { (Real)0, (Real)0, (Real)0 };
+            for (int i = 0; i < mNumPoints; ++i)
+            {
+                Vector<3, Real> diff = mPoints[i] - center;
+                coneAxis += Dot(diff, diff) * diff;
+            }
+            coneAxis *= invNumPoints;
+            Normalize(coneAxis);
+
+            // Compute the averages of powers and products of powers of
+            // a[i] = Dot(U,X[i]-C) and b[i] = Dot(X[i]-C,X[i]-C).
+            Real c10 = (Real)0, c20 = (Real)0, c30 = (Real)0, c01 = (Real)0;
+            Real c02 = (Real)0, c11 = (Real)0, c21 = (Real)0;
+            for (int i = 0; i < mNumPoints; ++i)
+            {
+                Vector<3, Real> diff = mPoints[i] - center;
+                Real ai = Dot(coneAxis, diff);
+                Real aisqr = ai * ai;
+                Real bi = Dot(diff, diff);
+                c10 += ai;
+                c20 += aisqr;
+                c30 += aisqr * ai;
+                c01 += bi;
+                c02 += bi * bi;
+                c11 += ai * bi;
+                c21 += aisqr * bi;
+            }
+            c10 *= invNumPoints;
+            c20 *= invNumPoints;
+            c30 *= invNumPoints;
+            c01 *= invNumPoints;
+            c02 *= invNumPoints;
+            c11 *= invNumPoints;
+            c21 *= invNumPoints;
+
+            // Compute the coefficients of p3(t) and q3(t).
+            Real e0 = (Real)3 * c10;
+            Real e1 = (Real)2 * c20 + c01;
+            Real e2 = c11;
+            Real e3 = (Real)3 * c20;
+            Real e4 = c30;
+
+            // Compute the coefficients of g(t).
+            Real g0 = c11 * c21 - c02 * c30;
+            Real g1 = c01 * c21 - (Real)3 * c02 * c20 + (Real)2 * (c20 * c21 - c11 * (c30 - c11));
+            Real g2 = (Real)3 * (c11 * (c01 - c20) + c10 * (c21 - c02));
+            Real g3 = c21 - c02 + c01 * (c01 + c20) + (Real)2 * (c10 * (c30 - c11) - c20 * c20);
+            Real g4 = c30 - c11 + c10 * (c01 - c20);
+
+            // Compute the roots of g(t) = 0.
+            std::map<Real, int> rmMap;
+            RootsPolynomial<Real>::SolveQuartic(g0, g1, g2, g3, g4, rmMap);
+
+            // Locate the positive root t that leads to an s = cos(theta)
+            // in (0,1) and that has minimum least-squares error.  In theory,
+            // there must always be such a root, but floating-point rounding
+            // errors might lead to no such root.  The implementation returns
+            // the center as the estimate of V and pi/4 as the estimate of
+            // the angle (s = 1/2).
+            std::vector<std::array<Real, 3>> info;
+            Real s, t;
+            for (auto const& element : rmMap)
+            {
+                t = element.first;
+                if (t > (Real)0)
+                {
+                    Real p3 = e2 + t * (e1 + t * (e0 + t));
+                    if (p3 != (Real)0)
+                    {
+                        Real q3 = e4 + t * (e3 + t * (e0 + t));
+                        s = q3 / p3;
+                        if ((Real)0 < s && s < (Real)1)
+                        {
+                            Real error(0);
+                            for (int i = 0; i < mNumPoints; ++i)
+                            {
+                                Vector<3, Real> diff = mPoints[i] - center;
+                                Real ai = Dot(coneAxis, diff);
+                                Real bi = Dot(diff, diff);
+                                Real tpai = t + ai;
+                                Real Fi = s * (bi + t * ((Real)2 * ai + t)) - tpai * tpai;
+                                error += Fi * Fi;
+                            }
+                            error *= invNumPoints;
+                            std::array<Real, 3> item = { s, t, error };
+                            info.push_back(item);
+                        }
+                    }
+                }
+            }
+
+            Real minError = std::numeric_limits<Real>::max();
+            std::array<Real, 3> minItem = { (Real)0, (Real)0, minError };
+            for (auto const& item : info)
+            {
+                if (item[2] < minError)
+                {
+                    minItem = item;
+                }
+            }
+
+            if (minItem[2] < std::numeric_limits<Real>::max())
+            {
+                // minItem = { minS, minT, minError }
+                coneVertex = center - minItem[1] * coneAxis;
+                coneCosAngle = std::sqrt(minItem[0]);
+            }
+            else
+            {
+                coneVertex = center;
+                coneCosAngle = (Real)GTE_C_INV_SQRT_2;
+            }
+        }
+
+        int mNumPoints;
+        Vector<3, Real> const* mPoints;
+        std::function<void(GVector<Real> const&, GVector<Real>&)> mFFunction;
+        std::function<void(GVector<Real> const&, GMatrix<Real>&)> mJFunction;
+    };
+}
