move GeometryProcessing plugin from NotForLicensees to Experimental

diff --git a/Engine/Plugins/Experimental/GeometryProcessing/Source/GeometryAlgorithms/Private/GTEngine/Mathematics/GteApprCylinder3.h b/Engine/Plugins/Experimental/GeometryProcessing/Source/GeometryAlgorithms/Private/GTEngine/Mathematics/GteApprCylinder3.h
new file mode 100644
--- /dev/null
+++ b/Engine/Plugins/Experimental/GeometryProcessing/Source/GeometryAlgorithms/Private/GTEngine/Mathematics/GteApprCylinder3.h
@@ -0,0 +1,462 @@
+// David Eberly, Geometric Tools, Redmond WA 98052
+// Copyright (c) 1998-2019
+// Distributed under the Boost Software License, Version 1.0.
+// http://www.boost.org/LICENSE_1_0.txt
+// http://www.geometrictools.com/License/Boost/LICENSE_1_0.txt
+// File Version: 3.3.5 (2018/10/26)
+
+#pragma once
+
+#include <GTEngine/Mathematics/GteCylinder3.h>
+#include <GTEngine/Mathematics/GteMatrix3x3.h>
+#include <GTEngine/Mathematics/GteSymmetricEigensolver3x3.h>
+#include <GTEngine/Mathematics/GteMath.h>
+#include <algorithm>
+#include <vector>
+#include <thread>
+
+// The algorithm for least-squares fitting of a point set by a cylinder is
+// described in
+//   http://www.geometrictools.com/Documentation/CylinderFitting.pdf
+// This document shows how to compute the cylinder radius r and the cylinder
+// axis as a line C+t*W with origin C, unit-length direction W, and any
+// real-valued t.  The implementation here adds one addition step.  It
+// projects the point set onto the cylinder axis, computes the bounding
+// t-interval [tmin,tmax] for the projections, and sets the cylinder center
+// to C + 0.5*(tmin+tmax)*W and the cylinder height to tmax-tmin.
+
+namespace gte
+{
+    template <typename Real>
+    class ApprCylinder3
+    {
+    public:
+        // Search the hemisphere for a minimum, choose numThetaSamples and
+        // numPhiSamples to be positive (and preferably large).  These are used
+        // to generate a hemispherical grid of samples to be evaluated to find
+        // the cylinder axis-direction W.  If the grid samples is quite large
+        // and the number of points to be fitted is large, you most likely will
+        // want to run multithreaded.  Set numThreads to 0 to run single-threaded
+        // in the main process.  Set numThreads > 0 to run multithreaded.  If
+        // either of numThetaSamples or numPhiSamples is zero, the operator() sets
+        // the cylinder origin and axis to the zero vectors, the radius and height
+        // to zero, and returns std::numeric_limits<Real>::max().
+        ApprCylinder3(unsigned int numThreads, unsigned int numThetaSamples, unsigned int numPhiSamples)
+            :
+            mConstructorType(FIT_BY_HEMISPHERE_SEARCH),
+            mNumThreads(numThreads),
+            mNumThetaSamples(numThetaSamples),
+            mNumPhiSamples(numPhiSamples),
+            mEigenIndex(0),
+            mInvNumPoints((Real)0)
+        {
+            mCylinderAxis = { (Real)0, (Real)0, (Real)0 };
+        }
+
+        // Choose one of the eigenvectors for the covariance matrix as the
+        // cylinder axis direction.  If eigenIndex is 0, the eigenvector
+        // associated with the smallest eigenvalue is chosen.  If eigenIndex is 2,
+        // the eigenvector associated with the largest eigenvalue is chosen.  If
+        // eigenIndex is 1, the eigenvector associated with the median eigenvalue
+        // is chosen; keep in mind that this could be the minimum or maximum
+        // eigenvalue if the eigenspace has dimension 2 or 3.  If eigenIndex is
+        // 3 or larger, the operator() sets the cylinder origin and axis to the
+        // zero vectors, the radius and height to zero, and returns
+        // std::numeric_limits<Real>::max().
+        ApprCylinder3(unsigned int eigenIndex)
+            :
+            mConstructorType(FIT_USING_COVARIANCE_EIGENVECTOR),
+            mNumThreads(0),
+            mNumThetaSamples(0),
+            mNumPhiSamples(0),
+            mEigenIndex(eigenIndex),
+            mInvNumPoints((Real)0)
+        {
+            mCylinderAxis = { (Real)0, (Real)0, (Real)0 };
+        }
+
+        // Choose the cylinder axis.  If cylinderAxis is not the zero vector,
+        // the constructor will normalize it.  If cylinderAxis is the zero vector,
+        // the operator() sets the cylinder origin and axis to the zero vectors,
+        // the radius and height to zero, and returns
+        // std::numeric_limits<Real>::max().
+        ApprCylinder3(Vector3<Real> const& cylinderAxis)
+            :
+            mConstructorType(FIT_USING_SPECIFIED_AXIS),
+            mNumThreads(0),
+            mNumThetaSamples(0),
+            mNumPhiSamples(0),
+            mEigenIndex(0),
+            mCylinderAxis(cylinderAxis),
+            mInvNumPoints((Real)0)
+        {
+            Normalize(mCylinderAxis, true);
+        }
+
+        // The algorithm must estimate 6 parameters, so the number of points must
+        // be at least 6 but preferably larger.  The returned value is the
+        // root-mean-square of the least-squares error.  If numPoints is less than
+        // 6 or if points is a null pointer, the operator() sets the cylinder origin
+        // and axis to the zero vectors, the radius and height to zero, and returns
+        // std::numeric_limits<Real>::max().
+        Real operator()(unsigned int numPoints, Vector3<Real> const* points, Cylinder3<Real>& cylinder)
+        {
+            mX.clear();
+            mInvNumPoints = (Real)0;
+            cylinder.axis.origin = Vector3<Real>::Zero();
+            cylinder.axis.direction = Vector3<Real>::Zero();
+            cylinder.radius = (Real)0;
+            cylinder.height = (Real)0;
+
+            // Validate the input parameters.
+            if (numPoints < 6 || !points)
+            {
+                return std::numeric_limits<Real>::max();
+            }
+
+            Vector3<Real> average;
+            Preprocess(numPoints, points, average);
+
+            // Fit the points based on which constructor the caller used.  The
+            // direction is either estimated or selected directly or indirectly
+            // by the caller.  The center and squared radius are estimated.
+            Vector3<Real> minPC, minW;
+            Real minRSqr, minError;
+
+            if (mConstructorType == FIT_BY_HEMISPHERE_SEARCH)
+            {
+                // Validate the relevant internal parameters.
+                if (mNumThetaSamples == 0 || mNumPhiSamples == 0)
+                {
+                    return std::numeric_limits<Real>::max();
+                }
+
+                // Search the hemisphere for the vector that leads to minimum error
+                // and use it for the cylinder axis.
+                if (mNumThreads == 0)
+                {
+                    // Execute the algorithm in the main process.
+                    minError = ComputeSingleThreaded(minPC, minW, minRSqr);
+                }
+                else
+                {
+                    // Execute the algorithm in multiple threads.
+                    minError = ComputeMultiThreaded(minPC, minW, minRSqr);
+                }
+            }
+            else if (mConstructorType == FIT_USING_COVARIANCE_EIGENVECTOR)
+            {
+                // Validate the relevant internal parameters.
+                if (mEigenIndex >= 3)
+                {
+                    return std::numeric_limits<Real>::max();
+                }
+
+                // Use the eigenvector corresponding to the mEigenIndex of the
+                // eigenvectors of the covariance matrix as the cylinder axis
+                // direction.  The eigenvectors are sorted from smallest
+                // eigenvalue (mEigenIndex = 0) to largest eigenvalue
+                // (mEigenIndex = 2).
+                minError = ComputeUsingCovariance(minPC, minW, minRSqr);
+            }
+            else  // mConstructorType == FIT_USING_SPECIFIED_AXIS
+            {
+                // Validate the relevant internal parameters.
+                if (mCylinderAxis == Vector3<Real>::Zero())
+                {
+                    return std::numeric_limits<Real>::max();
+                }
+
+                minError = ComputeUsingDirection(minPC, minW, minRSqr);
+            }
+
+            // Translate back to the original space by the average of the points.
+            cylinder.axis.origin = minPC + average;
+            cylinder.axis.direction = minW;
+
+            // Compute the cylinder radius.
+            cylinder.radius = std::sqrt(minRSqr);
+
+            // Project the points onto the cylinder axis and choose the cylinder
+            // center and cylinder height as described in the comments at the top of
+            // this header file.
+            Real tmin = (Real)0, tmax = (Real)0;
+            for (unsigned int i = 0; i < numPoints; ++i)
+            {
+                Real t = Dot(cylinder.axis.direction, points[i] - cylinder.axis.origin);
+                tmin = std::min(t, tmin);
+                tmax = std::max(t, tmax);
+            }
+
+            cylinder.axis.origin += ((tmin + tmax) * (Real)0.5) * cylinder.axis.direction;
+            cylinder.height = tmax - tmin;
+            return minError;
+        }
+
+    private:
+        enum ConstructorType
+        {
+            FIT_BY_HEMISPHERE_SEARCH,
+            FIT_USING_COVARIANCE_EIGENVECTOR,
+            FIT_USING_SPECIFIED_AXIS
+        };
+
+        void Preprocess(unsigned int numPoints, Vector3<Real> const* points, Vector3<Real>& average)
+        {
+            mX.resize(numPoints);
+            mInvNumPoints = (Real)1 / (Real)numPoints;
+
+            // Copy the points and translate by the average for numerical robustness.
+            average.MakeZero();
+            for (unsigned int i = 0; i < numPoints; ++i)
+            {
+                average += points[i];
+            }
+            average *= mInvNumPoints;
+            for (unsigned int i = 0; i < numPoints; ++i)
+            {
+                mX[i] = points[i] - average;
+            }
+
+            Vector<6, Real> zero{ (Real)0 };
+            std::vector<Vector<6, Real>> products(mX.size(), zero);
+            mMu = zero;
+            for (size_t i = 0; i < mX.size(); ++i)
+            {
+                products[i][0] = mX[i][0] * mX[i][0];
+                products[i][1] = mX[i][0] * mX[i][1];
+                products[i][2] = mX[i][0] * mX[i][2];
+                products[i][3] = mX[i][1] * mX[i][1];
+                products[i][4] = mX[i][1] * mX[i][2];
+                products[i][5] = mX[i][2] * mX[i][2];
+                mMu[0] += products[i][0];
+                mMu[1] += (Real)2 * products[i][1];
+                mMu[2] += (Real)2 * products[i][2];
+                mMu[3] += products[i][3];
+                mMu[4] += (Real)2 * products[i][4];
+                mMu[5] += products[i][5];
+            }
+            mMu *= mInvNumPoints;
+
+            mF0.MakeZero();
+            mF1.MakeZero();
+            mF2.MakeZero();
+            for (size_t i = 0; i < mX.size(); ++i)
+            {
+                Vector<6, Real> delta;
+                delta[0] = products[i][0] - mMu[0];
+                delta[1] = (Real)2 * products[i][1] - mMu[1];
+                delta[2] = (Real)2 * products[i][2] - mMu[2];
+                delta[3] = products[i][3] - mMu[3];
+                delta[4] = (Real)2 * products[i][4] - mMu[4];
+                delta[5] = products[i][5] - mMu[5];
+                mF0(0, 0) += products[i][0];
+                mF0(0, 1) += products[i][1];
+                mF0(0, 2) += products[i][2];
+                mF0(1, 1) += products[i][3];
+                mF0(1, 2) += products[i][4];
+                mF0(2, 2) += products[i][5];
+                mF1 += OuterProduct(mX[i], delta);
+                mF2 += OuterProduct(delta, delta);
+            }
+            mF0 *= mInvNumPoints;
+            mF0(1, 0) = mF0(0, 1);
+            mF0(2, 0) = mF0(0, 2);
+            mF0(2, 1) = mF0(1, 2);
+            mF1 *= mInvNumPoints;
+            mF2 *= mInvNumPoints;
+        }
+
+        Real ComputeUsingDirection(Vector3<Real>& minPC, Vector3<Real>& minW, Real& minRSqr)
+        {
+            minW = mCylinderAxis;
+            return G(minW, minPC, minRSqr);
+        }
+
+        Real ComputeUsingCovariance(Vector3<Real>& minPC, Vector3<Real>& minW, Real& minRSqr)
+        {
+            Matrix3x3<Real> covar = Matrix3x3<Real>::Zero();
+            for (auto const& X : mX)
+            {
+                covar += OuterProduct(X, X);
+            }
+            covar *= mInvNumPoints;
+            std::array<Real, 3> eval;
+            std::array<std::array<Real, 3>, 3> evec;
+            SymmetricEigensolver3x3<Real>()(
+                covar(0, 0), covar(0, 1), covar(0, 2), covar(1, 1), covar(1, 2), covar(2, 2),
+                true, +1, eval, evec);
+            minW = evec[mEigenIndex];
+            return G(minW, minPC, minRSqr);
+        }
+
+        Real ComputeSingleThreaded(Vector3<Real>& minPC, Vector3<Real>& minW, Real& minRSqr)
+        {
+            Real const iMultiplier = (Real)GTE_C_TWO_PI / (Real)mNumThetaSamples;
+            Real const jMultiplier = (Real)GTE_C_HALF_PI / (Real)mNumPhiSamples;
+
+            // Handle the north pole (0,0,1) separately.
+            minW = { (Real)0, (Real)0, (Real)1 };
+            Real minError = G(minW, minPC, minRSqr);
+
+            for (unsigned int j = 1; j <= mNumPhiSamples; ++j)
+            {
+                Real phi = jMultiplier * static_cast<Real>(j);  // in [0,pi/2]
+                Real csphi = std::cos(phi);
+                Real snphi = std::sin(phi);
+                for (unsigned int i = 0; i < mNumThetaSamples; ++i)
+                {
+                    Real theta = iMultiplier * static_cast<Real>(i);  // in [0,2*pi)
+                    Real cstheta = std::cos(theta);
+                    Real sntheta = std::sin(theta);
+                    Vector3<Real> W{ cstheta * snphi, sntheta * snphi, csphi };
+                    Vector3<Real> PC;
+                    Real rsqr;
+                    Real error = G(W, PC, rsqr);
+                    if (error < minError)
+                    {
+                        minError = error;
+                        minRSqr = rsqr;
+                        minW = W;
+                        minPC = PC;
+                    }
+                }
+            }
+
+            return minError;
+        }
+
+        Real ComputeMultiThreaded(Vector3<Real>& minPC, Vector3<Real>& minW, Real& minRSqr)
+        {
+            Real const iMultiplier = (Real)GTE_C_TWO_PI / (Real)mNumThetaSamples;
+            Real const jMultiplier = (Real)GTE_C_HALF_PI / (Real)mNumPhiSamples;
+
+            // Handle the north pole (0,0,1) separately.
+            minW = { (Real)0, (Real)0, (Real)1 };
+            Real minError = G(minW, minPC, minRSqr);
+
+            struct Local
+            {
+                Real error;
+                Real rsqr;
+                Vector3<Real> W;
+                Vector3<Real> PC;
+                unsigned int jmin;
+                unsigned int jmax;
+            };
+
+            std::vector<Local> local(mNumThreads);
+            unsigned int numPhiSamplesPerThread = mNumPhiSamples / mNumThreads;
+            for (unsigned int t = 0; t < mNumThreads; ++t)
+            {
+                local[t].error = std::numeric_limits<Real>::max();
+                local[t].rsqr = (Real)0;
+                local[t].W = Vector3<Real>::Zero();
+                local[t].PC = Vector3<Real>::Zero();
+                local[t].jmin = numPhiSamplesPerThread * t;
+                local[t].jmax = numPhiSamplesPerThread * (t + 1);
+            }
+            local[mNumThreads - 1].jmax = mNumPhiSamples + 1;
+
+            std::vector<std::thread> process(mNumThreads);
+            for (unsigned int t = 0; t < mNumThreads; ++t)
+            {
+                process[t] = std::thread
+                (
+                    [this, t, iMultiplier, jMultiplier, &local]()
+                {
+                    for (unsigned int j = local[t].jmin; j < local[t].jmax; ++j)
+                    {
+                        Real phi = jMultiplier * static_cast<Real>(j);  // in [0,pi/2]
+                        Real csphi = std::cos(phi);
+                        Real snphi = std::sin(phi);
+                        for (unsigned int i = 0; i < mNumThetaSamples; ++i)
+                        {
+                            Real theta = iMultiplier * static_cast<Real>(i);  // in [0,2*pi)
+                            Real cstheta = std::cos(theta);
+                            Real sntheta = std::sin(theta);
+                            Vector3<Real> W{ cstheta * snphi, sntheta * snphi, csphi };
+                            Vector3<Real> PC;
+                            Real rsqr;
+                            Real error = G(W, PC, rsqr);
+                            if (error < local[t].error)
+                            {
+                                local[t].error = error;
+                                local[t].rsqr = rsqr;
+                                local[t].W = W;
+                                local[t].PC = PC;
+                            }
+                        }
+                    }
+                }
+                );
+            }
+
+            for (unsigned int t = 0; t < mNumThreads; ++t)
+            {
+                process[t].join();
+
+                if (local[t].error < minError)
+                {
+                    minError = local[t].error;
+                    minRSqr = local[t].rsqr;
+                    minW = local[t].W;
+                    minPC = local[t].PC;
+                }
+            }
+
+            return minError;
+        }
+
+        Real G(Vector3<Real> const& W, Vector3<Real>& PC, Real& rsqr)
+        {
+            Matrix3x3<Real> P = Matrix3x3<Real>::Identity() - OuterProduct(W, W);
+            Matrix3x3<Real> S
+            {
+                (Real)0, -W[2], W[1],
+                W[2], (Real)0, -W[0],
+                -W[1], W[0], (Real)0
+            };
+
+            Matrix<3, 3, Real> A = P * mF0 * P;
+            Matrix<3, 3, Real> hatA = -(S * A * S);
+            Matrix<3, 3, Real> hatAA = hatA * A;
+            Real trace = Trace(hatAA);
+            Matrix<3, 3, Real> Q = hatA / trace;
+            Vector<6, Real> pVec{ P(0, 0), P(0, 1), P(0, 2), P(1, 1), P(1, 2), P(2, 2) };
+            Vector<3, Real> alpha = mF1 * pVec;
+            Vector<3, Real> beta = Q * alpha;
+            Real G = (Dot(pVec, mF2 * pVec) - (Real)4 * Dot(alpha, beta) + (Real)4 * Dot(beta, mF0 * beta)) / (Real)mX.size();
+
+            PC = beta;
+            rsqr = Dot(pVec, mMu) + Dot(PC, PC);
+            return G;
+        }
+
+        ConstructorType mConstructorType;
+
+        // Parameters for the hemisphere-search constructor.
+        unsigned int mNumThreads;
+        unsigned int mNumThetaSamples;
+        unsigned int mNumPhiSamples;
+
+        // Parameters for the eigenvector-index constructor.
+        unsigned int mEigenIndex;
+
+        // Parameters for the specified-axis constructor.
+        Vector3<Real> mCylinderAxis;
+
+        // A copy of the input points but translated by their average for
+        // numerical robustness.
+        std::vector<Vector3<Real>> mX;
+        Real mInvNumPoints;
+
+        // Preprocessed information that depends only on the sample points.
+        // This allows precomputed summations so that G(...) can be evaluated
+        // extremely fast.
+        Vector<6, Real> mMu;
+        Matrix<3, 3, Real> mF0;
+        Matrix<3, 6, Real> mF1;
+        Matrix<6, 6, Real> mF2;
+    };
+}
diff --git a/Engine/Plugins/Experimental/GeometryProcessing/Source/GeometryAlgorithms/Private/GTEngine/Mathematics/GteApprEllipse2.h b/Engine/Plugins/Experimental/GeometryProcessing/Source/GeometryAlgorithms/Private/GTEngine/Mathematics/GteApprEllipse2.h
new file mode 100644
--- /dev/null
+++ b/Engine/Plugins/Experimental/GeometryProcessing/Source/GeometryAlgorithms/Private/GTEngine/Mathematics/GteApprEllipse2.h
@@ -0,0 +1,144 @@
+// David Eberly, Geometric Tools, Redmond WA 98052
+// Copyright (c) 1998-2019
+// Distributed under the Boost Software License, Version 1.0.
+// http://www.boost.org/LICENSE_1_0.txt
+// http://www.geometrictools.com/License/Boost/LICENSE_1_0.txt
+// File Version: 3.0.1 (2018/10/05)
+
+#pragma once
+
+// The ellipse in general form is  X^t A X + B^t X + C = 0 where A is a
+// positive definite 2x2 matrix, B is a 2x1 vector, C is a scalar, and X is
+// a 2x1 vector X.  Completing the square, (X-U)^t A (X-U) = U^t A U - C
+// where U = -0.5 A^{-1} B.  Define M = A/(U^t A U - C).  The ellipse is
+// (X-U)^t M (X-U) = 1.  Factor M = R^t D R where R is orthonormal and D is
+// diagonal with positive diagonal terms.  The ellipse in factored form is
+// (X-U)^t R^t D^t R (X-U) = 1.
+//
+// Find the least squares fit of a set of N points P[0] through P[N-1].
+// The return value is the least-squares energy function at (U,R,D).
+
+#include <GTEngine/Mathematics/GteMatrix2x2.h>
+#include <GTEngine/Mathematics/GteContOrientedBox2.h>
+#include <GTEngine/Mathematics/GteDistPointHyperellipsoid.h>
+#include <GTEngine/Mathematics/GteMath.h>
+#include <GTEngine/Mathematics/GteMinimizeN.h>
+
+namespace gte
+{
+
+template <typename Real>
+class ApprEllipse2
+{
+public:
+    Real operator()(int numPoints, Vector2<Real> const* points,
+        Vector2<Real>& center, Matrix2x2<Real>& rotate, Real diagonal[2]);
+
+private:
+    static Real Energy(int numPoints, Vector2<Real> const* points,
+        Real const* input);
+};
+
+
+template <typename Real>
+Real ApprEllipse2<Real>::operator()(int numPoints,
+    Vector2<Real> const* points, Vector2<Real>& center,
+    Matrix2x2<Real>& rotate, Real diagonal[2])
+{
+    // Energy function is E : R^5 -> R where
+    //   V = (V0, V1, V2, V3, V4)
+    //     = (D[0], D[1], U.x, U.y, atan2(R(0,1),R(0,0))).
+    std::function<Real(Real const*)> energy =
+        [numPoints, points](Real const* input)
+    {
+        return Energy(numPoints, points, input);
+    };
+
+    MinimizeN<Real> minimizer(5, energy, 8, 8, 32);
+
+    // The initial guess for the minimizer is based on an oriented box that
+    // contains the points.
+    OrientedBox2<Real> box;
+    GetContainer(numPoints, points, box);
+    center = box.center;
+    for (int i = 0; i < 2; ++i)
+    {
+        rotate.SetRow(i, box.axis[i]);
+        diagonal[i] = box.extent[i];
+    }
+
+    Real angle = std::atan2(rotate(0, 1), rotate(0, 0));
+    Real e0 =
+        diagonal[0] * std::abs(rotate(0, 0)) +
+        diagonal[1] * std::abs(rotate(1, 0));
+    Real e1 =
+        diagonal[0] * std::abs(rotate(0, 1)) +
+        diagonal[1] * std::abs(rotate(1, 1));
+
+    Real v0[5] =
+    {
+        ((Real)0.5)*diagonal[0],
+        ((Real)0.5)*diagonal[1],
+        center[0] - e0,
+        center[1] - e1,
+        -(Real)GTE_C_PI
+    };
+
+    Real v1[5] =
+    {
+        ((Real)2)*diagonal[0],
+        ((Real)2)*diagonal[1],
+        center[0] + e0,
+        center[1] + e1,
+        (Real)GTE_C_PI
+    };
+
+    Real vInitial[5] =
+    {
+        diagonal[0],
+        diagonal[1],
+        center[0],
+        center[1],
+        angle
+    };
+
+    Real vMin[5], error;
+    minimizer.GetMinimum(v0, v1, vInitial, vMin, error);
+
+    diagonal[0] = vMin[0];
+    diagonal[1] = vMin[1];
+    center[0] = vMin[2];
+    center[1] = vMin[3];
+    MakeRotation(-vMin[4], rotate);
+
+    return error;
+}
+
+template <typename Real>
+Real ApprEllipse2<Real>::Energy(int numPoints, Vector2<Real> const* points,
+    Real const* input)
+{
+    // Build rotation matrix.
+    Matrix2x2<Real> rotate;
+    MakeRotation(-input[4], rotate);
+
+    Ellipse2<Real> ellipse(Vector2<Real>::Zero(), { Vector2<Real>::Unit(0),
+        Vector2<Real>::Unit(1) }, { input[0], input[1] });
+
+    // Transform the points to the coordinate system of center C and
+    // columns of rotation R.
+    DCPQuery<Real, Vector2<Real>, Ellipse2<Real>> peQuery;
+    Real energy = (Real)0;
+    for (int i = 0; i < numPoints; ++i)
+    {
+        Vector2<Real> diff = points[i] - Vector2<Real>{ input[2], input[3] };
+        Vector2<Real> prod = rotate * diff;
+        Real dist = peQuery(prod, ellipse).distance;
+        energy += dist;
+    }
+
+    return energy;
+}
+
+
+}
diff --git a/Engine/Plugins/Experimental/GeometryProcessing/Source/GeometryAlgorithms/Private/GTEngine/Mathematics/GteApprEllipseByArcs.h b/Engine/Plugins/Experimental/GeometryProcessing/Source/GeometryAlgorithms/Private/GTEngine/Mathematics/GteApprEllipseByArcs.h
new file mode 100644
--- /dev/null
+++ b/Engine/Plugins/Experimental/GeometryProcessing/Source/GeometryAlgorithms/Private/GTEngine/Mathematics/GteApprEllipseByArcs.h
@@ -0,0 +1,127 @@
+// David Eberly, Geometric Tools, Redmond WA 98052
+// Copyright (c) 1998-2019
+// Distributed under the Boost Software License, Version 1.0.
+// http://www.boost.org/LICENSE_1_0.txt
+// http://www.geometrictools.com/License/Boost/LICENSE_1_0.txt
+// File Version: 3.0.2 (2018/10/05)
+
+#pragma once
+
+#include <GTEngine/Mathematics/GteContScribeCircle2.h>
+#include <vector>
+
+namespace gte
+{
+
+// The ellipse is (x/a)^2 + (y/b)^2 = 1, but only the portion in the first
+// quadrant (x >= 0 and y >= 0) is approximated.  Generate numArcs >= 2 arcs
+// by constructing points corresponding to the weighted averages of the
+// curvatures at the ellipse points (a,0) and (0,b).  The returned input point
+// array has numArcs+1 elements and the returned input center and radius
+// arrays each have numArc elements.  The arc associated with points[i] and
+// points[i+1] has center centers[i] and radius radii[i].  The algorithm
+// is described in
+//   http://www.geometrictools.com/Documentation/ApproximateEllipse.pdf
+//
+// The function returns 'true' when the approximation succeeded, in which
+// case the output arrays are nonempty.  If the 'numArcs' is smaller than
+// 2 or a == b or one of the calls to Circumscribe fails, the function
+// returns 'false'.
+
+template <typename Real>
+bool ApproximateEllipseByArcs(Real a, Real b, int numArcs,
+    std::vector<Vector2<Real>>& points, std::vector<Vector2<Real>>& centers,
+    std::vector<Real>& radii);
+
+
+template <typename Real>
+bool ApproximateEllipseByArcs(Real a, Real b, int numArcs,
+    std::vector<Vector2<Real>>& points, std::vector<Vector2<Real>>& centers,
+    std::vector<Real>& radii)
+{
+    if (numArcs < 2 || a == b)
+    {
+        // At least 2 arcs are required.  The ellipse cannot already be a
+        // circle.
+        points.clear();
+        centers.clear();
+        radii.clear();
+        return false;
+    }
+
+    points.resize(numArcs + 1);
+    centers.resize(numArcs);
+    radii.resize(numArcs);
+
+    // Compute intermediate ellipse quantities.
+    Real a2 = a * a, b2 = b * b, ab = a * b;
+    Real invB2mA2 = ((Real)1) / (b2 - a2);
+
+    // Compute the endpoints of the ellipse in the first quadrant.  The
+    // points are generated in counterclockwise order.
+    points[0] = { a, (Real)0 };
+    points[numArcs] = { (Real)0, b };
+
+    // Compute the curvature at the endpoints.  These are used when computing
+    // the arcs.
+    Real curv0 = a / b2;
+    Real curv1 = b / a2;
+
+    // Select the ellipse points based on curvature properties.
+    Real invNumArcs = ((Real)1) / numArcs;
+    for (int i = 1; i < numArcs; ++i)
+    {
+        // The curvature at a new point is a weighted average of curvature
+        // at the endpoints.
+        Real weight1 = static_cast<Real>(i) * invNumArcs;
+        Real weight0 = (Real)1 - weight1;
+        Real curv = weight0 * curv0 + weight1 * curv1;
+
+        // Compute point having this curvature.
+        Real tmp = std::pow(ab / curv, (Real)2 / (Real)3);
+        points[i][0] = a * std::sqrt(std::abs((tmp - a2) * invB2mA2));
+        points[i][1] = b * std::sqrt(std::abs((tmp - b2) * invB2mA2));
+    }
+
+    // Compute the arc at (a,0).
+    Circle2<Real> circle;
+    Vector2<Real> const& p0 = points[0];
+    Vector2<Real> const& p1 = points[1];
+    if (!Circumscribe(Vector2<Real>{ p1[0], -p1[1] }, p0, p1, circle))
+    {
+        // This should not happen for the arc-fitting algorithm.
+        points.clear();
+        centers.clear();
+        radii.clear();
+        return false;
+    }
+    centers[0] = circle.center;
+    radii[0] = circle.radius;
+
+    // Compute arc at (0,b).
+    int last = numArcs - 1;
+    Vector2<Real> const& pNm1 = points[last];
+    Vector2<Real> const& pN = points[numArcs];
+    if (!Circumscribe(Vector2<Real>{ -pNm1[0], pNm1[1] }, pN, pNm1, circle))
+    {
+        // This should not happen for the arc-fitting algorithm.
+        points.clear();
+        centers.clear();
+        radii.clear();
+        return false;
+    }
+
+    centers[last] = circle.center;
+    radii[last] = circle.radius;
+
+    // Compute arcs at intermediate points between (a,0) and (0,b).
+    for (int iM = 0, i = 1, iP = 2; i < last; ++iM, ++i, ++iP)
+    {
+        Circumscribe(points[iM], points[i], points[iP], circle);
+        centers[i] = circle.center;
+        radii[i] = circle.radius;
+    }
+    return true;
+}
+
+}
diff --git a/Engine/Plugins/Experimental/GeometryProcessing/Source/GeometryAlgorithms/Private/GTEngine/Mathematics/GteApprEllipsoid3.h b/Engine/Plugins/Experimental/GeometryProcessing/Source/GeometryAlgorithms/Private/GTEngine/Mathematics/GteApprEllipsoid3.h
new file mode 100644
--- /dev/null
+++ b/Engine/Plugins/Experimental/GeometryProcessing/Source/GeometryAlgorithms/Private/GTEngine/Mathematics/GteApprEllipsoid3.h
@@ -0,0 +1,227 @@
+// David Eberly, Geometric Tools, Redmond WA 98052
+// Copyright (c) 1998-2019
+// Distributed under the Boost Software License, Version 1.0.
+// http://www.boost.org/LICENSE_1_0.txt
+// http://www.geometrictools.com/License/Boost/LICENSE_1_0.txt
+// File Version: 3.0.1 (2018/10/05)
+
+#pragma once
+
+// The ellipsoid in general form is  X^t A X + B^t X + C = 0 where
+// A is a positive definite 3x3 matrix, B is a 3x1 vector, C is a
+// scalar, and X is a 3x1 vector.  Completing the square,
+// (X-U)^t A (X-U) = U^t A U - C where U = -0.5 A^{-1} B.  Define
+// M = A/(U^t A U - C).  The ellipsoid is (X-U)^t M (X-U) = 1.  Factor
+// M = R^t D R where R is orthonormal and D is diagonal with positive
+// diagonal terms.  The ellipsoid in factored form is
+//     (X-U)^t R^t D^t R (X-U) = 1
+//
+// Find the least squares fit of a set of N points P[0] through P[N-1].
+// The error return value is the least-squares energy function at (U,R,D).
+
+#include <GTEngine/Mathematics/GteMatrix3x3.h>
+#include <GTEngine/Mathematics/GteRotation.h>
+#include <GTEngine/Mathematics/GteContOrientedBox3.h>
+#include <GTEngine/Mathematics/GteDistPointHyperellipsoid.h>
+#include <GTEngine/Mathematics/GteMath.h>
+#include <GTEngine/Mathematics/GteMinimizeN.h>
+
+namespace gte
+{
+
+template <typename Real>
+class ApprEllipsoid3
+{
+public:
+    Real operator()(int numPoints, Vector3<Real> const* points,
+        Vector3<Real>& center, Matrix3x3<Real>& rotate, Real diagonal[3]);
+
+private:
+    static void MatrixToAngles(Matrix3x3<Real> const& rotate, Real angle[3]);
+    static void AnglesToMatrix(Real const angle[3], Matrix3x3<Real>& rotate);
+    static Real Energy(int numPoints, Vector3<Real> const* points,
+        Real const* input);
+};
+
+
+template <typename Real>
+Real ApprEllipsoid3<Real>::operator()(int numPoints,
+    Vector3<Real> const* points, Vector3<Real>& center,
+    Matrix3x3<Real>& rotate, Real diagonal[3])
+{
+    // Energy function is E : R^9 -> R where
+    //   V = (V0,V1,V2,V3,V4,V5,V6,V7,V8)
+    //     = (D[0],D[1],D[2],U[0],U,y,U[2],A0,A1,A2). 
+    std::function<Real(Real const*)> energy =
+        [numPoints, points](Real const* input)
+    {
+        return Energy(numPoints, points, input);
+    };
+
+    MinimizeN<Real> minimizer(9, energy, 8, 8, 32);
+
+    // The initial guess for the minimizer is based on an oriented box that
+    // contains the points.
+    OrientedBox3<Real> box;
+    GetContainer(numPoints, points, box);
+    center = box.center;
+    for (int i = 0; i < 3; ++i)
+    {
+        rotate.SetRow(i, box.axis[i]);
+        diagonal[i] = box.extent[i];
+    }
+
+    Real angle[3];
+    MatrixToAngles(rotate, angle);
+
+    Real extent[3] =
+    {
+        diagonal[0] * std::abs(rotate(0, 0)) +
+        diagonal[1] * std::abs(rotate(0, 1)) +
+        diagonal[2] * std::abs(rotate(0, 2)),
+
+        diagonal[0] * std::abs(rotate(1, 0)) +
+        diagonal[1] * std::abs(rotate(1, 1)) +
+        diagonal[2] * std::abs(rotate(1, 2)),
+
+        diagonal[0] * std::abs(rotate(2, 0)) +
+        diagonal[1] * std::abs(rotate(2, 1)) +
+        diagonal[2] * std::abs(rotate(2, 2))
+    };
+
+    Real v0[9] =
+    {
+        ((Real)0.5)*diagonal[0],
+        ((Real)0.5)*diagonal[1],
+        ((Real)0.5)*diagonal[2],
+        center[0] - extent[0],
+        center[1] - extent[1],
+        center[2] - extent[2],
+        -(Real)GTE_C_PI,
+        (Real)0,
+        (Real)0
+    };
+
+    Real v1[9] =
+    {
+        ((Real)2)*diagonal[0],
+        ((Real)2)*diagonal[1],
+        ((Real)2)*diagonal[2],
+        center[0] + extent[0],
+        center[1] + extent[1],
+        center[2] + extent[2],
+        (Real)GTE_C_PI,
+        (Real)GTE_C_PI,
+        (Real)GTE_C_PI
+    };
+
+    Real vInitial[9] =
+    {
+        diagonal[0],
+        diagonal[1],
+        diagonal[2],
+        center[0],
+        center[1],
+        center[2],
+        angle[0],
+        angle[1],
+        angle[2]
+    };
+
+    Real vMin[9], error;
+    minimizer.GetMinimum(v0, v1, vInitial, vMin, error);
+
+    diagonal[0] = vMin[0];
+    diagonal[1] = vMin[1];
+    diagonal[2] = vMin[2];
+    center[0] = vMin[3];
+    center[1] = vMin[4];
+    center[2] = vMin[5];
+    AnglesToMatrix(&vMin[6], rotate);
+
+    return error;
+}
+
+template <typename Real>
+void ApprEllipsoid3<Real>::MatrixToAngles(Matrix3x3<Real> const& rotate,
+    Real angle[3])
+{
+    // rotation axis = (cos(a0)sin(a1),sin(a0)sin(a1),cos(a1))
+    // a0 in [-pi,pi], a1 in [0,pi], a2 in [0,pi]
+
+    Real const zero = (Real)0;
+    Real const one = (Real)1;
+    AxisAngle<3, Real> aa = Rotation<3, Real>(rotate);
+
+    if (-one < aa.axis[2])
+    {
+        if (aa.axis[2] < one)
+        {
+            angle[0] = std::atan2(aa.axis[1], aa.axis[0]);
+            angle[1] = std::acos(aa.axis[2]);
+        }
+        else
+        {
+            angle[0] = zero;
+            angle[1] = zero;
+        }
+    }
+    else
+    {
+        angle[0] = zero;
+        angle[1] = (Real)GTE_C_PI;
+    }
+}
+
+template <typename Real>
+void ApprEllipsoid3<Real>::AnglesToMatrix(Real const angle[3],
+    Matrix3x3<Real>& rotate)
+{
+    // rotation axis = (cos(a0)sin(a1),sin(a0)sin(a1),cos(a1))
+    // a0 in [-pi,pi], a1 in [0,pi], a2 in [0,pi]
+
+    Real cs0 = std::cos(angle[0]);
+    Real sn0 = std::sin(angle[0]);
+    Real cs1 = std::cos(angle[1]);
+    Real sn1 = std::sin(angle[1]);
+    AxisAngle<3, Real> aa;
+    aa.axis = { cs0*sn1, sn0*sn1, cs1 };
+    aa.angle = angle[2];
+    rotate = Rotation<3, Real>(aa);
+}
+
+template <typename Real>
+Real ApprEllipsoid3<Real>::Energy(int numPoints, Vector3<Real> const* points,
+    Real const* input)
+{
+    // Build rotation matrix.
+    Matrix3x3<Real> rotate;
+    AnglesToMatrix(&input[6], rotate);
+
+    // Uniformly scale the extents to keep reasonable floating point values
+    // in the distance calculations.
+    Real maxValue = std::max(std::max(input[0], input[1]), input[2]);
+    Real invMax = ((Real)1) / maxValue;
+    Ellipsoid3<Real> ellipsoid(Vector3<Real>::Zero(), { Vector3<Real>::Unit(0),
+        Vector3<Real>::Unit(1), Vector3<Real>::Unit(2) }, { invMax*input[0],
+        invMax*input[1], invMax*input[2] });
+
+    // Transform the points to the coordinate system of center C and columns
+    // of rotation R.
+    DCPQuery<Real, Vector3<Real>, Ellipsoid3<Real>> peQuery;
+    Real energy = (Real)0;
+    for (int i = 0; i < numPoints; ++i)
+    {
+        Vector3<Real> diff = points[i] -
+            Vector3<Real>{ input[3], input[4], input[5] };
+
+        Vector3<Real> prod = invMax * (diff * rotate);
+        Real dist = peQuery(prod, ellipsoid).distance;
+        energy += maxValue * dist;
+    }
+
+    return energy;
+}
+
+
+}
diff --git a/Engine/Plugins/Experimental/GeometryProcessing/Source/GeometryAlgorithms/Private/GTEngine/Mathematics/GteApprGaussian2.h b/Engine/Plugins/Experimental/GeometryProcessing/Source/GeometryAlgorithms/Private/GTEngine/Mathematics/GteApprGaussian2.h
new file mode 100644
--- /dev/null
+++ b/Engine/Plugins/Experimental/GeometryProcessing/Source/GeometryAlgorithms/Private/GTEngine/Mathematics/GteApprGaussian2.h
@@ -0,0 +1,181 @@
+// David Eberly, Geometric Tools, Redmond WA 98052
+// Copyright (c) 1998-2019
+// Distributed under the Boost Software License, Version 1.0.
+// http://www.boost.org/LICENSE_1_0.txt
+// http://www.geometrictools.com/License/Boost/LICENSE_1_0.txt
+// File Version: 3.0.0 (2016/06/19)
+
+#pragma once
+
+#include <GTEngine/Mathematics/GteVector2.h>
+#include <GTEngine/Mathematics/GteApprQuery.h>
+#include <GTEngine/Mathematics/GteOrientedBox.h>
+#include <GTEngine/Mathematics/GteSymmetricEigensolver2x2.h>
+
+// Fit points with a Gaussian distribution.  The center is the mean of the
+// points, the axes are the eigenvectors of the covariance matrix, and the
+// extents are the eigenvalues of the covariance matrix and are returned in
+// increasing order.  An oriented box is used to store the mean, axes, and
+// extents.
+
+namespace gte
+{
+
+template <typename Real>
+class ApprGaussian2
+    :
+    public ApprQuery<Real, ApprGaussian2<Real>, Vector2<Real>>
+{
+public:
+    // Initialize the model parameters to zero.
+    ApprGaussian2();
+
+    // Basic fitting algorithm.
+    bool Fit(int numPoints, Vector2<Real> const* points);
+    OrientedBox2<Real> const& GetParameters() const;
+
+    // Functions called by ApprQuery::RANSAC.  See GteApprQuery.h for a
+    // detailed description.
+    int GetMinimumRequired() const;
+    Real Error(Vector2<Real> const& observation) const;
+    bool Fit(std::vector<Vector2<Real>> const& observations,
+        std::vector<int> const& indices);
+
+private:
+    OrientedBox2<Real> mParameters;
+};
+
+
+template <typename Real>
+ApprGaussian2<Real>::ApprGaussian2()
+{
+    mParameters.center = Vector2<Real>::Zero();
+    mParameters.axis[0] = Vector2<Real>::Zero();
+    mParameters.axis[1] = Vector2<Real>::Zero();
+    mParameters.extent = Vector2<Real>::Zero();
+}
+
+template <typename Real>
+bool ApprGaussian2<Real>::Fit(int numPoints, Vector2<Real> const* points)
+{
+    if (numPoints >= GetMinimumRequired() && points)
+    {
+        // Compute the mean of the points.
+        Vector2<Real> mean = Vector2<Real>::Zero();
+        for (int i = 0; i < numPoints; ++i)
+        {
+            mean += points[i];
+        }
+        Real invSize = ((Real)1) / (Real)numPoints;
+        mean *= invSize;
+
+        // Compute the covariance matrix of the points.
+        Real covar00 = (Real)0, covar01 = (Real)0, covar11 = (Real)0;
+        for (int i = 0; i < numPoints; ++i)
+        {
+            Vector2<Real> diff = points[i] - mean;
+            covar00 += diff[0] * diff[0];
+            covar01 += diff[0] * diff[1];
+            covar11 += diff[1] * diff[1];
+        }
+        covar00 *= invSize;
+        covar01 *= invSize;
+        covar11 *= invSize;
+
+        // Solve the eigensystem.
+        SymmetricEigensolver2x2<Real> es;
+        std::array<Real, 2> eval;
+        std::array<std::array<Real, 2>, 2> evec;
+        es(covar00, covar01, covar11, +1, eval, evec);
+        mParameters.center = mean;
+        mParameters.axis[0] = evec[0];
+        mParameters.axis[1] = evec[1];
+        mParameters.extent = eval;
+        return true;
+    }
+
+    mParameters.center = Vector2<Real>::Zero();
+    mParameters.axis[0] = Vector2<Real>::Zero();
+    mParameters.axis[1] = Vector2<Real>::Zero();
+    mParameters.extent = Vector2<Real>::Zero();
+    return false;
+}
+
+template <typename Real>
+OrientedBox2<Real> const& ApprGaussian2<Real>::GetParameters() const
+{
+    return mParameters;
+}
+
+template <typename Real>
+int ApprGaussian2<Real>::GetMinimumRequired() const
+{
+    return 2;
+}
+
+template <typename Real>
+bool ApprGaussian2<Real>::Fit(
+    std::vector<Vector2<Real>> const& observations,
+    std::vector<int> const& indices)
+{
+    if (static_cast<int>(indices.size()) >= GetMinimumRequired())
+    {
+        // Compute the mean of the points.
+        Vector2<Real> mean = Vector2<Real>::Zero();
+        for (auto index : indices)
+        {
+            mean += observations[index];
+        }
+        Real invSize = ((Real)1) / (Real)indices.size();
+        mean *= invSize;
+
+        // Compute the covariance matrix of the points.
+        Real covar00 = (Real)0, covar01 = (Real)0, covar11 = (Real)0;
+        for (auto index : indices)
+        {
+            Vector2<Real> diff = observations[index] - mean;
+            covar00 += diff[0] * diff[0];
+            covar01 += diff[0] * diff[1];
+            covar11 += diff[1] * diff[1];
+        }
+        covar00 *= invSize;
+        covar01 *= invSize;
+        covar11 *= invSize;
+
+        // Solve the eigensystem.
+        SymmetricEigensolver2x2<Real> es;
+        std::array<Real, 2> eval;
+        std::array<std::array<Real, 2>, 2> evec;
+        es(covar00, covar01, covar11, +1, eval, evec);
+        mParameters.center = mean;
+        mParameters.axis[0] = evec[0];
+        mParameters.axis[1] = evec[1];
+        mParameters.extent = eval;
+    }
+
+    mParameters.center = Vector2<Real>::Zero();
+    mParameters.axis[0] = Vector2<Real>::Zero();
+    mParameters.axis[1] = Vector2<Real>::Zero();
+    mParameters.extent = Vector2<Real>::Zero();
+    return false;
+}
+
+template <typename Real>
+Real ApprGaussian2<Real>::Error(Vector2<Real> const& observation) const
+{
+    Vector2<Real> diff = observation - mParameters.center;
+    Real error = (Real)0;
+    for (int i = 0; i < 2; ++i)
+    {
+        if (mParameters.extent[i] >(Real)0)
+        {
+            Real ratio = Dot(diff, mParameters.axis[i]) /
+                mParameters.extent[i];
+            error += ratio * ratio;
+        }
+    }
+    return error;
+}
+
+
+}
diff --git a/Engine/Plugins/Experimental/GeometryProcessing/Source/GeometryAlgorithms/Private/GTEngine/Mathematics/GteApprGaussian3.h b/Engine/Plugins/Experimental/GeometryProcessing/Source/GeometryAlgorithms/Private/GTEngine/Mathematics/GteApprGaussian3.h
new file mode 100644
--- /dev/null
+++ b/Engine/Plugins/Experimental/GeometryProcessing/Source/GeometryAlgorithms/Private/GTEngine/Mathematics/GteApprGaussian3.h
@@ -0,0 +1,202 @@
+// David Eberly, Geometric Tools, Redmond WA 98052
+// Copyright (c) 1998-2019
+// Distributed under the Boost Software License, Version 1.0.
+// http://www.boost.org/LICENSE_1_0.txt
+// http://www.geometrictools.com/License/Boost/LICENSE_1_0.txt
+// File Version: 3.0.0 (2016/06/19)
+
+#pragma once
+
+#include <GTEngine/Mathematics/GteVector3.h>
+#include <GTEngine/Mathematics/GteApprQuery.h>
+#include <GTEngine/Mathematics/GteOrientedBox.h>
+#include <GTEngine/Mathematics/GteSymmetricEigensolver3x3.h>
+
+// Fit points with a Gaussian distribution.  The center is the mean of the
+// points, the axes are the eigenvectors of the covariance matrix, and the
+// extents are the eigenvalues of the covariance matrix and are returned in
+// increasing order.  An oriented box is used to store the mean, axes, and
+// extents.
+
+namespace gte
+{
+
+template <typename Real>
+class ApprGaussian3
+    :
+    public ApprQuery<Real, ApprGaussian3<Real>, Vector3<Real>>
+{
+public:
+    // Initialize the model parameters to zero.
+    ApprGaussian3();
+
+    // Basic fitting algorithm.
+    bool Fit(int numPoints, Vector3<Real> const* points);
+    OrientedBox3<Real> const& GetParameters() const;
+
+    // Functions called by ApprQuery::RANSAC.  See GteApprQuery.h for a
+    // detailed description.
+    int GetMinimumRequired() const;
+    Real Error(Vector3<Real> const& observation) const;
+    bool Fit(std::vector<Vector3<Real>> const& observations,
+        std::vector<int> const& indices);
+
+private:
+    OrientedBox3<Real> mParameters;
+};
+
+
+template <typename Real>
+ApprGaussian3<Real>::ApprGaussian3()
+{
+    mParameters.center = Vector3<Real>::Zero();
+    mParameters.axis[0] = Vector3<Real>::Zero();
+    mParameters.axis[1] = Vector3<Real>::Zero();
+    mParameters.axis[2] = Vector3<Real>::Zero();
+    mParameters.extent = Vector3<Real>::Zero();
+}
+
+template <typename Real>
+bool ApprGaussian3<Real>::Fit(int numPoints, Vector3<Real> const* points)
+{
+    if (numPoints >= GetMinimumRequired() && points)
+    {
+        // Compute the mean of the points.
+        Vector3<Real> mean = Vector3<Real>::Zero();
+        for (int i = 0; i < numPoints; ++i)
+        {
+            mean += points[i];
+        }
+        Real invSize = ((Real)1) / (Real)numPoints;
+        mean *= invSize;
+
+        // Compute the covariance matrix of the points.
+        Real covar00 = (Real)0, covar01 = (Real)0, covar02 = (Real)0;
+        Real covar11 = (Real)0, covar12 = (Real)0, covar22 = (Real)0;
+        for (int i = 0; i < numPoints; ++i)
+        {
+            Vector3<Real> diff = points[i] - mean;
+            covar00 += diff[0] * diff[0];
+            covar01 += diff[0] * diff[1];
+            covar02 += diff[0] * diff[2];
+            covar11 += diff[1] * diff[1];
+            covar12 += diff[1] * diff[2];
+            covar22 += diff[2] * diff[2];
+        }
+        covar00 *= invSize;
+        covar01 *= invSize;
+        covar02 *= invSize;
+        covar11 *= invSize;
+        covar12 *= invSize;
+        covar22 *= invSize;
+
+        // Solve the eigensystem.
+        SymmetricEigensolver3x3<Real> es;
+        std::array<Real, 3> eval;
+        std::array<std::array<Real, 3>, 3> evec;
+        es(covar00, covar01, covar02, covar11, covar12, covar22, false, +1,
+            eval, evec);
+        mParameters.center = mean;
+        mParameters.axis[0] = evec[0];
+        mParameters.axis[1] = evec[1];
+        mParameters.axis[2] = evec[2];
+        mParameters.extent = eval;
+        return true;
+    }
+
+    mParameters.center = Vector3<Real>::Zero();
+    mParameters.axis[0] = Vector3<Real>::Zero();
+    mParameters.axis[1] = Vector3<Real>::Zero();
+    mParameters.axis[2] = Vector3<Real>::Zero();
+    mParameters.extent = Vector3<Real>::Zero();
+    return false;
+}
+
+template <typename Real>
+OrientedBox3<Real> const& ApprGaussian3<Real>::GetParameters() const
+{
+    return mParameters;
+}
+
+template <typename Real>
+int ApprGaussian3<Real>::GetMinimumRequired() const
+{
+    return 2;
+}
+
+template <typename Real>
+Real ApprGaussian3<Real>::Error(Vector3<Real> const& observation) const
+{
+    Vector3<Real> diff = observation - mParameters.center;
+    Real error = (Real)0;
+    for (int i = 0; i < 3; ++i)
+    {
+        if (mParameters.extent[i] >(Real)0)
+        {
+            Real ratio = Dot(diff, mParameters.axis[i]) /
+                mParameters.extent[i];
+            error += ratio * ratio;
+        }
+    }
+    return error;
+}
+
+template <typename Real>
+bool ApprGaussian3<Real>::Fit(
+    std::vector<Vector3<Real>> const& observations,
+    std::vector<int> const& indices)
+{
+    if (static_cast<int>(indices.size()) >= GetMinimumRequired())
+    {
+        // Compute the mean of the points.
+        Vector3<Real> mean = Vector3<Real>::Zero();
+        for (auto index : indices)
+        {
+            mParameters.center += observations[index];
+        }
+        Real invSize = ((Real)1) / (Real)indices.size();
+        mean *= invSize;
+
+        // Compute the covariance matrix of the points.
+        Real covar00 = (Real)0, covar01 = (Real)0, covar02 = (Real)0;
+        Real covar11 = (Real)0, covar12 = (Real)0, covar22 = (Real)0;
+        for (auto index : indices)
+        {
+            Vector3<Real> diff = observations[index] - mean;
+            covar00 += diff[0] * diff[0];
+            covar01 += diff[0] * diff[1];
+            covar02 += diff[0] * diff[2];
+            covar11 += diff[1] * diff[1];
+            covar12 += diff[1] * diff[2];
+            covar22 += diff[2] * diff[2];
+        }
+        covar00 *= invSize;
+        covar01 *= invSize;
+        covar02 *= invSize;
+        covar11 *= invSize;
+        covar12 *= invSize;
+        covar22 *= invSize;
+
+        // Solve the eigensystem.
+        SymmetricEigensolver3x3<Real> es;
+        std::array<Real, 3> eval;
+        std::array<std::array<Real, 3>, 3> evec;
+        es(covar00, covar01, covar02, covar11, covar12, covar22, false, +1,
+            eval, evec);
+        mParameters.center = mean;
+        mParameters.axis[0] = evec[0];
+        mParameters.axis[1] = evec[1];
+        mParameters.axis[2] = evec[2];
+        mParameters.extent = eval;
+    }
+
+    mParameters.center = Vector3<Real>::Zero();
+    mParameters.axis[0] = Vector3<Real>::Zero();
+    mParameters.axis[1] = Vector3<Real>::Zero();
+    mParameters.axis[2] = Vector3<Real>::Zero();
+    mParameters.extent = Vector3<Real>::Zero();
+    return false;
+}
+
+
+}
diff --git a/Engine/Plugins/Experimental/GeometryProcessing/Source/GeometryAlgorithms/Private/GTEngine/Mathematics/GteApprGreatCircle3.h b/Engine/Plugins/Experimental/GeometryProcessing/Source/GeometryAlgorithms/Private/GTEngine/Mathematics/GteApprGreatCircle3.h
new file mode 100644
--- /dev/null
+++ b/Engine/Plugins/Experimental/GeometryProcessing/Source/GeometryAlgorithms/Private/GTEngine/Mathematics/GteApprGreatCircle3.h
@@ -0,0 +1,141 @@
+// David Eberly, Geometric Tools, Redmond WA 98052
+// Copyright (c) 1998-2019
+// Distributed under the Boost Software License, Version 1.0.
+// http://www.boost.org/LICENSE_1_0.txt
+// http://www.geometrictools.com/License/Boost/LICENSE_1_0.txt
+// File Version: 3.0.1 (2018/10/05)
+
+#pragma once
+
+#include <GTEngine/Mathematics/GteMatrix3x3.h>
+#include <GTEngine/Mathematics/GteMath.h>
+#include <GTEngine/Mathematics/GteSymmetricEigensolver3x3.h>
+
+namespace gte
+{
+
+// Least-squares fit of a great circle to unit-length vectors (x,y,z) by
+// using distance measurements orthogonal (and measured along great circles)
+// to the proposed great circle.  The inputs akPoint[] are unit length.  The
+// returned value is unit length, call it N.  The fitted great circle is
+// defined by Dot(N,X) = 0, where X is a unit-length vector on the great
+// circle.
+template <typename Real>
+class ApprGreatCircle3
+{
+public:
+    void operator()(int numPoints, Vector3<Real> const* points,
+        Vector3<Real>& normal);
+};
+
+
+// In addition to the least-squares fit of a great circle, the input vectors
+// are projected onto that circle.  The sector of smallest angle (possibly
+// obtuse) that contains the points is computed.  The endpoints of the arc
+// of the sector are returned.  The returned endpoints A0 and A1 are
+// perpendicular to the returned normal N.  Moreover, when you view the
+// arc by looking at the plane of the great circle with a view direction
+// of -N, the arc is traversed counterclockwise starting at A0 and ending
+// at A1.
+template <typename Real>
+class ApprGreatArc3
+{
+public:
+    void operator()(int numPoints, Vector3<Real> const* points,
+        Vector3<Real>& normal, Vector3<Real>& arcEnd0,
+        Vector3<Real>& arcEnd1);
+};
+
+
+template <typename Real>
+void ApprGreatCircle3<Real>::operator()(int numPoints,
+    Vector3<Real> const* points, Vector3<Real>& normal)
+{
+    // Compute the covariance matrix of the vectors.
+    Real covar00 = (Real)0, covar01 = (Real)0, covar02 = (Real)0;
+    Real covar11 = (Real)0, covar12 = (Real)0, covar22 = (Real)0;
+    for (int i = 0; i < numPoints; i++)
+    {
+        Vector3<Real> diff = points[i];
+        covar00 += diff[0] * diff[0];
+        covar01 += diff[0] * diff[1];
+        covar02 += diff[0] * diff[2];
+        covar11 += diff[1] * diff[1];
+        covar12 += diff[1] * diff[2];
+        covar22 += diff[2] * diff[2];
+    }
+
+    Real invNumPoints = ((Real)1) / static_cast<Real>(numPoints);
+    covar00 *= invNumPoints;
+    covar01 *= invNumPoints;
+    covar02 *= invNumPoints;
+    covar11 *= invNumPoints;
+    covar12 *= invNumPoints;
+    covar22 *= invNumPoints;
+
+    // Solve the eigensystem.
+    SymmetricEigensolver3x3<Real> es;
+    std::array<Real, 3> eval;
+    std::array<std::array<Real, 3>, 3> evec;
+    es(covar00, covar01, covar02, covar11, covar12, covar22, false, +1,
+        eval, evec);
+    normal = evec[0];
+}
+
+template <typename Real>
+void ApprGreatArc3<Real>::operator()(int numPoints,
+    Vector3<Real> const* points, Vector3<Real>& normal,
+    Vector3<Real>& arcEnd0, Vector3<Real>& arcEnd1)
+{
+    // Get the least-squares great circle for the vectors.  The circle is on
+    // the plane Dot(N,X) = 0.  Generate a basis from N.
+    Vector3<Real> basis[3];  // { N, U, V }
+    ApprGreatCircle3<Real>()(numPoints, points, basis[0]);
+    ComputeOrthogonalComplement(1, basis);
+
+    // The vectors are X[i] = u[i]*U + v[i]*V + w[i]*N.  The projections
+    // are P[i] = (u[i]*U + v[i]*V)/sqrt(u[i]*u[i] + v[i]*v[i]).  The great
+    // circle is parameterized by C(t) = cos(t)*U + sin(t)*V.  Compute the
+    // angles t in [-pi,pi] for the projections onto the great circle.  It
+    // is not necesarily to normalize (u[i],v[i]), instead computing
+    // t = atan2(v[i],u[i]).
+    std::vector<std::array<Real, 3>> items(numPoints);  // item (u, v, angle)
+    for (int i = 0; i < numPoints; ++i)
+    {
+        items[i][0] = Dot(basis[1], points[i]);
+        items[i][1] = Dot(basis[2], points[i]);
+        items[i][2] = std::atan2(items[i][1], items[i][0]);
+    }
+    std::sort(items.begin(), items.end(),
+        [](std::array<Real, 3> const& item0, std::array<Real, 3> const& item1)
+    {
+        return item0[2] < item1[2];
+    }
+    );
+
+    // Locate the pair of consecutive angles whose difference is a maximum.
+    // Effectively, we are constructing a cone of minimum angle that contains
+    // the unit-length vectors.
+    int numPointsM1 = numPoints - 1;
+    Real maxDiff = (Real)GTE_C_TWO_PI + items[0][2] - items[numPointsM1][2];
+    int end0 = 0, end1 = numPointsM1;
+    for (int i0 = 0, i1 = 1; i0 < numPointsM1; i0 = i1++)
+    {
+        Real diff = items[i1][2] - items[i0][2];
+        if (diff > maxDiff)
+        {
+            maxDiff = diff;
+            end0 = i1;
+            end1 = i0;
+        }
+    }
+
+    normal = basis[0];
+    arcEnd0 = items[end0][0] * basis[1] + items[end0][1] * basis[2];
+    arcEnd1 = items[end1][0] * basis[1] + items[end1][1] * basis[2];
+    Normalize(arcEnd0);
+    Normalize(arcEnd1);
+}
+
+
+}
diff --git a/Engine/Plugins/Experimental/GeometryProcessing/Source/GeometryAlgorithms/Private/GTEngine/Mathematics/GteApprHeightLine2.h b/Engine/Plugins/Experimental/GeometryProcessing/Source/GeometryAlgorithms/Private/GTEngine/Mathematics/GteApprHeightLine2.h
new file mode 100644
--- /dev/null
+++ b/Engine/Plugins/Experimental/GeometryProcessing/Source/GeometryAlgorithms/Private/GTEngine/Mathematics/GteApprHeightLine2.h
@@ -0,0 +1,153 @@
+// David Eberly, Geometric Tools, Redmond WA 98052
+// Copyright (c) 1998-2019
+// Distributed under the Boost Software License, Version 1.0.
+// http://www.boost.org/LICENSE_1_0.txt
+// http://www.geometrictools.com/License/Boost/LICENSE_1_0.txt
+// File Version: 3.0.0 (2016/06/19)
+
+#pragma once
+
+#include <GTEngine/Mathematics/GteMatrix2x2.h>
+#include <GTEngine/Mathematics/GteApprQuery.h>
+
+// Least-squares fit of a line to height data (x,f(x)).  The line is of the
+// form: (y - yAvr) = a*(x - xAvr), where (xAvr,yAvr) is the average of the
+// sample points.  The return value of Fit is 'true' iff the fit is successful
+// (the input points are not degenerate to a single point).  The mParameters
+// values are ((xAvr,yAvr),(a,-1)) on success and ((0,0),(0,0)) on failure.
+// The error for (x0,y0) is [a*(x0-xAvr)-(y0-yAvr)]^2.
+
+namespace gte
+{
+
+template <typename Real>
+class ApprHeightLine2
+    :
+    public ApprQuery<Real, ApprHeightLine2<Real>, Vector2<Real>>
+{
+public:
+    // Initialize the model parameters to zero.
+    ApprHeightLine2();
+
+    // Basic fitting algorithm.
+    bool Fit(int numPoints, Vector2<Real> const* points);
+    std::pair<Vector2<Real>, Vector2<Real>> const& GetParameters() const;
+
+    // Functions called by ApprQuery::RANSAC.  See GteApprQuery.h for a
+    // detailed description.
+    int GetMinimumRequired() const;
+    Real Error(Vector2<Real> const& observation) const;
+    bool Fit(std::vector<Vector2<Real>> const& observations,
+        std::vector<int> const& indices);
+
+private:
+    std::pair<Vector2<Real>, Vector2<Real>> mParameters;
+};
+
+
+template <typename Real>
+ApprHeightLine2<Real>::ApprHeightLine2()
+{
+    mParameters.first = Vector2<Real>::Zero();
+    mParameters.second = Vector2<Real>::Zero();
+}
+
+template <typename Real>
+bool ApprHeightLine2<Real>::Fit(int numPoints, Vector2<Real> const* points)
+{
+    if (numPoints >= GetMinimumRequired() && points)
+    {
+        // Compute the mean of the points.
+        Vector2<Real> mean = Vector2<Real>::Zero();
+        for (int i = 0; i < numPoints; ++i)
+        {
+            mean += points[i];
+        }
+        mean /= (Real)numPoints;
+
+        // Compute the covariance matrix of the points.
+        Real covar00 = (Real)0, covar01 = (Real)0;
+        for (int i = 0; i < numPoints; ++i)
+        {
+            Vector2<Real> diff = points[i] - mean;
+            covar00 += diff[0] * diff[0];
+            covar01 += diff[0] * diff[1];
+        }
+
+        // Decompose the covariance matrix.
+        if (covar00 >(Real)0)
+        {
+            mParameters.first = mean;
+            mParameters.second[0] = covar01 / covar00;
+            mParameters.second[1] = (Real)-1;
+            return true;
+        }
+    }
+
+    mParameters.first = Vector2<Real>::Zero();
+    mParameters.second = Vector2<Real>::Zero();
+    return false;
+}
+
+template <typename Real>
+std::pair<Vector2<Real>, Vector2<Real>> const&
+ApprHeightLine2<Real>::GetParameters() const
+{
+    return mParameters;
+}
+
+template <typename Real>
+int ApprHeightLine2<Real>::GetMinimumRequired() const
+{
+    return 2;
+}
+
+template <typename Real>
+Real ApprHeightLine2<Real>::Error(Vector2<Real> const& observation) const
+{
+    Real d = Dot(observation - mParameters.first, mParameters.second);
+    Real error = d*d;
+    return error;
+}
+
+template <typename Real>
+bool ApprHeightLine2<Real>::Fit(
+    std::vector<Vector2<Real>> const& observations,
+    std::vector<int> const& indices)
+{
+    if (static_cast<int>(indices.size()) >= GetMinimumRequired())
+    {
+        // Compute the mean of the points.
+        Vector2<Real> mean = Vector2<Real>::Zero();
+        for (auto index : indices)
+        {
+            mean += observations[index];
+        }
+        mean /= (Real)indices.size();
+
+        // Compute the covariance matrix of the points.
+        Real covar00 = (Real)0, covar01 = (Real)0;
+        for (auto index : indices)
+        {
+            Vector2<Real> diff = observations[index] - mean;
+            covar00 += diff[0] * diff[0];
+            covar01 += diff[0] * diff[1];
+        }
+
+        // Decompose the covariance matrix.
+        if (covar00 > (Real)0)
+        {
+            mParameters.first = mean;
+            mParameters.second[0] = covar01 / covar00;
+            mParameters.second[1] = (Real)-1;
+            return true;
+        }
+    }
+
+    mParameters.first = Vector2<Real>::Zero();
+    mParameters.second = Vector2<Real>::Zero();
+    return false;
+}
+
+
+}
diff --git a/Engine/Plugins/Experimental/GeometryProcessing/Source/GeometryAlgorithms/Private/GTEngine/Mathematics/GteApprHeightPlane3.h b/Engine/Plugins/Experimental/GeometryProcessing/Source/GeometryAlgorithms/Private/GTEngine/Mathematics/GteApprHeightPlane3.h
new file mode 100644
--- /dev/null
+++ b/Engine/Plugins/Experimental/GeometryProcessing/Source/GeometryAlgorithms/Private/GTEngine/Mathematics/GteApprHeightPlane3.h
@@ -0,0 +1,171 @@
+// David Eberly, Geometric Tools, Redmond WA 98052
+// Copyright (c) 1998-2019
+// Distributed under the Boost Software License, Version 1.0.
+// http://www.boost.org/LICENSE_1_0.txt
+// http://www.geometrictools.com/License/Boost/LICENSE_1_0.txt
+// File Version: 3.0.0 (2016/06/19)
+
+#pragma once
+
+#include <GTEngine/Mathematics/GteMatrix3x3.h>
+#include <GTEngine/Mathematics/GteApprQuery.h>
+
+// Least-squares fit of a plane to height data (x,y,f(x,y)).  The plane is of
+// the form (z - zAvr) = a*(x - xAvr) + b*(y - yAvr), where (xAvr,yAvr,zAvr)
+// is the average of the sample points.  The return value is 'true' iff the
+// fit is successful (the input points are noncollinear).  The mParameters
+// values are ((xAvr,yAvr,zAvr),(a,b,-1)) on success and ((0,0,0),(0,0,0)) on
+// failure.  The error for (x0,y0,z0) is [a*(x0-xAvr)+b*(y0-yAvr)-(z0-zAvr)]^2.
+
+namespace gte
+{
+
+template <typename Real>
+class ApprHeightPlane3
+    :
+    public ApprQuery<Real, ApprHeightPlane3<Real>, Vector3<Real>>
+{
+public:
+    // Initialize the model parameters to zero.
+    ApprHeightPlane3();
+
+    // Basic fitting algorithm.
+    bool Fit(int numPoints, Vector3<Real> const* points);
+    std::pair<Vector3<Real>, Vector3<Real>> const& GetParameters() const;
+
+    // Functions called by ApprQuery::RANSAC.  See GteApprQuery.h for a
+    // detailed description.
+    int GetMinimumRequired() const;
+    Real Error(Vector3<Real> const& observation) const;
+    bool Fit(std::vector<Vector3<Real>> const& observations,
+        std::vector<int> const& indices);
+
+private:
+    std::pair<Vector3<Real>, Vector3<Real>> mParameters;
+};
+
+
+template <typename Real>
+ApprHeightPlane3<Real>::ApprHeightPlane3()
+{
+    mParameters.first = Vector3<Real>::Zero();
+    mParameters.second = Vector3<Real>::Zero();
+}
+
+template <typename Real>
+bool ApprHeightPlane3<Real>::Fit(int numPoints, Vector3<Real> const* points)
+{
+    if (numPoints >= GetMinimumRequired() && points)
+    {
+        // Compute the mean of the points.
+        Vector3<Real> mean = Vector3<Real>::Zero();
+        for (int i = 0; i < numPoints; ++i)
+        {
+            mean += points[i];
+        }
+        mean /= (Real)numPoints;
+
+        // Compute the covariance matrix of the points.
+        Real covar00 = (Real)0, covar01 = (Real)0, covar02 = (Real)0;
+        Real covar11 = (Real)0, covar12 = (Real)0;
+        for (int i = 0; i < numPoints; ++i)
+        {
+            Vector3<Real> diff = points[i] - mean;
+            covar00 += diff[0] * diff[0];
+            covar01 += diff[0] * diff[1];
+            covar02 += diff[0] * diff[2];
+            covar11 += diff[1] * diff[1];
+            covar12 += diff[1] * diff[2];
+        }
+
+        // Decompose the covariance matrix.
+        Real det = covar00*covar11 - covar01*covar01;
+        if (det != (Real)0)
+        {
+            Real invDet = ((Real)1) / det;
+            mParameters.first = mean;
+            mParameters.second[0] =
+                (covar11*covar02 - covar01*covar12)*invDet;
+            mParameters.second[1] =
+                (covar00*covar12 - covar01*covar02)*invDet;
+            mParameters.second[2] = (Real)-1;
+            return true;
+        }
+    }
+
+    mParameters.first = Vector3<Real>::Zero();
+    mParameters.second = Vector3<Real>::Zero();
+    return false;
+}
+
+template <typename Real>
+std::pair<Vector3<Real>, Vector3<Real>> const&
+ApprHeightPlane3<Real>::GetParameters() const
+{
+    return mParameters;
+}
+
+template <typename Real>
+int ApprHeightPlane3<Real>::GetMinimumRequired() const
+{
+    return 3;
+}
+
+template <typename Real>
+Real ApprHeightPlane3<Real>::Error(Vector3<Real> const& observation) const
+{
+    Real d = Dot(observation - mParameters.first, mParameters.second);
+    Real error = d*d;
+    return error;
+}
+
+template <typename Real>
+bool ApprHeightPlane3<Real>::Fit(
+    std::vector<Vector3<Real>> const& observations,
+    std::vector<int> const& indices)
+{
+    if (static_cast<int>(indices.size()) >= GetMinimumRequired())
+    {
+        // Compute the mean of the points.
+        Vector3<Real> mean = Vector3<Real>::Zero();
+        for (auto index : indices)
+        {
+            mean += observations[index];
+        }
+        mean /= (Real)indices.size();
+
+        // Compute the covariance matrix of the points.
+        Real covar00 = (Real)0, covar01 = (Real)0, covar02 = (Real)0;
+        Real covar11 = (Real)0, covar12 = (Real)0;
+        for (auto index : indices)
+        {
+            Vector3<Real> diff = observations[index] - mean;
+            covar00 += diff[0] * diff[0];
+            covar01 += diff[0] * diff[1];
+            covar02 += diff[0] * diff[2];
+            covar11 += diff[1] * diff[1];
+            covar12 += diff[1] * diff[2];
+        }
+
+        // Decompose the covariance matrix.
+        Real det = covar00*covar11 - covar01*covar01;
+        if (det != (Real)0)
+        {
+            Real invDet = ((Real)1) / det;
+            mParameters.first = mean;
+            mParameters.second[0] =
+                (covar11*covar02 - covar01*covar12)*invDet;
+            mParameters.second[1] =
+                (covar00*covar12 - covar01*covar02)*invDet;
+            mParameters.second[2] = (Real)-1;
+            return true;
+        }
+    }
+
+    mParameters.first = Vector3<Real>::Zero();
+    mParameters.second = Vector3<Real>::Zero();
+    return false;
+}
+
+
+}
diff --git a/Engine/Plugins/Experimental/GeometryProcessing/Source/GeometryAlgorithms/Private/GTEngine/Mathematics/GteApprOrthogonalLine2.h b/Engine/Plugins/Experimental/GeometryProcessing/Source/GeometryAlgorithms/Private/GTEngine/Mathematics/GteApprOrthogonalLine2.h
new file mode 100644
--- /dev/null
+++ b/Engine/Plugins/Experimental/GeometryProcessing/Source/GeometryAlgorithms/Private/GTEngine/Mathematics/GteApprOrthogonalLine2.h
@@ -0,0 +1,169 @@
+// David Eberly, Geometric Tools, Redmond WA 98052
+// Copyright (c) 1998-2019
+// Distributed under the Boost Software License, Version 1.0.
+// http://www.boost.org/LICENSE_1_0.txt
+// http://www.geometrictools.com/License/Boost/LICENSE_1_0.txt
+// File Version: 3.0.0 (2016/06/19)
+
+#pragma once
+
+#include <GTEngine/Mathematics/GteVector2.h>
+#include <GTEngine/Mathematics/GteApprQuery.h>
+#include <GTEngine/Mathematics/GteLine.h>
+#include <GTEngine/Mathematics/GteSymmetricEigensolver2x2.h>
+
+// Least-squares fit of a line to (x,y) data by using distance measurements
+// orthogonal to the proposed line.  The return value is 'true' iff the fit
+// is unique (always successful, 'true' when a minimum eigenvalue is unique).
+// The mParameters value is a line with (P,D) = (origin,direction).  The
+// error for S = (x0,y0) is (S-P)^T*(I - D*D^T)*(S-P).
+
+namespace gte
+{
+
+template <typename Real>
+class ApprOrthogonalLine2
+    :
+    public ApprQuery<Real, ApprOrthogonalLine2<Real>, Vector2<Real>>
+{
+public:
+    // Initialize the model parameters to zero.
+    ApprOrthogonalLine2();
+
+    // Basic fitting algorithm.
+    bool Fit(int numPoints, Vector2<Real> const* points);
+    Line2<Real> const& GetParameters() const;
+
+    // Functions called by ApprQuery::RANSAC.  See GteApprQuery.h for a
+    // detailed description.
+    int GetMinimumRequired() const;
+    Real Error(Vector2<Real> const& observation) const;
+    bool Fit(std::vector<Vector2<Real>> const& observations,
+        std::vector<int> const& indices);
+
+private:
+    Line2<Real> mParameters;
+};
+
+
+template <typename Real>
+ApprOrthogonalLine2<Real>::ApprOrthogonalLine2()
+    :
+    mParameters(Vector2<Real>::Zero(), Vector2<Real>::Zero())
+{
+}
+
+template <typename Real>
+bool ApprOrthogonalLine2<Real>::Fit(int numPoints,
+    Vector2<Real> const* points)
+{
+    if (numPoints >= GetMinimumRequired() && points)
+    {
+        // Compute the mean of the points.
+        Vector2<Real> mean = Vector2<Real>::Zero();
+        for (int i = 0; i < numPoints; ++i)
+        {
+            mean += points[i];
+        }
+        mean /= (Real)numPoints;
+
+        // Compute the covariance matrix of the points.
+        Real covar00 = (Real)0, covar01 = (Real)0, covar11 = (Real)0;
+        for (int i = 0; i < numPoints; ++i)
+        {
+            Vector2<Real> diff = points[i] - mean;
+            covar00 += diff[0] * diff[0];
+            covar01 += diff[0] * diff[1];
+            covar11 += diff[1] * diff[1];
+        }
+
+        // Solve the eigensystem.
+        SymmetricEigensolver2x2<Real> es;
+        std::array<Real, 2> eval;
+        std::array<std::array<Real, 2>, 2> evec;
+        es(covar00, covar01, covar11, +1, eval, evec);
+
+        // The line direction is the eigenvector in the direction of largest
+        // variance of the points.
+        mParameters.origin = mean;
+        mParameters.direction = evec[1];
+
+        // The fitted line is unique when the maximum eigenvalue has
+        // multiplicity 1.
+        return eval[0] < eval[1];
+    }
+
+    mParameters = Line2<Real>(Vector2<Real>::Zero(), Vector2<Real>::Zero());
+    return false;
+}
+
+template <typename Real>
+Line2<Real> const& ApprOrthogonalLine2<Real>::GetParameters() const
+{
+    return mParameters;
+}
+
+template <typename Real>
+int ApprOrthogonalLine2<Real>::GetMinimumRequired() const
+{
+    return 2;
+}
+
+template <typename Real>
+Real ApprOrthogonalLine2<Real>::Error(Vector2<Real> const& observation) const
+{
+    Vector2<Real> diff = observation - mParameters.origin;
+    Real sqrlen = Dot(diff, diff);
+    Real dot = Dot(diff, mParameters.direction);
+    Real error = std::abs(sqrlen - dot*dot);
+    return error;
+}
+
+template <typename Real>
+bool ApprOrthogonalLine2<Real>::Fit(
+    std::vector<Vector2<Real>> const& observations,
+    std::vector<int> const& indices)
+{
+    if (static_cast<int>(indices.size()) >= GetMinimumRequired())
+    {
+        // Compute the mean of the points.
+        Vector2<Real> mean = Vector2<Real>::Zero();
+        for (auto index : indices)
+        {
+            mean += observations[index];
+        }
+        mean /= (Real)indices.size();
+
+        // Compute the covariance matrix of the points.
+        Real covar00 = (Real)0, covar01 = (Real)0, covar11 = (Real)0;
+        for (auto index : indices)
+        {
+            Vector2<Real> diff = observations[index] - mean;
+            covar00 += diff[0] * diff[0];
+            covar01 += diff[0] * diff[1];
+            covar11 += diff[1] * diff[1];
+        }
+
+        // Solve the eigensystem.
+        // Solve the eigensystem.
+        SymmetricEigensolver2x2<Real> es;
+        std::array<Real, 2> eval;
+        std::array<std::array<Real, 2>, 2> evec;
+        es(covar00, covar01, covar11, +1, eval, evec);
+
+        // The line direction is the eigenvector in the direction of largest
+        // variance of the points.
+        mParameters.origin = mean;
+        mParameters.direction = evec[1];
+
+        // The fitted line is unique when the maximum eigenvalue has
+        // multiplicity 1.
+        return eval[0] < eval[1];
+    }
+
+    mParameters = Line2<Real>(Vector2<Real>::Zero(), Vector2<Real>::Zero());
+    return false;
+}
+
+
+}
diff --git a/Engine/Plugins/Experimental/GeometryProcessing/Source/GeometryAlgorithms/Private/GTEngine/Mathematics/GteApprOrthogonalLine3.h b/Engine/Plugins/Experimental/GeometryProcessing/Source/GeometryAlgorithms/Private/GTEngine/Mathematics/GteApprOrthogonalLine3.h
new file mode 100644
--- /dev/null
+++ b/Engine/Plugins/Experimental/GeometryProcessing/Source/GeometryAlgorithms/Private/GTEngine/Mathematics/GteApprOrthogonalLine3.h
@@ -0,0 +1,191 @@
+// David Eberly, Geometric Tools, Redmond WA 98052
+// Copyright (c) 1998-2019
+// Distributed under the Boost Software License, Version 1.0.
+// http://www.boost.org/LICENSE_1_0.txt
+// http://www.geometrictools.com/License/Boost/LICENSE_1_0.txt
+// File Version: 3.0.0 (2016/06/19)
+
+#pragma once
+
+#include <GTEngine/Mathematics/GteVector3.h>
+#include <GTEngine/Mathematics/GteApprQuery.h>
+#include <GTEngine/Mathematics/GteLine.h>
+#include <GTEngine/Mathematics/GteSymmetricEigensolver3x3.h>
+
+// Least-squares fit of a line to (x,y,z) data by using distance measurements
+// orthogonal to the proposed line.  The return value is 'true' iff the fit
+// is unique (always successful, 'true' when a minimum eigenvalue is unique).
+// The mParameters value is a line with (P,D) = (origin,direction).  The
+// error for S = (x0,y0,z0) is (S-P)^T*(I - D*D^T)*(S-P).
+
+namespace gte
+{
+
+template <typename Real>
+class ApprOrthogonalLine3
+    :
+    public ApprQuery<Real, ApprOrthogonalLine3<Real>, Vector3<Real>>
+{
+public:
+    // Initialize the model parameters to zero.
+    ApprOrthogonalLine3();
+
+    // Basic fitting algorithm.
+    bool Fit(int numPoints, Vector3<Real> const* points);
+    Line3<Real> const& GetParameters() const;
+
+    // Functions called by ApprQuery::RANSAC.  See GteApprQuery.h for a
+    // detailed description.
+    int GetMinimumRequired() const;
+    Real Error(Vector3<Real> const& observation) const;
+    bool Fit(std::vector<Vector3<Real>> const& observations,
+        std::vector<int> const& indices);
+
+private:
+    Line3<Real> mParameters;
+};
+
+
+template <typename Real>
+ApprOrthogonalLine3<Real>::ApprOrthogonalLine3()
+    :
+    mParameters(Vector3<Real>::Zero(), Vector3<Real>::Zero())
+{
+}
+
+template <typename Real>
+bool ApprOrthogonalLine3<Real>::Fit(int numPoints, Vector3<Real> const* points)
+{
+    if (numPoints >= GetMinimumRequired() && points)
+    {
+        // Compute the mean of the points.
+        Vector3<Real> mean = Vector3<Real>::Zero();
+        for (int i = 0; i < numPoints; ++i)
+        {
+            mean += points[i];
+        }
+        Real invSize = ((Real)1) / (Real)numPoints;
+        mean *= invSize;
+
+        // Compute the covariance matrix of the points.
+        Real covar00 = (Real)0, covar01 = (Real)0, covar02 = (Real)0;
+        Real covar11 = (Real)0, covar12 = (Real)0, covar22 = (Real)0;
+        for (int i = 0; i < numPoints; ++i)
+        {
+            Vector3<Real> diff = points[i] - mean;
+            covar00 += diff[0] * diff[0];
+            covar01 += diff[0] * diff[1];
+            covar02 += diff[0] * diff[2];
+            covar11 += diff[1] * diff[1];
+            covar12 += diff[1] * diff[2];
+            covar22 += diff[2] * diff[2];
+        }
+        covar00 *= invSize;
+        covar01 *= invSize;
+        covar02 *= invSize;
+        covar11 *= invSize;
+        covar12 *= invSize;
+        covar22 *= invSize;
+
+        // Solve the eigensystem.
+        SymmetricEigensolver3x3<Real> es;
+        std::array<Real, 3> eval;
+        std::array<std::array<Real, 3>, 3> evec;
+        es(covar00, covar01, covar02, covar11, covar12, covar22, false, +1,
+            eval, evec);
+
+        // The line direction is the eigenvector in the direction of largest
+        // variance of the points.
+        mParameters.origin = mean;
+        mParameters.direction = evec[2];
+
+        // The fitted line is unique when the maximum eigenvalue has
+        // multiplicity 1.
+        return eval[1]< eval[2];
+    }
+
+    mParameters = Line3<Real>(Vector3<Real>::Zero(), Vector3<Real>::Zero());
+    return false;
+}
+
+template <typename Real>
+Line3<Real> const& ApprOrthogonalLine3<Real>::GetParameters() const
+{
+    return mParameters;
+}
+
+template <typename Real>
+int ApprOrthogonalLine3<Real>::GetMinimumRequired() const
+{
+    return 2;
+}
+
+template <typename Real>
+Real ApprOrthogonalLine3<Real>::Error(Vector3<Real> const& observation) const
+{
+    Vector3<Real> diff = observation - mParameters.origin;
+    Real sqrlen = Dot(diff, diff);
+    Real dot = Dot(diff, mParameters.direction);
+    Real error = std::abs(sqrlen - dot*dot);
+    return error;
+}
+
+template <typename Real>
+bool ApprOrthogonalLine3<Real>::Fit(
+    std::vector<Vector3<Real>> const& observations,
+    std::vector<int> const& indices)
+{
+    if (static_cast<int>(indices.size()) >= GetMinimumRequired())
+    {
+        // Compute the mean of the points.
+        Vector3<Real> mean = Vector3<Real>::Zero();
+        for (auto index : indices)
+        {
+            mean += observations[index];
+        }
+        Real invSize = ((Real)1) / (Real)indices.size();
+        mean *= invSize;
+
+        // Compute the covariance matrix of the points.
+        Real covar00 = (Real)0, covar01 = (Real)0, covar02 = (Real)0;
+        Real covar11 = (Real)0, covar12 = (Real)0, covar22 = (Real)0;
+        for (auto index : indices)
+        {
+            Vector3<Real> diff = observations[index] - mean;
+            covar00 += diff[0] * diff[0];
+            covar01 += diff[0] * diff[1];
+            covar02 += diff[0] * diff[2];
+            covar11 += diff[1] * diff[1];
+            covar12 += diff[1] * diff[2];
+            covar22 += diff[2] * diff[2];
+        }
+        covar00 *= invSize;
+        covar01 *= invSize;
+        covar02 *= invSize;
+        covar11 *= invSize;
+        covar12 *= invSize;
+        covar22 *= invSize;
+
+        // Solve the eigensystem.
+        SymmetricEigensolver3x3<Real> es;
+        std::array<Real, 3> eval;
+        std::array<std::array<Real, 3>, 3> evec;
+        es(covar00, covar01, covar02, covar11, covar12, covar22, false, +1,
+            eval, evec);
+
+        // The line direction is the eigenvector in the direction of largest
+        // variance of the points.
+        mParameters.origin = mean;
+        mParameters.direction = evec[2];
+
+        // The fitted line is unique when the maximum eigenvalue has
+        // multiplicity 1.
+        return eval[1]< eval[2];
+    }
+
+    mParameters = Line3<Real>(Vector3<Real>::Zero(), Vector3<Real>::Zero());
+    return false;
+}
+
+
+}
diff --git a/Engine/Plugins/Experimental/GeometryProcessing/Source/GeometryAlgorithms/Private/GTEngine/Mathematics/GteApprOrthogonalPlane3.h b/Engine/Plugins/Experimental/GeometryProcessing/Source/GeometryAlgorithms/Private/GTEngine/Mathematics/GteApprOrthogonalPlane3.h
new file mode 100644
--- /dev/null
+++ b/Engine/Plugins/Experimental/GeometryProcessing/Source/GeometryAlgorithms/Private/GTEngine/Mathematics/GteApprOrthogonalPlane3.h
@@ -0,0 +1,181 @@
+// David Eberly, Geometric Tools, Redmond WA 98052
+// Copyright (c) 1998-2019
+// Distributed under the Boost Software License, Version 1.0.
+// http://www.boost.org/LICENSE_1_0.txt
+// http://www.geometrictools.com/License/Boost/LICENSE_1_0.txt
+// File Version: 3.0.0 (2016/06/19)
+
+#pragma once
+
+#include <GTEngine/Mathematics/GteVector3.h>
+#include <GTEngine/Mathematics/GteApprQuery.h>
+#include <GTEngine/Mathematics/GteSymmetricEigensolver3x3.h>
+
+// Least-squares fit of a plane to (x,y,z) data by using distance measurements
+// orthogonal to the proposed plane.  The return value is 'true' iff the fit
+// is unique (always successful, 'true' when a minimum eigenvalue is unique).
+// The mParameters value is (P,N) = (origin,normal).  The error for
+// S = (x0,y0,z0) is (S-P)^T*(I - N*N^T)*(S-P).
+
+namespace gte
+{
+
+template <typename Real>
+class ApprOrthogonalPlane3
+    :
+    public ApprQuery<Real, ApprOrthogonalPlane3<Real>, Vector3<Real>>
+{
+public:
+    // Initialize the model parameters to zero.
+    ApprOrthogonalPlane3();
+
+    // Basic fitting algorithm.
+    bool Fit(int numPoints, Vector3<Real> const* points);
+    std::pair<Vector3<Real>, Vector3<Real>> const& GetParameters() const;
+
+    // Functions called by ApprQuery::RANSAC.  See GteApprQuery.h for a
+    // detailed description.
+    int GetMinimumRequired() const;
+    Real Error(Vector3<Real> const& observation) const;
+    bool Fit(std::vector<Vector3<Real>> const& observations,
+        std::vector<int> const& indices);
+
+private:
+    std::pair<Vector3<Real>, Vector3<Real>> mParameters;
+};
+
+
+template <typename Real>
+ApprOrthogonalPlane3<Real>::ApprOrthogonalPlane3()
+{
+    mParameters.first = Vector3<Real>::Zero();
+    mParameters.second = Vector3<Real>::Zero();
+}
+
+template <typename Real>
+bool ApprOrthogonalPlane3<Real>::Fit(int numPoints,
+    Vector3<Real> const* points)
+{
+    if (numPoints >= GetMinimumRequired() && points)
+    {
+        // Compute the mean of the points.
+        Vector3<Real> mean = Vector3<Real>::Zero();
+        for (int i = 0; i < numPoints; ++i)
+        {
+            mean += points[i];
+        }
+        mean /= (Real)numPoints;
+
+        // Compute the covariance matrix of the points.
+        Real covar00 = (Real)0, covar01 = (Real)0, covar02 = (Real)0;
+        Real covar11 = (Real)0, covar12 = (Real)0, covar22 = (Real)0;
+        for (int i = 0; i < numPoints; ++i)
+        {
+            Vector3<Real> diff = points[i] - mean;
+            covar00 += diff[0] * diff[0];
+            covar01 += diff[0] * diff[1];
+            covar02 += diff[0] * diff[2];
+            covar11 += diff[1] * diff[1];
+            covar12 += diff[1] * diff[2];
+            covar22 += diff[2] * diff[2];
+        }
+
+        // Solve the eigensystem.
+        SymmetricEigensolver3x3<Real> es;
+        std::array<Real, 3> eval;
+        std::array<std::array<Real, 3>, 3> evec;
+        es(covar00, covar01, covar02, covar11, covar12, covar22, false, +1,
+            eval, evec);
+
+        // The plane normal is the eigenvector in the direction of smallest
+        // variance of the points.
+        mParameters.first = mean;
+        mParameters.second = evec[0];
+
+        // The fitted plane is unique when the minimum eigenvalue has
+        // multiplicity 1.
+        return eval[0] < eval[1];
+    }
+
+    mParameters.first = Vector3<Real>::Zero();
+    mParameters.second = Vector3<Real>::Zero();
+    return false;
+}
+
+template <typename Real>
+std::pair<Vector3<Real>, Vector3<Real>> const&
+ApprOrthogonalPlane3<Real>::GetParameters() const
+{
+    return mParameters;
+}
+
+template <typename Real>
+int ApprOrthogonalPlane3<Real>::GetMinimumRequired() const
+{
+    return 3;
+}
+
+template <typename Real>
+Real ApprOrthogonalPlane3<Real>::Error(Vector3<Real> const& observation)
+const
+{
+    Vector3<Real> diff = observation - mParameters.first;
+    Real sqrlen = Dot(diff, diff);
+    Real dot = Dot(diff, mParameters.second);
+    Real error = std::abs(sqrlen - dot*dot);
+    return error;
+}
+
+template <typename Real>
+bool ApprOrthogonalPlane3<Real>::Fit(
+    std::vector<Vector3<Real>> const& observations,
+    std::vector<int> const& indices)
+{
+    if (static_cast<int>(indices.size()) >= GetMinimumRequired())
+    {
+        // Compute the mean of the points.
+        Vector3<Real> mean = Vector3<Real>::Zero();
+        for (auto index : indices)
+        {
+            mean += observations[index];
+        }
+        mean /= (Real)indices.size();
+
+        // Compute the covariance matrix of the points.
+        Real covar00 = (Real)0, covar01 = (Real)0, covar02 = (Real)0;
+        Real covar11 = (Real)0, covar12 = (Real)0, covar22 = (Real)0;
+        for (auto index : indices)
+        {
+            Vector3<Real> diff = observations[index] - mean;
+            covar00 += diff[0] * diff[0];
+            covar01 += diff[0] * diff[1];
+            covar02 += diff[0] * diff[2];
+            covar11 += diff[1] * diff[1];
+            covar12 += diff[1] * diff[2];
+            covar22 += diff[2] * diff[2];
+        }
+
+        // Solve the eigensystem.
+        SymmetricEigensolver3x3<Real> es;
+        std::array<Real, 3> eval;
+        std::array<std::array<Real, 3>, 3> evec;
+        es(covar00, covar01, covar02, covar11, covar12, covar22, false, +1,
+            eval, evec);
+
+        // The plane normal is the eigenvector in the direction of smallest
+        // variance of the points.
+        mParameters.first = mean;
+        mParameters.second = evec[0];
+
+        // The fitted plane is unique when the minimum eigenvalue has
+        // multiplicity 1.
+        return eval[0] < eval[1];
+    }
+
+    mParameters.first = Vector3<Real>::Zero();
+    mParameters.second = Vector3<Real>::Zero();
+    return false;
+}
+
+
+}
diff --git a/Engine/Plugins/Experimental/GeometryProcessing/Source/GeometryAlgorithms/Private/GTEngine/Mathematics/GteApprParaboloid3.h b/Engine/Plugins/Experimental/GeometryProcessing/Source/GeometryAlgorithms/Private/GTEngine/Mathematics/GteApprParaboloid3.h
new file mode 100644
--- /dev/null
+++ b/Engine/Plugins/Experimental/GeometryProcessing/Source/GeometryAlgorithms/Private/GTEngine/Mathematics/GteApprParaboloid3.h
@@ -0,0 +1,136 @@
+// David Eberly, Geometric Tools, Redmond WA 98052
+// Copyright (c) 1998-2019
+// Distributed under the Boost Software License, Version 1.0.
+// http://www.boost.org/LICENSE_1_0.txt
+// http://www.geometrictools.com/License/Boost/LICENSE_1_0.txt
+// File Version: 3.0.0 (2016/06/19)
+
+#pragma once
+
+#include <GTEngine/Mathematics/GteMatrix.h>
+#include <GTEngine/Mathematics/GteVector3.h>
+#include <GTEngine/Mathematics/GteLinearSystem.h>
+
+// Least-squares fit of a paraboloid to a set of point.  The paraboloid is
+// of the form z = c0*x^2+c1*x*y+c2*y^2+c3*x+c4*y+c5.  A successful fit is
+// indicated by return value of 'true'.
+//
+// Given a set of samples (x_i,y_i,z_i) for 0 <= i < N, and assuming
+// that the true values lie on a paraboloid
+//     z = p0*x*x + p1*x*y + p2*y*y + p3*x + p4*y + p5  = Dot(P,Q(x,y))
+// where P = (p0,p1,p2,p3,p4,p5) and Q(x,y) = (x*x,x*y,y*y,x,y,1),
+// select P to minimize the sum of squared errors
+//     E(P) = sum_{i=0}^{N-1} [Dot(P,Q_i)-z_i]^2
+// where Q_i = Q(x_i,y_i).
+//
+// The minimum occurs when the gradient of E is the zero vector,
+//     grad(E) = 2 sum_{i=0}^{N-1} [Dot(P,Q_i)-z_i] Q_i = 0
+// Some algebra converts this to a system of 6 equations in 6 unknowns:
+//     [(sum_{i=0}^{N-1} Q_i Q_i^t] P = sum_{i=0}^{N-1} z_i Q_i
+// The product Q_i Q_i^t is a product of the 6x1 matrix Q_i with the
+// 1x6 matrix Q_i^t, the result being a 6x6 matrix.
+//
+// Define the 6x6 symmetric matrix A = sum_{i=0}^{N-1} Q_i Q_i^t and the 6x1
+// vector B = sum_{i=0}^{N-1} z_i Q_i.  The choice for P is the solution to
+// the linear system of equations A*P = B.  The entries of A and B indicate
+// summations over the appropriate product of variables.  For example,
+// s(x^3 y) = sum_{i=0}^{N-1} x_i^3 y_i.
+//
+// +-                                                     -++  +   +-      -+
+// | s(x^4) s(x^3 y)   s(x^2 y^2) s(x^3)   s(x^2 y) s(x^2) ||p0|   |s(z x^2)|
+// |        s(x^2 y^2) s(x y^3)   s(x^2 y) s(x y^2) s(x y) ||p1|   |s(z x y)|
+// |                   s(y^4)     s(x y^2) s(y^3)   s(y^2) ||p2| = |s(z y^2)|
+// |                              s(x^2)   s(x y)   s(x)   ||p3|   |s(z x)  |
+// |                                       s(y^2)   s(y)   ||p4|   |s(z y)  |
+// |                                                s(1)   ||p5|   |s(z)    |
+// +-                                                     -++  +   +-      -+
+
+namespace gte
+{
+
+template <typename Real>
+class ApprParaboloid3
+{
+public:
+    bool operator()(int numPoints, Vector3<Real> const* points,
+        Real coefficients[6]);
+};
+
+
+template <typename Real>
+bool ApprParaboloid3<Real>::operator()(int numPoints,
+    Vector3<Real> const* points, Real coefficients[6])
+{
+    Matrix<6, 6, Real> A;
+    Vector<6, Real> B;
+    for (int i = 0; i < numPoints; i++)
+    {
+        Real x2 = points[i][0] * points[i][0];
+        Real xy = points[i][0] * points[i][1];
+        Real y2 = points[i][1] * points[i][1];
+        Real zx = points[i][2] * points[i][0];
+        Real zy = points[i][2] * points[i][1];
+        Real x3 = points[i][0] * x2;
+        Real x2y = x2*points[i][1];
+        Real xy2 = points[i][0] * y2;
+        Real y3 = points[i][1] * y2;
+        Real zx2 = points[i][2] * x2;
+        Real zxy = points[i][2] * xy;
+        Real zy2 = points[i][2] * y2;
+        Real x4 = x2*x2;
+        Real x3y = x3*points[i][1];
+        Real x2y2 = x2*y2;
+        Real xy3 = points[i][0] * y3;
+        Real y4 = y2*y2;
+
+        A(0, 0) += x4;
+        A(0, 1) += x3y;
+        A(0, 2) += x2y2;
+        A(0, 3) += x3;
+        A(0, 4) += x2y;
+        A(0, 5) += x2;
+        A(1, 2) += xy3;
+        A(1, 4) += xy2;
+        A(1, 5) += xy;
+        A(2, 2) += y4;
+        A(2, 4) += y3;
+        A(2, 5) += y2;
+        A(3, 3) += x2;
+        A(3, 5) += points[i][0];
+        A(4, 5) += points[i][1];
+
+        B[0] += zx2;
+        B[1] += zxy;
+        B[2] += zy2;
+        B[3] += zx;
+        B[4] += zy;
+        B[5] += points[i][2];
+    }
+
+    A(1, 0) = A(0, 1);
+    A(1, 1) = A(0, 2);
+    A(1, 3) = A(0, 4);
+    A(2, 0) = A(0, 2);
+    A(2, 1) = A(1, 2);
+    A(2, 3) = A(1, 4);
+    A(3, 0) = A(0, 3);
+    A(3, 1) = A(1, 3);
+    A(3, 2) = A(2, 3);
+    A(3, 4) = A(1, 5);
+    A(4, 0) = A(0, 4);
+    A(4, 1) = A(1, 4);
+    A(4, 2) = A(2, 4);
+    A(4, 3) = A(3, 4);
+    A(4, 4) = A(2, 5);
+    A(5, 0) = A(0, 5);
+    A(5, 1) = A(1, 5);
+    A(5, 2) = A(2, 5);
+    A(5, 3) = A(3, 5);
+    A(5, 4) = A(4, 5);
+    A(5, 5) = static_cast<Real>(numPoints);
+
+    return LinearSystem<Real>().Solve(6, &A[0], &B[0], &coefficients[0]);
+}
+
+
+}
diff --git a/Engine/Plugins/Experimental/GeometryProcessing/Source/GeometryAlgorithms/Private/GTEngine/Mathematics/GteApprPolynomial2.h b/Engine/Plugins/Experimental/GeometryProcessing/Source/GeometryAlgorithms/Private/GTEngine/Mathematics/GteApprPolynomial2.h
new file mode 100644
--- /dev/null
+++ b/Engine/Plugins/Experimental/GeometryProcessing/Source/GeometryAlgorithms/Private/GTEngine/Mathematics/GteApprPolynomial2.h
@@ -0,0 +1,200 @@
+// David Eberly, Geometric Tools, Redmond WA 98052
+// Copyright (c) 1998-2019
+// Distributed under the Boost Software License, Version 1.0.
+// http://www.boost.org/LICENSE_1_0.txt
+// http://www.geometrictools.com/License/Boost/LICENSE_1_0.txt
+// File Version: 3.0.0 (2016/06/19)
+
+#pragma once
+
+#include <GTEngine/LowLevel/GteArray2.h>
+#include <GTEngine/Mathematics/GteGMatrix.h>
+#include <GTEngine/Mathematics/GteApprQuery.h>
+#include <array>
+
+// The samples are (x[i],w[i]) for 0 <= i < S.  Think of w as a function of
+// x, say w = f(x).  The function fits the samples with a polynomial of
+// degree d, say w = sum_{i=0}^d c[i]*x^i.  The method is a least-squares
+// fitting algorithm.  The mParameters stores the coefficients c[i] for
+// 0 <= i <= d.  The observation type is std::array<Real,2>, which represents
+// a pair (x,w).
+//
+// WARNING.  The fitting algorithm for polynomial terms
+//   (1,x,x^2,...,x^d)
+// is known to be nonrobust for large degrees and for large magnitude data.
+// One alternative is to use orthogonal polynomials
+//   (f[0](x),...,f[d](x))
+// and apply the least-squares algorithm to these.  Another alternative is to
+// transform
+//   (x',w') = ((x-xcen)/rng, w/rng)
+// where xmin = min(x[i]), xmax = max(x[i]), xcen = (xmin+xmax)/2, and
+// rng = xmax-xmin.  Fit the (x',w') points,
+//   w' = sum_{i=0}^d c'[i]*(x')^i.
+// The original polynomial is evaluated as
+//   w = rng*sum_{i=0}^d c'[i]*((x-xcen)/rng)^i
+
+namespace gte
+{
+
+template <typename Real>
+class ApprPolynomial2
+    :
+    public ApprQuery<Real, ApprPolynomial2<Real>, std::array<Real, 2>>
+{
+public:
+    // Initialize the model parameters to zero.
+    ApprPolynomial2(int degree);
+
+    // The minimum number of observations required to fit the model.
+    int GetMinimumRequired() const;
+
+    // Estimate the model parameters for all observations specified by the
+    // indices.  This function is called by the base-class Fit(...) functions.
+    bool Fit(std::vector<std::array<Real, 2>> const& observations,
+        std::vector<int> const& indices);
+
+    // Compute the model error for the specified observation for the current
+    // model parameters.  The returned value for observation (x0,w0) is
+    // |w(x0) - w0|, where w(x) is the fitted polynomial.
+    Real Error(std::array<Real, 2> const& observation) const;
+
+    // Get the parameters of the model.
+    std::vector<Real> const& GetParameters() const;
+
+    // Evaluate the polynomial.  The domain interval is provided so you can
+    // interpolate (x in domain) or extrapolate (x not in domain).
+    std::array<Real, 2> const& GetXDomain() const;
+    Real Evaluate(Real x) const;
+
+private:
+    int mDegree, mSize;
+    std::array<Real, 2> mXDomain;
+    std::vector<Real> mParameters;
+};
+
+
+template <typename Real>
+ApprPolynomial2<Real>::ApprPolynomial2(int degree)
+    :
+    mDegree(degree),
+    mSize(degree + 1),
+    mParameters(mSize)
+{
+    mXDomain[0] = std::numeric_limits<Real>::max();
+    mXDomain[1] = -mXDomain[0];
+    std::fill(mParameters.begin(), mParameters.end(), (Real)0);
+}
+
+template <typename Real>
+int ApprPolynomial2<Real>::GetMinimumRequired() const
+{
+    return mSize;
+}
+
+template <typename Real>
+bool ApprPolynomial2<Real>::Fit(
+    std::vector<std::array<Real, 2>> const& observations,
+    std::vector<int> const& indices)
+{
+    if (indices.size() > 0)
+    {
+        int s, i0, i1;
+
+        // Compute the powers of x.
+        int numSamples = static_cast<int>(indices.size());
+        int twoDegree = 2 * mDegree;
+        Array2<Real> xPower(twoDegree + 1, numSamples);
+        for (s = 0; s < numSamples; ++s)
+        {
+            Real x = observations[indices[s]][0];
+            mXDomain[0] = std::min(x, mXDomain[0]);
+            mXDomain[1] = std::max(x, mXDomain[1]);
+
+            xPower[s][0] = (Real)1;
+            for (i0 = 1; i0 <= twoDegree; ++i0)
+            {
+                xPower[s][i0] = x * xPower[s][i0 - 1];
+            }
+        }
+
+        // Matrix A is the Vandermonde matrix and vector B is the right-hand
+        // side of the linear system A*X = B.
+        GMatrix<Real> A(mSize, mSize);
+        GVector<Real> B(mSize);
+        for (i0 = 0; i0 <= mDegree; ++i0)
+        {
+            Real sum = (Real)0;
+            for (s = 0; s < numSamples; ++s)
+            {
+                Real w = observations[indices[s]][1];
+                sum += w * xPower[s][i0];
+            }
+
+            B[i0] = sum;
+
+            for (i1 = 0; i1 <= mDegree; ++i1)
+            {
+                sum = (Real)0;
+                for (s = 0; s < numSamples; ++s)
+                {
+                    sum += xPower[s][i0 + i1];
+                }
+
+                A(i0, i1) = sum;
+            }
+        }
+
+        // Solve for the polynomial coefficients.
+        GVector<Real> coefficients = Inverse(A) * B;
+        bool hasNonzero = false;
+        for (int i = 0; i < mSize; ++i)
+        {
+            mParameters[i] = coefficients[i];
+            if (coefficients[i] != (Real)0)
+            {
+                hasNonzero = true;
+            }
+        }
+        return hasNonzero;
+
+    }
+
+    std::fill(mParameters.begin(), mParameters.end(), (Real)0);
+    return false;
+}
+
+template <typename Real>
+Real ApprPolynomial2<Real>::Error(std::array<Real, 2> const& observation)
+const
+{
+    Real w = Evaluate(observation[0]);
+    Real error = std::abs(w - observation[1]);
+    return error;
+}
+
+template <typename Real>
+std::vector<Real> const& ApprPolynomial2<Real>::GetParameters() const
+{
+    return mParameters;
+}
+
+template <typename Real>
+std::array<Real, 2> const& ApprPolynomial2<Real>::GetXDomain() const
+{
+    return mXDomain;
+}
+
+template <typename Real>
+Real ApprPolynomial2<Real>::Evaluate(Real x) const
+{
+    int i = mDegree;
+    Real w = mParameters[i];
+    while (--i >= 0)
+    {
+        w = mParameters[i] + w * x;
+    }
+    return w;
+}
+
+
+}
