move GeometryProcessing plugin from NotForLicensees to Experimental

diff --git a/Engine/Plugins/Experimental/GeometryProcessing/Source/GeometryAlgorithms/Private/GTEngine/Mathematics/GteApprPolynomial3.h b/Engine/Plugins/Experimental/GeometryProcessing/Source/GeometryAlgorithms/Private/GTEngine/Mathematics/GteApprPolynomial3.h
new file mode 100644
--- /dev/null
+++ b/Engine/Plugins/Experimental/GeometryProcessing/Source/GeometryAlgorithms/Private/GTEngine/Mathematics/GteApprPolynomial3.h
@@ -0,0 +1,256 @@
+// David Eberly, Geometric Tools, Redmond WA 98052
+// Copyright (c) 1998-2019
+// Distributed under the Boost Software License, Version 1.0.
+// http://www.boost.org/LICENSE_1_0.txt
+// http://www.geometrictools.com/License/Boost/LICENSE_1_0.txt
+// File Version: 3.0.0 (2016/06/19)
+
+#pragma once
+
+#include <GTEngine/LowLevel/GteArray2.h>
+#include <GTEngine/Mathematics/GteGMatrix.h>
+#include <GTEngine/Mathematics/GteApprQuery.h>
+#include <array>
+
+// The samples are (x[i],y[i],w[i]) for 0 <= i < S.  Think of w as a function
+// of x and y, say w = f(x,y).  The function fits the samples with a
+// polynomial of degree d0 in x and degree d1 in y, say
+//   w = sum_{i=0}^{d0} sum_{j=0}^{d1} c[i][j]*x^i*y^j
+// The method is a least-squares fitting algorithm.  The mParameters stores
+// c[i][j] = mParameters[i+(d0+1)*j] for a total of (d0+1)*(d1+1)
+// coefficients.  The observation type is std::array<Real,3>, which represents
+// a triple (x,y,w).
+//
+// WARNING.  The fitting algorithm for polynomial terms
+//   (1,x,x^2,...,x^d0), (1,y,y^2,...,y^d1)
+// is known to be nonrobust for large degrees and for large magnitude data.
+// One alternative is to use orthogonal polynomials
+//   (f[0](x),...,f[d0](x)), (g[0](y),...,g[d1](y))
+// and apply the least-squares algorithm to these.  Another alternative is to
+// transform
+//   (x',y',w') = ((x-xcen)/rng, (y-ycen)/rng, w/rng)
+// where xmin = min(x[i]), xmax = max(x[i]), xcen = (xmin+xmax)/2,
+// ymin = min(y[i]), ymax = max(y[i]), ycen = (ymin+ymax)/2, and
+// rng = max(xmax-xmin,ymax-ymin).  Fit the (x',y',w') points,
+//   w' = sum_{i=0}^{d0} sum_{j=0}^{d1} c'[i][j]*(x')^i*(y')^j
+// The original polynomial is evaluated as
+//   w = rng * sum_{i=0}^{d0} sum_{j=0}^{d1} c'[i][j] *
+//         ((x-xcen)/rng)^i * ((y-ycen)/rng)^j
+
+namespace gte
+{
+
+template <typename Real>
+class ApprPolynomial3
+    :
+    public ApprQuery<Real, ApprPolynomial3<Real>, std::array<Real, 3>>
+{
+public:
+    // Initialize the model parameters to zero.
+    ApprPolynomial3(int xDegree, int yDegree);
+
+    // The minimum number of observations required to fit the model.
+    int GetMinimumRequired() const;
+
+    // Estimate the model parameters for all observations specified by the
+    // indices.  This function is called by the base-class Fit(...) functions.
+    bool Fit(std::vector<std::array<Real, 3>> const& observations,
+        std::vector<int> const& indices);
+
+    // Compute the model error for the specified observation for the current
+    // model parameters.  The returned value for observation (x0,y0,w0) is
+    // |w(x0,y0) - w0|, where w(x,y) is the fitted polynomial.
+    Real Error(std::array<Real, 3> const& observation) const;
+
+    // Get the parameters of the model.
+    std::vector<Real> const& GetParameters() const;
+
+    // Evaluate the polynomial.  The domain intervals are provided so you can
+    // interpolate ((x,y) in domain) or extrapolate ((x,y) not in domain).
+    std::array<Real, 2> const& GetXDomain() const;
+    std::array<Real, 2> const& GetYDomain() const;
+    Real Evaluate(Real x, Real y) const;
+
+private:
+    int mXDegree, mYDegree, mXDegreeP1, mYDegreeP1, mSize;
+    std::array<Real, 2> mXDomain, mYDomain;
+    std::vector<Real> mParameters;
+
+    // This array is used by Evaluate() to avoid reallocation of the 'vector'
+    // for each call.  The member is mutable because, to the user, the call
+    // to Evaluate does not modify the polynomial.
+    mutable std::vector<Real> mYCoefficient;
+};
+
+
+template <typename Real>
+ApprPolynomial3<Real>::ApprPolynomial3(int xDegree, int yDegree)
+    :
+    mXDegree(xDegree),
+    mYDegree(yDegree),
+    mXDegreeP1(xDegree + 1),
+    mYDegreeP1(yDegree + 1),
+    mSize(mXDegreeP1 * mYDegreeP1),
+    mParameters(mSize),
+    mYCoefficient(mYDegreeP1)
+{
+    mXDomain[0] = std::numeric_limits<Real>::max();
+    mXDomain[1] = -mXDomain[0];
+    mYDomain[0] = std::numeric_limits<Real>::max();
+    mYDomain[1] = -mYDomain[0];
+    std::fill(mParameters.begin(), mParameters.end(), (Real)0);
+    std::fill(mYCoefficient.begin(), mYCoefficient.end(), (Real)0);
+}
+
+template <typename Real>
+int ApprPolynomial3<Real>::GetMinimumRequired() const
+{
+    return mSize;
+}
+
+template <typename Real>
+bool ApprPolynomial3<Real>::Fit(
+    std::vector<std::array<Real, 3>> const& observations,
+    std::vector<int> const& indices)
+{
+    if (indices.size() > 0)
+    {
+        int s, i0, j0, k0, i1, j1, k1;
+
+        // Compute the powers of x and y.
+        int numSamples = static_cast<int>(indices.size());
+        int twoXDegree = 2 * mXDegree;
+        int twoYDegree = 2 * mYDegree;
+        Array2<Real> xPower(twoXDegree + 1, numSamples);
+        Array2<Real> yPower(twoYDegree + 1, numSamples);
+        for (s = 0; s < numSamples; ++s)
+        {
+            Real x = observations[indices[s]][0];
+            Real y = observations[indices[s]][1];
+            mXDomain[0] = std::min(x, mXDomain[0]);
+            mXDomain[1] = std::max(x, mXDomain[1]);
+            mYDomain[0] = std::min(y, mYDomain[0]);
+            mYDomain[1] = std::max(y, mYDomain[1]);
+
+            xPower[s][0] = (Real)1;
+            for (i0 = 1; i0 <= twoXDegree; ++i0)
+            {
+                xPower[s][i0] = x * xPower[s][i0 - 1];
+            }
+
+            yPower[s][0] = (Real)1;
+            for (j0 = 1; j0 <= twoYDegree; ++j0)
+            {
+                yPower[s][j0] = y * yPower[s][j0 - 1];
+            }
+        }
+
+        // Matrix A is the Vandermonde matrix and vector B is the right-hand
+        // side of the linear system A*X = B.
+        GMatrix<Real> A(mSize, mSize);
+        GVector<Real> B(mSize);
+        for (j0 = 0; j0 <= mYDegree; ++j0)
+        {
+            for (i0 = 0; i0 <= mXDegree; ++i0)
+            {
+                Real sum = (Real)0;
+                k0 = i0 + mXDegreeP1 * j0;
+                for (s = 0; s < numSamples; ++s)
+                {
+                    Real w = observations[indices[s]][2];
+                    sum += w * xPower[s][i0] * yPower[s][j0];
+                }
+
+                B[k0] = sum;
+
+                for (j1 = 0; j1 <= mYDegree; ++j1)
+                {
+                    for (i1 = 0; i1 <= mXDegree; ++i1)
+                    {
+                        sum = (Real)0;
+                        k1 = i1 + mXDegreeP1 * j1;
+                        for (s = 0; s < numSamples; ++s)
+                        {
+                            sum += xPower[s][i0 + i1] * yPower[s][j0 + j1];
+                        }
+
+                        A(k0, k1) = sum;
+                    }
+                }
+            }
+        }
+
+        // Solve for the polynomial coefficients.
+        GVector<Real> coefficients = Inverse(A) * B;
+        bool hasNonzero = false;
+        for (int i = 0; i < mSize; ++i)
+        {
+            mParameters[i] = coefficients[i];
+            if (coefficients[i] != (Real)0)
+            {
+                hasNonzero = true;
+            }
+        }
+        return hasNonzero;
+    }
+
+    std::fill(mParameters.begin(), mParameters.end(), (Real)0);
+    return false;
+}
+
+template <typename Real>
+Real ApprPolynomial3<Real>::Error(std::array<Real, 3> const& observation)
+const
+{
+    Real w = Evaluate(observation[0], observation[1]);
+    Real error = std::abs(w - observation[2]);
+    return error;
+}
+
+template <typename Real>
+std::vector<Real> const& ApprPolynomial3<Real>::GetParameters() const
+{
+    return mParameters;
+}
+
+template <typename Real>
+std::array<Real, 2> const& ApprPolynomial3<Real>::GetXDomain() const
+{
+    return mXDomain;
+}
+
+template <typename Real>
+std::array<Real, 2> const& ApprPolynomial3<Real>::GetYDomain() const
+{
+    return mYDomain;
+}
+
+template <typename Real>
+Real ApprPolynomial3<Real>::Evaluate(Real x, Real y) const
+{
+    int i0, i1;
+    Real w;
+
+    for (i1 = 0; i1 <= mYDegree; ++i1)
+    {
+        i0 = mXDegree;
+        w = mParameters[i0 + mXDegreeP1 * i1];
+        while (--i0 >= 0)
+        {
+            w = mParameters[i0 + mXDegreeP1 * i1] + w * x;
+        }
+        mYCoefficient[i1] = w;
+    }
+
+    i1 = mYDegree;
+    w = mYCoefficient[i1];
+    while (--i1 >= 0)
+    {
+        w = mYCoefficient[i1] + w * y;
+    }
+
+    return w;
+}
+
+
+}
diff --git a/Engine/Plugins/Experimental/GeometryProcessing/Source/GeometryAlgorithms/Private/GTEngine/Mathematics/GteApprPolynomial4.h b/Engine/Plugins/Experimental/GeometryProcessing/Source/GeometryAlgorithms/Private/GTEngine/Mathematics/GteApprPolynomial4.h
new file mode 100644
--- /dev/null
+++ b/Engine/Plugins/Experimental/GeometryProcessing/Source/GeometryAlgorithms/Private/GTEngine/Mathematics/GteApprPolynomial4.h
@@ -0,0 +1,309 @@
+// David Eberly, Geometric Tools, Redmond WA 98052
+// Copyright (c) 1998-2019
+// Distributed under the Boost Software License, Version 1.0.
+// http://www.boost.org/LICENSE_1_0.txt
+// http://www.geometrictools.com/License/Boost/LICENSE_1_0.txt
+// File Version: 3.0.0 (2016/06/19)
+
+#pragma once
+
+#include <GTEngine/LowLevel/GteArray2.h>
+#include <GTEngine/Mathematics/GteGMatrix.h>
+#include <GTEngine/Mathematics/GteApprQuery.h>
+#include <array>
+
+// The samples are (x[i],y[i],z[i],w[i]) for 0 <= i < S.  Think of w as a
+// function of x, y, and z, say w = f(x,y,z).  The function fits the samples
+// with a polynomial of degree d0 in x, degree d1 in y, and degree d2 in z,
+// say
+//   w = sum_{i=0}^{d0} sum_{j=0}^{d1} sum_{k=0}^{d2} c[i][j][k]*x^i*y^j*z^k
+// The method is a least-squares fitting algorithm.  The mParameters stores
+// c[i][j][k] = mParameters[i+(d0+1)*(j+(d1+1)*k)] for a total of
+// (d0+1)*(d1+1)*(d2+1) coefficients.  The observation type is
+// std::array<Real,4>, which represents a tuple (x,y,z,w).
+//
+// WARNING.  The fitting algorithm for polynomial terms
+//   (1,x,x^2,...,x^d0), (1,y,y^2,...,y^d1), (1,z,z^2,...,z^d2)
+// is known to be nonrobust for large degrees and for large magnitude data.
+// One alternative is to use orthogonal polynomials
+//   (f[0](x),...,f[d0](x)), (g[0](y),...,g[d1](y)), (h[0](z),...,h[d2](z))
+// and apply the least-squares algorithm to these.  Another alternative is to
+// transform
+//   (x',y',z',w') = ((x-xcen)/rng, (y-ycen)/rng, (z-zcen)/rng, w/rng)
+// where xmin = min(x[i]), xmax = max(x[i]), xcen = (xmin+xmax)/2,
+// ymin = min(y[i]), ymax = max(y[i]), ycen = (ymin+ymax)/2, zmin = min(z[i]),
+// zmax = max(z[i]), zcen = (zmin+zmax)/2, and
+// rng = max(xmax-xmin,ymax-ymin,zmax-zmin).  Fit the (x',y',z',w') points,
+//   w' = sum_{i=0}^{d0} sum_{j=0}^{d1} sum_{k=0}^{d2} c'[i][j][k] *
+//          (x')^i*(y')^j*(z')^k
+// The original polynomial is evaluated as
+//   w = rng * sum_{i=0}^{d0} sum_{j=0}^{d1} sum_{k=0}^{d2} c'[i][j][k] *
+//         ((x-xcen)/rng)^i * ((y-ycen)/rng)^j * ((z-zcen)/rng)^k
+
+namespace gte
+{
+
+template <typename Real>
+class ApprPolynomial4
+    :
+    public ApprQuery<Real, ApprPolynomial4<Real>, std::array<Real, 4>>
+{
+public:
+    // Initialize the model parameters to zero.
+    ApprPolynomial4(int xDegree, int yDegree, int zDegree);
+
+    // The minimum number of observations required to fit the model.
+    int GetMinimumRequired() const;
+
+    // Estimate the model parameters for all observations specified by the
+    // indices.  This function is called by the base-class Fit(...) functions.
+    bool Fit(std::vector<std::array<Real, 4>> const& observations,
+        std::vector<int> const& indices);
+
+    // Compute the model error for the specified observation for the current
+    // model parameters.  The returned value for observation (x0,y0,z0,w0) is
+    // |w(x0,y0,z0) - w0|, where w(x,y,z) is the fitted polynomial.
+    Real Error(std::array<Real, 4> const& observation) const;
+
+    // Get the parameters of the model.
+    std::vector<Real> const& GetParameters() const;
+
+    // Evaluate the polynomial.  The domain intervals are provided so you can
+    // interpolate ((x,y,z) in domain) or extrapolate ((x,y,z) not in domain).
+    std::array<Real, 2> const& GetXDomain() const;
+    std::array<Real, 2> const& GetYDomain() const;
+    std::array<Real, 2> const& GetZDomain() const;
+    Real Evaluate(Real x, Real y, Real z) const;
+
+private:
+    int mXDegree, mYDegree, mZDegree;
+    int mXDegreeP1, mYDegreeP1, mZDegreeP1, mSize;
+    std::array<Real, 2> mXDomain, mYDomain, mZDomain;
+    std::vector<Real> mParameters;
+
+    // These arrays are used by Evaluate() to avoid reallocation of the
+    // 'vector's for each call.  The member is mutable because, to the
+    // user, the call to Evaluate does not modify the polynomial.
+    mutable std::vector<Real> mYZCoefficient;
+    mutable std::vector<Real> mZCoefficient;
+};
+
+
+template <typename Real>
+ApprPolynomial4<Real>::ApprPolynomial4(int xDegree, int yDegree, int zDegree)
+    :
+    mXDegree(xDegree),
+    mYDegree(yDegree),
+    mZDegree(zDegree),
+    mXDegreeP1(xDegree + 1),
+    mYDegreeP1(yDegree + 1),
+    mZDegreeP1(zDegree + 1),
+    mSize(mXDegreeP1 * mYDegreeP1 * mZDegreeP1),
+    mParameters(mSize),
+    mYZCoefficient(mYDegreeP1 * mZDegreeP1),
+    mZCoefficient(mZDegreeP1)
+{
+    mXDomain[0] = std::numeric_limits<Real>::max();
+    mXDomain[1] = -mXDomain[0];
+    mYDomain[0] = std::numeric_limits<Real>::max();
+    mYDomain[1] = -mYDomain[0];
+    mZDomain[0] = std::numeric_limits<Real>::max();
+    mZDomain[1] = -mZDomain[0];
+    std::fill(mParameters.begin(), mParameters.end(), (Real)0);
+    std::fill(mYZCoefficient.begin(), mYZCoefficient.end(), (Real)0);
+    std::fill(mZCoefficient.begin(), mZCoefficient.end(), (Real)0);
+}
+
+template <typename Real>
+int ApprPolynomial4<Real>::GetMinimumRequired() const
+{
+    return mSize;
+}
+
+template <typename Real>
+bool ApprPolynomial4<Real>::Fit(
+    std::vector<std::array<Real, 4>> const& observations,
+    std::vector<int> const& indices)
+{
+    if (indices.size() > 0)
+    {
+        int s, i0, j0, k0, n0, i1, j1, k1, n1;
+
+        // Compute the powers of x, y, and z.
+        int numSamples = static_cast<int>(indices.size());
+        int twoXDegree = 2 * mXDegree;
+        int twoYDegree = 2 * mYDegree;
+        int twoZDegree = 2 * mZDegree;
+        Array2<Real> xPower(twoXDegree + 1, numSamples);
+        Array2<Real> yPower(twoYDegree + 1, numSamples);
+        Array2<Real> zPower(twoZDegree + 1, numSamples);
+        for (s = 0; s < numSamples; ++s)
+        {
+            Real x = observations[indices[s]][0];
+            Real y = observations[indices[s]][1];
+            Real z = observations[indices[s]][2];
+            mXDomain[0] = std::min(x, mXDomain[0]);
+            mXDomain[1] = std::max(x, mXDomain[1]);
+            mYDomain[0] = std::min(y, mYDomain[0]);
+            mYDomain[1] = std::max(y, mYDomain[1]);
+            mZDomain[0] = std::min(z, mZDomain[0]);
+            mZDomain[1] = std::max(z, mZDomain[1]);
+
+            xPower[s][0] = (Real)1;
+            for (i0 = 1; i0 <= twoXDegree; ++i0)
+            {
+                xPower[s][i0] = x * xPower[s][i0 - 1];
+            }
+
+            yPower[s][0] = (Real)1;
+            for (j0 = 1; j0 <= twoYDegree; ++j0)
+            {
+                yPower[s][j0] = y * yPower[s][j0 - 1];
+            }
+
+            zPower[s][0] = (Real)1;
+            for (k0 = 1; k0 <= twoZDegree; ++k0)
+            {
+                zPower[s][k0] = z * zPower[s][k0 - 1];
+            }
+        }
+
+        // Matrix A is the Vandermonde matrix and vector B is the right-hand
+        // side of the linear system A*X = B.
+        GMatrix<Real> A(mSize, mSize);
+        GVector<Real> B(mSize);
+        for (k0 = 0; k0 <= mZDegree; ++k0)
+        {
+            for (j0 = 0; j0 <= mYDegree; ++j0)
+            {
+                for (i0 = 0; i0 <= mXDegree; ++i0)
+                {
+                    Real sum = (Real)0;
+                    n0 = i0 + mXDegreeP1*(j0 + mYDegreeP1*k0);
+                    for (s = 0; s < numSamples; ++s)
+                    {
+                        Real w = observations[indices[s]][3];
+                        sum += w * xPower[s][i0] * yPower[s][j0] *
+                            zPower[s][k0];
+                    }
+
+                    B[n0] = sum;
+
+                    for (k1 = 0; k1 <= mZDegree; ++k1)
+                    {
+                        for (j1 = 0; j1 <= mYDegree; ++j1)
+                        {
+                            for (i1 = 0; i1 <= mXDegree; ++i1)
+                            {
+                                sum = (Real)0;
+                                n1 = i1 + mXDegreeP1*(j1 + mYDegreeP1*k1);
+                                for (s = 0; s < numSamples; ++s)
+                                {
+                                    sum += xPower[s][i0 + i1] *
+                                        yPower[s][j0 + j1] *
+                                        zPower[s][k0 + k1];
+                                }
+
+                                A(n0, n1) = sum;
+                            }
+                        }
+                    }
+                }
+            }
+        }
+
+        // Solve for the polynomial coefficients.
+        GVector<Real> coefficients = Inverse(A) * B;
+        bool hasNonzero = false;
+        for (int i = 0; i < mSize; ++i)
+        {
+            mParameters[i] = coefficients[i];
+            if (coefficients[i] != (Real)0)
+            {
+                hasNonzero = true;
+            }
+        }
+        return hasNonzero;
+    }
+
+    std::fill(mParameters.begin(), mParameters.end(), (Real)0);
+    return false;
+}
+
+template <typename Real>
+Real ApprPolynomial4<Real>::Error(std::array<Real, 4> const& observation)
+const
+{
+    Real w = Evaluate(observation[0], observation[1], observation[2]);
+    Real error = std::abs(w - observation[3]);
+    return error;
+}
+
+template <typename Real>
+std::vector<Real> const& ApprPolynomial4<Real>::GetParameters() const
+{
+    return mParameters;
+}
+
+template <typename Real>
+std::array<Real, 2> const& ApprPolynomial4<Real>::GetXDomain() const
+{
+    return mXDomain;
+}
+
+template <typename Real>
+std::array<Real, 2> const& ApprPolynomial4<Real>::GetYDomain() const
+{
+    return mYDomain;
+}
+
+template <typename Real>
+std::array<Real, 2> const& ApprPolynomial4<Real>::GetZDomain() const
+{
+    return mZDomain;
+}
+
+template <typename Real>
+Real ApprPolynomial4<Real>::Evaluate(Real x, Real y, Real z) const
+{
+    int i0, i1, i2;
+    Real w;
+
+    for (i2 = 0; i2 <= mZDegree; ++i2)
+    {
+        for (i1 = 0; i1 <= mYDegree; ++i1)
+        {
+            i0 = mXDegree;
+            w = mParameters[i0 + mXDegreeP1 * (i1 + mYDegreeP1 * i2)];
+            while (--i0 >= 0)
+            {
+                w = mParameters[i0 + mXDegreeP1 * (i1 + mYDegreeP1 * i2)]
+                    + w * x;
+            }
+            mYZCoefficient[i1 + mYDegree * i2];
+        }
+    }
+
+    for (i2 = 0; i2 <= mZDegree; ++i2)
+    {
+        i1 = mYDegree;
+        w = mYZCoefficient[i1 + mYDegreeP1 * i2];
+        while (--i1 >= 0)
+        {
+            w = mParameters[i1 + mYDegreeP1 * i2] + w * y;
+        }
+        mZCoefficient[i2] = w;
+    }
+
+    i2 = mZDegree;
+    w = mZCoefficient[i2];
+    while (--i2 >= 0)
+    {
+        w = mZCoefficient[i2] + w * z;
+    }
+
+    return w;
+}
+
+
+}
diff --git a/Engine/Plugins/Experimental/GeometryProcessing/Source/GeometryAlgorithms/Private/GTEngine/Mathematics/GteApprPolynomialSpecial2.h b/Engine/Plugins/Experimental/GeometryProcessing/Source/GeometryAlgorithms/Private/GTEngine/Mathematics/GteApprPolynomialSpecial2.h
new file mode 100644
--- /dev/null
+++ b/Engine/Plugins/Experimental/GeometryProcessing/Source/GeometryAlgorithms/Private/GTEngine/Mathematics/GteApprPolynomialSpecial2.h
@@ -0,0 +1,305 @@
+// David Eberly, Geometric Tools, Redmond WA 98052
+// Copyright (c) 1998-2019
+// Distributed under the Boost Software License, Version 1.0.
+// http://www.boost.org/LICENSE_1_0.txt
+// http://www.geometrictools.com/License/Boost/LICENSE_1_0.txt
+// File Version: 3.0.0 (2016/06/19)
+
+#pragma once
+
+#include <GTEngine/LowLevel/GteLogger.h>
+#include <GTEngine/Mathematics/GteGMatrix.h>
+#include <GTEngine/Mathematics/GteApprQuery.h>
+#include <array>
+
+// Fit the data with a polynomial of the form
+//     w = sum_{i=0}^{n-1} c[i]*x^{p[i]}
+// where p[i] are distinct nonnegative powers provided by the caller.  A
+// least-squares fitting algorithm is used, but the input data is first
+// mapped to (x,w) in [-1,1]^2 for numerical robustness.
+
+namespace gte
+{
+
+template <typename Real>
+class ApprPolynomialSpecial2
+    :
+    public ApprQuery<Real, ApprPolynomialSpecial2<Real>, std::array<Real, 2>>
+{
+public:
+    // Initialize the model parameters to zero.  The degrees must be
+    // nonnegative and strictly increasing.
+    ApprPolynomialSpecial2(std::vector<int> const& degrees);
+
+    // The minimum number of observations required to fit the model.
+    int GetMinimumRequired() const;
+
+    // Estimate the model parameters for all observations specified by the
+    // indices.  This function is called by the base-class Fit(...) functions.
+    bool Fit(std::vector<std::array<Real, 2>> const& observations,
+        std::vector<int> const& indices);
+
+    // Compute the model error for the specified observation for the current
+    // model parameters.  The returned value for observation (x0,w0) is
+    // |w(x0) - w0|, where w(x) is the fitted polynomial.
+    Real Error(std::array<Real, 2> const& observation) const;
+
+    // Get the parameters of the model.
+    std::vector<Real> const& GetParameters() const;
+
+    // Evaluate the polynomial.  The domain interval is provided so you can
+    // interpolate (x in domain) or extrapolate (x not in domain).
+    std::array<Real, 2> const& GetXDomain() const;
+    Real Evaluate(Real x) const;
+
+private:
+    // Transform the (x,w) values to (x',w') in [-1,1]^2.
+    void Transform(std::vector<std::array<Real, 2>> const& observations,
+        std::vector<int> const& indices,
+        std::vector<std::array<Real, 2>>& transformed);
+
+    // The least-squares fitting algorithm for the transformed data.
+    bool DoLeastSquares(std::vector<std::array<Real, 2>>& transformed);
+
+    std::vector<int> mDegrees;
+    std::vector<Real> mParameters;
+
+    // Support for evaluation.  The coefficients were generated for the
+    // samples mapped to [-1,1]^2.  The Evaluate() function must transform
+    // x to x' in [-1,1], compute w' in [-1,1], then transform w' to w.
+    std::array<Real, 2> mXDomain, mWDomain;
+    std::array<Real, 2> mScale;
+    Real mInvTwoWScale;
+
+    // This array is used by Evaluate() to avoid reallocation of the 'vector'
+    // for each call.  The member is mutable because, to the user, the call
+    // to Evaluate does not modify the polynomial.
+    mutable std::vector<Real> mXPowers;
+};
+
+
+template <typename Real>
+ApprPolynomialSpecial2<Real>::ApprPolynomialSpecial2(
+    std::vector<int> const& degrees)
+    :
+    mDegrees(degrees),
+    mParameters(degrees.size())
+{
+#if !defined(GTE_NO_LOGGER)
+    LogAssert(mDegrees.size() > 0, "The input array must have elements.");
+    int lastDegree = -1;
+    for (auto degree : mDegrees)
+    {
+        LogAssert(degree > lastDegree, "Degrees must be increasing.");
+        lastDegree = degree;
+    }
+#endif
+
+    mXDomain[0] = std::numeric_limits<Real>::max();
+    mXDomain[1] = -mXDomain[0];
+    mWDomain[0] = std::numeric_limits<Real>::max();
+    mWDomain[1] = -mWDomain[0];
+    std::fill(mParameters.begin(), mParameters.end(), (Real)0);
+
+    mScale[0] = (Real)0;
+    mScale[1] = (Real)0;
+    mInvTwoWScale = (Real)0;
+
+    // Powers of x are computed up to twice the powers when constructing the
+    // fitted polynomial.  Powers of x are computed up to the powers for the
+    // evaluation of the fitted polynomial.
+    mXPowers.resize(2 * mDegrees.back() + 1);
+    mXPowers[0] = (Real)1;
+}
+
+template <typename Real>
+int ApprPolynomialSpecial2<Real>::GetMinimumRequired() const
+{
+    return static_cast<int>(mParameters.size());
+}
+
+template <typename Real>
+bool ApprPolynomialSpecial2<Real>::Fit(
+    std::vector<std::array<Real, 2>> const& observations,
+    std::vector<int> const& indices)
+{
+    if (indices.size() > 0)
+    {
+        // Transform the observations to [-1,1]^2 for numerical robustness.
+        std::vector<std::array<Real, 2>> transformed;
+        Transform(observations, indices, transformed);
+
+        // Fit the transformed data using a least-squares algorithm.
+        return DoLeastSquares(transformed);
+    }
+
+    std::fill(mParameters.begin(), mParameters.end(), (Real)0);
+    return false;
+}
+
+template <typename Real>
+Real ApprPolynomialSpecial2<Real>::Error(
+    std::array<Real, 2> const& observation) const
+{
+    Real w = Evaluate(observation[0]);
+    Real error = std::abs(w - observation[1]);
+    return error;
+}
+
+template <typename Real>
+std::vector<Real> const& ApprPolynomialSpecial2<Real>::GetParameters() const
+{
+    return mParameters;
+}
+
+template <typename Real>
+std::array<Real, 2> const& ApprPolynomialSpecial2<Real>::GetXDomain() const
+{
+    return mXDomain;
+}
+
+template <typename Real>
+Real ApprPolynomialSpecial2<Real>::Evaluate(Real x) const
+{
+    // Transform x to x' in [-1,1].
+    x = (Real)-1 + ((Real)2) * mScale[0] * (x - mXDomain[0]);
+
+    // Compute relevant powers of x.
+    int jmax = mDegrees.back();
+    for (int j = 1; j <= jmax; ++j)
+    {
+        mXPowers[j] = mXPowers[j - 1] * x;
+    }
+
+    Real w = (Real)0;
+    int isup = static_cast<int>(mDegrees.size());
+    for (int i = 0; i < isup; ++i)
+    {
+        Real xp = mXPowers[mDegrees[i]];
+        w += mParameters[i] * xp;
+    }
+
+    // Transform w from [-1,1] back to the original space.
+    w = (w + (Real)1) * mInvTwoWScale + mWDomain[0];
+    return w;
+}
+
+template <typename Real>
+void ApprPolynomialSpecial2<Real>::Transform(
+    std::vector<std::array<Real, 2>> const& observations,
+    std::vector<int> const& indices,
+    std::vector<std::array<Real, 2>>& transformed)
+{
+    int numSamples = static_cast<int>(indices.size());
+    transformed.resize(numSamples);
+
+    std::array<Real, 2> omin = observations[indices[0]];
+    std::array<Real, 2> omax = omin;
+    std::array<Real, 2> obs;
+    int s, i;
+    for (s = 1; s < numSamples; ++s)
+    {
+        obs = observations[indices[s]];
+        for (i = 0; i < 2; ++i)
+        {
+            if (obs[i] < omin[i])
+            {
+                omin[i] = obs[i];
+            }
+            else if (obs[i] > omax[i])
+            {
+                omax[i] = obs[i];
+            }
+        }
+    }
+
+    mXDomain[0] = omin[0];
+    mXDomain[1] = omax[0];
+    mWDomain[0] = omin[1];
+    mWDomain[1] = omax[1];
+    for (i = 0; i < 2; ++i)
+    {
+        mScale[i] = ((Real)1) / (omax[i] - omin[i]);
+    }
+
+    for (s = 0; s < numSamples; ++s)
+    {
+        obs = observations[indices[s]];
+        for (i = 0; i < 2; ++i)
+        {
+            transformed[s][i] = (Real)-1 + ((Real)2) * mScale[i] *
+                (obs[i] - omin[i]);
+        }
+    }
+    mInvTwoWScale = ((Real)0.5) / mScale[1];
+}
+
+template <typename Real>
+bool ApprPolynomialSpecial2<Real>::DoLeastSquares(
+    std::vector<std::array<Real, 2>>& transformed)
+{
+    // Set up a linear system A*X = B, where X are the polynomial
+    // coefficients.
+    int size = static_cast<int>(mDegrees.size());
+    GMatrix<Real> A(size, size);
+    A.MakeZero();
+    GVector<Real> B(size);
+    B.MakeZero();
+
+    int numSamples = static_cast<int>(transformed.size());
+    int twoMaxXDegree = 2 * mDegrees.back();
+    int row, col;
+    for (int i = 0; i < numSamples; ++i)
+    {
+        // Compute relevant powers of x.
+        Real x = transformed[i][0];
+        Real w = transformed[i][1];
+        for (int j = 0; j <= twoMaxXDegree; ++j)
+        {
+            mXPowers[j] = mXPowers[j - 1] * x;
+        }
+
+        for (row = 0; row < size; ++row)
+        {
+            // Update the upper-triangular portion of the symmetric matrix.
+            for (col = row; col < size; ++col)
+            {
+                A(row, col) += mXPowers[mDegrees[row] + mDegrees[col]];
+            }
+
+            // Update the right-hand side of the system.
+            B[row] += mXPowers[mDegrees[row]] * w;
+        }
+    }
+
+    // Copy the upper-triangular portion of the symmetric matrix to the
+    // lower-triangular portion.
+    for (row = 0; row < size; ++row)
+    {
+        for (col = 0; col < row; ++col)
+        {
+            A(row, col) = A(col, row);
+        }
+    }
+
+    // Precondition by normalizing the sums.
+    Real invNumSamples = ((Real)1) / (Real)numSamples;
+    A *= invNumSamples;
+    B *= invNumSamples;
+
+    // Solve for the polynomial coefficients.
+    GVector<Real> coefficients = Inverse(A) * B;
+    bool hasNonzero = false;
+    for (int i = 0; i < size; ++i)
+    {
+        mParameters[i] = coefficients[i];
+        if (coefficients[i] != (Real)0)
+        {
+            hasNonzero = true;
+        }
+    }
+    return hasNonzero;
+}
+
+
+}
diff --git a/Engine/Plugins/Experimental/GeometryProcessing/Source/GeometryAlgorithms/Private/GTEngine/Mathematics/GteApprPolynomialSpecial3.h b/Engine/Plugins/Experimental/GeometryProcessing/Source/GeometryAlgorithms/Private/GTEngine/Mathematics/GteApprPolynomialSpecial3.h
new file mode 100644
--- /dev/null
+++ b/Engine/Plugins/Experimental/GeometryProcessing/Source/GeometryAlgorithms/Private/GTEngine/Mathematics/GteApprPolynomialSpecial3.h
@@ -0,0 +1,352 @@
+// David Eberly, Geometric Tools, Redmond WA 98052
+// Copyright (c) 1998-2019
+// Distributed under the Boost Software License, Version 1.0.
+// http://www.boost.org/LICENSE_1_0.txt
+// http://www.geometrictools.com/License/Boost/LICENSE_1_0.txt
+// File Version: 3.0.0 (2016/06/19)
+
+#pragma once
+
+#include <GTEngine/LowLevel/GteLogger.h>
+#include <GTEngine/Mathematics/GteGMatrix.h>
+#include <GTEngine/Mathematics/GteApprQuery.h>
+#include <array>
+
+// Fit the data with a polynomial of the form
+//     w = sum_{i=0}^{n-1} c[i]*x^{p[i]}*y^{q[i]}
+// where <p[i],q[i]> are distinct pairs of nonnegative powers provided by the
+// caller.  A least-squares fitting algorithm is used, but the input data is
+// first mapped to (x,y,w) in [-1,1]^3 for numerical robustness.
+
+namespace gte
+{
+
+template <typename Real>
+class ApprPolynomialSpecial3
+    :
+    public ApprQuery<Real, ApprPolynomialSpecial3<Real>, std::array<Real, 3>>
+{
+public:
+    // Initialize the model parameters to zero.  The degrees must be
+    // nonnegative and strictly increasing.
+    ApprPolynomialSpecial3(std::vector<int> const& xDegrees,
+        std::vector<int> const& yDegrees);
+
+    // The minimum number of observations required to fit the model.
+    int GetMinimumRequired() const;
+
+    // Estimate the model parameters for all observations specified by the
+    // indices.  This function is called by the base-class Fit(...) functions.
+    bool Fit(std::vector<std::array<Real, 3>> const& observations,
+        std::vector<int> const& indices);
+
+    // Compute the model error for the specified observation for the current
+    // model parameters.  The returned value for observation (x0,w0) is
+    // |w(x0) - w0|, where w(x) is the fitted polynomial.
+    Real Error(std::array<Real, 3> const& observation) const;
+
+    // Get the parameters of the model.
+    std::vector<Real> const& GetParameters() const;
+
+    // Evaluate the polynomial.  The domain interval is provided so you can
+    // interpolate ((x,y) in domain) or extrapolate ((x,y) not in domain).
+    std::array<Real, 2> const& GetXDomain() const;
+    std::array<Real, 2> const& GetYDomain() const;
+    Real Evaluate(Real x, Real y) const;
+
+private:
+    // Transform the (x,y,w) values to (x',y',w') in [-1,1]^3.
+    void Transform(std::vector<std::array<Real, 3>> const& observations,
+        std::vector<int> const& indices,
+        std::vector<std::array<Real, 3>>& transformed);
+
+    // The least-squares fitting algorithm for the transformed data.
+    bool DoLeastSquares(std::vector<std::array<Real, 3>>& transformed);
+
+    std::vector<int> mXDegrees, mYDegrees;
+    std::vector<Real> mParameters;
+
+    // Support for evaluation.  The coefficients were generated for the
+    // samples mapped to [-1,1]^3.  The Evaluate() function must transform
+    // (x,y) to (x',y') in [-1,1]^2, compute w' in [-1,1], then transform w'
+    // to w.
+    std::array<Real, 2> mXDomain, mYDomain, mWDomain;
+    std::array<Real, 3> mScale;
+    Real mInvTwoWScale;
+
+    // This array is used by Evaluate() to avoid reallocation of the 'vector's
+    // for each call.  The members are mutable because, to the user, the call
+    // to Evaluate does not modify the polynomial.
+    mutable std::vector<Real> mXPowers, mYPowers;
+};
+
+
+template <typename Real>
+ApprPolynomialSpecial3<Real>::ApprPolynomialSpecial3(
+    std::vector<int> const& xDegrees, std::vector<int> const& yDegrees)
+    :
+    mXDegrees(xDegrees),
+    mYDegrees(yDegrees),
+    mParameters(mXDegrees.size() * mYDegrees.size())
+{
+#if !defined(GTE_NO_LOGGER)
+    LogAssert(mXDegrees.size() == mYDegrees.size(),
+        "The input arrays must have the same size.");
+
+    LogAssert(mXDegrees.size() > 0, "The input array must have elements.");
+    int lastDegree = -1;
+    for (auto degree : mXDegrees)
+    {
+        LogAssert(degree > lastDegree, "Degrees must be increasing.");
+        lastDegree = degree;
+    }
+
+    LogAssert(mYDegrees.size() > 0, "The input array must have elements.");
+    lastDegree = -1;
+    for (auto degree : mYDegrees)
+    {
+        LogAssert(degree > lastDegree, "Degrees must be increasing.");
+        lastDegree = degree;
+    }
+#endif
+
+    mXDomain[0] = std::numeric_limits<Real>::max();
+    mXDomain[1] = -mXDomain[0];
+    mYDomain[0] = std::numeric_limits<Real>::max();
+    mYDomain[1] = -mYDomain[0];
+    mWDomain[0] = std::numeric_limits<Real>::max();
+    mWDomain[1] = -mWDomain[0];
+    std::fill(mParameters.begin(), mParameters.end(), (Real)0);
+
+    mScale[0] = (Real)0;
+    mScale[1] = (Real)0;
+    mScale[2] = (Real)0;
+    mInvTwoWScale = (Real)0;
+
+    // Powers of x and y are computed up to twice the powers when constructing
+    // the fitted polynomial.  Powers of x and y are computed up to the powers
+    // for the evaluation of the fitted polynomial.
+    mXPowers.resize(2 * mXDegrees.back() + 1);
+    mXPowers[0] = (Real)1;
+    mYPowers.resize(2 * mYDegrees.back() + 1);
+    mYPowers[0] = (Real)1;
+}
+
+template <typename Real>
+int ApprPolynomialSpecial3<Real>::GetMinimumRequired() const
+{
+    return static_cast<int>(mParameters.size());
+}
+
+template <typename Real>
+bool ApprPolynomialSpecial3<Real>::Fit(
+    std::vector<std::array<Real, 3>> const& observations,
+    std::vector<int> const& indices)
+{
+    if (indices.size() > 0)
+    {
+        // Transform the observations to [-1,1]^3 for numerical robustness.
+        std::vector<std::array<Real, 3>> transformed;
+        Transform(observations, indices, transformed);
+
+        // Fit the transformed data using a least-squares algorithm.
+        return DoLeastSquares(transformed);
+    }
+
+    std::fill(mParameters.begin(), mParameters.end(), (Real)0);
+    return false;
+}
+
+template <typename Real>
+Real ApprPolynomialSpecial3<Real>::Error(
+    std::array<Real, 3> const& observation) const
+{
+    Real w = Evaluate(observation[0], observation[1]);
+    Real error = std::abs(w - observation[2]);
+    return error;
+}
+
+template <typename Real>
+std::vector<Real> const& ApprPolynomialSpecial3<Real>::GetParameters() const
+{
+    return mParameters;
+}
+
+template <typename Real>
+std::array<Real, 2> const& ApprPolynomialSpecial3<Real>::GetXDomain() const
+{
+    return mXDomain;
+}
+
+template <typename Real>
+std::array<Real, 2> const& ApprPolynomialSpecial3<Real>::GetYDomain() const
+{
+    return mYDomain;
+}
+
+template <typename Real>
+Real ApprPolynomialSpecial3<Real>::Evaluate(Real x, Real y) const
+{
+    // Transform (x,y) to (x',y') in [-1,1]^2.
+    x = (Real)-1 + ((Real)2) * mScale[0] * (x - mXDomain[0]);
+    y = (Real)-1 + ((Real)2) * mScale[1] * (y - mYDomain[0]);
+
+    // Compute relevant powers of x and y.
+    int jmax = mXDegrees.back();
+    for (int j = 1; j <= jmax; ++j)
+    {
+        mXPowers[j] = mXPowers[j - 1] * x;
+    }
+
+    jmax = mYDegrees.back();
+    for (int j = 1; j <= jmax; ++j)
+    {
+        mYPowers[j] = mYPowers[j - 1] * y;
+    }
+
+    Real w = (Real)0;
+    int isup = static_cast<int>(mXDegrees.size());
+    for (int i = 0; i < isup; ++i)
+    {
+        Real xp = mXPowers[mXDegrees[i]];
+        Real yp = mYPowers[mYDegrees[i]];
+        w += mParameters[i] * xp * yp;
+    }
+
+    // Transform w from [-1,1] back to the original space.
+    w = (w + (Real)1) * mInvTwoWScale + mWDomain[0];
+    return w;
+}
+
+template <typename Real>
+void ApprPolynomialSpecial3<Real>::Transform(
+    std::vector<std::array<Real, 3>> const& observations,
+    std::vector<int> const& indices,
+    std::vector<std::array<Real, 3>>& transformed)
+{
+    int numSamples = static_cast<int>(indices.size());
+    transformed.resize(numSamples);
+
+    std::array<Real, 3> omin = observations[indices[0]];
+    std::array<Real, 3> omax = omin;
+    std::array<Real, 3> obs;
+    int s, i;
+    for (s = 1; s < numSamples; ++s)
+    {
+        obs = observations[indices[s]];
+        for (i = 0; i < 3; ++i)
+        {
+            if (obs[i] < omin[i])
+            {
+                omin[i] = obs[i];
+            }
+            else if (obs[i] > omax[i])
+            {
+                omax[i] = obs[i];
+            }
+        }
+    }
+
+    mXDomain[0] = omin[0];
+    mXDomain[1] = omax[0];
+    mYDomain[0] = omin[1];
+    mYDomain[1] = omax[1];
+    mWDomain[0] = omin[2];
+    mWDomain[1] = omax[2];
+    for (i = 0; i < 3; ++i)
+    {
+        mScale[i] = ((Real)1) / (omax[i] - omin[i]);
+    }
+
+    for (s = 0; s < numSamples; ++s)
+    {
+        obs = observations[indices[s]];
+        for (i = 0; i < 3; ++i)
+        {
+            transformed[s][i] = (Real)-1 + ((Real)2) * mScale[i] *
+                (obs[i] - omin[i]);
+        }
+    }
+    mInvTwoWScale = ((Real)0.5) / mScale[2];
+}
+
+template <typename Real>
+bool ApprPolynomialSpecial3<Real>::DoLeastSquares(
+    std::vector<std::array<Real, 3>>& transformed)
+{
+    // Set up a linear system A*X = B, where X are the polynomial
+    // coefficients.
+    int size = static_cast<int>(mXDegrees.size());
+    GMatrix<Real> A(size, size);
+    A.MakeZero();
+    GVector<Real> B(size);
+    B.MakeZero();
+
+    int numSamples = static_cast<int>(transformed.size());
+    int twoMaxXDegree = 2 * mXDegrees.back();
+    int twoMaxYDegree = 2 * mYDegrees.back();
+    int row, col;
+    for (int i = 0; i < numSamples; ++i)
+    {
+        // Compute relevant powers of x and y.
+        Real x = transformed[i][0];
+        Real y = transformed[i][1];
+        Real w = transformed[i][2];
+        for (int j = 1; j <= 2 * twoMaxXDegree; ++j)
+        {
+            mXPowers[j] = mXPowers[j - 1] * x;
+        }
+        for (int j = 1; j <= 2 * twoMaxYDegree; ++j)
+        {
+            mYPowers[j] = mYPowers[j - 1] * y;
+        }
+
+        for (row = 0; row < size; ++row)
+        {
+            // Update the upper-triangular portion of the symmetric matrix.
+            Real xp, yp;
+            for (col = row; col < size; ++col)
+            {
+                xp = mXPowers[mXDegrees[row] + mXDegrees[col]];
+                yp = mYPowers[mYDegrees[row] + mYDegrees[col]];
+                A(row, col) += xp * yp;
+            }
+
+            // Update the right-hand side of the system.
+            xp = mXPowers[mXDegrees[row]];
+            yp = mYPowers[mYDegrees[row]];
+            B[row] += xp * yp * w;
+        }
+    }
+
+    // Copy the upper-triangular portion of the symmetric matrix to the
+    // lower-triangular portion.
+    for (row = 0; row < size; ++row)
+    {
+        for (col = 0; col < row; ++col)
+        {
+            A(row, col) = A(col, row);
+        }
+    }
+
+    // Precondition by normalizing the sums.
+    Real invNumSamples = ((Real)1) / (Real)numSamples;
+    A *= invNumSamples;
+    B *= invNumSamples;
+
+    // Solve for the polynomial coefficients.
+    GVector<Real> coefficients = Inverse(A) * B;
+    bool hasNonzero = false;
+    for (int i = 0; i < size; ++i)
+    {
+        mParameters[i] = coefficients[i];
+        if (coefficients[i] != (Real)0)
+        {
+            hasNonzero = true;
+        }
+    }
+    return hasNonzero;
+}
+
+
+}
diff --git a/Engine/Plugins/Experimental/GeometryProcessing/Source/GeometryAlgorithms/Private/GTEngine/Mathematics/GteApprPolynomialSpecial4.h b/Engine/Plugins/Experimental/GeometryProcessing/Source/GeometryAlgorithms/Private/GTEngine/Mathematics/GteApprPolynomialSpecial4.h
new file mode 100644
--- /dev/null
+++ b/Engine/Plugins/Experimental/GeometryProcessing/Source/GeometryAlgorithms/Private/GTEngine/Mathematics/GteApprPolynomialSpecial4.h
@@ -0,0 +1,393 @@
+// David Eberly, Geometric Tools, Redmond WA 98052
+// Copyright (c) 1998-2019
+// Distributed under the Boost Software License, Version 1.0.
+// http://www.boost.org/LICENSE_1_0.txt
+// http://www.geometrictools.com/License/Boost/LICENSE_1_0.txt
+// File Version: 3.0.0 (2016/06/19)
+
+#pragma once
+
+#include <GTEngine/LowLevel/GteLogger.h>
+#include <GTEngine/Mathematics/GteGMatrix.h>
+#include <GTEngine/Mathematics/GteApprQuery.h>
+#include <array>
+
+// Fit the data with a polynomial of the form
+//     w = sum_{i=0}^{n-1} c[i]*x^{p[i]}*y^{q[i]}*z^{r[i]}
+// where <p[i],q[i],r[i]> are distinct triples of nonnegative powers provided
+// by the caller.  A least-squares fitting algorithm is used, but the input
+// data is first mapped to (x,y,z,w) in [-1,1]^4 for numerical robustness.
+
+namespace gte
+{
+
+template <typename Real>
+class ApprPolynomialSpecial4
+    :
+    public ApprQuery<Real, ApprPolynomialSpecial4<Real>, std::array<Real, 3>>
+{
+public:
+    // Initialize the model parameters to zero.  The degrees must be
+    // nonnegative and strictly increasing.
+    ApprPolynomialSpecial4(std::vector<int> const& xDegrees,
+        std::vector<int> const& yDegrees, std::vector<int> const& zDegrees);
+
+    // The minimum number of observations required to fit the model.
+    int GetMinimumRequired() const;
+
+    // Estimate the model parameters for all observations specified by the
+    // indices.  This function is called by the base-class Fit(...) functions.
+    bool Fit(std::vector<std::array<Real, 4>> const& observations,
+        std::vector<int> const& indices);
+
+    // Compute the model error for the specified observation for the current
+    // model parameters.  The returned value for observation (x0,w0) is
+    // |w(x0) - w0|, where w(x) is the fitted polynomial.
+    Real Error(std::array<Real, 4> const& observation) const;
+
+    // Get the parameters of the model.
+    std::vector<Real> const& GetParameters() const;
+
+    // Evaluate the polynomial.  The domain interval is provided so you can
+    // interpolate ((x,y,z) in domain) or extrapolate ((x,y,z) not in domain).
+    std::array<Real, 2> const& GetXDomain() const;
+    std::array<Real, 2> const& GetYDomain() const;
+    std::array<Real, 2> const& GetZDomain() const;
+    Real Evaluate(Real x, Real y, Real z) const;
+
+private:
+    // Transform the (x,y,z,w) values to (x',y',z',w') in [-1,1]^4.
+    void Transform(std::vector<std::array<Real, 4>> const& observations,
+        std::vector<int> const& indices,
+        std::vector<std::array<Real, 4>>& transformed);
+
+    // The least-squares fitting algorithm for the transformed data.
+    bool DoLeastSquares(std::vector<std::array<Real, 4>>& transformed);
+
+    std::vector<int> mXDegrees, mYDegrees, mZDegrees;
+    std::vector<Real> mParameters;
+
+    // Support for evaluation.  The coefficients were generated for the
+    // samples mapped to [-1,1]^4.  The Evaluate() function must transform
+    // (x,y,z) to (x',y',z') in [-1,1]^3, compute w' in [-1,1], then transform
+    // w' to w.
+    std::array<Real, 2> mXDomain, mYDomain, mZDomain, mWDomain;
+    std::array<Real, 4> mScale;
+    Real mInvTwoWScale;
+
+    // This array is used by Evaluate() to avoid reallocation of the 'vector's
+    // for each call.  The members are mutable because, to the user, the call
+    // to Evaluate does not modify the polynomial.
+    mutable std::vector<Real> mXPowers, mYPowers, mZPowers;
+};
+
+
+template <typename Real>
+ApprPolynomialSpecial4<Real>::ApprPolynomialSpecial4(
+    std::vector<int> const& xDegrees, std::vector<int> const& yDegrees,
+    std::vector<int> const& zDegrees)
+    :
+    mXDegrees(xDegrees),
+    mYDegrees(yDegrees),
+    mZDegrees(zDegrees),
+    mParameters(mXDegrees.size() * mYDegrees.size() * mZDegrees.size())
+{
+#if !defined(GTE_NO_LOGGER)
+    LogAssert(mXDegrees.size() == mYDegrees.size()
+        && mXDegrees.size() == mZDegrees.size(),
+        "The input arrays must have the same size.");
+
+    LogAssert(mXDegrees.size() > 0, "The input array must have elements.");
+    int lastDegree = -1;
+    for (auto degree : mXDegrees)
+    {
+        LogAssert(degree > lastDegree, "Degrees must be increasing.");
+        lastDegree = degree;
+    }
+
+    LogAssert(mYDegrees.size() > 0, "The input array must have elements.");
+    lastDegree = -1;
+    for (auto degree : mYDegrees)
+    {
+        LogAssert(degree > lastDegree, "Degrees must be increasing.");
+        lastDegree = degree;
+    }
+
+    LogAssert(mZDegrees.size() > 0, "The input array must have elements.");
+    lastDegree = -1;
+    for (auto degree : mZDegrees)
+    {
+        LogAssert(degree > lastDegree, "Degrees must be increasing.");
+        lastDegree = degree;
+    }
+#endif
+
+    mXDomain[0] = std::numeric_limits<Real>::max();
+    mXDomain[1] = -mXDomain[0];
+    mYDomain[0] = std::numeric_limits<Real>::max();
+    mYDomain[1] = -mYDomain[0];
+    mZDomain[0] = std::numeric_limits<Real>::max();
+    mZDomain[1] = -mZDomain[0];
+    mWDomain[0] = std::numeric_limits<Real>::max();
+    mWDomain[1] = -mWDomain[0];
+    std::fill(mParameters.begin(), mParameters.end(), (Real)0);
+
+    mScale[0] = (Real)0;
+    mScale[1] = (Real)0;
+    mScale[2] = (Real)0;
+    mScale[3] = (Real)0;
+    mInvTwoWScale = (Real)0;
+
+    // Powers of x, y, and z are computed up to twice the powers when
+    // constructing the fitted polynomial.  Powers of x, y, and z are
+    // computed up to the powers for the evaluation of the fitted polynomial.
+    mXPowers.resize(2 * mXDegrees.back() + 1);
+    mXPowers[0] = (Real)1;
+    mYPowers.resize(2 * mYDegrees.back() + 1);
+    mYPowers[0] = (Real)1;
+    mZPowers.resize(2 * mZDegrees.back() + 1);
+    mZPowers[0] = (Real)1;
+}
+
+template <typename Real>
+int ApprPolynomialSpecial4<Real>::GetMinimumRequired() const
+{
+    return static_cast<int>(mParameters.size());
+}
+
+template <typename Real>
+bool ApprPolynomialSpecial4<Real>::Fit(
+    std::vector<std::array<Real, 4>> const& observations,
+    std::vector<int> const& indices)
+{
+    if (indices.size() > 0)
+    {
+        // Transform the observations to [-1,1]^4 for numerical robustness.
+        std::vector<std::array<Real, 4>> transformed;
+        Transform(observations, indices, transformed);
+
+        // Fit the transformed data using a least-squares algorithm.
+        return DoLeastSquares(transformed);
+    }
+
+    std::fill(mParameters.begin(), mParameters.end(), (Real)0);
+    return false;
+}
+
+template <typename Real>
+Real ApprPolynomialSpecial4<Real>::Error(
+    std::array<Real, 4> const& observation) const
+{
+    Real w = Evaluate(observation[0], observation[1], observation[2]);
+    Real error = std::abs(w - observation[3]);
+    return error;
+}
+
+template <typename Real>
+std::vector<Real> const& ApprPolynomialSpecial4<Real>::GetParameters() const
+{
+    return mParameters;
+}
+
+template <typename Real>
+std::array<Real, 2> const& ApprPolynomialSpecial4<Real>::GetXDomain() const
+{
+    return mXDomain;
+}
+
+template <typename Real>
+std::array<Real, 2> const& ApprPolynomialSpecial4<Real>::GetYDomain() const
+{
+    return mYDomain;
+}
+
+template <typename Real>
+std::array<Real, 2> const& ApprPolynomialSpecial4<Real>::GetZDomain() const
+{
+    return mZDomain;
+}
+
+template <typename Real>
+Real ApprPolynomialSpecial4<Real>::Evaluate(Real x, Real y, Real z) const
+{
+    // Transform (x,y,z) to (x',y',z') in [-1,1]^3.
+    x = (Real)-1 + ((Real)2) * mScale[0] * (x - mXDomain[0]);
+    y = (Real)-1 + ((Real)2) * mScale[1] * (y - mYDomain[0]);
+    z = (Real)-1 + ((Real)2) * mScale[2] * (z - mZDomain[0]);
+
+    // Compute relevant powers of x, y, and z.
+    int jmax = mXDegrees.back();;
+    for (int j = 1; j <= jmax; ++j)
+    {
+        mXPowers[j] = mXPowers[j - 1] * x;
+    }
+
+    jmax = mYDegrees.back();;
+    for (int j = 1; j <= jmax; ++j)
+    {
+        mYPowers[j] = mYPowers[j - 1] * y;
+    }
+
+    jmax = mZDegrees.back();;
+    for (int j = 1; j <= jmax; ++j)
+    {
+        mZPowers[j] = mZPowers[j - 1] * z;
+    }
+
+    Real w = (Real)0;
+    int isup = static_cast<int>(mXDegrees.size());
+    for (int i = 0; i < isup; ++i)
+    {
+        Real xp = mXPowers[mXDegrees[i]];
+        Real yp = mYPowers[mYDegrees[i]];
+        Real zp = mYPowers[mZDegrees[i]];
+        w += mParameters[i] * xp * yp * zp;
+    }
+
+    // Transform w from [-1,1] back to the original space.
+    w = (w + (Real)1) * mInvTwoWScale + mWDomain[0];
+    return w;
+}
+
+template <typename Real>
+void ApprPolynomialSpecial4<Real>::Transform(
+    std::vector<std::array<Real, 4>> const& observations,
+    std::vector<int> const& indices,
+    std::vector<std::array<Real, 4>>& transformed)
+{
+    int numSamples = static_cast<int>(indices.size());
+    transformed.resize(numSamples);
+
+    std::array<Real, 4> omin = observations[indices[0]];
+    std::array<Real, 4> omax = omin;
+    std::array<Real, 4> obs;
+    int s, i;
+    for (s = 1; s < numSamples; ++s)
+    {
+        obs = observations[indices[s]];
+        for (i = 0; i < 4; ++i)
+        {
+            if (obs[i] < omin[i])
+            {
+                omin[i] = obs[i];
+            }
+            else if (obs[i] > omax[i])
+            {
+                omax[i] = obs[i];
+            }
+        }
+    }
+
+    mXDomain[0] = omin[0];
+    mXDomain[1] = omax[0];
+    mYDomain[0] = omin[1];
+    mYDomain[1] = omax[1];
+    mZDomain[0] = omin[2];
+    mZDomain[1] = omax[2];
+    mWDomain[0] = omin[3];
+    mWDomain[1] = omax[3];
+    for (i = 0; i < 4; ++i)
+    {
+        mScale[i] = ((Real)1) / (omax[i] - omin[i]);
+    }
+
+    for (s = 0; s < numSamples; ++s)
+    {
+        obs = observations[indices[s]];
+        for (i = 0; i < 4; ++i)
+        {
+            transformed[s][i] = (Real)-1 + ((Real)2) * mScale[i] *
+                (obs[i] - omin[i]);
+        }
+    }
+    mInvTwoWScale = ((Real)0.5) / mScale[3];
+}
+
+template <typename Real>
+bool ApprPolynomialSpecial4<Real>::DoLeastSquares(
+    std::vector<std::array<Real, 4>>& transformed)
+{
+    // Set up a linear system A*X = B, where X are the polynomial
+    // coefficients.
+    int size = static_cast<int>(mXDegrees.size());
+    GMatrix<Real> A(size, size);
+    A.MakeZero();
+    GVector<Real> B(size);
+    B.MakeZero();
+
+    int numSamples = static_cast<int>(transformed.size());
+    int twoMaxXDegree = 2 * mXDegrees.back();
+    int twoMaxYDegree = 2 * mYDegrees.back();
+    int twoMaxZDegree = 2 * mZDegrees.back();
+    int row, col;
+    for (int i = 0; i < numSamples; ++i)
+    {
+        // Compute relevant powers of x, y, and z.
+        Real x = transformed[i][0];
+        Real y = transformed[i][1];
+        Real z = transformed[i][2];
+        Real w = transformed[i][3];
+        for (int j = 1; j <= twoMaxXDegree; ++j)
+        {
+            mXPowers[j] = mXPowers[j - 1] * x;
+        }
+        for (int j = 1; j <= twoMaxYDegree; ++j)
+        {
+            mYPowers[j] = mYPowers[j - 1] * y;
+        }
+        for (int j = 1; j <= twoMaxZDegree; ++j)
+        {
+            mZPowers[j] = mZPowers[j - 1] * z;
+        }
+
+        for (row = 0; row < size; ++row)
+        {
+            // Update the upper-triangular portion of the symmetric matrix.
+            Real xp, yp, zp;
+            for (col = row; col < size; ++col)
+            {
+                xp = mXPowers[mXDegrees[row] + mXDegrees[col]];
+                yp = mYPowers[mYDegrees[row] + mYDegrees[col]];
+                zp = mZPowers[mZDegrees[row] + mZDegrees[col]];
+                A(row, col) += xp * yp * zp;
+            }
+
+            // Update the right-hand side of the system.
+            xp = mXPowers[mXDegrees[row]];
+            yp = mYPowers[mYDegrees[row]];
+            zp = mZPowers[mZDegrees[row]];
+            B[row] += xp * yp * zp * w;
+        }
+    }
+
+    // Copy the upper-triangular portion of the symmetric matrix to the
+    // lower-triangular portion.
+    for (row = 0; row < size; ++row)
+    {
+        for (col = 0; col < row; ++col)
+        {
+            A(row, col) = A(col, row);
+        }
+    }
+
+    // Precondition by normalizing the sums.
+    Real invNumSamples = ((Real)1) / (Real)numSamples;
+    A *= invNumSamples;
+    B *= invNumSamples;
+
+    // Solve for the polynomial coefficients.
+    GVector<Real> coefficients = Inverse(A) * B;
+    bool hasNonzero = false;
+    for (int i = 0; i < size; ++i)
+    {
+        mParameters[i] = coefficients[i];
+        if (coefficients[i] != (Real)0)
+        {
+            hasNonzero = true;
+        }
+    }
+    return hasNonzero;
+}
+
+
+}
diff --git a/Engine/Plugins/Experimental/GeometryProcessing/Source/GeometryAlgorithms/Private/GTEngine/Mathematics/GteApprQuadratic2.h b/Engine/Plugins/Experimental/GeometryProcessing/Source/GeometryAlgorithms/Private/GTEngine/Mathematics/GteApprQuadratic2.h
new file mode 100644
--- /dev/null
+++ b/Engine/Plugins/Experimental/GeometryProcessing/Source/GeometryAlgorithms/Private/GTEngine/Mathematics/GteApprQuadratic2.h
@@ -0,0 +1,218 @@
+// David Eberly, Geometric Tools, Redmond WA 98052
+// Copyright (c) 1998-2019
+// Distributed under the Boost Software License, Version 1.0.
+// http://www.boost.org/LICENSE_1_0.txt
+// http://www.geometrictools.com/License/Boost/LICENSE_1_0.txt
+// File Version: 3.0.1 (2018/10/05)
+
+#pragma once
+
+#include <GTEngine/Mathematics/GteMatrix.h>
+#include <GTEngine/Mathematics/GteVector2.h>
+#include <GTEngine/Mathematics/GteHypersphere.h>
+#include <GTEngine/Mathematics/GteSymmetricEigensolver.h>
+
+namespace gte
+{
+
+// The quadratic fit is
+//
+//   0 = C[0] + C[1]*X + C[2]*Y + C[3]*X^2 + C[4]*Y^2 + C[5]*X*Y
+//
+// subject to Length(C) = 1.  Minimize E(C) = C^t M C with Length(C) = 1
+// and M = (sum_i V_i)(sum_i V_i)^t where
+//
+//   V = (1, X, Y, X^2, Y^2, X*Y)
+//         
+// The minimum value is the smallest eigenvalue of M and C is a corresponding
+// unit length eigenvector.
+//
+// Input:
+//   n = number of points to fit
+//   p[0..n-1] = array of points to fit
+//
+// Output:
+//   c[0..5] = coefficients of quadratic fit (the eigenvector)
+//   return value of function is nonnegative and a measure of the fit
+//   (the minimum eigenvalue; 0 = exact fit, positive otherwise)
+
+// Canonical forms.  The quadratic equation can be factored into
+// P^T A P + B^T P + K = 0 where P = (X,Y,Z), K = C[0], B = (C[1],C[2],C[3]),
+// and A is a 3x3 symmetric matrix with A00 = C[4], A11 = C[5], A22 = C[6],
+// A01 = C[7]/2, A02 = C[8]/2, and A12 = C[9]/2.  Matrix A = R^T D R where
+// R is orthogonal and D is diagonal (using an eigendecomposition).  Define
+// V = R P = (v0,v1,v2), E = R B = (e0,e1,e2), D = diag(d0,d1,d2), and f = K
+// to obtain
+//
+//   d0 v0^2 + d1 v1^2 + d2 v^2 + e0 v0 + e1 v1 + e2 v2 + f = 0
+//
+// The characterization depends on the signs of the d_i.
+
+template <typename Real>
+class ApprQuadratic2
+{
+public:
+    Real operator()(int numPoints, Vector2<Real> const* points,
+        Real coefficients[6]);
+};
+
+
+// If you think your points are nearly circular, use this.  The circle is of
+// the form C'[0]+C'[1]*X+C'[2]*Y+C'[3]*(X^2+Y^2), where Length(C') = 1.  The
+// function returns C = (C'[0]/C'[3],C'[1]/C'[3],C'[2]/C'[3]), so the fitted
+// circle is C[0]+C[1]*X+C[2]*Y+X^2+Y^2.  The center is (xc,yc) =
+// -0.5*(C[1],C[2]) and the radius is r = sqrt(xc*xc+yc*yc-C[0]).
+
+template <typename Real>
+class ApprQuadraticCircle2
+{
+public:
+    Real operator()(int numPoints, Vector2<Real> const* points,
+        Circle2<Real>& circle);
+};
+
+
+template <typename Real>
+Real ApprQuadratic2<Real>::operator()(int numPoints,
+    Vector2<Real> const* points, Real coefficients[6])
+{
+    Matrix<6, 6, Real> A;
+    for (int i = 0; i < numPoints; ++i)
+    {
+        Real x = points[i][0];
+        Real y = points[i][1];
+        Real x2 = x*x;
+        Real y2 = y*y;
+        Real xy = x*y;
+        Real x3 = x*x2;
+        Real xy2 = x*y2;
+        Real x2y = x*xy;
+        Real y3 = y*y2;
+        Real x4 = x*x3;
+        Real x2y2 = x*xy2;
+        Real x3y = x*x2y;
+        Real y4 = y*y3;
+        Real xy3 = x*y3;
+
+        A(0, 1) += x;
+        A(0, 2) += y;
+        A(0, 3) += x2;
+        A(0, 4) += y2;
+        A(0, 5) += xy;
+        A(1, 3) += x3;
+        A(1, 4) += xy2;
+        A(1, 5) += x2y;
+        A(2, 4) += y3;
+        A(3, 3) += x4;
+        A(3, 4) += x2y2;
+        A(3, 5) += x3y;
+        A(4, 4) += y4;
+        A(4, 5) += xy3;
+    }
+
+    A(0, 0) = static_cast<Real>(numPoints);
+    A(1, 1) = A(0, 3);
+    A(1, 2) = A(0, 5);
+    A(2, 2) = A(0, 4);
+    A(2, 3) = A(1, 5);
+    A(2, 5) = A(1, 4);
+    A(5, 5) = A(3, 4);
+
+    for (int row = 0; row < 6; ++row)
+    {
+        for (int col = 0; col < row; ++col)
+        {
+            A(row, col) = A(col, row);
+        }
+    }
+
+    Real invNumPoints = ((Real)1) / static_cast<Real>(numPoints);
+    for (int row = 0; row < 6; ++row)
+    {
+        for (int col = 0; col < 6; ++col)
+        {
+            A(row, col) *= invNumPoints;
+        }
+    }
+
+    SymmetricEigensolver<Real> es(6, 1024);
+    es.Solve(&A[0], +1);
+    es.GetEigenvector(0, &coefficients[0]);
+
+    // For an exact fit, numeric round-off errors might make the minimum
+    // eigenvalue just slightly negative.  Return the absolute value since
+    // the application might rely on the return value being nonnegative.
+    return std::abs(es.GetEigenvalue(0));
+}
+
+template <typename Real>
+Real ApprQuadraticCircle2<Real>::operator()(int numPoints,
+    Vector2<Real> const* points, Circle2<Real>& circle)
+{
+    Matrix<4, 4, Real> A;
+    for (int i = 0; i < numPoints; ++i)
+    {
+        Real x = points[i][0];
+        Real y = points[i][1];
+        Real x2 = x*x;
+        Real y2 = y*y;
+        Real xy = x*y;
+        Real r2 = x2 + y2;
+        Real xr2 = x*r2;
+        Real yr2 = y*r2;
+        Real r4 = r2*r2;
+
+        A(0, 1) += x;
+        A(0, 2) += y;
+        A(0, 3) += r2;
+        A(1, 1) += x2;
+        A(1, 2) += xy;
+        A(1, 3) += xr2;
+        A(2, 2) += y2;
+        A(2, 3) += yr2;
+        A(3, 3) += r4;
+    }
+
+    A(0, 0) = static_cast<Real>(numPoints);
+
+    for (int row = 0; row < 4; ++row)
+    {
+        for (int col = 0; col < row; ++col)
+        {
+            A(row, col) = A(col, row);
+        }
+    }
+
+    Real invNumPoints = ((Real)1) / static_cast<Real>(numPoints);
+    for (int row = 0; row < 4; ++row)
+    {
+        for (int col = 0; col < 4; ++col)
+        {
+            A(row, col) *= invNumPoints;
+        }
+    }
+
+    SymmetricEigensolver<Real> es(4, 1024);
+    es.Solve(&A[0], +1);
+    Vector<4, Real> evector;
+    es.GetEigenvector(0, &evector[0]);
+
+    Real inv = ((Real)1) / evector[3];  // TODO: Guard against zero divide?
+    Real coefficients[3];
+    for (int row = 0; row < 3; ++row)
+    {
+        coefficients[row] = inv * evector[row];
+    }
+
+    circle.center[0] = ((Real)-0.5) * coefficients[1];
+    circle.center[1] = ((Real)-0.5) * coefficients[2];
+    circle.radius = std::sqrt(std::abs(Dot(circle.center, circle.center) - coefficients[0]));
+
+    // For an exact fit, numeric round-off errors might make the minimum
+    // eigenvalue just slightly negative.  Return the absolute value since
+    // the application might rely on the return value being nonnegative.
+    return std::abs(es.GetEigenvalue(0));
+}
+
+
+}
diff --git a/Engine/Plugins/Experimental/GeometryProcessing/Source/GeometryAlgorithms/Private/GTEngine/Mathematics/GteApprQuadratic3.h b/Engine/Plugins/Experimental/GeometryProcessing/Source/GeometryAlgorithms/Private/GTEngine/Mathematics/GteApprQuadratic3.h
new file mode 100644
--- /dev/null
+++ b/Engine/Plugins/Experimental/GeometryProcessing/Source/GeometryAlgorithms/Private/GTEngine/Mathematics/GteApprQuadratic3.h
@@ -0,0 +1,285 @@
+// David Eberly, Geometric Tools, Redmond WA 98052
+// Copyright (c) 1998-2019
+// Distributed under the Boost Software License, Version 1.0.
+// http://www.boost.org/LICENSE_1_0.txt
+// http://www.geometrictools.com/License/Boost/LICENSE_1_0.txt
+// File Version: 3.0.1 (2018/10/05)
+
+#pragma once
+
+#include <GTEngine/Mathematics/GteMatrix.h>
+#include <GTEngine/Mathematics/GteVector3.h>
+#include <GTEngine/Mathematics/GteHypersphere.h>
+#include <GTEngine/Mathematics/GteSymmetricEigensolver.h>
+
+namespace gte
+{
+
+// The quadratic fit is
+//
+//   0 = C[0] + C[1]*X + C[2]*Y + C[3]*Z + C[4]*X^2 + C[5]*Y^2
+//       + C[6]*Z^2 + C[7]*X*Y + C[8]*X*Z + C[9]*Y*Z
+//
+// subject to Length(C) = 1.  Minimize E(C) = C^t M C with Length(C) = 1
+// and M = (sum_i V_i)(sum_i V_i)^t where
+//
+//   V = (1, X, Y, Z, X^2, Y^2, Z^2, X*Y, X*Z, Y*Z)
+//         
+// The minimum value is the smallest eigenvalue of M and C is a corresponding
+// unit length eigenvector.
+//
+// Input:
+//   n = number of points to fit
+//   p[0..n-1] = array of points to fit
+//
+// Output:
+//   c[0..9] = coefficients of quadratic fit (the eigenvector)
+//   return value of function is nonnegative and a measure of the fit
+//   (the minimum eigenvalue; 0 = exact fit, positive otherwise)
+
+// Canonical forms.  The quadratic equation can be factored into
+// P^T A P + B^T P + K = 0 where P = (X,Y,Z), K = C[0], B = (C[1],C[2],C[3]),
+// and A is a 3x3 symmetric matrix with A00 = C[4], A11 = C[5], A22 = C[6],
+// A01 = C[7]/2, A02 = C[8]/2, and A12 = C[9]/2.  Matrix A = R^T D R where
+// R is orthogonal and D is diagonal (using an eigendecomposition).  Define
+// V = R P = (v0,v1,v2), E = R B = (e0,e1,e2), D = diag(d0,d1,d2), and f = K
+// to obtain
+//
+//   d0 v0^2 + d1 v1^2 + d2 v^2 + e0 v0 + e1 v1 + e2 v2 + f = 0
+//
+// The characterization depends on the signs of the d_i.
+
+template <typename Real>
+class ApprQuadratic3
+{
+public:
+    Real operator()(int numPoints, Vector3<Real> const* points,
+        Real coefficidents[10]);
+};
+
+
+// If you think your points are nearly spherical, use this.  Sphere is of form
+// C'[0]+C'[1]*X+C'[2]*Y+C'[3]*Z+C'[4]*(X^2+Y^2+Z^2) where Length(C') = 1.
+// Function returns C = (C'[0]/C'[4],C'[1]/C'[4],C'[2]/C'[4],C'[3]/C'[4]), so
+// fitted sphere is C[0]+C[1]*X+C[2]*Y+C[3]*Z+X^2+Y^2+Z^2.  Center is
+// (xc,yc,zc) = -0.5*(C[1],C[2],C[3]) and radius is rad =
+// sqrt(xc*xc+yc*yc+zc*zc-C[0]).
+template <typename Real>
+class ApprQuadraticSphere3
+{
+public:
+    Real operator()(int numPoints, Vector3<Real> const* points,
+        Sphere3<Real>& sphere);
+};
+
+
+template <typename Real>
+Real ApprQuadratic3<Real>::operator()(int numPoints,
+    Vector3<Real> const* points, Real coefficients[10])
+{
+    Matrix<10, 10, Real> A;
+    for (int i = 0; i < numPoints; ++i)
+    {
+        Real x = points[i][0];
+        Real y = points[i][1];
+        Real z = points[i][2];
+        Real x2 = x*x;
+        Real y2 = y*y;
+        Real z2 = z*z;
+        Real xy = x*y;
+        Real xz = x*z;
+        Real yz = y*z;
+        Real x3 = x*x2;
+        Real xy2 = x*y2;
+        Real xz2 = x*z2;
+        Real x2y = x*xy;
+        Real x2z = x*xz;
+        Real xyz = x*y*z;
+        Real y3 = y*y2;
+        Real yz2 = y*z2;
+        Real y2z = y*yz;
+        Real z3 = z*z2;
+        Real x4 = x*x3;
+        Real x2y2 = x*xy2;
+        Real x2z2 = x*xz2;
+        Real x3y = x*x2y;
+        Real x3z = x*x2z;
+        Real x2yz = x*xyz;
+        Real y4 = y*y3;
+        Real y2z2 = y*yz2;
+        Real xy3 = x*y3;
+        Real xy2z = x*y2z;
+        Real y3z = y*y2z;
+        Real z4 = z*z3;
+        Real xyz2 = x*yz2;
+        Real xz3 = x*z3;
+        Real yz3 = y*z3;
+
+        A(0, 1) += x;
+        A(0, 2) += y;
+        A(0, 3) += z;
+        A(0, 4) += x2;
+        A(0, 5) += y2;
+        A(0, 6) += z2;
+        A(0, 7) += xy;
+        A(0, 8) += xz;
+        A(0, 9) += yz;
+        A(1, 4) += x3;
+        A(1, 5) += xy2;
+        A(1, 6) += xz2;
+        A(1, 7) += x2y;
+        A(1, 8) += x2z;
+        A(1, 9) += xyz;
+        A(2, 5) += y3;
+        A(2, 6) += yz2;
+        A(2, 9) += y2z;
+        A(3, 6) += z3;
+        A(4, 4) += x4;
+        A(4, 5) += x2y2;
+        A(4, 6) += x2z2;
+        A(4, 7) += x3y;
+        A(4, 8) += x3z;
+        A(4, 9) += x2yz;
+        A(5, 5) += y4;
+        A(5, 6) += y2z2;
+        A(5, 7) += xy3;
+        A(5, 8) += xy2z;
+        A(5, 9) += y3z;
+        A(6, 6) += z4;
+        A(6, 7) += xyz2;
+        A(6, 8) += xz3;
+        A(6, 9) += yz3;
+        A(9, 9) += y2z2;
+    }
+
+    A(0, 0) = static_cast<Real>(numPoints);
+    A(1, 1) = A(0, 4);
+    A(1, 2) = A(0, 7);
+    A(1, 3) = A(0, 8);
+    A(2, 2) = A(0, 5);
+    A(2, 3) = A(0, 9);
+    A(2, 4) = A(1, 7);
+    A(2, 7) = A(1, 5);
+    A(2, 8) = A(1, 9);
+    A(3, 3) = A(0, 6);
+    A(3, 4) = A(1, 8);
+    A(3, 5) = A(2, 9);
+    A(3, 7) = A(1, 9);
+    A(3, 8) = A(1, 6);
+    A(3, 9) = A(2, 6);
+    A(7, 7) = A(4, 5);
+    A(7, 8) = A(4, 9);
+    A(7, 9) = A(5, 8);
+    A(8, 8) = A(4, 6);
+    A(8, 9) = A(6, 7);
+    A(9, 9) = A(5, 6);
+
+    for (int row = 0; row < 10; ++row)
+    {
+        for (int col = 0; col < row; ++col)
+        {
+            A(row, col) = A(col, row);
+        }
+    }
+
+    Real invNumPoints = ((Real)1) / static_cast<Real>(numPoints);
+    for (int row = 0; row < 10; ++row)
+    {
+        for (int col = 0; col < 10; ++col)
+        {
+            A(row, col) *= invNumPoints;
+        }
+    }
+
+    SymmetricEigensolver<Real> es(10, 1024);
+    es.Solve(&A[0], +1);
+    es.GetEigenvector(0, &coefficients[0]);
+
+    // For an exact fit, numeric round-off errors might make the minimum
+    // eigenvalue just slightly negative.  Return the absolute value since
+    // the application might rely on the return value being nonnegative.
+    return std::abs(es.GetEigenvalue(0));
+}
+
+template <typename Real>
+Real ApprQuadraticSphere3<Real>::operator()(int numPoints,
+    Vector3<Real> const* points, Sphere3<Real>& sphere)
+{
+    Matrix<5, 5, Real> A;
+    for (int i = 0; i < numPoints; ++i)
+    {
+        Real x = points[i][0];
+        Real y = points[i][1];
+        Real z = points[i][2];
+        Real x2 = x*x;
+        Real y2 = y*y;
+        Real z2 = z*z;
+        Real xy = x*y;
+        Real xz = x*z;
+        Real yz = y*z;
+        Real r2 = x2 + y2 + z2;
+        Real xr2 = x*r2;
+        Real yr2 = y*r2;
+        Real zr2 = z*r2;
+        Real r4 = r2*r2;
+
+        A(0, 1) += x;
+        A(0, 2) += y;
+        A(0, 3) += z;
+        A(0, 4) += r2;
+        A(1, 1) += x2;
+        A(1, 2) += xy;
+        A(1, 3) += xz;
+        A(1, 4) += xr2;
+        A(2, 2) += y2;
+        A(2, 3) += yz;
+        A(2, 4) += yr2;
+        A(3, 3) += z2;
+        A(3, 4) += zr2;
+        A(4, 4) += r4;
+    }
+
+    A(0, 0) = static_cast<Real>(numPoints);
+
+    for (int row = 0; row < 5; ++row)
+    {
+        for (int col = 0; col < row; ++col)
+        {
+            A(row, col) = A(col, row);
+        }
+    }
+
+    Real invNumPoints = ((Real)1) / static_cast<Real>(numPoints);
+    for (int row = 0; row < 5; ++row)
+    {
+        for (int col = 0; col < 5; ++col)
+        {
+            A(row, col) *= invNumPoints;
+        }
+    }
+
+    SymmetricEigensolver<Real> es(5, 1024);
+    es.Solve(&A[0], +1);
+    Vector<5, Real> evector;
+    es.GetEigenvector(0, &evector[0]);
+
+    Real inv = ((Real)1) / evector[4];  // TODO: Guard against zero divide?
+    Real coefficients[4];
+    for (int row = 0; row < 4; ++row)
+    {
+        coefficients[row] = inv * evector[row];
+    }
+
+    sphere.center[0] = ((Real)-0.5) * coefficients[1];
+    sphere.center[1] = ((Real)-0.5) * coefficients[2];
+    sphere.center[2] = ((Real)-0.5) * coefficients[3];
+    sphere.radius = std::sqrt(std::abs(Dot(sphere.center, sphere.center) - coefficients[0]));
+
+    // For an exact fit, numeric round-off errors might make the minimum
+    // eigenvalue just slightly negative.  Return the absolute value since
+    // the application might rely on the return value being nonnegative.
+    return std::abs(es.GetEigenvalue(0));
+}
+
+
+}
diff --git a/Engine/Plugins/Experimental/GeometryProcessing/Source/GeometryAlgorithms/Private/GTEngine/Mathematics/GteApprQuery.h b/Engine/Plugins/Experimental/GeometryProcessing/Source/GeometryAlgorithms/Private/GTEngine/Mathematics/GteApprQuery.h
new file mode 100644
--- /dev/null
+++ b/Engine/Plugins/Experimental/GeometryProcessing/Source/GeometryAlgorithms/Private/GTEngine/Mathematics/GteApprQuery.h
@@ -0,0 +1,207 @@
+// David Eberly, Geometric Tools, Redmond WA 98052
+// Copyright (c) 1998-2019
+// Distributed under the Boost Software License, Version 1.0.
+// http://www.boost.org/LICENSE_1_0.txt
+// http://www.geometrictools.com/License/Boost/LICENSE_1_0.txt
+// File Version: 3.0.0 (2016/06/19)
+
+#pragma once
+
+#include <GTEngine/GTEngineDEF.h>
+#include <algorithm>
+#include <random>
+#include <vector>
+
+// Class ApprQuery supports the RANSAC algorithm for fitting and uses the
+// Curiously Recurring Template Pattern.  The ModelType must be a class or
+// struct with the following interfaces:
+//
+//   // The minimum number of observations required to fit the model.
+//   int ModelType::GetMinimumRequired() const;
+//
+//   // Compute the model error for the specified observation for the current
+//   // model parameters.
+//   Real Error(ObservationType const& observation) const;
+//
+//   // Estimate the model parameters for all observations specified by the
+//   // indices.  The three Fit() functions of ApprQuery manipulate their
+//   // inputs in order to pass them to ModelType::Fit().
+//   ModelType::Fit(std::vector<ObservationType> const& observations,
+//       std::vector<size_t> const& indices);
+
+namespace gte
+{
+
+template <typename Real, typename ModelType, typename ObservationType>
+class ApprQuery
+{
+public:
+    // Estimate the model parameters for all observations.
+    bool Fit(std::vector<ObservationType> const& observations);
+
+    // Estimate the model parameters for a contiguous subset of observations.
+    bool Fit(std::vector<ObservationType> const& observations,
+        int const imin, int const imax);
+
+    // Estimate the model parameters for the subset of observations specified
+    // by the indices and the number of indices that is possibly smaller than
+    // indices.size().
+    bool Fit(std::vector<ObservationType> const& observations,
+        std::vector<int> const& indices, int const numIndices);
+
+    // Apply the RANdom SAmple Consensus algorithm for fitting a model to
+    // observations.
+    static bool RANSAC(
+        ModelType& candidateModel,
+        std::vector<ObservationType> const& observations,
+        int const numRequiredForGoodFit, Real const maxErrorForGoodFit,
+        int const numIterations, std::vector<int>& bestConsensus,
+        ModelType& bestModel);
+};
+
+
+template <typename Real, typename ModelType, typename ObservationType>
+bool ApprQuery<Real, ModelType, ObservationType>::Fit(
+    std::vector<ObservationType> const& observations)
+{
+    std::vector<int> indices(observations.size());
+    int i = 0;
+    for (auto& index : indices)
+    {
+        index = i++;
+    }
+
+    return ((ModelType*)this)->Fit(observations, indices);
+}
+
+template <typename Real, typename ModelType, typename ObservationType>
+bool ApprQuery<Real, ModelType, ObservationType>::Fit(
+    std::vector<ObservationType> const& observations,
+    int const imin, int const imax)
+{
+    if (imin <= imax)
+    {
+        int numIndices = imax - imin + 1;
+        std::vector<int> indices(numIndices);
+        int i = imin;
+        for (auto& index : indices)
+        {
+            index = i++;
+        }
+
+        return ((ModelType*)this)->Fit(observations, indices);
+    }
+    else
+    {
+        return false;
+    }
+}
+
+template <typename Real, typename ModelType, typename ObservationType>
+bool ApprQuery<Real, ModelType, ObservationType>::Fit(
+    std::vector<ObservationType> const& observations,
+    std::vector<int> const& indices, int const numIndices)
+{
+    int imax = std::min(numIndices, static_cast<int>(observations.size()));
+    std::vector<int> localindices(imax);
+    int i = 0;
+    for (auto& index : localindices)
+    {
+        index = indices[i++];
+    }
+
+    return ((ModelType*)this)->Fit(observations, indices);
+}
+
+template <typename Real, typename ModelType, typename ObservationType>
+bool ApprQuery<Real, ModelType, ObservationType>::RANSAC(
+    ModelType& candidateModel,
+    std::vector<ObservationType> const& observations,
+    int const numRequiredForGoodFit, Real const maxErrorForGoodFit,
+    int const numIterations, std::vector<int>& bestConsensus,
+    ModelType& bestModel)
+{
+    int const numObservations = static_cast<int>(observations.size());
+    int const minRequired = candidateModel.GetMinimumRequired();
+    if (numObservations < minRequired)
+    {
+        // Too few observations for model fitting.
+        return false;
+    }
+
+    // The first part of the array will store the consensus set, initially
+    // filled with the minimumu number of indices that correspond to the
+    // candidate inliers.  The last part will store the remaining indices.
+    // These points are tested against the model and are added to the
+    // consensus set when they fit.  All the index manipulation is done
+    // in place.  Initially, the candidates are the identity permutation.
+    std::vector<int> candidates(numObservations);
+    int j = 0;
+    for (auto& c : candidates)
+    {
+        c = j++;
+    }
+
+    if (numObservations == minRequired)
+    {
+        // We have the minimum number of observations to generate the model,
+        // so RANSAC cannot be used.  Compute the model with the entire set
+        // of observations.
+        bestConsensus = candidates;
+        return bestModel.Fit(observations);
+    }
+
+    int bestNumFittedObservations = minRequired;
+
+    for (int i = 0; i < numIterations; ++i)
+    {
+        // Randomly permute the previous candidates, partitioning the array
+        // into GetMinimumRequired() indices (the candidate inliers) followed
+        // by the remaining indices (candidates for testing against the
+        // model).
+        std::shuffle(candidates.begin(), candidates.end(),
+            std::default_random_engine());
+
+        // Fit the model to the inliers.
+        if (candidateModel.Fit(observations, candidates, minRequired))
+        {
+            // Test each remaining observation whether it fits the model.  If
+            // it does, include it in the consensus set.
+            int numFittedObservations = minRequired;
+            for (j = minRequired; j < numObservations; ++j)
+            {
+                if (candidateModel.Error(observations[candidates[j]])
+                    <= maxErrorForGoodFit)
+                {
+                    std::swap(candidates[j],
+                        candidates[numFittedObservations]);
+                    ++numFittedObservations;
+                }
+            }
+
+            if (numFittedObservations >= numRequiredForGoodFit)
+            {
+                // We have observations that fit the model.  Update the best
+                // model using the consensus set.
+                candidateModel.Fit(observations, candidates,
+                    numFittedObservations);
+                if (numFittedObservations > bestNumFittedObservations)
+                {
+                    // The consensus set is larger than the previous consensus
+                    // set, so its model becomes the best one.
+                    bestModel = candidateModel;
+                    bestConsensus.resize(numFittedObservations);
+                    std::copy(candidates.begin(),
+                        candidates.begin() + numFittedObservations,
+                        bestConsensus.begin());
+                    bestNumFittedObservations = numFittedObservations;
+                }
+            }
+        }
+    }
+
+    return bestNumFittedObservations >= numRequiredForGoodFit;
+}
+
+
+}
diff --git a/Engine/Plugins/Experimental/GeometryProcessing/Source/GeometryAlgorithms/Private/GTEngine/Mathematics/GteApprSphere3.h b/Engine/Plugins/Experimental/GeometryProcessing/Source/GeometryAlgorithms/Private/GTEngine/Mathematics/GteApprSphere3.h
new file mode 100644
--- /dev/null
+++ b/Engine/Plugins/Experimental/GeometryProcessing/Source/GeometryAlgorithms/Private/GTEngine/Mathematics/GteApprSphere3.h
@@ -0,0 +1,101 @@
+// David Eberly, Geometric Tools, Redmond WA 98052
+// Copyright (c) 1998-2019
+// Distributed under the Boost Software License, Version 1.0.
+// http://www.boost.org/LICENSE_1_0.txt
+// http://www.geometrictools.com/License/Boost/LICENSE_1_0.txt
+// File Version: 3.0.0 (2016/06/19)
+
+#pragma once
+
+#include <GTEngine/Mathematics/GteApprQuadratic3.h>
+
+// Least-squares fit of a sphere to a set of points.  A successful fit is
+// indicated by the return value of 'true'.  The return value is the number
+// of iterations used.  If the number is the maximum, you can either accept
+// the result or try increasing the maximum number and calling the function
+// again.  TODO: Currently, the test for terminating the algorithm is exact
+// (diff == (0,0,0)).  Expose an epsilon for this?
+//
+// If initialCenterIsAverage is set to 'true', the initial guess for the
+// sphere center is the average of the data points.  If the data points are
+// clustered along a solid angle, ApprSphere32 is very slow to converge.  If
+// initialCenterIsAverage is set to 'false', the initial guess for the
+// sphere center is computed using a least-squares estimate of the
+// coefficients for a quadratic equation that represents a sphere.  This
+// approach tends to converge rapidly.
+
+namespace gte
+{
+
+template <typename Real>
+class ApprSphere3
+{
+public:
+    unsigned int operator()(int numPoints, Vector3<Real> const* points,
+        unsigned int maxIterations, bool initialCenterIsAverage,
+        Sphere3<Real>& sphere);
+};
+
+
+template <typename Real>
+unsigned int ApprSphere3<Real>::operator()(int numPoints,
+    Vector3<Real> const* points, unsigned int maxIterations,
+    bool initialCenterIsAverage, Sphere3<Real>& sphere)
+{
+    // Compute the average of the data points.
+    Vector3<Real> average = points[0];
+    for (int i = 1; i < numPoints; ++i)
+    {
+        average += points[i];
+    }
+    Real invNumPoints = ((Real)1) / static_cast<Real>(numPoints);
+    average *= invNumPoints;
+
+    // The initial guess for the center.
+    if (initialCenterIsAverage)
+    {
+        sphere.center = average;
+    }
+    else
+    {
+        ApprQuadraticSphere3<Real>()(numPoints, points, sphere);
+    }
+
+    unsigned int iteration;
+    for (iteration = 0; iteration < maxIterations; ++iteration)
+    {
+        // Update the iterates.
+        Vector3<Real> current = sphere.center;
+
+        // Compute average L, dL/da, dL/db, dL/dc.
+        Real lenAverage = (Real)0;
+        Vector3<Real> derLenAverage = Vector3<Real>::Zero();
+        for (int i = 0; i < numPoints; ++i)
+        {
+            Vector3<Real> diff = points[i] - sphere.center;
+            Real length = Length(diff);
+            if (length >(Real)0)
+            {
+                lenAverage += length;
+                Real invLength = ((Real)1) / length;
+                derLenAverage -= invLength * diff;
+            }
+        }
+        lenAverage *= invNumPoints;
+        derLenAverage *= invNumPoints;
+
+        sphere.center = average + lenAverage * derLenAverage;
+        sphere.radius = lenAverage;
+
+        Vector3<Real> diff = sphere.center - current;
+        if (diff == Vector3<Real>::Zero())
+        {
+            break;
+        }
+    }
+
+    return ++iteration;
+}
+
+
+}
diff --git a/Engine/Plugins/Experimental/GeometryProcessing/Source/GeometryAlgorithms/Private/GTEngine/Mathematics/GteApprTorus3.h b/Engine/Plugins/Experimental/GeometryProcessing/Source/GeometryAlgorithms/Private/GTEngine/Mathematics/GteApprTorus3.h
new file mode 100644
--- /dev/null
+++ b/Engine/Plugins/Experimental/GeometryProcessing/Source/GeometryAlgorithms/Private/GTEngine/Mathematics/GteApprTorus3.h
@@ -0,0 +1,397 @@
+// David Eberly, Geometric Tools, Redmond WA 98052
+// Copyright (c) 1998-2019
+// Distributed under the Boost Software License, Version 1.0.
+// http://www.boost.org/LICENSE_1_0.txt
+// http://www.geometrictools.com/License/Boost/LICENSE_1_0.txt
+// File Version: 3.14.3 (2018/10/30)
+
+#pragma once
+
+#include <GTEngine/Mathematics/GteApprOrthogonalPlane3.h>
+#include <GTEngine/Mathematics/GteUIntegerAP32.h>
+#include <GTEngine/Mathematics/GteBSRational.h>
+#include <GTEngine/Mathematics/GteRootsPolynomial.h>
+#include <GTEngine/Mathematics/GteGaussNewtonMinimizer.h>
+#include <GTEngine/Mathematics/GteLevenbergMarquardtMinimizer.h>
+
+// Let the torus center be C with plane of symmetry containing C and having
+// directions D0 and D1.  The axis of symmetry is the line containing C and
+// having direction N (the plane normal).  The radius from the center of the
+// torus is r0 and the radius of the tube of the torus is r1.  A point P may
+// be written as P = C + x*D0 + y*D1 + z*N, where matrix [D0 D1 N] is
+// orthogonal and has determinant 1.  Thus, x = Dot(D0,P-C), y = Dot(D1,P-C)
+// and z = Dot(N,P-C).  The implicit equation defining the torus is
+//     (|P-C|^2 + r0^2 - r1^2)^2 - 4*r0^2*(|P-C|^2 - (Dot(N,P-C))^2) = 0
+// Observe that D0 and D1 are not present in the equation, which is to be
+// expected by the symmetry.
+//
+// Define u = r0^2 and v = r0^2 - r1^2.  Define
+//     F(X;C,N,u,v) = (|P-C|^2 + v)^2 - 4*u*(|P-C|^2 - (Dot(N,P-C))^2)
+// The nonlinear least-squares fitting of points {X[i]}_{i=0}^{n-1} computes
+// C, N, u and v to minimize the error function
+//     E(C,N,u,v) = sum_{i=0}^{n-1} F(X[i];C,N,u,v)^2
+// When the sample points are distributed so that there is large coverage
+// by a purported fitted torus, a variation on fitting is the following.
+// Compute the least-squares plane with origin C and normal N that fits the
+// points.  Define G(X;u,v) = F(X;C,N,u,v); the only variables now are u and
+// v.  Define L[i] = |X[i]-C|^2 and S[i] = 4 * (L[i] - (Dot(N,X[i]-C))^2).
+// Define the error function
+//     H(u,v) = sum_{i=0}^{n-1} G(X[i];u,v)^2
+//            = sum_{i=0}^{n-1} ((v + L[i])^2 - S[i]*u)^2
+// The first-order partial derivatives are
+//     dH/du = -2 sum_{i=0}^{n-1} ((v + L[i])^2 - S[i]*u) * S[i]
+//     dH/dv =  4 sum_{i=0}^{n-1} ((v + L[i])^2 - S[i]*u) * (v + L[i])
+// Setting these to zero and expanding the terms, we have
+//     0 = a2 * v^2 + a1 * v + a0 - b0 * u
+//     0 = c3 * v^3 + c2 * v^2 + c1 * v + c0 - u * (d1 * v + d0)
+// where a2 = sum(S[i]), a1 = 2*sum(S[i]*L[i]), a2 = sum(S[i]*L[i]^2),
+// b0 = sum(S[i]^2), c3 = sum(1) = n, c2 = 3*sum(L[i]), c1 = 3*sum(L[i]^2),
+// c0 = sum(L[i]^3), d1 = sum(S[i]) = a2 and d0 = sum(S[i]*L[i]) = a1/2.
+// The first equation is solved for
+//     u = (a2 * v^2 + a1 * v + a0) / b0 = e2 * v^2 + e1 * v + e0
+// and substituted into the second equation to obtain a cubic polynomial
+// equation
+//     0 = f3 * v^3 + f2 * v^2 + f1 * v + f0
+// where f3 = c3 - d1 * e2, f2 = c2 - d1 * e1 - d0 * e2,
+// f1 = c1 - d1 * e0 - d0 * e1 and f0 = c0 - d0 * e0.  The positive v-roots
+// are computed.  For each root compute the corresponding u.  For all pairs
+// (u,v) with u > v > 0, evaluate H(u,v) and choose the pair that minimizes
+// H(u,v).  The torus radii are r0 = sqrt(u) and r1 = sqrt(u - v).
+
+namespace gte
+{
+    template <typename Real>
+    class ApprTorus3
+    {
+    public:
+        ApprTorus3()
+        {
+            // The unit-length normal is
+            //   N = (cos(theta)*sin(phi), sin(theta)*sin(phi), cos(phi)
+            // for theta in [0,2*pi) and phi in [0,*pi).  The radii are
+            // encoded as
+            //   u = r0^2, v = r0^2 - r1^2
+            // with 0 < v < u.  Let D = C - X[i] where X[i] is a sample point.
+            // The parameters P = (C0,C1,C2,theta,phi,u,v).
+
+            // F[i](C,theta,phi,u,v) =
+            //   (|D|^2 + v)^2 - 4*u*(|D|^2 - Dot(N,D)^2)
+            mFFunction = [this](GVector<Real> const& P, GVector<Real>& F)
+            {
+                Real csTheta = std::cos(P[3]);
+                Real snTheta = std::sin(P[3]);
+                Real csPhi = std::cos(P[4]);
+                Real snPhi = std::sin(P[4]);
+                Vector<3, Real> C = { P[0], P[1], P[2] };
+                Vector<3, Real> N = { csTheta * snPhi, snTheta * snPhi, csPhi };
+                Real u = P[5];
+                Real v = P[6];
+                for (int i = 0; i < mNumPoints; ++i)
+                {
+                    Vector<3, Real> D = C - mPoints[i];
+                    Real DdotD = Dot(D, D), NdotD = Dot(N, D);
+                    Real sum = DdotD + v;
+                    F[i] = sum * sum - (Real)4 * u * (DdotD - NdotD * NdotD);
+                }
+            };
+
+            // dF[i]/dC = 4 * (|D|^2 + v) * D - 8 * u * (I - N*N^T) * D
+            // dF[i]/dTheta = 8 * u * Dot(dN/dTheta, D)
+            // dF[i]/dPhi = 8 * u * Dot(dN/dPhi, D)
+            // dF[i]/du = -4 * u * (|D|^2 - Dot(N,D)^2)
+            // dF[i]/dv = 2 * (|D|^2 + v)
+            mJFunction = [this](GVector<Real> const& P, GMatrix<Real>& J)
+            {
+                Real const r2(2), r4(4), r8(8);
+                Real csTheta = std::cos(P[3]);
+                Real snTheta = std::sin(P[3]);
+                Real csPhi = std::cos(P[4]);
+                Real snPhi = std::sin(P[4]);
+                Vector<3, Real> C = { P[0], P[1], P[2] };
+                Vector<3, Real> N = { csTheta * snPhi, snTheta * snPhi, csPhi };
+                Real u = P[5];
+                Real v = P[6];
+                for (int row = 0; row < mNumPoints; ++row)
+                {
+                    Vector<3, Real> D = C - mPoints[row];
+                    Real DdotD = Dot(D, D), NdotD = Dot(N, D);
+                    Real sum = DdotD + v;
+                    Vector<3, Real> dNdTheta{ -snTheta * snPhi, csTheta * snPhi, (Real)0 };
+                    Vector<3, Real> dNdPhi{ csTheta * csPhi, snTheta * csPhi, -snPhi };
+                    Vector<3, Real> temp = r4 * sum * D - r8 * u * (D - NdotD * N);
+                    J(row, 0) = temp[0];
+                    J(row, 1) = temp[1];
+                    J(row, 2) = temp[2];
+                    J(row, 3) = r8 * u * Dot(dNdTheta, D);
+                    J(row, 4) = r8 * u * Dot(dNdPhi, D);
+                    J(row, 5) = -r4 * u * (DdotD - NdotD * NdotD);
+                    J(row, 6) = r2 * sum;
+                }
+            };
+        }
+
+        // When the samples are distributed approximately uniformly near a
+        // torus, use this method.  For example, if the purported torus has
+        // center (0,0,0) and normal (0,0,1), you want the (x,y,z) samples
+        // to occur in all 8 octants.  If the samples occur, say, only in
+        // one octant, this method will estimate a C and N that are nowhere
+        // near (0,0,0) and (0,0,1).  The function sets the output variables
+        // C, N, r0 and r1 as the fitted torus.
+        //
+        // The return value is a pair <bool,Real>.  The first element is
+        // 'true' when the estimate is valid, in which case the second
+        // element is the least-squares error for that estimate.  If any
+        // unexpected condition occurs that prevents computing an estimate,
+        // the first element is 'false' and the second element is
+        // std::numeric_limits<Real>::max().
+        std::pair<bool, Real>
+        operator()(int numPoints, Vector<3, Real> const* points,
+            Vector<3, Real>& C, Vector<3, Real>& N, Real& r0, Real& r1)
+        {
+            ApprOrthogonalPlane3<Real> fitter;
+            if (!fitter.Fit(numPoints, points))
+            {
+                return std::make_pair(false, std::numeric_limits<Real>::max());
+            }
+            C = fitter.GetParameters().first;
+            N = fitter.GetParameters().second;
+
+            Real const zero(0);
+            Real a0 = zero, a1 = zero, a2 = zero, b0 = zero;
+            Real c0 = zero, c1 = zero, c2 = zero, c3 = (Real)numPoints;
+            for (int i = 0; i < numPoints; ++i)
+            {
+                Vector<3, Real> delta = points[i] - C;
+                Real dot = Dot(N, delta);
+                Real L = Dot(delta, delta), L2 = L * L, L3 = L * L2;
+                Real S = (Real)4 * (L - dot * dot), S2 = S * S;
+                a2 += S;
+                a1 += S * L;
+                a0 += S * L2;
+                b0 += S2;
+                c2 += L;
+                c1 += L2;
+                c0 += L3;
+            }
+            Real d1 = a2;
+            Real d0 = a1;
+            a1 *= (Real)2;
+            c2 *= (Real)3;
+            c1 *= (Real)3;
+            Real invB0 = (Real)1 / b0;
+            Real e0 = a0 * invB0;
+            Real e1 = a1 * invB0;
+            Real e2 = a2 * invB0;
+
+            Rational f0 = c0 - d0 * e0;
+            Rational f1 = c1 - d1 * e0 - d0 * e1;
+            Rational f2 = c2 - d1 * e1 - d0 * e2;
+            Rational f3 = c3 - d1 * e2;
+            std::map<Real, int> rmMap;
+            RootsPolynomial<Real>::SolveCubic(f0, f1, f2, f3, rmMap);
+
+            Real hmin = std::numeric_limits<Real>::max();
+            Real umin = zero, vmin = zero;
+            for (auto const& element : rmMap)
+            {
+                Real v = element.first;
+                if (v > zero)
+                {
+                    Real u = e0 + v * (e1 + v * e2);
+                    if (u > v)
+                    {
+                        Real h = zero;
+                        for (int i = 0; i < numPoints; ++i)
+                        {
+                            Vector<3, Real> delta = points[i] - C;
+                            Real dot = Dot(N, delta);
+                            Real L = Dot(delta, delta);
+                            Real S = (Real)4 * (L - dot * dot);
+                            Real sum = v + L;
+                            Real term = sum * sum - S * u;
+                            h += term * term;
+                        }
+                        if (h < hmin)
+                        {
+                            hmin = h;
+                            umin = u;
+                            vmin = v;
+                        }
+                    }
+                }
+            }
+
+            if (hmin == std::numeric_limits<Real>::max())
+            {
+                return std::make_pair(false, std::numeric_limits<Real>::max());
+            }
+
+            r0 = std::sqrt(umin);
+            r1 = std::sqrt(umin - vmin);
+            return std::make_pair(true, hmin);
+        }
+
+        // If you want to specify that C, N, r0 and r1 are the initial guesses
+        // for the minimizer, set the parameter useTorusInputAsInitialGuess to
+        // 'true'.  If you want the function to compute initial guesses, set
+        // useTorusInputAsInitialGuess to 'false'.  A Gauss-Newton minimizer
+        // is used to fit a torus using nonlinear least-squares.  The fitted
+        // torus is returned in C, N, r0 and r1. See GteGaussNewtonMinimizer.h
+        // for a description of the least-squares algorithm and the parameters
+        // that it requires.
+        typename GaussNewtonMinimizer<Real>::Result
+        operator()(int numPoints, Vector<3, Real> const* points,
+            size_t maxIterations, Real updateLengthTolerance, Real errorDifferenceTolerance,
+            bool useTorusInputAsInitialGuess,
+            Vector<3, Real>& C, Vector<3, Real>& N, Real& r0, Real& r1)
+        {
+            mNumPoints = numPoints;
+            mPoints = points;
+            GaussNewtonMinimizer<Real> minimizer(7, mNumPoints, mFFunction, mJFunction);
+
+            if (!useTorusInputAsInitialGuess)
+            {
+                operator()(numPoints, points, C, N, r0, r1);
+            }
+
+            GVector<Real> initial(7);
+
+            // The initial guess for the plane origin.
+            initial[0] = C[0];
+            initial[1] = C[1];
+            initial[2] = C[2];
+
+            // The initial guess for the plane normal.  The angles must be
+            // extracted for spherical coordinates.
+            if (std::abs(N[2]) < (Real)1)
+            {
+                initial[3] = std::atan2(N[1], N[0]);
+                initial[4] = std::acos(N[2]);
+            }
+            else
+            {
+                initial[3] = (Real)0;
+                initial[4] = (Real)0;
+            }
+
+            // The initial guess for the radii-related parameters.
+            initial[5] = r0 * r0;
+            initial[6] = initial[5] - r1 * r1;
+
+            auto result = minimizer(initial, maxIterations, updateLengthTolerance,
+                errorDifferenceTolerance);
+
+            // No test is made for result.converged so that we return some
+            // estimates of the torus.  The caller can decide how to respond
+            // when result.converged is false.
+            C[0] = result.minLocation[0];
+            C[1] = result.minLocation[1];
+            C[2] = result.minLocation[2];
+
+            Real theta = result.minLocation[3];
+            Real phi = result.minLocation[4];
+            Real csTheta = std::cos(theta);
+            Real snTheta = std::sin(theta);
+            Real csPhi = std::cos(phi);
+            Real snPhi = std::sin(phi);
+            N[0] = csTheta * snPhi;
+            N[1] = snTheta * snPhi;
+            N[2] = csPhi;
+
+            Real u = result.minLocation[5];
+            Real v = result.minLocation[6];
+            r0 = std::sqrt(u);
+            r1 = std::sqrt(u - v);
+
+            mNumPoints = 0;
+            mPoints = nullptr;
+            return result;
+        }
+
+        // If you want to specify that C, N, r0 and r1 are the initial guesses
+        // for the minimizer, set the parameter useTorusInputAsInitialGuess to
+        // 'true'.  If you want the function to compute initial guesses, set
+        // useTorusInputAsInitialGuess to 'false'.  A Gauss-Newton minimizer
+        // is used to fit a torus using nonlinear least-squares.  The fitted
+        // torus is returned in C, N, r0 and r1. See GteGaussNewtonMinimizer.h
+        // for a description of the least-squares algorithm and the parameters
+        // that it requires.
+        typename LevenbergMarquardtMinimizer<Real>::Result
+        operator()(int numPoints, Vector<3, Real> const* points,
+            size_t maxIterations, Real updateLengthTolerance, Real errorDifferenceTolerance,
+            Real lambdaFactor, Real lambdaAdjust, size_t maxAdjustments,
+            bool useTorusInputAsInitialGuess,
+            Vector<3, Real>& C, Vector<3, Real>& N, Real& r0, Real& r1)
+        {
+            mNumPoints = numPoints;
+            mPoints = points;
+            LevenbergMarquardtMinimizer<Real> minimizer(7, mNumPoints, mFFunction, mJFunction);
+
+            if (!useTorusInputAsInitialGuess)
+            {
+                operator()(numPoints, points, C, N, r0, r1);
+            }
+
+            GVector<Real> initial(7);
+
+            // The initial guess for the plane origin.
+            initial[0] = C[0];
+            initial[1] = C[1];
+            initial[2] = C[2];
+
+            // The initial guess for the plane normal.  The angles must be
+            // extracted for spherical coordinates.
+            if (std::abs(N[2]) < (Real)1)
+            {
+                initial[3] = std::atan2(N[1], N[0]);
+                initial[4] = std::acos(N[2]);
+            }
+            else
+            {
+                initial[3] = (Real)0;
+                initial[4] = (Real)0;
+            }
+
+            // The initial guess for the radii-related parameters.
+            initial[5] = r0 * r0;
+            initial[6] = initial[5] - r1 * r1;
+
+            auto result = minimizer(initial, maxIterations, updateLengthTolerance,
+                errorDifferenceTolerance, lambdaFactor, lambdaAdjust, maxAdjustments);
+
+            // No test is made for result.converged so that we return some
+            // estimates of the torus.  The caller can decide how to respond
+            // when result.converged is false.
+            C[0] = result.minLocation[0];
+            C[1] = result.minLocation[1];
+            C[2] = result.minLocation[2];
+
+            Real theta = result.minLocation[3];
+            Real phi = result.minLocation[4];
+            Real csTheta = std::cos(theta);
+            Real snTheta = std::sin(theta);
+            Real csPhi = std::cos(phi);
+            Real snPhi = std::sin(phi);
+            N[0] = csTheta * snPhi;
+            N[1] = snTheta * snPhi;
+            N[2] = csPhi;
+
+            Real u = result.minLocation[5];
+            Real v = result.minLocation[6];
+            r0 = std::sqrt(u);
+            r1 = std::sqrt(u - v);
+
+            mNumPoints = 0;
+            mPoints = nullptr;
+            return result;
+        }
+
+    private:
+        typedef BSRational<UIntegerAP32> Rational;
+
+        int mNumPoints;
+        Vector<3, Real> const* mPoints;
+        std::function<void(GVector<Real> const&, GVector<Real>&)> mFFunction;
+        std::function<void(GVector<Real> const&, GMatrix<Real>&)> mJFunction;
+    };
+}
