move GeometryProcessing plugin from NotForLicensees to Experimental

diff --git a/Engine/Plugins/Experimental/GeometryProcessing/Source/GeometryAlgorithms/Private/GTEngine/Mathematics/GteArbitraryPrecision.h b/Engine/Plugins/Experimental/GeometryProcessing/Source/GeometryAlgorithms/Private/GTEngine/Mathematics/GteArbitraryPrecision.h
new file mode 100644
--- /dev/null
+++ b/Engine/Plugins/Experimental/GeometryProcessing/Source/GeometryAlgorithms/Private/GTEngine/Mathematics/GteArbitraryPrecision.h
@@ -0,0 +1,15 @@
+// David Eberly, Geometric Tools, Redmond WA 98052
+// Copyright (c) 1998-2019
+// Distributed under the Boost Software License, Version 1.0.
+// http://www.boost.org/LICENSE_1_0.txt
+// http://www.geometrictools.com/License/Boost/LICENSE_1_0.txt
+// File Version: 3.21.0 (2019/01/17)
+
+#pragma once
+
+#include <GTEngine/Mathematics/GteUIntegerALU32.h>
+#include <GTEngine/Mathematics/GteUIntegerAP32.h>
+#include <GTEngine/Mathematics/GteUIntegerFP32.h>
+#include <GTEngine/Mathematics/GteBSNumber.h>
+#include <GTEngine/Mathematics/GteBSRational.h>
+#include <GTEngine/Mathematics/GteBSPrecision.h>
diff --git a/Engine/Plugins/Experimental/GeometryProcessing/Source/GeometryAlgorithms/Private/GTEngine/Mathematics/GteArc2.h b/Engine/Plugins/Experimental/GeometryProcessing/Source/GeometryAlgorithms/Private/GTEngine/Mathematics/GteArc2.h
new file mode 100644
--- /dev/null
+++ b/Engine/Plugins/Experimental/GeometryProcessing/Source/GeometryAlgorithms/Private/GTEngine/Mathematics/GteArc2.h
@@ -0,0 +1,154 @@
+// David Eberly, Geometric Tools, Redmond WA 98052
+// Copyright (c) 1998-2019
+// Distributed under the Boost Software License, Version 1.0.
+// http://www.boost.org/LICENSE_1_0.txt
+// http://www.geometrictools.com/License/Boost/LICENSE_1_0.txt
+// File Version: 3.0.0 (2016/06/19)
+
+#pragma once
+
+#include <GTEngine/Mathematics/GteVector2.h>
+
+namespace gte
+{
+
+// The circle containing the arc is represented as |X-C| = R where C is the
+// center and R is the radius.  The arc is defined by two points end0 and end1
+// on the circle so that end1 is obtained from end0 by traversing
+// counterclockwise.  The application is responsible for ensuring that end0
+// and end1 are on the circle and that they are properly ordered.
+
+template <typename Real>
+class Arc2
+{
+public:
+    // Construction and destruction.  The default constructor sets the center
+    // to (0,0), radius to 1, end0 to (1,0), and end1 to (0,1).
+    Arc2();
+    Arc2(Vector2<Real> const& inCenter, Real inRadius,
+        Vector2<Real>const& inEnd0, Vector2<Real>const & inEnd1);
+
+    // Test whether P is on the arc.  The application must ensure that P is on
+    // the circle; that is, |P-C| = R.  This test works for any angle between
+    // B-C and A-C, not just those between 0 and pi radians.
+    bool Contains(Vector2<Real> const& p) const;
+
+    Vector2<Real> center;
+    Real radius;
+    Vector2<Real> end[2];
+
+public:
+    // Comparisons to support sorted containers.
+    bool operator==(Arc2 const& arc) const;
+    bool operator!=(Arc2 const& arc) const;
+    bool operator< (Arc2 const& arc) const;
+    bool operator<=(Arc2 const& arc) const;
+    bool operator> (Arc2 const& arc) const;
+    bool operator>=(Arc2 const& arc) const;
+};
+
+
+template <typename Real>
+Arc2<Real>::Arc2()
+    :
+    center(Vector2<Real>::Zero()),
+    radius((Real)1)
+{
+    end[0] = Vector2<Real>::Unit(0);
+    end[1] = Vector2<Real>::Unit(1);
+}
+
+template <typename Real>
+Arc2<Real>::Arc2(Vector2<Real> const& inCenter, Real inRadius,
+    Vector2<Real>const& inEnd0, Vector2<Real>const & inEnd1)
+    :
+    center(inCenter),
+    radius(inRadius)
+{
+    end[0] = inEnd0;
+    end[1] = inEnd1;
+}
+
+template <typename Real>
+bool Arc2<Real>::Contains(Vector2<Real> const& p) const
+{
+    // Assert: |P-C| = R where P is the input point, C is the circle center,
+    // and R is the circle radius.  For P to be on the arc from A to B, it
+    // must be on the side of the plane containing A with normal N = Perp(B-A)
+    // where Perp(u,v) = (v,-u).
+
+    Vector2<Real> diffPE0 = p - end[0];
+    Vector2<Real> diffE1E0 = end[1] - end[0];
+    Real dotPerp = DotPerp(diffPE0, diffE1E0);
+    return dotPerp >= (Real)0;
+}
+
+template <typename Real>
+bool Arc2<Real>::operator==(Arc2 const& arc) const
+{
+    return center == arc.center && radius == arc.radius
+        && end[0] == arc.end[0] && end[1] == arc.end[1];
+}
+
+template <typename Real>
+bool Arc2<Real>::operator!=(Arc2 const& arc) const
+{
+    return !operator==(arc);
+}
+
+template <typename Real>
+bool Arc2<Real>::operator<(Arc2 const& arc) const
+{
+    if (center < arc.center)
+    {
+        return true;
+    }
+
+    if (center > arc.center)
+    {
+        return false;
+    }
+
+    if (radius < arc.radius)
+    {
+        return true;
+    }
+
+    if (radius > arc.radius)
+    {
+        return false;
+    }
+
+    if (end[0] < arc.end[0])
+    {
+        return true;
+    }
+
+    if (end[0] > arc.end[0])
+    {
+        return false;
+    }
+
+    return end[1] < arc.end[1];
+}
+
+template <typename Real>
+bool Arc2<Real>::operator<=(Arc2 const& arc) const
+{
+    return operator<(arc) || operator==(arc);
+}
+
+template <typename Real>
+bool Arc2<Real>::operator>(Arc2 const& arc) const
+{
+    return !operator<=(arc);
+}
+
+template <typename Real>
+bool Arc2<Real>::operator>=(Arc2 const& arc) const
+{
+    return !operator<(arc);
+}
+
+
+}
diff --git a/Engine/Plugins/Experimental/GeometryProcessing/Source/GeometryAlgorithms/Private/GTEngine/Mathematics/GteAxisAngle.h b/Engine/Plugins/Experimental/GeometryProcessing/Source/GeometryAlgorithms/Private/GTEngine/Mathematics/GteAxisAngle.h
new file mode 100644
--- /dev/null
+++ b/Engine/Plugins/Experimental/GeometryProcessing/Source/GeometryAlgorithms/Private/GTEngine/Mathematics/GteAxisAngle.h
@@ -0,0 +1,48 @@
+// David Eberly, Geometric Tools, Redmond WA 98052
+// Copyright (c) 1998-2019
+// Distributed under the Boost Software License, Version 1.0.
+// http://www.boost.org/LICENSE_1_0.txt
+// http://www.geometrictools.com/License/Boost/LICENSE_1_0.txt
+// File Version: 3.0.0 (2016/06/19)
+
+#pragma once
+
+#include <GTEngine/Mathematics/GteVector.h>
+
+namespace gte
+{
+
+// Axis-angle representation for N = 3 or N = 4.  When N = 4, the axis
+// must be a vector of the form (x,y,z,0) [affine representation of the
+// 3-tuple direction].
+
+template <int N, typename Real>
+class AxisAngle
+{
+public:
+    AxisAngle();
+    AxisAngle(Vector<N,Real> const& inAxis, Real inAngle);
+
+    Vector<N,Real> axis;
+    Real angle;
+};
+
+
+template <int N, typename Real>
+AxisAngle<N, Real>::AxisAngle()
+{
+    static_assert(N == 3 || N == 4, "Dimension must be 3 or 4.");
+    // Uninitialized.
+}
+
+template <int N, typename Real>
+AxisAngle<N, Real>::AxisAngle(Vector<N, Real> const& inAxis, Real inAngle)
+    :
+    axis(inAxis),
+    angle(inAngle)
+{
+    static_assert(N == 3 || N == 4, "Dimension must be 3 or 4.");
+}
+
+
+}
diff --git a/Engine/Plugins/Experimental/GeometryProcessing/Source/GeometryAlgorithms/Private/GTEngine/Mathematics/GteBSNumber.h b/Engine/Plugins/Experimental/GeometryProcessing/Source/GeometryAlgorithms/Private/GTEngine/Mathematics/GteBSNumber.h
new file mode 100644
--- /dev/null
+++ b/Engine/Plugins/Experimental/GeometryProcessing/Source/GeometryAlgorithms/Private/GTEngine/Mathematics/GteBSNumber.h
@@ -0,0 +1,1079 @@
+// David Eberly, Geometric Tools, Redmond WA 98052
+// Copyright (c) 1998-2019
+// Distributed under the Boost Software License, Version 1.0.
+// http://www.boost.org/LICENSE_1_0.txt
+// http://www.geometrictools.com/License/Boost/LICENSE_1_0.txt
+// File Version: 3.0.3 (2019/01/17)
+
+#pragma once
+
+#include <GTEngine/LowLevel/GteLogger.h>
+#include <GTEngine/Mathematics/GteBitHacks.h>
+#include <GTEngine/Mathematics/GteMath.h>
+#include <GTEngine/Mathematics/GteIEEEBinary.h>
+#include <algorithm>
+#include <fstream>
+
+// The class BSNumber (binary scientific number) is designed to provide exact
+// arithmetic for robust algorithms, typically those for which we need to know
+// the exact sign of determinants.  The template parameter UIntegerType must
+// have support for at least the following public interface.  The fstream
+// objects for Write/Read must be created using std::ios::binary.  The return
+// value of Write/Read is 'true' iff the operation was successful.
+//
+//      class UIntegerType
+//      {
+//      public:
+//          UIntegerType();
+//          UIntegerType(UIntegerType const& number);
+//          UIntegerType(uint32_t number);
+//          UIntegerType(uint64_t number);
+//          UIntegerType(int numBits);
+//          UIntegerType& operator=(UIntegerType const& number);
+//          UIntegerType(UIntegerType&& number);
+//          UIntegerType& operator=(UIntegerType&& number);
+//          int32_t GetNumBits() const;
+//          bool operator==(UIntegerType const& number) const;
+//          bool operator< (UIntegerType const& number) const;
+//          void Add(UIntegerType const& n0, UIntegerType const& n1);
+//          void Sub(UIntegerType const& n0, UIntegerType const& n1);
+//          void Mul(UIntegerType const& n0, UIntegerType const& n1);
+//          void ShiftLeft(UIntegerType const& number, int shift);
+//          int32_t ShiftRightToOdd(UIntegerType const& number);
+//          uint64_t GetPrefix(int numRequested) const;
+//          bool Write(std::ofstream& output) const;
+//          bool Read(std::ifstream& input);
+//      };
+//
+// GTEngine currently has 32-bits-per-word storage for UIntegerType.  See the
+// classes UIntegerAP32 (arbitrary precision), UIntegerFP32<N> (fixed
+// precision), and UIntegerALU32 (arithmetic logic unit shared by the previous
+// two classes).  The document at the following link describes the design,
+// implementation, and use of BSNumber and BSRational.
+//   http://www.geometrictools.com/Documentation/ArbitraryPrecision.pdf
+//
+// Support for debugging algorithms that use exact rational arithmetic.  Each
+// BSNumber and BSRational has a double-precision member that is exposed when
+// the conditional define is enabled.  Be aware that this can be very slow
+// because of the conversion to double-precision whenever new objects are
+// created by arithmetic operations.  As a faster alternative, you can add
+// temporary code in your algorithms that explicitly convert specific rational
+// numbers to double precision.
+//
+//#define GTE_BINARY_SCIENTIFIC_SHOW_DOUBLE
+
+namespace gte
+{
+    template <typename UIntegerType> class BSRational;
+
+    template <typename UIntegerType>
+    class BSNumber
+    {
+    public:
+        // Construction.  The default constructor generates the zero BSNumber.
+        BSNumber()
+            :
+            mSign(0),
+            mBiasedExponent(0)
+        {
+#if defined(GTE_BINARY_SCIENTIFIC_SHOW_DOUBLE)
+            mValue = (double)*this;
+#endif
+        }
+
+        BSNumber(BSNumber const& number)
+        {
+            *this = number;
+        }
+
+        BSNumber(float number)
+        {
+            ConvertFrom<IEEEBinary32>(number);
+#if defined(GTE_BINARY_SCIENTIFIC_SHOW_DOUBLE)
+            mValue = (double)*this;
+#endif
+        }
+
+        BSNumber(double number)
+        {
+            ConvertFrom<IEEEBinary64>(number);
+#if defined(GTE_BINARY_SCIENTIFIC_SHOW_DOUBLE)
+            mValue = (double)*this;
+#endif
+        }
+
+        BSNumber(int32_t number)
+        {
+            if (number == 0)
+            {
+                mSign = 0;
+                mBiasedExponent = 0;
+            }
+            else
+            {
+                if (number < 0)
+                {
+                    mSign = -1;
+                    number = -number;
+                }
+                else
+                {
+                    mSign = 1;
+                }
+
+                mBiasedExponent = GetTrailingBit(number);
+                mUInteger = (uint32_t)number;
+            }
+#if defined(GTE_BINARY_SCIENTIFIC_SHOW_DOUBLE)
+            mValue = (double)*this;
+#endif
+        }
+
+        BSNumber(uint32_t number)
+        {
+            if (number == 0)
+            {
+                mSign = 0;
+                mBiasedExponent = 0;
+            }
+            else
+            {
+                mSign = 1;
+                mBiasedExponent = GetTrailingBit(number);
+                mUInteger = number;
+            }
+#if defined(GTE_BINARY_SCIENTIFIC_SHOW_DOUBLE)
+            mValue = (double)*this;
+#endif
+        }
+
+        BSNumber(int64_t number)
+        {
+            if (number == 0)
+            {
+                mSign = 0;
+                mBiasedExponent = 0;
+            }
+            else
+            {
+                if (number < 0)
+                {
+                    mSign = -1;
+                    number = -number;
+                }
+                else
+                {
+                    mSign = 1;
+                }
+
+                mBiasedExponent = GetTrailingBit(number);
+                mUInteger = (uint64_t)number;
+            }
+#if defined(GTE_BINARY_SCIENTIFIC_SHOW_DOUBLE)
+            mValue = (double)*this;
+#endif
+        }
+
+        BSNumber(uint64_t number)
+        {
+            if (number == 0)
+            {
+                mSign = 0;
+                mBiasedExponent = 0;
+            }
+            else
+            {
+                mSign = 1;
+                mBiasedExponent = GetTrailingBit(number);
+                mUInteger = number;
+            }
+#if defined(GTE_BINARY_SCIENTIFIC_SHOW_DOUBLE)
+            mValue = (double)*this;
+#endif
+        }
+
+        // Implicit conversions.
+        inline operator float() const
+        {
+            return ConvertTo<IEEEBinary32>();
+        }
+
+        inline operator double() const
+        {
+            return ConvertTo<IEEEBinary64>();
+        }
+
+        // Assignment.
+        BSNumber& operator=(BSNumber const& number)
+        {
+            mSign = number.mSign;
+            mBiasedExponent = number.mBiasedExponent;
+            mUInteger = number.mUInteger;
+#if defined(GTE_BINARY_SCIENTIFIC_SHOW_DOUBLE)
+            mValue = number.mValue;
+#endif
+            return *this;
+        }
+
+        // Support for std::move.
+        BSNumber(BSNumber&& number)
+        {
+            *this = std::move(number);
+        }
+
+        BSNumber& operator=(BSNumber&& number)
+        {
+            mSign = number.mSign;
+            mBiasedExponent = number.mBiasedExponent;
+            mUInteger = std::move(number.mUInteger);
+            number.mSign = 0;
+            number.mBiasedExponent = 0;
+#if defined(GTE_BINARY_SCIENTIFIC_SHOW_DOUBLE)
+            mValue = number.mValue;
+#endif
+            return *this;
+        }
+
+        // Member access.
+        inline int32_t GetSign() const
+        {
+            return mSign;
+        }
+
+        inline int32_t GetBiasedExponent() const
+        {
+            return mBiasedExponent;
+        }
+
+        inline int32_t GetExponent() const
+        {
+            return mBiasedExponent + mUInteger.GetNumBits() - 1;
+        }
+
+        inline UIntegerType const& GetUInteger() const
+        {
+            return mUInteger;
+        }
+
+        // Comparisons.
+        bool operator==(BSNumber const& number) const
+        {
+            return (mSign == number.mSign ? EqualIgnoreSign(*this, number) : false);
+        }
+
+        bool operator!=(BSNumber const& number) const
+        {
+            return !operator==(number);
+        }
+
+        bool operator< (BSNumber const& number) const
+        {
+            if (mSign > 0)
+            {
+                if (number.mSign <= 0)
+                {
+                    return false;
+                }
+
+                // Both numbers are positive.
+                return LessThanIgnoreSign(*this, number);
+            }
+            else if (mSign < 0)
+            {
+                if (number.mSign >= 0)
+                {
+                    return true;
+                }
+
+                // Both numbers are negative.
+                return LessThanIgnoreSign(number, *this);
+            }
+            else
+            {
+                return number.mSign > 0;
+            }
+        }
+
+        bool operator<=(BSNumber const& number) const
+        {
+            return !number.operator<(*this);
+        }
+
+        bool operator> (BSNumber const& number) const
+        {
+            return number.operator<(*this);
+        }
+
+        bool operator>=(BSNumber const& number) const
+        {
+            return !operator<(number);
+        }
+
+        // Unary operations.
+        BSNumber operator+() const
+        {
+            return *this;
+        }
+
+        BSNumber operator-() const
+        {
+            BSNumber result = *this;
+            result.mSign = -result.mSign;
+#if defined(GTE_BINARY_SCIENTIFIC_SHOW_DOUBLE)
+            result.mValue = (double)result;
+#endif
+            return result;
+        }
+
+        // Arithmetic.
+        BSNumber operator+(BSNumber const& n1) const
+        {
+            BSNumber const& n0 = *this;
+
+            if (n0.mSign == 0)
+            {
+                return n1;
+            }
+
+            if (n1.mSign == 0)
+            {
+                return n0;
+            }
+
+            if (n0.mSign > 0)
+            {
+                if (n1.mSign > 0)
+                {
+                    // n0 + n1 = |n0| + |n1|
+                    return AddIgnoreSign(n0, n1, +1);
+                }
+                else // n1.mSign < 0
+                {
+                    if (!EqualIgnoreSign(n0, n1))
+                    {
+                        if (LessThanIgnoreSign(n1, n0))
+                        {
+                            // n0 + n1 = |n0| - |n1| > 0
+                            return SubIgnoreSign(n0, n1, +1);
+                        }
+                        else
+                        {
+                            // n0 + n1 = -(|n1| - |n0|) < 0
+                            return SubIgnoreSign(n1, n0, -1);
+                        }
+                    }
+                    // else n0 + n1 = 0
+                }
+            }
+            else // n0.mSign < 0
+            {
+                if (n1.mSign < 0)
+                {
+                    // n0 + n1 = -(|n0| + |n1|)
+                    return AddIgnoreSign(n0, n1, -1);
+                }
+                else // n1.mSign > 0
+                {
+                    if (!EqualIgnoreSign(n0, n1))
+                    {
+                        if (LessThanIgnoreSign(n1, n0))
+                        {
+                            // n0 + n1 = -(|n0| - |n1|) < 0
+                            return SubIgnoreSign(n0, n1, -1);
+                        }
+                        else
+                        {
+                            // n0 + n1 = |n1| - |n0| > 0
+                            return SubIgnoreSign(n1, n0, +1);
+                        }
+                    }
+                    // else n0 + n1 = 0
+                }
+            }
+
+            return BSNumber();  // = 0
+        }
+
+        BSNumber operator-(BSNumber const& n1) const
+        {
+            BSNumber const& n0 = *this;
+
+            if (n0.mSign == 0)
+            {
+                return -n1;
+            }
+
+            if (n1.mSign == 0)
+            {
+                return n0;
+            }
+
+            if (n0.mSign > 0)
+            {
+                if (n1.mSign < 0)
+                {
+                    // n0 - n1 = |n0| + |n1|
+                    return AddIgnoreSign(n0, n1, +1);
+                }
+                else // n1.mSign > 0
+                {
+                    if (!EqualIgnoreSign(n0, n1))
+                    {
+                        if (LessThanIgnoreSign(n1, n0))
+                        {
+                            // n0 - n1 = |n0| - |n1| > 0
+                            return SubIgnoreSign(n0, n1, +1);
+                        }
+                        else
+                        {
+                            // n0 - n1 = -(|n1| - |n0|) < 0
+                            return SubIgnoreSign(n1, n0, -1);
+                        }
+                    }
+                    // else n0 - n1 = 0
+                }
+            }
+            else // n0.mSign < 0
+            {
+                if (n1.mSign > 0)
+                {
+                    // n0 - n1 = -(|n0| + |n1|)
+                    return AddIgnoreSign(n0, n1, -1);
+                }
+                else // n1.mSign < 0
+                {
+                    if (!EqualIgnoreSign(n0, n1))
+                    {
+                        if (LessThanIgnoreSign(n1, n0))
+                        {
+                            // n0 - n1 = -(|n0| - |n1|) < 0
+                            return SubIgnoreSign(n0, n1, -1);
+                        }
+                        else
+                        {
+                            // n0 - n1 = |n1| - |n0| > 0
+                            return SubIgnoreSign(n1, n0, +1);
+                        }
+                    }
+                    // else n0 - n1 = 0
+                }
+            }
+
+            return BSNumber();  // = 0
+        }
+
+        BSNumber operator*(BSNumber const& number) const
+        {
+            BSNumber result;  // = 0
+            int sign = mSign * number.mSign;
+            if (sign != 0)
+            {
+                result.mSign = sign;
+                result.mBiasedExponent = mBiasedExponent + number.mBiasedExponent;
+                result.mUInteger.Mul(mUInteger, number.mUInteger);
+            }
+#if defined(GTE_BINARY_SCIENTIFIC_SHOW_DOUBLE)
+            result.mValue = (double)result;
+#endif
+            return result;
+        }
+
+        BSNumber& operator+=(BSNumber const& number)
+        {
+            *this = operator+(number);
+            return *this;
+        }
+
+        BSNumber& operator-=(BSNumber const& number)
+        {
+            *this = operator-(number);
+            return *this;
+        }
+
+        BSNumber& operator*=(BSNumber const& number)
+        {
+            *this = operator*(number);
+            return *this;
+        }
+
+        // Disk input/output.  The fstream objects should be created using
+        // std::ios::binary.  The return value is 'true' iff the operation
+        // was successful.
+        bool Write(std::ofstream& output) const
+        {
+            if (output.write((char const*)&mSign, sizeof(mSign)).bad())
+            {
+                return false;
+            }
+
+            if (output.write((char const*)&mBiasedExponent,
+                sizeof(mBiasedExponent)).bad())
+            {
+                return false;
+            }
+
+            return mUInteger.Write(output);
+        }
+
+        bool Read(std::ifstream& input)
+        {
+            if (input.read((char*)&mSign, sizeof(mSign)).bad())
+            {
+                return false;
+            }
+
+            if (input.read((char*)&mBiasedExponent, sizeof(mBiasedExponent)).bad())
+            {
+                return false;
+            }
+
+            return mUInteger.Read(input);
+        }
+
+    private:
+        // Helpers for operator==, operator<, operator+, operator-.
+        static bool EqualIgnoreSign(BSNumber const& n0, BSNumber const& n1)
+        {
+            return n0.mBiasedExponent == n1.mBiasedExponent && n0.mUInteger == n1.mUInteger;
+        }
+
+        static bool LessThanIgnoreSign(BSNumber const& n0, BSNumber const& n1)
+        {
+            int32_t e0 = n0.GetExponent(), e1 = n1.GetExponent();
+            if (e0 < e1)
+            {
+                return true;
+            }
+            if (e0 > e1)
+            {
+                return false;
+            }
+            return n0.mUInteger < n1.mUInteger;
+        }
+
+        // Add two positive numbers.
+        static BSNumber AddIgnoreSign(BSNumber const& n0, BSNumber const& n1, int32_t resultSign)
+        {
+            BSNumber result, temp;
+
+            int32_t diff = n0.mBiasedExponent - n1.mBiasedExponent;
+            if (diff > 0)
+            {
+                temp.mUInteger.ShiftLeft(n0.mUInteger, diff);
+                result.mUInteger.Add(temp.mUInteger, n1.mUInteger);
+                result.mBiasedExponent = n1.mBiasedExponent;
+            }
+            else if (diff < 0)
+            {
+                temp.mUInteger.ShiftLeft(n1.mUInteger, -diff);
+                result.mUInteger.Add(n0.mUInteger, temp.mUInteger);
+                result.mBiasedExponent = n0.mBiasedExponent;
+            }
+            else
+            {
+                temp.mUInteger.Add(n0.mUInteger, n1.mUInteger);
+                int32_t shift = result.mUInteger.ShiftRightToOdd(temp.mUInteger);
+                result.mBiasedExponent = n0.mBiasedExponent + shift;
+            }
+
+            result.mSign = resultSign;
+#if defined(GTE_BINARY_SCIENTIFIC_SHOW_DOUBLE)
+            result.mValue = (double)result;
+#endif
+            return result;
+        }
+
+        // Subtract two positive numbers where n0 > n1.
+        static BSNumber SubIgnoreSign(BSNumber const& n0, BSNumber const& n1, int32_t resultSign)
+        {
+            BSNumber result, temp;
+
+            int32_t diff = n0.mBiasedExponent - n1.mBiasedExponent;
+            if (diff > 0)
+            {
+                temp.mUInteger.ShiftLeft(n0.mUInteger, diff);
+                result.mUInteger.Sub(temp.mUInteger, n1.mUInteger);
+                result.mBiasedExponent = n1.mBiasedExponent;
+            }
+            else if (diff < 0)
+            {
+                temp.mUInteger.ShiftLeft(n1.mUInteger, -diff);
+                result.mUInteger.Sub(n0.mUInteger, temp.mUInteger);
+                result.mBiasedExponent = n0.mBiasedExponent;
+            }
+            else
+            {
+                temp.mUInteger.Sub(n0.mUInteger, n1.mUInteger);
+                int32_t shift = result.mUInteger.ShiftRightToOdd(temp.mUInteger);
+                result.mBiasedExponent = n0.mBiasedExponent + shift;
+            }
+
+            result.mSign = resultSign;
+#if defined(GTE_BINARY_SCIENTIFIC_SHOW_DOUBLE)
+            result.mValue = (double)result;
+#endif
+            return result;
+        }
+
+        // Support for conversions from floating-point numbers to BSNumber.
+        template <typename IEEE>
+        void ConvertFrom(typename IEEE::FloatType number)
+        {
+            IEEE x(number);
+
+            // Extract sign s, biased exponent e, and trailing significand t.
+            typename IEEE::UIntType s = x.GetSign();
+            typename IEEE::UIntType e = x.GetBiased();
+            typename IEEE::UIntType t = x.GetTrailing();
+
+            if (e == 0)
+            {
+                if (t == 0)  // zeros
+                {
+                    // x = (-1)^s * 0
+                    mSign = 0;
+                    mBiasedExponent = 0;
+                }
+                else  // subnormal numbers
+                {
+                    // x = (-1)^s * 0.t * 2^{1-EXPONENT_BIAS}
+                    int32_t last = GetTrailingBit(t);
+                    int32_t diff = IEEE::NUM_TRAILING_BITS - last;
+                    mSign = (s > 0 ? -1 : 1);
+                    mBiasedExponent = IEEE::MIN_SUB_EXPONENT - diff;
+                    mUInteger = (t >> last);
+                }
+            }
+            else if (e < IEEE::MAX_BIASED_EXPONENT)  // normal numbers
+            {
+                // x = (-1)^s * 1.t * 2^{e-EXPONENT_BIAS}
+                if (t > 0)
+                {
+                    int32_t last = GetTrailingBit(t);
+                    int32_t diff = IEEE::NUM_TRAILING_BITS - last;
+                    mSign = (s > 0 ? -1 : 1);
+                    mBiasedExponent =
+                        static_cast<int32_t>(e) - IEEE::EXPONENT_BIAS - diff;
+                    mUInteger = ((t | IEEE::SUP_TRAILING) >> last);
+                }
+                else
+                {
+                    mSign = (s > 0 ? -1 : 1);
+                    mBiasedExponent = static_cast<int32_t>(e) - IEEE::EXPONENT_BIAS;
+                    mUInteger = (typename IEEE::UIntType)1;
+                }
+            }
+            else  // e == MAX_BIASED_EXPONENT, special numbers
+            {
+                if (t == 0)  // infinities
+                {
+                    // x = (-1)^s * infinity
+                    LogWarning("Input is " + std::string(s > 0 ? "-" : "+") +
+                        "infinity.");
+
+                    // Return (-1)^s * 2^{1+EXPONENT_BIAS} for a graceful exit.
+                    mSign = (s > 0 ? -1 : 1);
+                    mBiasedExponent = 1 + IEEE::EXPONENT_BIAS;
+                    mUInteger = (typename IEEE::UIntType)1;
+                }
+                else  // not-a-number (NaN)
+                {
+                    LogError("Input is a " +
+                        std::string(t & IEEE::NAN_QUIET_MASK ?
+                            "quiet" : "signaling") + " NaN with payload " +
+                        std::to_string(t & IEEE::NAN_PAYLOAD_MASK) + ".");
+
+                    // Return 0 for a graceful exit.
+                    mSign = 0;
+                    mBiasedExponent = 0;
+                }
+            }
+        }
+
+        // Support for conversions from BSNumber to floating-point numbers.
+        template <typename IEEE>
+        typename IEEE::FloatType ConvertTo() const
+        {
+            typename IEEE::UIntType s = (mSign < 0 ? 1 : 0);
+            typename IEEE::UIntType e, t;
+
+            if (mSign != 0)
+            {
+                // The conversions use round-to-nearest-ties-to-even semantics.
+                int32_t exponent = GetExponent();
+                if (exponent < IEEE::MIN_EXPONENT)
+                {
+                    if (exponent < IEEE::MIN_EXPONENT - 1
+                        || mUInteger.GetNumBits() == 1)  // x = 1.0*2^{MIN_EXPONENT-1}
+                    {
+                        // Round to zero.
+                        e = 0;
+                        t = 0;
+                    }
+                    else
+                    {
+                        // Round to min subnormal.
+                        e = 0;
+                        t = 1;
+                    }
+                }
+                else if (exponent < IEEE::MIN_SUB_EXPONENT)
+                {
+                    // The second input is in {0, ..., NUM_TRAILING_BITS-1}.
+                    t = GetTrailing<IEEE>(0, IEEE::MIN_SUB_EXPONENT - exponent - 1);
+                    if (t & IEEE::SUP_TRAILING)
+                    {
+                        // Leading NUM_SIGNIFICAND_BITS bits were all 1, so round to
+                        // min normal.
+                        e = 1;
+                        t = 0;
+                    }
+                    else
+                    {
+                        e = 0;
+                    }
+                }
+                else if (exponent <= IEEE::EXPONENT_BIAS)
+                {
+                    e = static_cast<uint32_t>(exponent + IEEE::EXPONENT_BIAS);
+                    t = GetTrailing<IEEE>(1, 0);
+                    if (t & (IEEE::SUP_TRAILING << 1))
+                    {
+                        // Carry-out occurred, so increase exponent by 1 and
+                        // shift right to compensate.
+                        ++e;
+                        t >>= 1;
+                    }
+                    // Eliminate the leading 1 (implied for normals).
+                    t &= ~IEEE::SUP_TRAILING;
+                }
+                else
+                {
+                    // Set to infinity.
+                    e = IEEE::MAX_BIASED_EXPONENT;
+                    t = 0;
+                }
+            }
+            else
+            {
+                // The input is zero.
+                e = 0;
+                t = 0;
+            }
+
+            IEEE x(s, e, t);
+            return x.number;
+        }
+
+        template <typename IEEE>
+        typename IEEE::UIntType GetTrailing(int32_t normal, int32_t sigma) const
+        {
+            int32_t const numRequested = IEEE::NUM_SIGNIFICAND_BITS + normal;
+
+            // We need numRequested bits to determine rounding direction.  These are
+            // stored in the high-order bits of 'prefix'.
+            uint64_t prefix = mUInteger.GetPrefix(numRequested);
+
+            // The first bit index after the implied binary point for rounding.
+            int32_t diff = numRequested - sigma;
+            int32_t roundBitIndex = 64 - diff;
+
+            // Determine rounding value based on round-to-nearest-ties-to-even.
+            uint64_t mask = (1ull << roundBitIndex);
+            uint64_t round;
+            if (prefix & mask)
+            {
+                // The first bit of the remainder is 1.
+                if (mUInteger.GetNumBits() == diff)
+                {
+                    // The first bit of the remainder is the lowest-order bit of
+                    // mBits[0].  Apply the ties-to-even rule.
+                    if (prefix & (mask << 1))
+                    {
+                        // The last bit of the trailing significand is odd, so
+                        // round up.
+                        round = 1;
+                    }
+                    else
+                    {
+                        // The last bit of the trailing significand is even, so
+                        // round down.
+                        round = 0;
+                    }
+                }
+                else
+                {
+                    // The first bit of the remainder is not the lowest-order bit of
+                    // mBits[0].  The remainder as a fraction is larger than 1/2, so
+                    // round up.
+                    round = 1;
+                }
+            }
+            else
+            {
+                // The first bit of the remainder is 0, so round down.
+                round = 0;
+            }
+
+            // Get the unrounded trailing significand.
+            uint64_t trailing = prefix >> (roundBitIndex + 1);
+
+            // Apply the rounding.
+            trailing += round;
+            return static_cast<typename IEEE::UIntType>(trailing);
+        }
+
+#if defined(GTE_BINARY_SCIENTIFIC_SHOW_DOUBLE)
+    public:
+        // List this first so that it shows up first in the debugger watch window.
+        double mValue;
+    private:
+#endif
+
+        // The number 0 is represented by: mSign = 0, mBiasedExponent = 0, and
+        // mUInteger = 0.  For nonzero numbers, mSign != 0 and mUInteger > 0.
+        int32_t mSign;
+        int32_t mBiasedExponent;
+        UIntegerType mUInteger;
+
+        // Access to members to avoid exposing them publically when they are
+        // needed only internally.
+        friend class BSRational<UIntegerType>;
+        friend class UnitTestBSNumber;
+    };
+}
+
+namespace std
+{
+    // TODO: Allow for implementations of the math functions in which a
+    // specified precision is used when computing the result.
+
+    template <typename UIntegerType>
+    inline gte::BSNumber<UIntegerType> abs(gte::BSNumber<UIntegerType> const& x)
+    {
+        return (x.GetSign() >= 0 ? x : -x);
+    }
+
+    template <typename UIntegerType>
+    inline gte::BSNumber<UIntegerType> acos(gte::BSNumber<UIntegerType> const& x)
+    {
+        return (gte::BSNumber<UIntegerType>)std::acos((double)x);
+    }
+
+    template <typename UIntegerType>
+    inline gte::BSNumber<UIntegerType> acosh(gte::BSNumber<UIntegerType> const& x)
+    {
+        return (gte::BSNumber<UIntegerType>)std::acosh((double)x);
+    }
+
+    template <typename UIntegerType>
+    inline gte::BSNumber<UIntegerType> asin(gte::BSNumber<UIntegerType> const& x)
+    {
+        return (gte::BSNumber<UIntegerType>)std::asin((double)x);
+    }
+
+    template <typename UIntegerType>
+    inline gte::BSNumber<UIntegerType> asinh(gte::BSNumber<UIntegerType> const& x)
+    {
+        return (gte::BSNumber<UIntegerType>)std::asinh((double)x);
+    }
+
+    template <typename UIntegerType>
+    inline gte::BSNumber<UIntegerType> atan(gte::BSNumber<UIntegerType> const& x)
+    {
+        return (gte::BSNumber<UIntegerType>)std::atan((double)x);
+    }
+
+    template <typename UIntegerType>
+    inline gte::BSNumber<UIntegerType> atanh(gte::BSNumber<UIntegerType> const& x)
+    {
+        return (gte::BSNumber<UIntegerType>)std::atanh((double)x);
+    }
+
+    template <typename UIntegerType>
+    inline gte::BSNumber<UIntegerType> atan2(gte::BSNumber<UIntegerType> const& y, gte::BSNumber<UIntegerType> const& x)
+    {
+        return (gte::BSNumber<UIntegerType>)std::atan2((double)y, (double)x);
+    }
+
+    template <typename UIntegerType>
+    inline gte::BSNumber<UIntegerType> ceil(gte::BSNumber<UIntegerType> const& x)
+    {
+        return (gte::BSNumber<UIntegerType>)std::ceil((double)x);
+    }
+
+    template <typename UIntegerType>
+    inline gte::BSNumber<UIntegerType> cos(gte::BSNumber<UIntegerType> const& x)
+    {
+        return (gte::BSNumber<UIntegerType>)std::cos((double)x);
+    }
+
+    template <typename UIntegerType>
+    inline gte::BSNumber<UIntegerType> cosh(gte::BSNumber<UIntegerType> const& x)
+    {
+        return (gte::BSNumber<UIntegerType>)std::cosh((double)x);
+    }
+
+    template <typename UIntegerType>
+    inline gte::BSNumber<UIntegerType> exp(gte::BSNumber<UIntegerType> const& x)
+    {
+        return (gte::BSNumber<UIntegerType>)std::exp((double)x);
+    }
+
+    template <typename UIntegerType>
+    inline gte::BSNumber<UIntegerType> exp2(gte::BSNumber<UIntegerType> const& x)
+    {
+        return (gte::BSNumber<UIntegerType>)std::exp2((double)x);
+    }
+
+    template <typename UIntegerType>
+    inline gte::BSNumber<UIntegerType> floor(gte::BSNumber<UIntegerType> const& x)
+    {
+        return (gte::BSNumber<UIntegerType>)std::floor((double)x);
+    }
+
+    template <typename UIntegerType>
+    inline gte::BSNumber<UIntegerType> fmod(gte::BSNumber<UIntegerType> const& x, gte::BSNumber<UIntegerType> const& y)
+    {
+        return (gte::BSNumber<UIntegerType>)std::fmod((double)x, (double)y);
+    }
+
+    template <typename UIntegerType>
+    inline gte::BSNumber<UIntegerType> frexp(gte::BSNumber<UIntegerType> const& x, int* exponent)
+    {
+        return (gte::BSNumber<UIntegerType>)std::frexp((double)x, exponent);
+    }
+
+    template <typename UIntegerType>
+    inline gte::BSNumber<UIntegerType> ldexp(gte::BSNumber<UIntegerType> const& x, int exponent)
+    {
+        return (gte::BSNumber<UIntegerType>)std::ldexp((double)x, exponent);
+    }
+
+    template <typename UIntegerType>
+    inline gte::BSNumber<UIntegerType> log(gte::BSNumber<UIntegerType> const& x)
+    {
+        return (gte::BSNumber<UIntegerType>)std::log((double)x);
+    }
+
+    template <typename UIntegerType>
+    inline gte::BSNumber<UIntegerType> log2(gte::BSNumber<UIntegerType> const& x)
+    {
+        return (gte::BSNumber<UIntegerType>)std::log2((double)x);
+    }
+
+    template <typename UIntegerType>
+    inline gte::BSNumber<UIntegerType> log10(gte::BSNumber<UIntegerType> const& x)
+    {
+        return (gte::BSNumber<UIntegerType>)std::log10((double)x);
+    }
+
+    template <typename UIntegerType>
+    inline gte::BSNumber<UIntegerType> pow(gte::BSNumber<UIntegerType> const& x, gte::BSNumber<UIntegerType> const& y)
+    {
+        return (gte::BSNumber<UIntegerType>)std::pow((double)x, (double)y);
+    }
+
+    template <typename UIntegerType>
+    inline gte::BSNumber<UIntegerType> sin(gte::BSNumber<UIntegerType> const& x)
+    {
+        return (gte::BSNumber<UIntegerType>)std::sin((double)x);
+    }
+
+    template <typename UIntegerType>
+    inline gte::BSNumber<UIntegerType> sinh(gte::BSNumber<UIntegerType> const& x)
+    {
+        return (gte::BSNumber<UIntegerType>)std::sinh((double)x);
+    }
+
+    template <typename UIntegerType>
+    inline gte::BSNumber<UIntegerType> sqrt(gte::BSNumber<UIntegerType> const& x)
+    {
+        return (gte::BSNumber<UIntegerType>)std::sqrt((double)x);
+    }
+
+    template <typename UIntegerType>
+    inline gte::BSNumber<UIntegerType> tan(gte::BSNumber<UIntegerType> const& x)
+    {
+        return (gte::BSNumber<UIntegerType>)std::tan((double)x);
+    }
+
+    template <typename UIntegerType>
+    inline gte::BSNumber<UIntegerType> tanh(gte::BSNumber<UIntegerType> const& x)
+    {
+        return (gte::BSNumber<UIntegerType>)std::tanh((double)x);
+    }
+}
+
+namespace gte
+{
+    template <typename UIntegerType>
+    inline BSNumber<UIntegerType> atandivpi(BSNumber<UIntegerType> const& x)
+    {
+        return (BSNumber<UIntegerType>)atandivpi((double)x);
+    }
+
+    template <typename UIntegerType>
+    inline BSNumber<UIntegerType> atan2divpi(BSNumber<UIntegerType> const& y, BSNumber<UIntegerType> const& x)
+    {
+        return (BSNumber<UIntegerType>)atan2divpi((double)y, (double)x);
+    }
+
+    template <typename UIntegerType>
+    inline BSNumber<UIntegerType> clamp(BSNumber<UIntegerType> const& x, BSNumber<UIntegerType> const& xmin, BSNumber<UIntegerType> const& xmax)
+    {
+        return (BSNumber<UIntegerType>)clamp((double)x, (double)xmin, (double)xmax);
+    }
+
+    template <typename UIntegerType>
+    inline BSNumber<UIntegerType> cospi(BSNumber<UIntegerType> const& x)
+    {
+        return (BSNumber<UIntegerType>)cospi((double)x);
+    }
+
+    template <typename UIntegerType>
+    inline BSNumber<UIntegerType> exp10(BSNumber<UIntegerType> const& x)
+    {
+        return (BSNumber<UIntegerType>)exp10((double)x);
+    }
+
+    template <typename UIntegerType>
+    inline BSNumber<UIntegerType> invsqrt(BSNumber<UIntegerType> const& x)
+    {
+        return (BSNumber<UIntegerType>)invsqrt((double)x);
+    }
+
+    template <typename UIntegerType>
+    inline int isign(BSNumber<UIntegerType> const& x)
+    {
+        return isign((double)x);
+    }
+
+    template <typename UIntegerType>
+    inline BSNumber<UIntegerType> saturate(BSNumber<UIntegerType> const& x)
+    {
+        return (BSNumber<UIntegerType>)saturate((double)x);
+    }
+
+    template <typename UIntegerType>
+    inline BSNumber<UIntegerType> sign(BSNumber<UIntegerType> const& x)
+    {
+        return (BSNumber<UIntegerType>)sign((double)x);
+    }
+
+    template <typename UIntegerType>
+    inline BSNumber<UIntegerType> sinpi(BSNumber<UIntegerType> const& x)
+    {
+        return (BSNumber<UIntegerType>)sinpi((double)x);
+    }
+
+    template <typename UIntegerType>
+    inline BSNumber<UIntegerType> sqr(BSNumber<UIntegerType> const& x)
+    {
+        return (BSNumber<UIntegerType>)sqr((double)x);
+    }
+
+    // See the comments in GteMath.h about trait is_arbitrary_precision.
+    template <typename UIntegerType>
+    struct is_arbitrary_precision_internal<BSNumber<UIntegerType>> : std::true_type {};
+}
diff --git a/Engine/Plugins/Experimental/GeometryProcessing/Source/GeometryAlgorithms/Private/GTEngine/Mathematics/GteBSPrecision.cpp b/Engine/Plugins/Experimental/GeometryProcessing/Source/GeometryAlgorithms/Private/GTEngine/Mathematics/GteBSPrecision.cpp
new file mode 100644
--- /dev/null
+++ b/Engine/Plugins/Experimental/GeometryProcessing/Source/GeometryAlgorithms/Private/GTEngine/Mathematics/GteBSPrecision.cpp
@@ -0,0 +1,232 @@
+// David Eberly, Geometric Tools, Redmond WA 98052
+// Copyright (c) 1998-2019
+// Distributed under the Boost Software License, Version 1.0.
+// http://www.boost.org/LICENSE_1_0.txt
+// http://www.geometrictools.com/License/Boost/LICENSE_1_0.txt
+// File Version: 3.0.0 (2016/06/19)
+
+#include <GTEngine/Mathematics/GteBSPrecision.h>
+#include <GTEngine/GTEnginePCH.h>
+#include <algorithm>
+using namespace gte;
+
+
+BSPrecision::BSPrecision(bool isFloat, bool forBSNumber)
+    :
+    mForBSNumber(forBSNumber)
+{
+    // NOTE:  It is not clear what the rationale is for the C++ Standard
+    // Library to set numeric_limits<float>::min_exponent to -125
+    // instead of -126; same issue with numeric_limits<double>::min_exponent
+    // set to -1021 instead of -1022.  Similarly, it is not clear why
+    // numeric_limits<float>::max_exponent is set to 128 when the maximum
+    // finite 'float' has exponent 127.  The exponent 128 is used in the
+    // 'float' infinity, but treating this as 2^{128} does not seem to be
+    // consistent with the IEEE 754-2008 standard.  Same issue with
+    // numeric_limits<double>::max_exponent set to 1024 rather than 1023.
+
+    if (isFloat)
+    {
+        mNumBits = std::numeric_limits<float>::digits;  // 24
+        mMinBiasedExponent =
+            std::numeric_limits<float>::min_exponent - mNumBits;  // -149
+        mMaxExponent = std::numeric_limits<float>::max_exponent - 1;  // 127
+    }
+    else
+    {
+        mNumBits = std::numeric_limits<double>::digits;  // 53
+        mMinBiasedExponent =
+            std::numeric_limits<double>::min_exponent - mNumBits;  // -1074
+        mMaxExponent = std::numeric_limits<double>::max_exponent - 1;  // 1023
+    }
+}
+
+BSPrecision::BSPrecision(bool isFloat, int32_t maxExponent, bool forBSNumber)
+    :
+    mForBSNumber(forBSNumber)
+{
+    if (isFloat)
+    {
+        mNumBits = std::numeric_limits<float>::digits;  // 24
+        mMinBiasedExponent =
+            std::numeric_limits<float>::min_exponent - mNumBits;  // -149
+        mMaxExponent = maxExponent;
+    }
+    else
+    {
+        mNumBits = std::numeric_limits<double>::digits;  // 53
+        mMinBiasedExponent =
+            std::numeric_limits<double>::min_exponent - mNumBits;  // -1074
+        mMaxExponent = maxExponent;
+    }
+}
+
+BSPrecision::BSPrecision(int32_t numBits, int32_t minBiasedExponent,
+    int32_t maxExponent, bool forBSNumber)
+    :
+    mNumBits(numBits),
+    mMinBiasedExponent(minBiasedExponent),
+    mMaxExponent(maxExponent),
+    mForBSNumber(forBSNumber)
+{
+}
+
+int32_t BSPrecision::GetNumWords() const
+{
+    return mNumBits / 32 + ((mNumBits % 32) > 0 ? 1 : 0);
+}
+
+int32_t BSPrecision::GetNumBits() const
+{
+    return mNumBits;
+}
+
+int32_t BSPrecision::GetMinBiasedExponent() const
+{
+    return mMinBiasedExponent;
+}
+
+int32_t BSPrecision::GetMinExponent() const
+{
+    return mMinBiasedExponent + mNumBits - 1;
+}
+
+int32_t BSPrecision::GetMaxExponent() const
+{
+    return mMaxExponent;
+}
+
+BSPrecision BSPrecision::operator*(BSPrecision const& precision)
+{
+    // BSRational operator* involves a product of numerators and a product
+    // of denominators.  In worst case, the precision requirements are the
+    // same as a BSNumber operator*, so testing of mForBSNumber is not
+    // needed.
+    int32_t numBits = mNumBits + precision.mNumBits;
+    int32_t maxExponent = mMaxExponent + precision.mMaxExponent + 1;
+    int32_t minBiasedExponent =
+        mMinBiasedExponent + precision.mMinBiasedExponent;
+    return BSPrecision(numBits, minBiasedExponent, maxExponent, mForBSNumber);
+}
+
+BSPrecision BSPrecision::operator+(BSPrecision const& precision)
+{
+    if (mForBSNumber)
+    {
+        int32_t minBiasedExponent =
+            std::min(mMinBiasedExponent, precision.mMinBiasedExponent);
+        int32_t maxExponent =
+            std::max(mMaxExponent, precision.mMaxExponent) + 1;
+        int32_t numBits = maxExponent - minBiasedExponent;
+        return BSPrecision(numBits, minBiasedExponent, maxExponent,
+            mForBSNumber);
+    }
+    else
+    {
+        // n0/d0 + n1/d1 = (n0*d1 + n1*d0) / (d0*d1)
+        BSPrecision product = operator*(precision);
+        product.mForBSNumber = true;
+        BSPrecision sum = product + product;
+        sum.mForBSNumber = false;
+        BSPrecision division = sum / product;
+        division.mForBSNumber = false;
+        return division;
+    }
+}
+
+BSPrecision BSPrecision::operator-(BSPrecision const& precision)
+{
+    return operator+(precision);
+}
+
+BSPrecision BSPrecision::operator/(BSPrecision const& precision)
+{
+    if (mForBSNumber)
+    {
+        return precision;
+    }
+    else
+    {
+        // Division leads to multiplication of numerators and denominators.
+        return operator*(precision);
+    }
+}
+
+BSPrecision BSPrecision::operator==(BSPrecision const& precision)
+{
+    if (mForBSNumber)
+    {
+        return precision;
+    }
+    else
+    {
+        // Comparison leads to multiplication of numerators and denominators.
+        return operator*(precision);
+    }
+}
+
+BSPrecision BSPrecision::operator!=(BSPrecision const& precision)
+{
+    if (mForBSNumber)
+    {
+        return precision;
+    }
+    else
+    {
+        // Comparison leads to multiplication of numerators and denominators.
+        return operator*(precision);
+    }
+}
+
+BSPrecision BSPrecision::operator< (BSPrecision const& precision)
+{
+    if (mForBSNumber)
+    {
+        return precision;
+    }
+    else
+    {
+        // Comparison leads to multiplication of numerators and denominators.
+        return operator*(precision);
+    }
+}
+
+BSPrecision BSPrecision::operator<=(BSPrecision const& precision)
+{
+    if (mForBSNumber)
+    {
+        return precision;
+    }
+    else
+    {
+        // Comparison leads to multiplication of numerators and denominators.
+        return operator*(precision);
+    }
+}
+
+BSPrecision BSPrecision::operator> (BSPrecision const& precision)
+{
+    if (mForBSNumber)
+    {
+        return precision;
+    }
+    else
+    {
+        // Comparison leads to multiplication of numerators and denominators.
+        return operator*(precision);
+    }
+}
+
+BSPrecision BSPrecision::operator>=(BSPrecision const& precision)
+{
+    if (mForBSNumber)
+    {
+        return precision;
+    }
+    else
+    {
+        // Comparison leads to multiplication of numerators and denominators.
+        return operator*(precision);
+    }
+}
+
diff --git a/Engine/Plugins/Experimental/GeometryProcessing/Source/GeometryAlgorithms/Private/GTEngine/Mathematics/GteBSPrecision.h b/Engine/Plugins/Experimental/GeometryProcessing/Source/GeometryAlgorithms/Private/GTEngine/Mathematics/GteBSPrecision.h
new file mode 100644
--- /dev/null
+++ b/Engine/Plugins/Experimental/GeometryProcessing/Source/GeometryAlgorithms/Private/GTEngine/Mathematics/GteBSPrecision.h
@@ -0,0 +1,77 @@
+// David Eberly, Geometric Tools, Redmond WA 98052
+// Copyright (c) 1998-2019
+// Distributed under the Boost Software License, Version 1.0.
+// http://www.boost.org/LICENSE_1_0.txt
+// http://www.geometrictools.com/License/Boost/LICENSE_1_0.txt
+// File Version: 3.0.0 (2016/06/19)
+
+#pragma once
+
+#include <GTEngine/GTEngineDEF.h>
+#include <cstdint>
+#include <limits>
+
+// Support for determining the number of bits of precision required to compute
+// an expression.  See the document
+//   http://www.geometrictools.com/Documentation/ArbitraryPrecision.pdf
+// for an example of how to use this class.
+
+namespace gte
+{
+
+class BSPrecision
+{
+public:
+    // This constructor is used for 'float' or 'double'.  The floating-point
+    // inputs for the expressions have no restrictions; that is, the inputs
+    // can be any finite floating-point numbers (normal or subnormal).
+    BSPrecision(bool isFloat, bool forBSNumber);
+
+    // If you know that your inputs are limited in magnitude, use this
+    // constructor.  For example, if you know that your inputs x satisfy
+    // |x| <= 8, you can specify maxExponent of 3.  The minimum power will
+    // still be that for the smallest positive subnormal.
+    BSPrecision(bool isFloat, int32_t maxExponent, bool forBSNumber);
+
+    // You must use this constructor carefully based on knowledge of your
+    // expressions.  For example, if you know that your inputs are 'float'
+    // and in the interval [1,2), you would choose 24 for the number of bits
+    // of precision, a minimum biased exponent of -23 because the largest
+    // 'float' smaller than 2 is 1.1^{23}*2^0 = 1^{24}*2^{-23}, and a maximum
+    // exponent of 0.  These numbers work to determine bits of precision to
+    // compute x*y+z*w.  However, if you then compute an expression such as
+    // x-y for x and y in [1,2) and multiply by powers of 1/2, the bit
+    // counting will not be correct because the results can be subnormals
+    // where the minimum biased exponent is -149, not -23.
+    BSPrecision(int32_t numBits, int32_t minBiasedExponent,
+        int32_t maxExponent, bool forBSNumber);
+
+    // Member access.
+    int32_t GetNumWords() const;
+    int32_t GetNumBits() const;
+    int32_t GetMinBiasedExponent() const;
+    int32_t GetMinExponent() const;
+    int32_t GetMaxExponent() const;
+
+    // Support for determining the number of bits of precision required to
+    // compute an expression using BSNumber or BSRational.
+    BSPrecision operator*(BSPrecision const& precision);
+    BSPrecision operator+(BSPrecision const& precision);
+    BSPrecision operator-(BSPrecision const& precision);
+
+    // Support for determining the number of bits of precision required to
+    // compute a BSRational expression (operations not relevant for BSNumber).
+    BSPrecision operator/(BSPrecision const& precision);
+    BSPrecision operator==(BSPrecision const& precision);
+    BSPrecision operator!=(BSPrecision const& precision);
+    BSPrecision operator< (BSPrecision const& precision);
+    BSPrecision operator<=(BSPrecision const& precision);
+    BSPrecision operator> (BSPrecision const& precision);
+    BSPrecision operator>=(BSPrecision const& precision);
+
+private:
+    int32_t mNumBits, mMinBiasedExponent, mMaxExponent;
+    bool mForBSNumber;
+};
+
+}
diff --git a/Engine/Plugins/Experimental/GeometryProcessing/Source/GeometryAlgorithms/Private/GTEngine/Mathematics/GteBSRational.h b/Engine/Plugins/Experimental/GeometryProcessing/Source/GeometryAlgorithms/Private/GTEngine/Mathematics/GteBSRational.h
new file mode 100644
--- /dev/null
+++ b/Engine/Plugins/Experimental/GeometryProcessing/Source/GeometryAlgorithms/Private/GTEngine/Mathematics/GteBSRational.h
@@ -0,0 +1,748 @@
+// David Eberly, Geometric Tools, Redmond WA 98052
+// Copyright (c) 1998-2019
+// Distributed under the Boost Software License, Version 1.0.
+// http://www.boost.org/LICENSE_1_0.txt
+// http://www.geometrictools.com/License/Boost/LICENSE_1_0.txt
+// File Version: 3.0.3 (2019/01/17)
+
+#pragma once
+
+#include <GTEngine/Mathematics/GteBSNumber.h>
+
+// See the comments in GteBSNumber.h about the UIntegerType requirements.  The
+// denominator of a BSRational is chosen to be positive, which allows some
+// simplification of comparisons.  Also see the comments about exposing the
+// GTE_BINARY_SCIENTIFIC_SHOW_DOUBLE conditional define.
+
+namespace gte
+{
+    template <typename UIntegerType>
+    class BSRational
+    {
+    public:
+        // Construction.  The default constructor generates the zero BSRational.
+        // The constructors that take only numerators set the denominators to one.
+        BSRational()
+            :
+            mNumerator(0),
+            mDenominator(1)
+        {
+#if defined(GTE_BINARY_SCIENTIFIC_SHOW_DOUBLE)
+            mValue = (double)*this;
+#endif
+        }
+
+        BSRational(BSRational const& r)
+        {
+            *this = r;
+        }
+
+        BSRational(float numerator)
+            :
+            mNumerator(numerator),
+            mDenominator(1.0f)
+        {
+#if defined(GTE_BINARY_SCIENTIFIC_SHOW_DOUBLE)
+            mValue = (double)*this;
+#endif
+        }
+
+        BSRational(double numerator)
+            :
+            mNumerator(numerator),
+            mDenominator(1.0)
+        {
+#if defined(GTE_BINARY_SCIENTIFIC_SHOW_DOUBLE)
+            mValue = (double)*this;
+#endif
+        }
+
+        BSRational(int32_t numerator)
+            :
+            mNumerator(numerator),
+            mDenominator(1)
+        {
+#if defined(GTE_BINARY_SCIENTIFIC_SHOW_DOUBLE)
+            mValue = (double)*this;
+#endif
+        }
+
+        BSRational(uint32_t numerator)
+            :
+            mNumerator(numerator),
+            mDenominator(1)
+        {
+#if defined(GTE_BINARY_SCIENTIFIC_SHOW_DOUBLE)
+            mValue = (double)*this;
+#endif
+        }
+
+        BSRational(int64_t numerator)
+            :
+            mNumerator(numerator),
+            mDenominator(1)
+        {
+#if defined(GTE_BINARY_SCIENTIFIC_SHOW_DOUBLE)
+            mValue = (double)*this;
+#endif
+        }
+
+        BSRational(uint64_t numerator)
+            :
+            mNumerator(numerator),
+            mDenominator(1)
+        {
+#if defined(GTE_BINARY_SCIENTIFIC_SHOW_DOUBLE)
+            mValue = (double)*this;
+#endif
+        }
+
+        BSRational(BSNumber<UIntegerType> const& numerator)
+            :
+            mNumerator(numerator),
+            mDenominator(1)
+        {
+#if defined(GTE_BINARY_SCIENTIFIC_SHOW_DOUBLE)
+            mValue = (double)*this;
+#endif
+        }
+
+        BSRational(float numerator, float denominator)
+            :
+            mNumerator(numerator),
+            mDenominator(denominator)
+        {
+            LogAssert(mDenominator.mSign != 0, "Division by zero not allowed.");
+            if (mDenominator.mSign < 0)
+            {
+                mNumerator.mSign = -mNumerator.mSign;
+                mDenominator.mSign = 1;
+            }
+#if defined(GTE_BINARY_SCIENTIFIC_SHOW_DOUBLE)
+            mValue = (double)*this;
+#endif
+        }
+
+        BSRational(double numerator, double denominator)
+            :
+            mNumerator(numerator),
+            mDenominator(denominator)
+        {
+            LogAssert(mDenominator.mSign != 0, "Division by zero not allowed.");
+            if (mDenominator.mSign < 0)
+            {
+                mNumerator.mSign = -mNumerator.mSign;
+                mDenominator.mSign = 1;
+            }
+#if defined(GTE_BINARY_SCIENTIFIC_SHOW_DOUBLE)
+            mValue = (double)*this;
+#endif
+        }
+
+        BSRational(BSNumber<UIntegerType> const& numerator, BSNumber<UIntegerType> const& denominator)
+            :
+            mNumerator(numerator),
+            mDenominator(denominator)
+        {
+            LogAssert(mDenominator.mSign != 0, "Division by zero not allowed.");
+            if (mDenominator.mSign < 0)
+            {
+                mNumerator.mSign = -mNumerator.mSign;
+                mDenominator.mSign = 1;
+            }
+
+            // Set the exponent of the denominator to zero, but you can do so only
+            // by modifying the biased exponent.  Adjust the numerator accordingly.
+            // This prevents large growth of the exponents in both numerator and
+            // denominator simultaneously.
+            mNumerator.mBiasedExponent -= mDenominator.GetExponent();
+            mDenominator.mBiasedExponent =
+                -(mDenominator.GetUInteger().GetNumBits() - 1);
+
+#if defined(GTE_BINARY_SCIENTIFIC_SHOW_DOUBLE)
+            mValue = (double)*this;
+#endif
+        }
+
+        // Implicit conversions.
+        operator float() const
+        {
+            return Convert<uint32_t, float>();
+        }
+
+        operator double() const
+        {
+            return Convert<uint64_t, double>();
+        }
+
+        // Assignment.
+        BSRational& operator=(BSRational const& r)
+        {
+            mNumerator = r.mNumerator;
+            mDenominator = r.mDenominator;
+#if defined(GTE_BINARY_SCIENTIFIC_SHOW_DOUBLE)
+            mValue = (double)*this;
+#endif
+            return *this;
+        }
+
+        // Support for std::move.
+        BSRational(BSRational&& r)
+        {
+            *this = std::move(r);
+        }
+
+        BSRational& operator=(BSRational&& r)
+        {
+            mNumerator = std::move(r.mNumerator);
+            mDenominator = std::move(r.mDenominator);
+#if defined(GTE_BINARY_SCIENTIFIC_SHOW_DOUBLE)
+            mValue = (double)*this;
+#endif
+            return *this;
+        }
+
+        // Member access.
+        inline int GetSign() const
+        {
+            return mNumerator.GetSign() * mDenominator.GetSign();
+        }
+
+        inline BSNumber<UIntegerType> const& GetNumerator() const
+        {
+            return mNumerator;
+        }
+
+        inline BSNumber<UIntegerType> const& GetDenominator() const
+        {
+            return mDenominator;
+        }
+
+        // Comparisons.
+        bool operator==(BSRational const& r) const
+        {
+            // Do inexpensive sign tests first for optimum performance.
+            if (mNumerator.mSign != r.mNumerator.mSign)
+            {
+                return false;
+            }
+            if (mNumerator.mSign == 0)
+            {
+                // The numbers are both zero.
+                return true;
+            }
+
+            return mNumerator * r.mDenominator == mDenominator * r.mNumerator;
+        }
+
+        bool operator!=(BSRational const& r) const
+        {
+            return !operator==(r);
+        }
+
+        bool operator< (BSRational const& r) const
+        {
+            // Do inexpensive sign tests first for optimum performance.
+            if (mNumerator.mSign > 0)
+            {
+                if (r.mNumerator.mSign <= 0)
+                {
+                    return false;
+                }
+            }
+            else if (mNumerator.mSign == 0)
+            {
+                return r.mNumerator.mSign > 0;
+            }
+            else if (mNumerator.mSign < 0)
+            {
+                if (r.mNumerator.mSign >= 0)
+                {
+                    return true;
+                }
+            }
+
+            return mNumerator * r.mDenominator < mDenominator * r.mNumerator;
+        }
+
+        bool operator<=(BSRational const& r) const
+        {
+            return !r.operator<(*this);
+        }
+
+        bool operator> (BSRational const& r) const
+        {
+            return r.operator<(*this);
+        }
+
+        bool operator>=(BSRational const& r) const
+        {
+            return !operator<(r);
+        }
+
+        // Unary operations.
+        BSRational operator+() const
+        {
+            return *this;
+        }
+
+        BSRational operator-() const
+        {
+            return BSRational(-mNumerator, mDenominator);
+        }
+
+        // Arithmetic.
+        BSRational operator+(BSRational const& r) const
+        {
+            BSNumber<UIntegerType> product0 = mNumerator * r.mDenominator;
+            BSNumber<UIntegerType> product1 = mDenominator * r.mNumerator;
+            BSNumber<UIntegerType> numerator = product0 + product1;
+
+            // Complex expressions can lead to 0/denom, where denom is not 1.
+            if (numerator.mSign != 0)
+            {
+                BSNumber<UIntegerType> denominator = mDenominator * r.mDenominator;
+                return BSRational(numerator, denominator);
+            }
+            else
+            {
+                return BSRational(0);
+            }
+        }
+
+        BSRational operator-(BSRational const& r) const
+        {
+            BSNumber<UIntegerType> product0 = mNumerator * r.mDenominator;
+            BSNumber<UIntegerType> product1 = mDenominator * r.mNumerator;
+            BSNumber<UIntegerType> numerator = product0 - product1;
+
+            // Complex expressions can lead to 0/denom, where denom is not 1.
+            if (numerator.mSign != 0)
+            {
+                BSNumber<UIntegerType> denominator = mDenominator * r.mDenominator;
+                return BSRational(numerator, denominator);
+            }
+            else
+            {
+                return BSRational(0);
+            }
+        }
+
+        BSRational operator*(BSRational const& r) const
+        {
+            BSNumber<UIntegerType> numerator = mNumerator * r.mNumerator;
+
+            // Complex expressions can lead to 0/denom, where denom is not 1.
+            if (numerator.mSign != 0)
+            {
+                BSNumber<UIntegerType> denominator = mDenominator * r.mDenominator;
+                return BSRational(numerator, denominator);
+            }
+            else
+            {
+                return BSRational(0);
+            }
+        }
+
+        BSRational operator/(BSRational const& r) const
+        {
+            LogAssert(r.mNumerator.mSign != 0, "Division by zero not allowed.");
+
+            BSNumber<UIntegerType> numerator = mNumerator * r.mDenominator;
+
+            // Complex expressions can lead to 0/denom, where denom is not 1.
+            if (numerator.mSign != 0)
+            {
+                BSNumber<UIntegerType> denominator = mDenominator * r.mNumerator;
+                if (denominator.mSign < 0)
+                {
+                    numerator.mSign = -numerator.mSign;
+                    denominator.mSign = 1;
+                }
+                return BSRational(numerator, denominator);
+            }
+            else
+            {
+                return BSRational(0);
+            }
+        }
+
+        BSRational& operator+=(BSRational const& r)
+        {
+            *this = operator+(r);
+            return *this;
+        }
+
+        BSRational& operator-=(BSRational const& r)
+        {
+            *this = operator-(r);
+            return *this;
+        }
+
+        BSRational& operator*=(BSRational const& r)
+        {
+            *this = operator*(r);
+            return *this;
+        }
+
+        BSRational& operator/=(BSRational const& r)
+        {
+            *this = operator/(r);
+            return *this;
+        }
+
+        // Disk input/output.  The fstream objects should be created using
+        // std::ios::binary.  The return value is 'true' iff the operation
+        // was successful.
+        bool Write(std::ofstream& output) const
+        {
+            return mNumerator.Write(output) && mDenominator.Write(output);
+        }
+
+        bool Read(std::ifstream& input)
+        {
+            return mNumerator.Read(input) && mDenominator.Read(input);
+        }
+
+    private:
+        // Generic conversion code that converts to the correctly rounded
+        // result using round-to-nearest-ties-to-even.
+        template <typename UIntType, typename RealType>
+        RealType Convert() const
+        {
+            if (mNumerator.mSign == 0)
+            {
+                return (RealType)0;
+            }
+
+            // The ratio is abstractly of the form (1.u*2^p)/(1.v*2^q).
+            // Convert to the form (1.u/1.v)*2^{p-q}, if 1.u >= 1.v, or to the
+            // form (2*(1.u)/1.v)*2*{p-q-1}) if 1.u < 1.v.  The final form
+            // n/d must be in the interval [1,2).
+            BSNumber<UIntegerType> n = mNumerator, d = mDenominator;
+            int32_t sign = n.mSign * d.mSign;
+            n.mSign = 1;
+            d.mSign = 1;
+            int32_t pmq = n.GetExponent() - d.GetExponent();
+            n.mBiasedExponent = 1 - n.GetUInteger().GetNumBits();
+            d.mBiasedExponent = 1 - d.GetUInteger().GetNumBits();
+            if (BSNumber<UIntegerType>::LessThanIgnoreSign(n, d))
+            {
+                ++n.mBiasedExponent;
+                --pmq;
+            }
+
+            // At this time, n/d = 1.c in [1,2).  Define the sequence of bits
+            // w = 1c = w_{imax} w_{imax-1} ... w_0 w_{-1} w_{-2} ... where
+            // imax = precision(RealType)-1 and w_{imax} = 1.
+
+            // Compute 'precision' bits for w, the leading bit guaranteed to
+            // be 1 and occurring at index (1 << (precision-1)).
+            BSNumber<UIntegerType> one(1), two(2);
+            int const imax = std::numeric_limits<RealType>::digits - 1;
+            UIntType w = 0;
+            UIntType mask = ((UIntType)1 << imax);
+            for (int i = imax; i >= 0; --i, mask >>= 1)
+            {
+                if (BSNumber<UIntegerType>::LessThanIgnoreSign(n, d))
+                {
+                    n = two * n;
+                }
+                else
+                {
+                    n = two * (n - d);
+                    w |= mask;
+                }
+            }
+
+            // Apply the mode round-to-nearest-ties-to-even to decide whether
+            // to round down or up.  We computed w = w_{imax} ... w_0.  The
+            // remainder is n/d = w_{imax+1}.w_{imax+2}... in [0,2).  Compute
+            // n'/d = (n-d)/d in [-1,1).  Round-to-nearest-ties-to-even mode
+            // is the following, where we need only test the sign of n'.  A
+            // remainder of "half" is the case n' = 0.
+            //   Round down when n' < 0 or (n' = 0 and w_0 = 0):  use w
+            //   Round up when n' > 0 or (n' = 0 and w_0 == 1):  use w+1
+            n = n - d;
+            if (n.mSign > 0 || (n.mSign == 0 && (w & 1) == 1))
+            {
+                ++w;
+            }
+
+            if (w > 0)
+            {
+                // Ensure that the low-order bit of w is 1, which is required
+                // for the BSNumber integer part.
+                int32_t trailing = GetTrailingBit(w);
+                w >>= trailing;
+                pmq += trailing;
+
+                // Compute a BSNumber with integer part w and the appropriate
+                // number of bits and exponents.
+                BSNumber<UIntegerType> result(w);
+                result.mBiasedExponent = pmq - imax;
+                RealType converted = (RealType)result;
+                if (sign < 0)
+                {
+                    converted = -converted;
+                }
+                return converted;
+            }
+            else
+            {
+                return (RealType)0;
+            }
+            }
+
+#if defined(GTE_BINARY_SCIENTIFIC_SHOW_DOUBLE)
+    public:
+        // List this first so that it shows up first in the debugger watch
+        // window.
+        double mValue;
+    private:
+#endif
+
+        BSNumber<UIntegerType> mNumerator, mDenominator;
+
+        friend class UnitTestBSRational;
+    };
+}
+
+namespace std
+{
+    // TODO: Allow for implementations of the math functions in which a
+    // specified precision is used when computing the result.
+
+    template <typename UIntegerType>
+    inline gte::BSRational<UIntegerType> abs(gte::BSRational<UIntegerType> const& x)
+    {
+        return (x.GetSign() >= 0 ? x : -x);
+    }
+
+    template <typename UIntegerType>
+    inline gte::BSRational<UIntegerType> acos(gte::BSRational<UIntegerType> const& x)
+    {
+        return (gte::BSRational<UIntegerType>)std::acos((double)x);
+    }
+
+    template <typename UIntegerType>
+    inline gte::BSRational<UIntegerType> acosh(gte::BSRational<UIntegerType> const& x)
+    {
+        return (gte::BSRational<UIntegerType>)std::acosh((double)x);
+    }
+
+    template <typename UIntegerType>
+    inline gte::BSRational<UIntegerType> asin(gte::BSRational<UIntegerType> const& x)
+    {
+        return (gte::BSRational<UIntegerType>)std::asin((double)x);
+    }
+
+    template <typename UIntegerType>
+    inline gte::BSRational<UIntegerType> asinh(gte::BSRational<UIntegerType> const& x)
+    {
+        return (gte::BSRational<UIntegerType>)std::asinh((double)x);
+    }
+
+    template <typename UIntegerType>
+    inline gte::BSRational<UIntegerType> atan(gte::BSRational<UIntegerType> const& x)
+    {
+        return (gte::BSRational<UIntegerType>)std::atan((double)x);
+    }
+
+    template <typename UIntegerType>
+    inline gte::BSRational<UIntegerType> atanh(gte::BSRational<UIntegerType> const& x)
+    {
+        return (gte::BSRational<UIntegerType>)std::atanh((double)x);
+    }
+
+    template <typename UIntegerType>
+    inline gte::BSRational<UIntegerType> atan2(gte::BSRational<UIntegerType> const& y, gte::BSRational<UIntegerType> const& x)
+    {
+        return (gte::BSRational<UIntegerType>)std::atan2((double)y, (double)x);
+    }
+
+    template <typename UIntegerType>
+    inline gte::BSRational<UIntegerType> ceil(gte::BSRational<UIntegerType> const& x)
+    {
+        return (gte::BSRational<UIntegerType>)std::ceil((double)x);
+    }
+
+    template <typename UIntegerType>
+    inline gte::BSRational<UIntegerType> cos(gte::BSRational<UIntegerType> const& x)
+    {
+        return (gte::BSRational<UIntegerType>)std::cos((double)x);
+    }
+
+    template <typename UIntegerType>
+    inline gte::BSRational<UIntegerType> cosh(gte::BSRational<UIntegerType> const& x)
+    {
+        return (gte::BSRational<UIntegerType>)std::cosh((double)x);
+    }
+
+    template <typename UIntegerType>
+    inline gte::BSRational<UIntegerType> exp(gte::BSRational<UIntegerType> const& x)
+    {
+        return (gte::BSRational<UIntegerType>)std::exp((double)x);
+    }
+
+    template <typename UIntegerType>
+    inline gte::BSRational<UIntegerType> exp2(gte::BSRational<UIntegerType> const& x)
+    {
+        return (gte::BSRational<UIntegerType>)std::exp2((double)x);
+    }
+
+    template <typename UIntegerType>
+    inline gte::BSRational<UIntegerType> floor(gte::BSRational<UIntegerType> const& x)
+    {
+        return (gte::BSRational<UIntegerType>)std::floor((double)x);
+    }
+
+    template <typename UIntegerType>
+    inline gte::BSRational<UIntegerType> fmod(gte::BSRational<UIntegerType> const& x, gte::BSRational<UIntegerType> const& y)
+    {
+        return (gte::BSRational<UIntegerType>)std::fmod((double)x, (double)y);
+    }
+
+    template <typename UIntegerType>
+    inline gte::BSRational<UIntegerType> frexp(gte::BSRational<UIntegerType> const& x, int* exponent)
+    {
+        return (gte::BSRational<UIntegerType>)std::frexp((double)x, exponent);
+    }
+
+    template <typename UIntegerType>
+    inline gte::BSRational<UIntegerType> ldexp(gte::BSRational<UIntegerType> const& x, int exponent)
+    {
+        return (gte::BSRational<UIntegerType>)std::ldexp((double)x, exponent);
+    }
+
+    template <typename UIntegerType>
+    inline gte::BSRational<UIntegerType> log(gte::BSRational<UIntegerType> const& x)
+    {
+        return (gte::BSRational<UIntegerType>)std::log((double)x);
+    }
+
+    template <typename UIntegerType>
+    inline gte::BSRational<UIntegerType> log2(gte::BSRational<UIntegerType> const& x)
+    {
+        return (gte::BSRational<UIntegerType>)std::log2((double)x);
+    }
+
+    template <typename UIntegerType>
+    inline gte::BSRational<UIntegerType> log10(gte::BSRational<UIntegerType> const& x)
+    {
+        return (gte::BSRational<UIntegerType>)std::log10((double)x);
+    }
+
+    template <typename UIntegerType>
+    inline gte::BSRational<UIntegerType> pow(gte::BSRational<UIntegerType> const& x, gte::BSRational<UIntegerType> const& y)
+    {
+        return (gte::BSRational<UIntegerType>)std::pow((double)x, (double)y);
+    }
+
+    template <typename UIntegerType>
+    inline gte::BSRational<UIntegerType> sin(gte::BSRational<UIntegerType> const& x)
+    {
+        return (gte::BSRational<UIntegerType>)std::sin((double)x);
+    }
+
+    template <typename UIntegerType>
+    inline gte::BSRational<UIntegerType> sinh(gte::BSRational<UIntegerType> const& x)
+    {
+        return (gte::BSRational<UIntegerType>)std::sinh((double)x);
+    }
+
+    template <typename UIntegerType>
+    inline gte::BSRational<UIntegerType> sqrt(gte::BSRational<UIntegerType> const& x)
+    {
+        return (gte::BSRational<UIntegerType>)std::sqrt((double)x);
+    }
+
+    template <typename UIntegerType>
+    inline gte::BSRational<UIntegerType> tan(gte::BSRational<UIntegerType> const& x)
+    {
+        return (gte::BSRational<UIntegerType>)std::tan((double)x);
+    }
+
+    template <typename UIntegerType>
+    inline gte::BSRational<UIntegerType> tanh(gte::BSRational<UIntegerType> const& x)
+    {
+        return (gte::BSRational<UIntegerType>)std::tanh((double)x);
+    }
+}
+
+namespace gte
+{
+    template <typename UIntegerType>
+    inline BSRational<UIntegerType> atandivpi(BSRational<UIntegerType> const& x)
+    {
+        return (BSRational<UIntegerType>)atandivpi((double)x);
+    }
+
+    template <typename UIntegerType>
+    inline BSRational<UIntegerType> atan2divpi(BSRational<UIntegerType> const& y, BSRational<UIntegerType> const& x)
+    {
+        return (BSRational<UIntegerType>)atan2divpi((double)y, (double)x);
+    }
+
+    template <typename UIntegerType>
+    inline BSRational<UIntegerType> clamp(BSRational<UIntegerType> const& x, BSRational<UIntegerType> const& xmin, BSRational<UIntegerType> const& xmax)
+    {
+        return (BSRational<UIntegerType>)clamp((double)x, (double)xmin, (double)xmax);
+    }
+
+    template <typename UIntegerType>
+    inline BSRational<UIntegerType> cospi(BSRational<UIntegerType> const& x)
+    {
+        return (BSRational<UIntegerType>)cospi((double)x);
+    }
+
+    template <typename UIntegerType>
+    inline BSRational<UIntegerType> exp10(BSRational<UIntegerType> const& x)
+    {
+        return (BSRational<UIntegerType>)exp10((double)x);
+    }
+
+    template <typename UIntegerType>
+    inline BSRational<UIntegerType> invsqrt(BSRational<UIntegerType> const& x)
+    {
+        return (BSRational<UIntegerType>)invsqrt((double)x);
+    }
+
+    template <typename UIntegerType>
+    inline int isign(BSRational<UIntegerType> const& x)
+    {
+        return isign((double)x);
+    }
+
+    template <typename UIntegerType>
+    inline BSRational<UIntegerType> saturate(BSRational<UIntegerType> const& x)
+    {
+        return (BSRational<UIntegerType>)saturate((double)x);
+    }
+
+    template <typename UIntegerType>
+    inline BSRational<UIntegerType> sign(BSRational<UIntegerType> const& x)
+    {
+        return (BSRational<UIntegerType>)sign((double)x);
+    }
+
+    template <typename UIntegerType>
+    inline BSRational<UIntegerType> sinpi(BSRational<UIntegerType> const& x)
+    {
+        return (BSRational<UIntegerType>)sinpi((double)x);
+    }
+
+    template <typename UIntegerType>
+    inline BSRational<UIntegerType> sqr(BSRational<UIntegerType> const& x)
+    {
+        return (BSRational<UIntegerType>)sqr((double)x);
+    }
+
+    // See the comments in GteMath.h about traits is_arbitrary_precision
+    // and has_division_operator.
+    template <typename UIntegerType>
+    struct is_arbitrary_precision_internal<BSRational<UIntegerType>> : std::true_type {};
+
+    template <typename UIntegerType>
+    struct has_division_operator_internal<BSRational<UIntegerType>> : std::true_type {};
+}
diff --git a/Engine/Plugins/Experimental/GeometryProcessing/Source/GeometryAlgorithms/Private/GTEngine/Mathematics/GteBSplineCurve.h b/Engine/Plugins/Experimental/GeometryProcessing/Source/GeometryAlgorithms/Private/GTEngine/Mathematics/GteBSplineCurve.h
new file mode 100644
--- /dev/null
+++ b/Engine/Plugins/Experimental/GeometryProcessing/Source/GeometryAlgorithms/Private/GTEngine/Mathematics/GteBSplineCurve.h
@@ -0,0 +1,195 @@
+// David Eberly, Geometric Tools, Redmond WA 98052
+// Copyright (c) 1998-2019
+// Distributed under the Boost Software License, Version 1.0.
+// http://www.boost.org/LICENSE_1_0.txt
+// http://www.geometrictools.com/License/Boost/LICENSE_1_0.txt
+// File Version: 3.0.2 (2018/02/17)
+
+#pragma once
+
+#include <GTEngine/Mathematics/GteBasisFunction.h>
+#include <GTEngine/Mathematics/GteParametricCurve.h>
+
+namespace gte
+{
+
+template <int N, typename Real>
+class BSplineCurve : public ParametricCurve<N, Real>
+{
+public:
+    // Construction.  If the input controls is non-null, a copy is made of
+    // the controls.  To defer setting the control points, pass a null pointer
+    // and later access the control points via GetControls() or SetControl()
+    // member functions.  The domain is t in [t[d],t[n]], where t[d] and t[n]
+    // are knots with d the degree and n the number of control points.  To
+    // validate construction, create an object as shown:
+    //     BSplineCurve<N, Real> curve(parameters);
+    //     if (!curve) { <constructor failed, handle accordingly>; }
+    BSplineCurve(BasisFunctionInput<Real> const& input,
+        Vector<N, Real> const* controls);
+
+    // Member access.
+    inline BasisFunction<Real> const& GetBasisFunction() const;
+    inline int GetNumControls() const;
+    inline Vector<N, Real> const* GetControls() const;
+    inline Vector<N, Real>* GetControls();
+    void SetControl(int i, Vector<N, Real> const& control);
+    Vector<N, Real> const& GetControl(int i) const;
+
+    // Evaluation of the curve.  The function supports derivative calculation
+    // through order 3; that is, maxOrder <= 3 is required.  If you want
+    // only the position, pass in maxOrder of 0.  If you want the position and
+    // first derivative, pass in maxOrder of 1, and so on.  The output
+    // 'values' are ordered as: position, first derivative, second derivative,
+    // third derivative.
+    virtual void Evaluate(Real t, unsigned int maxOrder,
+        Vector<N, Real> values[4]) const;
+
+private:
+    // Support for Evaluate(...).
+    Vector<N, Real> Compute(unsigned int order, int imin, int imax) const;
+
+    BasisFunction<Real> mBasisFunction;
+    std::vector<Vector<N, Real>> mControls;
+};
+
+
+template <int N, typename Real>
+BSplineCurve<N, Real>::BSplineCurve(BasisFunctionInput<Real> const& input,
+    Vector<N, Real> const* controls)
+    :
+    ParametricCurve<N, Real>((Real)0, (Real)1),
+    mBasisFunction(input)
+{
+    if (!mBasisFunction)
+    {
+        // Errors were already generated during construction of the
+        // basis function.
+        return;
+    }
+
+    // The mBasisFunction stores the domain but so does ParametricCurve.
+    this->mTime.front() = mBasisFunction.GetMinDomain();
+    this->mTime.back() = mBasisFunction.GetMaxDomain();
+
+    // The replication of control points for periodic splines is avoided
+    // by wrapping the i-loop index in Evaluate.
+    mControls.resize(input.numControls);
+    if (controls)
+    {
+        std::copy(controls, controls + input.numControls, mControls.begin());
+    }
+    else
+    {
+        Vector<N, Real> zero{ (Real)0 };
+        std::fill(mControls.begin(), mControls.end(), zero);
+    }
+    this->mConstructed = true;
+}
+
+template <int N, typename Real>
+BasisFunction<Real> const& BSplineCurve<N, Real>::GetBasisFunction() const
+{
+    return mBasisFunction;
+}
+
+template <int N, typename Real>
+int BSplineCurve<N, Real>::GetNumControls() const
+{
+    return static_cast<int>(mControls.size());
+}
+
+template <int N, typename Real>
+Vector<N, Real> const* BSplineCurve<N, Real>::GetControls() const
+{
+    return mControls.data();
+}
+
+template <int N, typename Real>
+Vector<N, Real>* BSplineCurve<N, Real>::GetControls()
+{
+    return mControls.data();
+}
+
+template <int N, typename Real>
+void BSplineCurve<N, Real>::SetControl(int i, Vector<N, Real> const& control)
+{
+    if (0 <= i && i < GetNumControls())
+    {
+        mControls[i] = control;
+    }
+}
+
+template <int N, typename Real>
+Vector<N, Real> const& BSplineCurve<N, Real>::GetControl(int i) const
+{
+    if (0 <= i && i < GetNumControls())
+    {
+        return mControls[i];
+    }
+    else
+    {
+        return mControls[0];
+    }
+}
+
+template <int N, typename Real>
+void BSplineCurve<N, Real>::Evaluate(Real t, unsigned int maxOrder,
+    Vector<N, Real> values[4]) const
+{
+    if (!this->mConstructed)
+    {
+        // Errors were already generated during construction.
+        for (unsigned int order = 0; order < 4; ++order)
+        {
+            values[order].MakeZero();
+        }
+        return;
+    }
+
+    int imin, imax;
+    mBasisFunction.Evaluate(t, maxOrder, imin, imax);
+
+    // Compute position.
+    values[0] = Compute(0, imin, imax);
+    if (maxOrder >= 1)
+    {
+        // Compute first derivative.
+        values[1] = Compute(1, imin, imax);
+        if (maxOrder >= 2)
+        {
+            // Compute second derivative.
+            values[2] = Compute(2, imin, imax);
+            if (maxOrder == 3)
+            {
+                values[3] = Compute(3, imin, imax);
+            }
+            else
+            {
+                values[3].MakeZero();
+            }
+        }
+    }
+}
+
+template <int N, typename Real>
+Vector<N, Real> BSplineCurve<N, Real>::Compute(unsigned int order, int imin,
+    int imax) const
+{
+    // The j-index introduces a tiny amount of overhead in order to handle
+    // both aperiodic and periodic splines.  For aperiodic splines, j = i
+    // always.
+
+    int numControls = GetNumControls();
+    Vector<N, Real> result;
+    result.MakeZero();
+    for (int i = imin; i <= imax; ++i)
+    {
+        Real tmp = mBasisFunction.GetValue(order, i);
+        int j = (i >= numControls ? i - numControls : i);
+        result += tmp * mControls[j];
+    }
+    return result;
+}
+
+}
diff --git a/Engine/Plugins/Experimental/GeometryProcessing/Source/GeometryAlgorithms/Private/GTEngine/Mathematics/GteBSplineCurveFit.h b/Engine/Plugins/Experimental/GeometryProcessing/Source/GeometryAlgorithms/Private/GTEngine/Mathematics/GteBSplineCurveFit.h
new file mode 100644
--- /dev/null
+++ b/Engine/Plugins/Experimental/GeometryProcessing/Source/GeometryAlgorithms/Private/GTEngine/Mathematics/GteBSplineCurveFit.h
@@ -0,0 +1,257 @@
+// David Eberly, Geometric Tools, Redmond WA 98052
+// Copyright (c) 1998-2019
+// Distributed under the Boost Software License, Version 1.0.
+// http://www.boost.org/LICENSE_1_0.txt
+// http://www.geometrictools.com/License/Boost/LICENSE_1_0.txt
+// File Version: 3.0.0 (2016/06/19)
+
+#pragma once
+
+#include <GTEngine/Mathematics/GteBandedMatrix.h>
+#include <GTEngine/Mathematics/GteBasisFunction.h>
+
+namespace gte
+{
+
+template <typename Real>
+class BSplineCurveFit
+{
+public:
+    // Construction.  The preconditions for calling the constructor are
+    //   1 <= degree && degree < numControls <= numSamples
+    // The samples points are contiguous blocks of 'dimension' real values
+    // stored in sampleData.
+    BSplineCurveFit(int dimension, int numSamples, Real const* sampleData,
+        int degree, int numControls);
+
+    // Access to input sample information.
+    inline int GetDimension() const;
+    inline int GetNumSamples() const;
+    inline Real const* GetSampleData() const;
+
+    // Access to output control point and curve information.
+    inline int GetDegree() const;
+    inline int GetNumControls() const;
+    inline Real const* GetControlData() const;
+    inline BasisFunction<Real> const& GetBasis() const;
+
+    // Evaluation of the B-spline curve.  It is defined for 0 <= t <= 1.  If a
+    // t-value is outside [0,1], an open spline clamps it to [0,1].  The
+    // caller must ensure that position[] has at least 'dimension' elements.
+    void GetPosition(Real t, Real* position) const;
+
+private:
+    // Input sample information.
+    int mDimension;
+    int mNumSamples;
+    Real const* mSampleData;
+
+    // The fitted B-spline curve, open and with uniform knots.
+    int mDegree;
+    int mNumControls;
+    std::vector<Real> mControlData;
+    BasisFunction<Real> mBasis;
+};
+
+
+template <typename Real>
+BSplineCurveFit<Real>::BSplineCurveFit(int dimension, int numSamples,
+    Real const* sampleData, int degree, int numControls)
+    :
+    mDimension(dimension),
+    mNumSamples(numSamples),
+    mSampleData(sampleData),
+    mDegree(degree),
+    mNumControls(numControls),
+    mControlData(dimension * numControls)
+{
+    LogAssert(dimension >= 1, "Invalid dimension.");
+    LogAssert(1 <= degree && degree < numControls, "Invalid degree.");
+    LogAssert(sampleData, "Invalid sample data.");
+    LogAssert(numControls <= numSamples, "Invalid number of controls.");
+
+    BasisFunctionInput<Real> input;
+    input.numControls = numControls;
+    input.degree = degree;
+    input.uniform = true;
+    input.periodic = false;
+    input.numUniqueKnots = numControls - degree + 1;
+    input.uniqueKnots.resize(input.numUniqueKnots);
+    input.uniqueKnots[0].t = (Real)0;
+    input.uniqueKnots[0].multiplicity = degree + 1;
+    int last = input.numUniqueKnots - 1;
+    Real factor = ((Real)1) / (Real)last;
+    for (int i = 1; i < last; ++i)
+    {
+        input.uniqueKnots[i].t = factor * (Real)i;
+        input.uniqueKnots[i].multiplicity = 1;
+    }
+    input.uniqueKnots[last].t = (Real)1;
+    input.uniqueKnots[last].multiplicity = degree + 1;
+    mBasis.Create(input);
+
+    // Fit the data points with a B-spline curve using a least-squares error
+    // metric.  The problem is of the form A^T*A*Q = A^T*P, where A^T*A is a
+    // banded matrix, P contains the sample data, and Q is the unknown vector
+    // of control points.
+    Real tMultiplier = ((Real)1) / (Real)(mNumSamples - 1);
+    Real t;
+    int i0, i1, i2, imin, imax, j;
+
+    // Construct the matrix A^T*A.
+    int degp1 = mDegree + 1;
+    int numBands = (mNumControls > degp1 ? degp1 : mDegree);
+    BandedMatrix<Real> ATAMat(mNumControls, numBands, numBands);
+    for (i0 = 0; i0 < mNumControls; ++i0)
+    {
+        for (i1 = 0; i1 < i0; ++i1)
+        {
+            ATAMat(i0, i1) = ATAMat(i1, i0);
+        }
+
+        int i1Max = i0 + mDegree;
+        if (i1Max >= mNumControls)
+        {
+            i1Max = mNumControls - 1;
+        }
+
+        for (i1 = i0; i1 <= i1Max; ++i1)
+        {
+            Real value = (Real)0;
+            for (i2 = 0; i2 < mNumSamples; ++i2)
+            {
+                t = tMultiplier * (Real)i2;
+                mBasis.Evaluate(t, 0, imin, imax);
+                if (imin <= i0 && i0 <= imax && imin <= i1 && i1 <= imax)
+                {
+                    Real b0 = mBasis.GetValue(0, i0);
+                    Real b1 = mBasis.GetValue(0, i1);
+                    value += b0 * b1;
+                }
+            }
+            ATAMat(i0, i1) = value;
+        }
+    }
+
+    // Construct the matrix A^T.
+    Array2<Real> ATMat(mNumSamples, mNumControls);
+    memset(ATMat[0], 0, mNumControls * mNumSamples * sizeof(Real));
+    for (i0 = 0; i0 < mNumControls; ++i0)
+    {
+        for (i1 = 0; i1 < mNumSamples; ++i1)
+        {
+            t = tMultiplier * (Real)i1;
+            mBasis.Evaluate(t, 0, imin, imax);
+            if (imin <= i0 && i0 <= imax)
+            {
+                ATMat[i0][i1] = mBasis.GetValue(0, i0);
+            }
+        }
+    }
+
+    // Compute X0 = (A^T*A)^{-1}*A^T by solving the linear system
+    // A^T*A*X = A^T.
+    bool solved = ATAMat.template SolveSystem<true>(ATMat[0], mNumSamples);
+    LogAssert(solved, "Failed to solve linear system.");
+    (void)solved;
+
+    // The control points for the fitted curve are stored in the vector
+    // Q = X0*P, where P is the vector of sample data.
+    std::fill(mControlData.begin(), mControlData.end(), (Real)0);
+    for (i0 = 0; i0 < mNumControls; ++i0)
+    {
+        Real* Q = &mControlData[i0 * mDimension];
+        for (i1 = 0; i1 < mNumSamples; ++i1)
+        {
+            Real const* P = mSampleData + i1 * mDimension;
+            Real xValue = ATMat[i0][i1];
+            for (j = 0; j < mDimension; ++j)
+            {
+                Q[j] += xValue*P[j];
+            }
+        }
+    }
+
+    // Set the first and last output control points to match the first and
+    // last input samples.  This supports the application of fitting keyframe
+    // data with B-spline curves.  The user expects that the curve passes
+    // through the first and last positions in order to support matching two
+    // consecutive keyframe sequences.
+    Real* cEnd0 = &mControlData[0];
+    Real const* sEnd0 = mSampleData;
+    Real* cEnd1 = &mControlData[mDimension * (mNumControls - 1)];
+    Real const* sEnd1 = &mSampleData[mDimension * (mNumSamples - 1)];
+    for (j = 0; j < mDimension; ++j)
+    {
+        *cEnd0++ = *sEnd0++;
+        *cEnd1++ = *sEnd1++;
+    }
+}
+
+template <typename Real> inline
+int BSplineCurveFit<Real>::GetDimension() const
+{
+    return mDimension;
+}
+
+template <typename Real> inline
+int BSplineCurveFit<Real>::GetNumSamples() const
+{
+    return mNumSamples;
+}
+
+template <typename Real> inline
+Real const* BSplineCurveFit<Real>::GetSampleData() const
+{
+    return mSampleData;
+}
+
+template <typename Real> inline
+int BSplineCurveFit<Real>::GetDegree() const
+{
+    return mDegree;
+}
+
+template <typename Real> inline
+int BSplineCurveFit<Real>::GetNumControls() const
+{
+    return mNumControls;
+}
+
+template <typename Real> inline
+Real const* BSplineCurveFit<Real>::GetControlData() const
+{
+    return &mControlData[0];
+}
+
+template <typename Real> inline
+BasisFunction<Real> const& BSplineCurveFit<Real>::GetBasis() const
+{
+    return mBasis;
+}
+
+template <typename Real>
+void BSplineCurveFit<Real>::GetPosition(Real t, Real* position) const
+{
+    int imin, imax;
+    mBasis.Evaluate(t, 0, imin, imax);
+
+    Real const* source = &mControlData[mDimension * imin];
+    Real basisValue = mBasis.GetValue(0, imin);
+    for (int j = 0; j < mDimension; ++j)
+    {
+        position[j] = basisValue * (*source++);
+    }
+
+    for (int i = imin + 1; i <= imax; ++i)
+    {
+        basisValue = mBasis.GetValue(0, i);
+        for (int j = 0; j < mDimension; ++j)
+        {
+            position[j] += basisValue * (*source++);
+        }
+    }
+}
+
+
+}
