move GeometryProcessing plugin from NotForLicensees to Experimental

diff --git a/Engine/Plugins/Experimental/GeometryProcessing/Source/GeometryAlgorithms/Private/GTEngine/Mathematics/GteBSplineSurface.h b/Engine/Plugins/Experimental/GeometryProcessing/Source/GeometryAlgorithms/Private/GTEngine/Mathematics/GteBSplineSurface.h
new file mode 100644
--- /dev/null
+++ b/Engine/Plugins/Experimental/GeometryProcessing/Source/GeometryAlgorithms/Private/GTEngine/Mathematics/GteBSplineSurface.h
@@ -0,0 +1,208 @@
+// David Eberly, Geometric Tools, Redmond WA 98052
+// Copyright (c) 1998-2019
+// Distributed under the Boost Software License, Version 1.0.
+// http://www.boost.org/LICENSE_1_0.txt
+// http://www.geometrictools.com/License/Boost/LICENSE_1_0.txt
+// File Version: 3.0.2 (2018/02/17)
+
+#pragma once
+
+#include <GTEngine/Mathematics/GteBasisFunction.h>
+#include <GTEngine/Mathematics/GteParametricSurface.h>
+
+namespace gte
+{
+
+template <int N, typename Real>
+class BSplineSurface : public ParametricSurface<N, Real>
+{
+public:
+    // Construction.  If the input controls is non-null, a copy is made of
+    // the controls.  To defer setting the control points, pass a null pointer
+    // and later access the control points via GetControls() or SetControl()
+    // member functions.  The input 'controls' must be stored in row-major
+    // order, control[i0 + numControls0*i1].  As a 2D array, this corresponds
+    // to control2D[i1][i0].  To validate construction, create an object as
+    // shown:
+    //     BSplineSurface<N,Real> surface(parameters);
+    //     if (!surface) { <constructor failed, handle accordingly>; }
+    BSplineSurface(BasisFunctionInput<Real> const input[2],
+        Vector<N, Real> const* controls);
+
+    // Member access.  The index 'dim' must be in {0,1}.
+    inline BasisFunction<Real> const& GetBasisFunction(int dim) const;
+    inline int GetNumControls(int dim) const;
+    inline Vector<N, Real>* GetControls();
+    inline Vector<N, Real> const* GetControls() const;
+    void SetControl(int i0, int i1, Vector<N, Real> const& control);
+    Vector<N, Real> const& GetControl(int i0, int i1) const;
+
+    // Evaluation of the surface.  The function supports derivative
+    // calculation through order 2; that is, maxOrder <= 2 is required.  If
+    // you want only the position, pass in maxOrder of 0.  If you want the
+    // position and first-order derivatives, pass in maxOrder of 1, and so on.
+    // The output 'values' are ordered as: position X; first-order derivatives
+    // dX/du, dX/dv; second-order derivatives d2X/du2, d2X/dudv, d2X/dv2.
+    virtual void Evaluate(Real u, Real v, unsigned int maxOrder,
+        Vector<N, Real> values[6]) const;
+
+private:
+    // Support for Evaluate(...).
+    Vector<N, Real> Compute(unsigned int uOrder, unsigned int vOrder,
+        int iumin, int iumax, int ivmin, int ivmax) const;
+
+    std::array<BasisFunction<Real>, 2> mBasisFunction;
+    std::array<int, 2> mNumControls;
+    std::vector<Vector<N, Real>> mControls;
+};
+
+
+template <int N, typename Real>
+BSplineSurface<N, Real>::BSplineSurface(BasisFunctionInput<Real> const input[2],
+    Vector<N, Real> const* controls)
+    :
+    ParametricSurface<N, Real>((Real)0, (Real)1, (Real)0, (Real)1, true)
+{
+    for (int i = 0; i < 2; ++i)
+    {
+        mNumControls[i] = input[i].numControls;
+        mBasisFunction[i].Create(input[i]);
+        if (!mBasisFunction[i])
+        {
+            // Errors were already generated during construction of the
+            // basis functions.
+            return;
+        }
+    }
+
+    // The mBasisFunction stores the domain but so does ParametricCurve.
+    this->mUMin = mBasisFunction[0].GetMinDomain();
+    this->mUMax = mBasisFunction[0].GetMaxDomain();
+    this->mVMin = mBasisFunction[1].GetMinDomain();
+    this->mVMax = mBasisFunction[1].GetMaxDomain();
+
+    // The replication of control points for periodic splines is avoided
+    // by wrapping the i-loop index in Evaluate.
+    int numControls = mNumControls[0] * mNumControls[1];
+    mControls.resize(numControls);
+    if (controls)
+    {
+        std::copy(controls, controls + numControls, mControls.begin());
+    }
+    else
+    {
+        Vector<N, Real> zero{ (Real)0 };
+        std::fill(mControls.begin(), mControls.end(), zero);
+    }
+    this->mConstructed = true;
+}
+
+template <int N, typename Real>
+BasisFunction<Real> const& BSplineSurface<N, Real>::GetBasisFunction(int dim) const
+{
+    return mBasisFunction[dim];
+}
+
+template <int N, typename Real>
+int BSplineSurface<N, Real>::GetNumControls(int dim) const
+{
+    return mNumControls[dim];
+}
+
+template <int N, typename Real>
+Vector<N, Real> const* BSplineSurface<N, Real>::GetControls() const
+{
+    return mControls.data();
+}
+
+template <int N, typename Real>
+Vector<N, Real>* BSplineSurface<N, Real>::GetControls()
+{
+    return mControls.data();
+}
+
+template <int N, typename Real>
+void BSplineSurface<N, Real>::SetControl(int i0, int i1, Vector<N, Real> const& control)
+{
+    if (0 <= i0 && i0 < GetNumControls(0)
+        && 0 <= i1 && i1 < GetNumControls(1))
+    {
+        mControls[i0 + mNumControls[0] * i1] = control;
+    }
+}
+
+template <int N, typename Real>
+Vector<N, Real> const& BSplineSurface<N, Real>::GetControl(int i0, int i1) const
+{
+    if (0 <= i0 && i0 < GetNumControls(0) && 0 <= i1 && i1 < GetNumControls(1))
+    {
+        return mControls[i0 + mNumControls[0] * i1];
+    }
+    else
+    {
+        return mControls[0];
+    }
+}
+
+template <int N, typename Real>
+void BSplineSurface<N, Real>::Evaluate(Real u, Real v, unsigned int maxOrder,
+    Vector<N, Real> values[6]) const
+{
+    if (!this->mConstructed)
+    {
+        // Errors were already generated during construction.
+        for (int i = 0; i < 6; ++i)
+        {
+            values[i].MakeZero();
+        }
+        return;
+    }
+
+    int iumin, iumax, ivmin, ivmax;
+    mBasisFunction[0].Evaluate(u, maxOrder, iumin, iumax);
+    mBasisFunction[1].Evaluate(v, maxOrder, ivmin, ivmax);
+
+    // Compute position.
+    values[0] = Compute(0, 0, iumin, iumax, ivmin, ivmax);
+    if (maxOrder >= 1)
+    {
+        // Compute first-order derivatives.
+        values[1] = Compute(1, 0, iumin, iumax, ivmin, ivmax);
+        values[2] = Compute(0, 1, iumin, iumax, ivmin, ivmax);
+        if (maxOrder >= 2)
+        {
+            // Compute second-order derivatives.
+            values[3] = Compute(2, 0, iumin, iumax, ivmin, ivmax);
+            values[4] = Compute(1, 1, iumin, iumax, ivmin, ivmax);
+            values[5] = Compute(0, 2, iumin, iumax, ivmin, ivmax);
+        }
+    }
+}
+
+template <int N, typename Real>
+Vector<N, Real> BSplineSurface<N, Real>::Compute(unsigned int uOrder,
+    unsigned int vOrder, int iumin, int iumax, int ivmin, int ivmax) const
+{
+    // The j*-indices introduce a tiny amount of overhead in order to handle
+    // both aperiodic and periodic splines.  For aperiodic splines, j* = i*
+    // always.
+
+    int const numControls0 = mNumControls[0];
+    int const numControls1 = mNumControls[1];
+    Vector<N, Real> result;
+    result.MakeZero();
+    for (int iv = ivmin; iv <= ivmax; ++iv)
+    {
+        Real tmpv = mBasisFunction[1].GetValue(vOrder, iv);
+        int jv = (iv >= numControls1 ? iv - numControls1 : iv);
+        for (int iu = iumin; iu <= iumax; ++iu)
+        {
+            Real tmpu = mBasisFunction[0].GetValue(uOrder, iu);
+            int ju = (iu >= numControls0 ? iu - numControls0 : iu);
+            result += (tmpu * tmpv) * mControls[ju + numControls0 * jv];
+        }
+    }
+    return result;
+}
+
+}
diff --git a/Engine/Plugins/Experimental/GeometryProcessing/Source/GeometryAlgorithms/Private/GTEngine/Mathematics/GteBSplineSurfaceFit.h b/Engine/Plugins/Experimental/GeometryProcessing/Source/GeometryAlgorithms/Private/GTEngine/Mathematics/GteBSplineSurfaceFit.h
new file mode 100644
--- /dev/null
+++ b/Engine/Plugins/Experimental/GeometryProcessing/Source/GeometryAlgorithms/Private/GTEngine/Mathematics/GteBSplineSurfaceFit.h
@@ -0,0 +1,268 @@
+// David Eberly, Geometric Tools, Redmond WA 98052
+// Copyright (c) 1998-2019
+// Distributed under the Boost Software License, Version 1.0.
+// http://www.boost.org/LICENSE_1_0.txt
+// http://www.geometrictools.com/License/Boost/LICENSE_1_0.txt
+// File Version: 3.0.0 (2016/06/19)
+
+#pragma once
+
+#include <GTEngine/Mathematics/GteBandedMatrix.h>
+#include <GTEngine/Mathematics/GteVector3.h>
+#include <GTEngine/Mathematics/GteBasisFunction.h>
+
+namespace gte
+{
+
+template <typename Real>
+class BSplineSurfaceFit
+{
+public:
+    // Construction.  The preconditions for calling the constructor are
+    //   1 <= degree0 && degree0 + 1 < numControls0 <= numSamples0
+    //   1 <= degree1 && degree1 + 1 < numControls1 <= numSamples1
+    // The sample data must be in row-major order.  The control data is
+    // also stored in row-major order.
+    BSplineSurfaceFit(int degree0, int numControls0, int numSamples0,
+        int degree1, int numControls1, int numSamples1,
+        Vector3<Real> const* sampleData);
+
+    // Access to input sample information.
+    inline int GetNumSamples(int dimension) const;
+    inline Vector3<Real> const* GetSampleData() const;
+
+    // Access to output control point and surface information.
+    inline int GetDegree(int dimension) const;
+    inline int GetNumControls(int dimension) const;
+    inline Vector3<Real> const* GetControlData() const;
+    inline BasisFunction<Real> const& GetBasis(int dimension) const;
+
+    // Evaluation of the B-spline surface.  It is defined for 0 <= u <= 1
+    // and 0 <= v <= 1.  If a parameter value is outside [0,1], it is clamped
+    // to [0,1].
+    Vector3<Real> GetPosition(Real u, Real v) const;
+
+private:
+    // Input sample information.
+    int mNumSamples[2];
+    Vector3<Real> const* mSampleData;
+
+    // The fitted B-spline surface, open and with uniform knots.
+    int mDegree[2];
+    int mNumControls[2];
+    std::vector<Vector3<Real>> mControlData;
+    BasisFunction<Real> mBasis[2];
+};
+
+
+template <typename Real>
+BSplineSurfaceFit<Real>::BSplineSurfaceFit(int degree0, int numControls0,
+    int numSamples0, int degree1, int numControls1, int numSamples1,
+    Vector3<Real> const* sampleData)
+    :
+    mSampleData(sampleData),
+    mControlData(numControls0 * numControls1)
+{
+    LogAssert(1 <= degree0 && degree0 + 1 < numControls0, "Invalid degree.");
+    LogAssert(numControls0 <= numSamples0, "Invalid number of controls.");
+    LogAssert(1 <= degree1 && degree1 + 1 < numControls1, "Invalid degree.");
+    LogAssert(numControls1 <= numSamples1, "Invalid number of controls.");
+    LogAssert(sampleData, "Invalid sample data.");
+
+    mDegree[0] = degree0;
+    mNumSamples[0] = numSamples0;
+    mNumControls[0] = numControls0;
+    mDegree[1] = degree1;
+    mNumSamples[1] = numSamples1;
+    mNumControls[1] = numControls1;
+
+    BasisFunctionInput<Real> input;
+    Real tMultiplier[2];
+    int dim;
+    for (dim = 0; dim < 2; ++dim)
+    {
+        input.numControls = mNumControls[dim];
+        input.degree = mDegree[dim];
+        input.uniform = true;
+        input.periodic = false;
+        input.numUniqueKnots = mNumControls[dim] - mDegree[dim] + 1;
+        input.uniqueKnots.resize(input.numUniqueKnots);
+        input.uniqueKnots[0].t = (Real)0;
+        input.uniqueKnots[0].multiplicity = mDegree[dim] + 1;
+        int last = input.numUniqueKnots - 1;
+        Real factor = ((Real)1) / (Real)last;
+        for (int i = 1; i < last; ++i)
+        {
+            input.uniqueKnots[i].t = factor * (Real)i;
+            input.uniqueKnots[i].multiplicity = 1;
+        }
+        input.uniqueKnots[last].t = (Real)1;
+        input.uniqueKnots[last].multiplicity = mDegree[dim] + 1;
+        mBasis[dim].Create(input);
+
+        tMultiplier[dim] = ((Real)1) / (Real)(mNumSamples[dim] - 1);
+    }
+
+    // Fit the data points with a B-spline surface using a least-squares error
+    // metric.  The problem is of the form A0^T*A0*Q*A1^T*A1 = A0^T*P*A1, where
+    // A0^T*A0 and A1^T*A1 are banded matrices, P contains the sample data,
+    // and Q is the unknown matrix of control points.
+    Real t;
+    int i0, i1, i2, imin, imax;
+
+    // Construct the matrices A0^T*A0 and A1^T*A1.
+    BandedMatrix<Real> ATAMat[2] =
+    {
+        BandedMatrix<Real>(mNumControls[0], mDegree[0] + 1, mDegree[0] + 1),
+        BandedMatrix<Real>(mNumControls[1], mDegree[1] + 1, mDegree[1] + 1)
+    };
+
+    for (dim = 0; dim < 2; ++dim)
+    {
+        for (i0 = 0; i0 < mNumControls[dim]; ++i0)
+        {
+            for (i1 = 0; i1 < i0; ++i1)
+            {
+                ATAMat[dim](i0, i1) = ATAMat[dim](i1, i0);
+            }
+
+            int i1Max = i0 + mDegree[dim];
+            if (i1Max >= mNumControls[dim])
+            {
+                i1Max = mNumControls[dim] - 1;
+            }
+
+            for (i1 = i0; i1 <= i1Max; ++i1)
+            {
+                Real value = (Real)0;
+                for (i2 = 0; i2 < mNumSamples[dim]; ++i2)
+                {
+                    t = tMultiplier[dim] * (Real)i2;
+                    mBasis[dim].Evaluate(t, 0, imin, imax);
+                    if (imin <= i0 && i0 <= imax && imin <= i1 && i1 <= imax)
+                    {
+                        Real b0 = mBasis[dim].GetValue(0, i0);
+                        Real b1 = mBasis[dim].GetValue(0, i1);
+                        value += b0 * b1;
+                    }
+                }
+                ATAMat[dim](i0, i1) = value;
+            }
+        }
+    }
+
+    // Construct the matrices A0^T and A1^T.  A[d]^T has mNumControls[d]
+    // rows and mNumSamples[d] columns.
+    Array2<Real> ATMat[2];
+    for (dim = 0; dim < 2; dim++)
+    {
+        ATMat[dim] = Array2<Real>(mNumSamples[dim], mNumControls[dim]);
+        size_t numBytes = mNumControls[dim] * mNumSamples[dim] * sizeof(Real);
+        memset(ATMat[dim][0], 0, numBytes);
+        for (i0 = 0; i0 < mNumControls[dim]; ++i0)
+        {
+            for (i1 = 0; i1 < mNumSamples[dim]; ++i1)
+            {
+                t = tMultiplier[dim] * (Real)i1;
+                mBasis[dim].Evaluate(t, 0, imin, imax);
+                if (imin <= i0 && i0 <= imax)
+                {
+                    ATMat[dim][i0][i1] = mBasis[dim].GetValue(0, i0);
+                }
+            }
+        }
+    }
+
+    // Compute X0 = (A0^T*A0)^{-1}*A0^T and X1 = (A1^T*A1)^{-1}*A1^T by
+    // solving the linear systems A0^T*A0*X0 = A0^T and A1^T*A1*X1 = A1^T.
+    for (dim = 0; dim < 2; ++dim)
+    {
+        bool solved = ATAMat[dim].template SolveSystem<true>(ATMat[dim][0],
+            mNumSamples[dim]);
+        LogAssert(solved, "Failed to solve linear system.");
+        (void)solved;
+    }
+
+    // The control points for the fitted surface are stored in the matrix
+    // Q = X0*P*X1^T, where P is the matrix of sample data.
+    for (i1 = 0; i1 < mNumControls[1]; ++i1)
+    {
+        for (i0 = 0; i0 < mNumControls[0]; ++i0)
+        {
+            Vector3<Real> sum = Vector3<Real>::Zero();
+            for (int j1 = 0; j1 < mNumSamples[1]; ++j1)
+            {
+                Real x1Value = ATMat[1][i1][j1];
+                for (int j0 = 0; j0 < mNumSamples[0]; ++j0)
+                {
+                    Real x0Value = ATMat[0][i0][j0];
+                    Vector3<Real> sample =
+                        mSampleData[j0 + mNumSamples[0] * j1];
+                    sum += (x0Value * x1Value) * sample;
+                }
+            }
+            mControlData[i0 + mNumControls[0] * i1] = sum;
+        }
+    }
+}
+
+template <typename Real> inline
+int BSplineSurfaceFit<Real>::GetNumSamples(int dimension) const
+{
+    return mNumSamples[dimension];
+}
+
+template <typename Real> inline
+Vector3<Real> const* BSplineSurfaceFit<Real>::GetSampleData() const
+{
+    return mSampleData;
+}
+
+template <typename Real> inline
+int BSplineSurfaceFit<Real>::GetDegree(int dimension) const
+{
+    return mDegree[dimension];
+}
+
+template <typename Real> inline
+int BSplineSurfaceFit<Real>::GetNumControls(int dimension) const
+{
+    return mNumControls[dimension];
+}
+
+template <typename Real> inline
+Vector3<Real> const* BSplineSurfaceFit<Real>::GetControlData() const
+{
+    return &mControlData[0];
+}
+
+template <typename Real> inline
+BasisFunction<Real> const& BSplineSurfaceFit<Real>::GetBasis(int dimension)
+const
+{
+    return mBasis[dimension];
+}
+
+template <typename Real>
+Vector3<Real> BSplineSurfaceFit<Real>::GetPosition(Real u, Real v) const
+{
+    int iumin, iumax, ivmin, ivmax;
+    mBasis[0].Evaluate(u, 0, iumin, iumax);
+    mBasis[1].Evaluate(v, 0, ivmin, ivmax);
+
+    Vector3<Real> position = Vector3<Real>::Zero();
+    for (int iv = ivmin; iv <= ivmax; ++iv)
+    {
+        Real value1 = mBasis[1].GetValue(0, iv);
+        for (int iu = iumin; iu <= iumax; ++iu)
+        {
+            Real value0 = mBasis[0].GetValue(0, iu);
+            Vector3<Real> control = mControlData[iu + mNumControls[0] * iv];
+            position += (value0 * value1) * control;
+        }
+    }
+    return position;
+}
+
+
+}
diff --git a/Engine/Plugins/Experimental/GeometryProcessing/Source/GeometryAlgorithms/Private/GTEngine/Mathematics/GteBSplineVolume.h b/Engine/Plugins/Experimental/GeometryProcessing/Source/GeometryAlgorithms/Private/GTEngine/Mathematics/GteBSplineVolume.h
new file mode 100644
--- /dev/null
+++ b/Engine/Plugins/Experimental/GeometryProcessing/Source/GeometryAlgorithms/Private/GTEngine/Mathematics/GteBSplineVolume.h
@@ -0,0 +1,254 @@
+// David Eberly, Geometric Tools, Redmond WA 98052
+// Copyright (c) 1998-2019
+// Distributed under the Boost Software License, Version 1.0.
+// http://www.boost.org/LICENSE_1_0.txt
+// http://www.geometrictools.com/License/Boost/LICENSE_1_0.txt
+// File Version: 3.0.2 (2018/02/17)
+
+#pragma once
+
+#include <GTEngine/Mathematics/GteVector.h>
+#include <GTEngine/Mathematics/GteBasisFunction.h>
+
+namespace gte
+{
+
+template <int N, typename Real>
+class BSplineVolume
+{
+public:
+    // Construction.  If the input controls is non-null, a copy is made of
+    // the controls.  To defer setting the control points, pass a null pointer
+    // and later access the control points via GetControls() or SetControl()
+    // member functions.  The input 'controls' must be stored in
+    // lexicographical order, control[i0+numControls0*(i1+numControls1*i2)].
+    // As a 3D array, this corresponds to control3D[i2][i1][i0].
+    BSplineVolume(BasisFunctionInput<Real> const input[3],
+        Vector<N, Real> const* controls);
+
+    // To validate construction, create an object as shown:
+    //     BSplineVolume<N, Real> volume(parameters);
+    //     if (!volume) { <constructor failed, handle accordingly>; }
+    inline operator bool() const;
+
+    // Member access.  The index 'dim' must be in {0,1,2}.
+    inline BasisFunction<Real> const& GetBasisFunction(int dim) const;
+    inline Real GetMinDomain(int dim) const;
+    inline Real GetMaxDomain(int dim) const;
+    inline int GetNumControls(int dim) const;
+    inline Vector<N, Real> const* GetControls() const;
+    inline Vector<N, Real>* GetControls();
+    void SetControl(int i0, int i1, int i2, Vector<N, Real> const& control);
+    Vector<N, Real> const& GetControl(int i0, int i1, int i2) const;
+
+    // Evaluation of the volume.  The function supports derivative
+    // calculation through order 2; that is, maxOrder <= 2 is required.  If
+    // you want only the position, pass in maxOrder of 0.  If you want the
+    // position and first-order derivatives, pass in maxOrder of 1, and so on.
+    // The output 'values' are ordered as: position X; first-order derivatives
+    // dX/du, dX/dv, dX/dw; second-order derivatives d2X/du2, d2X/dv2,
+    // d2X/dw2, d2X/dudv, d2X/dudw, d2X/dvdw.
+    void Evaluate(Real u, Real v, Real w, unsigned int maxOrder,
+        Vector<N, Real> values[10]) const;
+
+private:
+    // Support for Evaluate(...).
+    Vector<N, Real> Compute(unsigned int uOrder, unsigned int vOrder,
+        unsigned int wOrder, int iumin, int iumax, int ivmin, int ivmax,
+        int iwmin, int iwmax) const;
+
+    std::array<BasisFunction<Real>, 3> mBasisFunction;
+    std::array<int, 3> mNumControls;
+    std::vector<Vector<N, Real>> mControls;
+    bool mConstructed;
+};
+
+
+template <int N, typename Real>
+BSplineVolume<N, Real>::BSplineVolume(BasisFunctionInput<Real> const input[3],
+    Vector<N, Real> const* controls)
+    :
+    mConstructed(false)
+{
+    for (int i = 0; i < 3; ++i)
+    {
+        mNumControls[i] = input[i].numControls;
+        mBasisFunction[i].Create(input[i]);
+        if (!mBasisFunction[i])
+        {
+            // Errors were already generated during construction of the
+            // basis functions.
+            return;
+        }
+    }
+
+    // The replication of control points for periodic splines is avoided
+    // by wrapping the i-loop index in Evaluate.
+    int numControls = mNumControls[0] * mNumControls[1] * mNumControls[2];
+    mControls.resize(numControls);
+    if (controls)
+    {
+        std::copy(controls, controls + numControls, mControls.begin());
+    }
+    else
+    {
+        Vector<N, Real> zero{ (Real)0 };
+        std::fill(mControls.begin(), mControls.end(), zero);
+    }
+    mConstructed = true;
+}
+
+template <int N, typename Real>
+BasisFunction<Real> const& BSplineVolume<N, Real>::GetBasisFunction(int dim) const
+{
+    return mBasisFunction[dim];
+}
+
+template <int N, typename Real>
+Real BSplineVolume<N, Real>::GetMinDomain(int dim) const
+{
+    return mBasisFunction[dim].GetMinDomain();
+}
+
+template <int N, typename Real>
+Real BSplineVolume<N, Real>::GetMaxDomain(int dim) const
+{
+    return mBasisFunction[dim].GetMaxDomain();
+}
+
+template <int N, typename Real>
+BSplineVolume<N, Real>::operator bool() const
+{
+    return mConstructed;
+}
+
+template <int N, typename Real>
+int BSplineVolume<N, Real>::GetNumControls(int dim) const
+{
+    return mNumControls[dim];
+}
+
+template <int N, typename Real>
+Vector<N, Real> const* BSplineVolume<N, Real>::GetControls() const
+{
+    return mControls.data();
+}
+
+template <int N, typename Real>
+Vector<N, Real>* BSplineVolume<N, Real>::GetControls()
+{
+    return mControls.data();
+}
+
+template <int N, typename Real>
+void BSplineVolume<N, Real>::SetControl(int i0, int i1, int i2,
+    Vector<N, Real> const& control)
+{
+    if (0 <= i0 && i0 < GetNumControls(0)
+        && 0 <= i1 && i1 < GetNumControls(1)
+        && 0 <= i2 && i2 < GetNumControls(2))
+    {
+        mControls[i0 + mNumControls[0] * (i1 + mNumControls[1] * i2)] = control;
+    }
+}
+
+template <int N, typename Real>
+Vector<N, Real> const& BSplineVolume<N, Real>::GetControl(int i0, int i1, int i2) const
+{
+    if (0 <= i0 && i0 < GetNumControls(0)
+        && 0 <= i1 && i1 < GetNumControls(1)
+        && 0 <= i2 && i2 < GetNumControls(2))
+    {
+        return mControls[i0 + mNumControls[0] * (i1 + mNumControls[1] * i2)];
+    }
+    else
+    {
+        return mControls[0];
+    }
+}
+
+template <int N, typename Real>
+void BSplineVolume<N, Real>::Evaluate(Real u, Real v, Real w,
+    unsigned int maxOrder, Vector<N, Real> values[10]) const
+{
+    if (!mConstructed)
+    {
+        // Errors were already generated during construction.
+        for (int i = 0; i < 10; ++i)
+        {
+            values[i].MakeZero();
+        }
+        return;
+    }
+
+    int iumin, iumax, ivmin, ivmax, iwmin, iwmax;
+    mBasisFunction[0].Evaluate(u, maxOrder, iumin, iumax);
+    mBasisFunction[1].Evaluate(v, maxOrder, ivmin, ivmax);
+    mBasisFunction[2].Evaluate(w, maxOrder, iwmin, iwmax);
+
+    // Compute position.
+    values[0] = Compute(0, 0, 0, iumin, iumax, ivmin, ivmax, iwmin, iwmax);
+    if (maxOrder >= 1)
+    {
+        // Compute first-order derivatives.
+        values[1] =
+            Compute(1, 0, 0, iumin, iumax, ivmin, ivmax, iwmin, iwmax);
+        values[2] =
+            Compute(0, 1, 0, iumin, iumax, ivmin, ivmax, iwmin, iwmax);
+        values[3] =
+            Compute(0, 0, 1, iumin, iumax, ivmin, ivmax, iwmin, iwmax);
+        if (maxOrder >= 2)
+        {
+            // Compute second-order derivatives.
+            values[4] =
+                Compute(2, 0, 0, iumin, iumax, ivmin, ivmax, iwmin, iwmax);
+            values[5] =
+                Compute(0, 2, 0, iumin, iumax, ivmin, ivmax, iwmin, iwmax);
+            values[6] =
+                Compute(0, 0, 2, iumin, iumax, ivmin, ivmax, iwmin, iwmax);
+            values[7] =
+                Compute(1, 1, 0, iumin, iumax, ivmin, ivmax, iwmin, iwmax);
+            values[8] =
+                Compute(1, 0, 1, iumin, iumax, ivmin, ivmax, iwmin, iwmax);
+            values[9] =
+                Compute(0, 1, 1, iumin, iumax, ivmin, ivmax, iwmin, iwmax);
+        }
+    }
+}
+
+template <int N, typename Real>
+Vector<N, Real> BSplineVolume<N, Real>::Compute(unsigned int uOrder,
+    unsigned int vOrder, unsigned int wOrder, int iumin, int iumax, int ivmin,
+    int ivmax, int iwmin, int iwmax) const
+{
+    // The j*-indices introduce a tiny amount of overhead in order to handle
+    // both aperiodic and periodic splines.  For aperiodic splines, j* = i*
+    // always.
+
+    int const numControls0 = mNumControls[0];
+    int const numControls1 = mNumControls[1];
+    int const numControls2 = mNumControls[2];
+    Vector<N, Real> result;
+    result.MakeZero();
+    for (int iw = iwmin; iw <= iwmax; ++iw)
+    {
+        Real tmpw = mBasisFunction[2].GetValue(wOrder, iw);
+        int jw = (iw >= numControls2 ? iw - numControls2 : iw);
+        for (int iv = ivmin; iv <= ivmax; ++iv)
+        {
+            Real tmpv = mBasisFunction[1].GetValue(vOrder, iv);
+            Real tmpvw = tmpv * tmpw;
+            int jv = (iv >= numControls1 ? iv - numControls1 : iv);
+            for (int iu = iumin; iu <= iumax; ++iu)
+            {
+                Real tmpu = mBasisFunction[0].GetValue(uOrder, iu);
+                int ju = (iu >= numControls0 ? iu - numControls0 : iu);
+                result += (tmpu * tmpvw) *
+                    mControls[ju + numControls0*(jv + numControls1*jw)];
+            }
+        }
+    }
+    return result;
+}
+
+}
diff --git a/Engine/Plugins/Experimental/GeometryProcessing/Source/GeometryAlgorithms/Private/GTEngine/Mathematics/GteBandedMatrix.h b/Engine/Plugins/Experimental/GeometryProcessing/Source/GeometryAlgorithms/Private/GTEngine/Mathematics/GteBandedMatrix.h
new file mode 100644
--- /dev/null
+++ b/Engine/Plugins/Experimental/GeometryProcessing/Source/GeometryAlgorithms/Private/GTEngine/Mathematics/GteBandedMatrix.h
@@ -0,0 +1,573 @@
+// David Eberly, Geometric Tools, Redmond WA 98052
+// Copyright (c) 1998-2019
+// Distributed under the Boost Software License, Version 1.0.
+// http://www.boost.org/LICENSE_1_0.txt
+// http://www.geometrictools.com/License/Boost/LICENSE_1_0.txt
+// File Version: 3.0.1 (2018/10/05)
+
+#pragma once
+
+#include <GTEngine/LowLevel/GteLexicoArray2.h>
+#include <cmath>
+#include <vector>
+
+namespace gte
+{
+
+template <typename Real>
+class BandedMatrix
+{
+public:
+    // Construction and destruction.
+    ~BandedMatrix();
+    BandedMatrix(int size, int numLBands, int numUBands);
+
+    // Member access.
+    inline int GetSize() const;
+    inline std::vector<Real>& GetDBand();
+    inline std::vector<Real> const& GetDBand() const;
+    inline std::vector<std::vector<Real>>& GetLBands();
+    inline std::vector<std::vector<Real>> const& GetLBands() const;
+    inline std::vector<std::vector<Real>>& GetUBands();
+    inline std::vector<std::vector<Real>> const& GetUBands() const;
+    Real& operator()(int r, int c);
+    Real const& operator()(int r, int c) const;
+
+    // Factor the square banded matrix A into A = L*L^T, where L is a
+    // lower-triangular matrix (L^T is an upper-triangular matrix).  This is
+    // an LU decomposition that allows for stable inversion of A to solve
+    // A*X = B.  The return value is 'true' iff the factorizing is successful
+    // (L is invertible).  If successful, A contains the Cholesky
+    // factorization: L in the lower-triangular part of A and L^T in the
+    // upper-triangular part of A.
+    bool CholeskyFactor();
+
+    // Solve the linear system A*X = B, where A is an NxN banded matrix and B
+    // is an Nx1 vector.  The unknown X is also Nx1.  The input to this
+    // function is B.  The output X is computed and stored in B.  The return
+    // value is 'true' iff the system has a solution.  The matrix A and the
+    // vector B are both modified by this function.  If successful, A contains
+    // the Cholesky factorization: L in the lower-triangular part of A and
+    // L^T in the upper-triangular part of A.
+    bool SolveSystem(Real* bVector);
+
+    // Solve the linear system A*X = B, where A is an NxN banded matrix and
+    // B is an NxM matrix.  The unknown X is also NxM.  The input to this
+    // function is B.  The output X is computed and stored in B.  The return
+    // value is 'true' iff the system has a solution.  The matrix A and the
+    // vector B are both modified by this function.  If successful, A contains
+    // the Cholesky factorization: L in the lower-triangular part of A and
+    // L^T in the upper-triangular part of A.
+    //
+    // 'bMatrix' must have the storage order specified by the template
+    // parameter.
+    template <bool RowMajor>
+    bool SolveSystem(Real* bMatrix, int numBColumns);
+
+    // Compute the inverse of the banded matrix.  The return value is 'true'
+    // when the matrix is invertible, in which case the 'inverse' output is
+    // valid.  The return value is 'false' when the matrix is not invertible,
+    // in which case 'inverse' is invalid and should not be used.  The
+    // input matrix 'inverse' must be the same size as 'this'.
+    //
+    // 'bMatrix' must have the storage order specified by the template
+    // parameter.
+    template <bool RowMajor>
+    bool ComputeInverse(Real* inverse) const;
+
+private:
+    // The linear system is L*U*X = B, where A = L*U and U = L^T,  Reduce this
+    // to U*X = L^{-1}*B.  The return value is 'true' iff the operation is
+    // successful.
+    bool SolveLower(Real* dataVector) const;
+
+    // The linear system is U*X = L^{-1}*B.  Reduce this to
+    // X = U^{-1}*L^{-1}*B.  The return value is 'true' iff the operation is
+    // successful.
+    bool SolveUpper(Real* dataVector) const;
+
+    // The linear system is L*U*X = B, where A = L*U and U = L^T,  Reduce this
+    // to U*X = L^{-1}*B.  The return value is 'true' iff the operation is
+    // successful.  See the comments for SolveSystem(Real*,int) about the
+    // storage for dataMatrix.
+    template <bool RowMajor>
+    bool SolveLower(Real* dataMatrix, int numColumns) const;
+
+    // The linear system is U*X = L^{-1}*B.  Reduce this to
+    // X = U^{-1}*L^{-1}*B.  The return value is 'true' iff the operation is
+    // successful.  See the comments for SolveSystem(Real*,int) about the
+    // storage for dataMatrix.
+    template <bool RowMajor>
+    bool SolveUpper(Real* dataMatrix, int numColumns) const;
+
+    int mSize;
+    std::vector<Real> mDBand;
+    std::vector<std::vector<Real>> mLBands, mUBands;
+
+    // For return by operator()(int,int) for valid indices not in the bands,
+    // in which case the matrix entries are zero,
+    mutable Real mZero;
+};
+
+
+template <typename Real>
+BandedMatrix<Real>::~BandedMatrix()
+{
+}
+
+template <typename Real>
+BandedMatrix<Real>::BandedMatrix(int size, int numLBands, int numUBands)
+    :
+    mSize(size),
+    mZero((Real)0)
+{
+    if (size > 0
+        && 0 <= numLBands && numLBands < size
+        && 0 <= numUBands && numUBands < size)
+    {
+        mDBand.resize(size);
+        std::fill(mDBand.begin(), mDBand.end(), (Real)0);
+
+        if (numLBands > 0)
+        {
+            mLBands.resize(numLBands);
+            int numElements = size - 1;
+            for (auto& band : mLBands)
+            {
+                band.resize(numElements--);
+                std::fill(band.begin(), band.end(), (Real)0);
+            }
+        }
+
+        if (numUBands > 0)
+        {
+            mUBands.resize(numUBands);
+            int numElements = size - 1;
+            for (auto& band : mUBands)
+            {
+                band.resize(numElements--);
+                std::fill(band.begin(), band.end(), (Real)0);
+            }
+        }
+    }
+    else
+    {
+        // Invalid argument to BandedMatrix constructor.
+        mSize = 0;
+    }
+}
+
+template <typename Real> inline
+int BandedMatrix<Real>::GetSize() const
+{
+    return mSize;
+}
+
+template <typename Real> inline
+std::vector<Real>& BandedMatrix<Real>::GetDBand()
+{
+    return mDBand;
+}
+
+template <typename Real> inline
+std::vector<Real> const& BandedMatrix<Real>::GetDBand() const
+{
+    return mDBand;
+}
+
+template <typename Real> inline
+std::vector<std::vector<Real>>& BandedMatrix<Real>::GetLBands()
+{
+    return mLBands;
+}
+
+template <typename Real> inline
+std::vector<std::vector<Real>> const& BandedMatrix<Real>::GetLBands() const
+{
+    return mLBands;
+}
+
+template <typename Real> inline
+std::vector<std::vector<Real>>& BandedMatrix<Real>::GetUBands()
+{
+    return mUBands;
+}
+
+template <typename Real> inline
+std::vector<std::vector<Real>> const& BandedMatrix<Real>::GetUBands() const
+{
+    return mUBands;
+}
+
+template <typename Real>
+Real& BandedMatrix<Real>::operator()(int r, int c)
+{
+    if (0 <= r && r < mSize && 0 <= c && c < mSize)
+    {
+        int band = c - r;
+        if (band > 0)
+        {
+            int const numUBands = static_cast<int>(mUBands.size());
+            if (--band < numUBands && r < mSize - 1 - band)
+            {
+                return mUBands[band][r];
+            }
+        }
+        else if (band < 0)
+        {
+            band = -band;
+            int const numLBands = static_cast<int>(mLBands.size());
+            if (--band < numLBands && c < mSize - 1 - band)
+            {
+                return mLBands[band][c];
+            }
+        }
+        else
+        {
+            return mDBand[r];
+        }
+    }
+    // else invalid index
+
+
+    // Set the value to zero in case someone unknowingly modified mZero on a
+    // previous call to operator(int,int).
+    mZero = (Real)0;
+    return mZero;
+}
+
+template <typename Real>
+Real const& BandedMatrix<Real>::operator()(int r, int c) const
+{
+    if (0 <= r && r < mSize && 0 <= c && c < mSize)
+    {
+        int band = c - r;
+        if (band > 0)
+        {
+            int const numUBands = static_cast<int>(mUBands.size());
+            if (--band < numUBands && r < mSize - 1 - band)
+            {
+                return mUBands[band][r];
+            }
+        }
+        else if (band < 0)
+        {
+            band = -band;
+            int const numLBands = static_cast<int>(mLBands.size());
+            if (--band < numLBands && c < mSize - 1 - band)
+            {
+                return mLBands[band][c];
+            }
+        }
+        else
+        {
+            return mDBand[r];
+        }
+    }
+    // else invalid index
+
+
+    // Set the value to zero in case someone unknowingly modified mZero on a
+    // previous call to operator(int,int).
+    mZero = (Real)0;
+    return mZero;
+}
+
+template <typename Real>
+bool BandedMatrix<Real>::CholeskyFactor()
+{
+    if (mDBand.size() == 0 || mLBands.size() != mUBands.size())
+    {
+        // Invalid number of bands.
+        return false;
+    }
+
+    int const sizeM1 = mSize - 1;
+    int const numBands = static_cast<int>(mLBands.size());
+
+    int k, kMax;
+    for (int i = 0; i < mSize; ++i)
+    {
+        int jMin = i - numBands;
+        if (jMin < 0)
+        {
+            jMin = 0;
+        }
+
+        int j;
+        for (j = jMin; j < i; ++j)
+        {
+            kMax = j + numBands;
+            if (kMax > sizeM1)
+            {
+                kMax = sizeM1;
+            }
+
+            for (k = i; k <= kMax; ++k)
+            {
+                operator()(k, i) -= operator()(i, j)*operator()(k, j);
+            }
+        }
+
+        kMax = j + numBands;
+        if (kMax > sizeM1)
+        {
+            kMax = sizeM1;
+        }
+
+        for (k = 0; k < i; ++k)
+        {
+            operator()(k, i) = operator()(i, k);
+        }
+
+        Real diagonal = operator()(i, i);
+        if (diagonal <= (Real)0)
+        {
+            return false;
+        }
+        Real invSqrt = ((Real)1) / std::sqrt(diagonal);
+        for (k = i; k <= kMax; ++k)
+        {
+            operator()(k, i) *= invSqrt;
+        }
+    }
+
+    return true;
+}
+
+template <typename Real>
+bool BandedMatrix<Real>::SolveSystem(Real* bVector)
+{
+    return CholeskyFactor()
+        && SolveLower(bVector)
+        && SolveUpper(bVector);
+}
+
+template <typename Real>
+template <bool RowMajor>
+bool BandedMatrix<Real>::SolveSystem(Real* bMatrix, int numBColumns)
+{
+    return CholeskyFactor()
+        && SolveLower<RowMajor>(bMatrix, numBColumns)
+        && SolveUpper<RowMajor>(bMatrix, numBColumns);
+}
+
+template <typename Real>
+template <bool RowMajor>
+bool BandedMatrix<Real>::ComputeInverse(Real* inverse) const
+{
+    LexicoArray2<RowMajor, Real> invA(mSize, mSize, inverse);
+
+    BandedMatrix<Real> tmpA = *this;
+    for (int row = 0; row < mSize; ++row)
+    {
+        for (int col = 0; col < mSize; ++col)
+        {
+            if (row != col)
+            {
+                invA(row, col) = (Real)0;
+            }
+            else
+            {
+                invA(row, row) = (Real)1;
+            }
+        }
+    }
+
+    // Forward elimination.
+    for (int row = 0; row < mSize; ++row)
+    {
+        // The pivot must be nonzero in order to proceed.
+        Real diag = tmpA(row, row);
+        if (diag == (Real)0)
+        {
+            return false;
+        }
+
+        Real invDiag = ((Real)1) / diag;
+        tmpA(row, row) = (Real)1;
+
+        // Multiply the row to be consistent with diagonal term of 1.
+        int colMin = row + 1;
+        int colMax = colMin + static_cast<int>(mUBands.size());
+        if (colMax > mSize)
+        {
+            colMax = mSize;
+        }
+
+        int c;
+        for (c = colMin; c < colMax; ++c)
+        {
+            tmpA(row, c) *= invDiag;
+        }
+        for (c = 0; c <= row; ++c)
+        {
+            invA(row, c) *= invDiag;
+        }
+
+        // Reduce the remaining rows.
+        int rowMin = row + 1;
+        int rowMax = rowMin + static_cast<int>(mLBands.size());
+        if (rowMax > mSize)
+        {
+            rowMax = mSize;
+        }
+
+        for (int r = rowMin; r < rowMax; ++r)
+        {
+            Real mult = tmpA(r, row);
+            tmpA(r, row) = (Real)0;
+            for (c = colMin; c < colMax; ++c)
+            {
+                tmpA(r, c) -= mult*tmpA(row, c);
+            }
+            for (c = 0; c <= row; ++c)
+            {
+                invA(r, c) -= mult*invA(row, c);
+            }
+        }
+    }
+
+    // Backward elimination.
+    for (int row = mSize - 1; row >= 1; --row)
+    {
+        int rowMax = row - 1;
+        int rowMin = row - static_cast<int>(mUBands.size());
+        if (rowMin < 0)
+        {
+            rowMin = 0;
+        }
+
+        for (int r = rowMax; r >= rowMin; --r)
+        {
+            Real mult = tmpA(r, row);
+            tmpA(r, row) = (Real)0;
+            for (int c = 0; c < mSize; ++c)
+            {
+                invA(r, c) -= mult*invA(row, c);
+            }
+        }
+    }
+
+    return true;
+}
+
+template <typename Real>
+bool BandedMatrix<Real>::SolveLower(Real* dataVector) const
+{
+    int const size = static_cast<int>(mDBand.size());
+    for (int r = 0; r < size; ++r)
+    {
+        Real lowerRR = operator()(r, r);
+        if (lowerRR >(Real)0)
+        {
+            for (int c = 0; c < r; ++c)
+            {
+                Real lowerRC = operator()(r, c);
+                dataVector[r] -= lowerRC * dataVector[c];
+            }
+            dataVector[r] /= lowerRR;
+        }
+        else
+        {
+            return false;
+        }
+    }
+    return true;
+}
+
+template <typename Real>
+bool BandedMatrix<Real>::SolveUpper(Real* dataVector) const
+{
+    int const size = static_cast<int>(mDBand.size());
+    for (int r = size - 1; r >= 0; --r)
+    {
+        Real upperRR = operator()(r, r);
+        if (upperRR > (Real)0)
+        {
+            for (int c = r + 1; c < size; ++c)
+            {
+                Real upperRC = operator()(r, c);
+                dataVector[r] -= upperRC * dataVector[c];
+            }
+
+            dataVector[r] /= upperRR;
+        }
+        else
+        {
+            return false;
+        }
+    }
+    return true;
+}
+
+template <typename Real>
+template <bool RowMajor>
+bool BandedMatrix<Real>::SolveLower(Real* dataMatrix, int numColumns) const
+{
+    LexicoArray2<RowMajor, Real> data(mSize, numColumns, dataMatrix);
+
+    for (int r = 0; r < mSize; ++r)
+    {
+        Real lowerRR = operator()(r, r);
+        if (lowerRR >(Real)0)
+        {
+            for (int c = 0; c < r; ++c)
+            {
+                Real lowerRC = operator()(r, c);
+                for (int bCol = 0; bCol < numColumns; ++bCol)
+                {
+                    data(r, bCol) -= lowerRC * data(c, bCol);
+                }
+            }
+
+            Real inverse = ((Real)1) / lowerRR;
+            for (int bCol = 0; bCol < numColumns; ++bCol)
+            {
+                data(r, bCol) *= inverse;
+            }
+        }
+        else
+        {
+            return false;
+        }
+    }
+    return true;
+}
+
+template <typename Real>
+template <bool RowMajor>
+bool BandedMatrix<Real>::SolveUpper(Real* dataMatrix, int numColumns) const
+{
+    LexicoArray2<RowMajor, Real> data(mSize, numColumns, dataMatrix);
+
+    for (int r = mSize - 1; r >= 0; --r)
+    {
+        Real upperRR = operator()(r, r);
+        if (upperRR > (Real)0)
+        {
+            for (int c = r + 1; c < mSize; ++c)
+            {
+                Real upperRC = operator()(r, c);
+                for (int bCol = 0; bCol < numColumns; ++bCol)
+                {
+                    data(r, bCol) -= upperRC * data(c, bCol);
+                }
+            }
+
+            Real inverse = ((Real)1) / upperRR;
+            for (int bCol = 0; bCol < numColumns; ++bCol)
+            {
+                data(r, bCol) *= inverse;
+            }
+        }
+        else
+        {
+            return false;
+        }
+    }
+    return true;
+}
+
+
+}
diff --git a/Engine/Plugins/Experimental/GeometryProcessing/Source/GeometryAlgorithms/Private/GTEngine/Mathematics/GteBasisFunction.h b/Engine/Plugins/Experimental/GeometryProcessing/Source/GeometryAlgorithms/Private/GTEngine/Mathematics/GteBasisFunction.h
new file mode 100644
--- /dev/null
+++ b/Engine/Plugins/Experimental/GeometryProcessing/Source/GeometryAlgorithms/Private/GTEngine/Mathematics/GteBasisFunction.h
@@ -0,0 +1,596 @@
+// David Eberly, Geometric Tools, Redmond WA 98052
+// Copyright (c) 1998-2019
+// Distributed under the Boost Software License, Version 1.0.
+// http://www.boost.org/LICENSE_1_0.txt
+// http://www.geometrictools.com/License/Boost/LICENSE_1_0.txt
+// File Version: 3.0.3 (2018/10/05)
+
+#pragma once
+
+#include <GTEngine/LowLevel/GteArray2.h>
+#include <GTEngine/LowLevel/GteLogger.h>
+#include <cmath>
+#include <cstring>
+
+namespace gte
+{
+
+template <typename Real>
+struct UniqueKnot
+{
+    Real t;
+    int multiplicity;
+};
+
+template <typename Real>
+struct BasisFunctionInput
+{
+    // The members are uninitialized.
+    BasisFunctionInput();
+
+    // Construct an open uniform curve with t in [0,1].
+    BasisFunctionInput(int inNumControls, int inDegree);
+
+    int numControls;
+    int degree;
+    bool uniform;
+    bool periodic;
+    int numUniqueKnots;
+    std::vector<UniqueKnot<Real>> uniqueKnots;
+};
+
+template <typename Real>
+class BasisFunction
+{
+public:
+    // Let n be the number of control points. Let d be the degree, where
+    // 1 <= d <= n-1.  The number of knots is k = n + d + 1.  The knots are
+    // t[i] for 0 <= i < k and must be nondecreasing, t[i] <= t[i+1], but a
+    // knot value can be repeated.  Let s be the number of distinct knots.
+    // Let the distinct knots be u[j] for 0 <= j < s, so u[j] < u[j+1] for
+    // all j.  The set of u[j] is called a 'breakpoint sequence'.  Let
+    // m[j] >= 1 be the multiplicity; that is, if t[i] is the first occurrence
+    // of u[j], then t[i+r] = t[i] for 1 <= r < m[j].  The multiplicities have
+    // the constraints m[0] <= d+1, m[s-1] <= d+1, and m[j] <= d for
+    // 1 <= j <= s-2.  Also, k = sum_{j=0}^{s-1} m[j], which says the
+    // multiplicities account for all k knots.
+    //
+    // Given a knot vector (t[0],...,t[n+d]), the domain of the corresponding
+    // B-spline curve is the interval [t[d],t[n]].
+    //
+    // The corresponding B-spline or NURBS curve is characterized as follows.
+    // See "Geometric Modeling with Splines: An Introduction" by Elaine Cohen,
+    // Richard F. Riesenfeld, and Gershon Elber, AK Peters, 2001, Natick MA.
+    // The curve is 'open' when m[0] = m[s-1] = d+1; otherwise, it is
+    // 'floating'.  An open curve is uniform when the knots t[d] through
+    // t[n] are equally spaced; that is, t[i+1] - t[i] are a common value
+    // for d <= i <= n-1.  By implication, s = n-d+1 and m[j] = 1 for
+    // 1 <= j <= s-2.  An open curve that does not satisfy these conditions
+    // is said to be nonuniform.  A floating curve is uniform when
+    // m[j] = 1 for 0 <= j <= s-1 and t[i+1] - t[i] are a common value for
+    // 0 <= i <= k-2; otherwise, the floating curve is nonuniform.
+    //
+    // A special case of a floating curve is a periodic curve.  The intent
+    // is that the curve is closed, so the first and last control points
+    // should be the same, which ensures C^{0} continuity.  Higher-order
+    // continuity is obtained by repeating more control points.  If the
+    // control points are P[0] through P[n-1], append the points P[0]
+    // through P[d-1] to ensure C^{d-1} continuity.  Additionally the knots
+    // must be chosen properly.  You may choose t[d] through t[n] as you
+    // wish.  The other knots are defined by
+    //   t[i] - t[i-1] = t[n-d+i] - t[n-d+i-1]
+    //   t[n+i] - t[n+i-1] = t[d+i] - t[d+i-1]
+    // for 1 <= i <= d.
+
+    // Construction and destruction.  The determination that the curve is
+    // open or floating is based on the multiplicities.  The 'uniform' input
+    // is used to avoid misclassifications due to floating-point rounding
+    // errors.  Specifically, the breakpoints might be equally spaced
+    // (uniform) as real numbers, but the floating-point representations can
+    // have rounding errors that cause the knot differences not to be exactly
+    // the same constant.  A periodic curve can have uniform or nonuniform
+    // knots.  This object makes copies of the input arrays.
+
+    ~BasisFunction();
+    BasisFunction();
+    BasisFunction(BasisFunctionInput<Real> const& input);
+
+    // No copying is allowed.
+    BasisFunction(BasisFunction const&) = delete;
+    BasisFunction& operator=(BasisFunction const&) = delete;
+
+    // Support for explicit creation in classes that have std::array
+    // members involving BasisFunction.  This is a call-once function.
+    void Create(BasisFunctionInput<Real> const& input);
+
+    // The inputs have complicated validation tests.  You should create a
+    // BasisFunction object as shown:
+    //     BasisFunction<Real> function(parameters);
+    //     if (!function) { <constructor failed, respond accordingly>; }
+    inline operator bool() const;
+
+    // Member access.
+    inline int GetNumControls() const;
+    inline int GetDegree() const;
+    inline int GetNumUniqueKnots() const;
+    inline int GetNumKnots() const;
+    inline Real GetMinDomain() const;
+    inline Real GetMaxDomain() const;
+    inline bool IsOpen() const;
+    inline bool IsUniform() const;
+    inline bool IsPeriodic() const;
+    inline UniqueKnot<Real> const* GetUniqueKnots() const;
+    inline Real const* GetKnots() const;
+
+    // Evaluation of the basis function and its derivatives through order 3.
+    // For the function value only, pass order 0.  For the function and first
+    // derivative, pass order 1, and so on.
+    void Evaluate(Real t, unsigned int order, int& minIndex, int& maxIndex)
+        const;
+
+    // Access the results of the call to Evaluate(...).  The index i must
+    // satisfy minIndex <= i <= maxIndex.  If it is not, the function returns
+    // zero.  The separation of evaluation and access is based on local
+    // control of the basis function; that is, only the accessible values are
+    // (potentially) not zero.
+    Real GetValue(unsigned int order, int i) const;
+
+private:
+    // Determine the index i for which knot[i] <= t < knot[i+1].  The t-value
+    // is modified (wrapped for periodic splines, clamped for nonperiodic
+    // splines).
+    int GetIndex(Real& t) const;
+
+    // Constructor inputs and values derived from them.
+    int mNumControls;
+    int mDegree;
+    Real mTMin, mTMax, mTLength;
+    bool mOpen;
+    bool mUniform;
+    bool mPeriodic;
+    bool mConstructed;
+    std::vector<UniqueKnot<Real>> mUniqueKnots;
+    std::vector<Real> mKnots;
+
+    // Lookup information for the GetIndex() function.  The first element of
+    // the pair is a unique knot value.  The second element is the index in
+    // mKnots[] for the last occurrence of that knot value.  NOTE:  This is
+    // a heavily used function during Evaluate calls.  This member was
+    // initially set to be std::vector<*>, but in debug mode the range-based
+    // for-loop in GetIndex was taking an extremely long time due to checked
+    // calls involving iterators.  We modified this to a regular array to
+    // ensure that debug performance is better.
+    std::vector<std::pair<Real, int>> mKeys;
+
+    // Storage for the basis functions and their first three derivatives.
+    mutable Array2<Real> mValue[4];  // mValue[i] is array[d+1][n+d]
+};
+
+
+template <typename Real>
+BasisFunctionInput<Real>::BasisFunctionInput()
+{
+}
+
+template <typename Real>
+BasisFunctionInput<Real>::BasisFunctionInput(int inNumControls, int inDegree)
+    :
+    numControls(inNumControls),
+    degree(inDegree),
+    uniform(true),
+    periodic(false),
+    numUniqueKnots(numControls - degree + 1),
+    uniqueKnots(numUniqueKnots)
+{
+    uniqueKnots.front().t = (Real)0;
+    uniqueKnots.front().multiplicity = degree + 1;
+    for (int i = 1; i <= numUniqueKnots - 2; ++i)
+    {
+        uniqueKnots[i].t = i / (numUniqueKnots - (Real)1);
+        uniqueKnots[i].multiplicity = 1;
+    }
+    uniqueKnots.back().t = (Real)1;
+    uniqueKnots.back().multiplicity = degree + 1;
+}
+
+template <typename Real>
+BasisFunction<Real>::~BasisFunction()
+{
+}
+
+template <typename Real>
+BasisFunction<Real>::BasisFunction()
+    :
+    mNumControls(0),
+    mDegree(0),
+    mTMin((Real)0),
+    mTMax((Real)0),
+    mTLength((Real)0),
+    mOpen(false),
+    mUniform(false),
+    mPeriodic(false),
+    mConstructed(false)
+{
+}
+
+template <typename Real>
+BasisFunction<Real>::BasisFunction(BasisFunctionInput<Real> const& input)
+    :
+    mConstructed(false)
+{
+    Create(input);
+}
+
+
+template <typename Real>
+void BasisFunction<Real>::Create(BasisFunctionInput<Real> const& input)
+{
+    if (mConstructed)
+    {
+        LogError("Object already created.");
+        return;
+    }
+
+    mNumControls = (input.periodic ? input.numControls + input.degree : input.numControls);
+    mDegree = input.degree;
+    mTMin = (Real)0;
+    mTMax = (Real)0;
+    mTLength = (Real)0;
+    mOpen = false;
+    mUniform = input.uniform;
+    mPeriodic = input.periodic;
+    for (int i = 0; i < 4; ++i)
+    {
+        mValue[i] = Array2<Real>();
+    }
+
+    if (input.numControls < 2)
+    {
+        LogError("Invalid number of control points.");
+        return;
+    }
+
+    if (input.degree < 1 || input.degree >= input.numControls)
+    {
+        LogError("Invalid degree.");
+        return;
+    }
+
+    if (input.numUniqueKnots < 2)
+    {
+        LogError("Invalid number of unique knots.");
+        return;
+    }
+
+    mUniqueKnots.resize(input.numUniqueKnots);
+    std::copy(input.uniqueKnots.begin(), input.uniqueKnots.begin() + input.numUniqueKnots,
+        mUniqueKnots.begin());
+
+    Real u = mUniqueKnots.front().t;
+    for (int i = 1; i < input.numUniqueKnots - 1; ++i)
+    {
+        Real uNext = mUniqueKnots[i].t;
+        if (u >= uNext)
+        {
+            LogError("Unique knots are not strictly increasing.");
+            return;
+        }
+        u = uNext;
+    }
+
+    int mult0 = mUniqueKnots.front().multiplicity;
+    if (mult0 < 1 || mult0 > mDegree + 1)
+    {
+        LogError("Invalid first multiplicity.");
+        return;
+    }
+
+    int mult1 = mUniqueKnots.back().multiplicity;
+    if (mult1 < 1 || mult1 > mDegree + 1)
+    {
+        LogError("Invalid last multiplicity.");
+        return;
+    }
+
+    for (int i = 1; i <= input.numUniqueKnots - 2; ++i)
+    {
+        int mult = mUniqueKnots[i].multiplicity;
+        if (mult < 1 || mult > mDegree)
+        {
+            LogError("Invalid interior multiplicity.");
+            return;
+        }
+    }
+
+    mOpen = (mult0 == mult1 && mult0 == mDegree + 1);
+
+    mKnots.resize(mNumControls + mDegree + 1);
+    mKeys.resize(input.numUniqueKnots);
+    int sum = 0;
+    for (int i = 0, j = 0; i < input.numUniqueKnots; ++i)
+    {
+        Real tCommon = mUniqueKnots[i].t;
+        int mult = mUniqueKnots[i].multiplicity;
+        for (int k = 0; k < mult; ++k, ++j)
+        {
+            mKnots[j] = tCommon;
+        }
+
+        mKeys[i].first = tCommon;
+        mKeys[i].second = sum - 1;
+        sum += mult;
+    }
+
+    mTMin = mKnots[mDegree];
+    mTMax = mKnots[mNumControls];
+    mTLength = mTMax - mTMin;
+
+    size_t numRows = mDegree + 1;
+    size_t numCols = mNumControls + mDegree;
+    size_t numBytes = numRows * numCols * sizeof(Real);
+    for (int i = 0; i < 4; ++i)
+    {
+        mValue[i] = Array2<Real>(numCols, numRows);
+        memset(mValue[i][0], 0, numBytes);
+    }
+
+    mConstructed = true;
+}
+
+template <typename Real> inline
+BasisFunction<Real>::operator bool() const
+{
+    return mConstructed;
+}
+
+template <typename Real> inline
+int BasisFunction<Real>::GetNumControls() const
+{
+    return mNumControls;
+}
+
+template <typename Real> inline
+int BasisFunction<Real>::GetDegree() const
+{
+    return mDegree;
+}
+
+template <typename Real> inline
+int BasisFunction<Real>::GetNumUniqueKnots() const
+{
+    return static_cast<int>(mUniqueKnots.size());
+}
+
+template <typename Real> inline
+int BasisFunction<Real>::GetNumKnots() const
+{
+    return static_cast<int>(mKnots.size());
+}
+
+template <typename Real> inline
+Real BasisFunction<Real>::GetMinDomain() const
+{
+    return mTMin;
+}
+
+template <typename Real> inline
+Real BasisFunction<Real>::GetMaxDomain() const
+{
+    return mTMax;
+}
+
+template <typename Real> inline
+bool BasisFunction<Real>::IsOpen() const
+{
+    return mOpen;
+}
+
+template <typename Real> inline
+bool BasisFunction<Real>::IsUniform() const
+{
+    return mUniform;
+}
+
+template <typename Real> inline
+bool BasisFunction<Real>::IsPeriodic() const
+{
+    return mPeriodic;
+}
+
+template <typename Real> inline
+UniqueKnot<Real> const* BasisFunction<Real>::GetUniqueKnots() const
+{
+    return &mUniqueKnots[0];
+}
+
+template <typename Real> inline
+Real const* BasisFunction<Real>::GetKnots() const
+{
+    return &mKnots[0];
+}
+
+template <typename Real>
+void BasisFunction<Real>::Evaluate(Real t, unsigned int order, int& minIndex,
+    int& maxIndex) const
+{
+    if (!mConstructed)
+    {
+        // Errors were already generated during construction.  Return an index
+        // range that leads to zero-valued positions and derivatives.
+        minIndex = -1;
+        maxIndex = -1;
+        return;
+    }
+
+    if (order > 3)
+    {
+        LogError("Only derivatives through order 3 are supported.");
+        minIndex = 0;
+        maxIndex = 0;
+        return;
+    }
+
+    int i = GetIndex(t);
+    mValue[0][0][i] = (Real)1;
+
+    if (order >= 1)
+    {
+        mValue[1][0][i] = (Real)0;
+        if (order >= 2)
+        {
+            mValue[2][0][i] = (Real)0;
+            if (order >= 3)
+            {
+                mValue[3][0][i] = (Real)0;
+            }
+        }
+    }
+
+    Real n0 = t - mKnots[i], n1 = mKnots[i + 1] - t;
+    Real e0, e1, d0, d1, invD0, invD1;
+    int j;
+    for (j = 1; j <= mDegree; j++)
+    {
+        d0 = mKnots[i + j] - mKnots[i];
+        d1 = mKnots[i + 1] - mKnots[i - j + 1];
+        invD0 = (d0 > (Real)0 ? (Real)1 / d0 : (Real)0);
+        invD1 = (d1 > (Real)0 ? (Real)1 / d1 : (Real)0);
+
+        e0 = n0*mValue[0][j - 1][i];
+        mValue[0][j][i] = e0*invD0;
+        e1 = n1*mValue[0][j - 1][i - j + 1];
+        mValue[0][j][i - j] = e1*invD1;
+
+        if (order >= 1)
+        {
+            e0 = n0 * mValue[1][j - 1][i] + mValue[0][j - 1][i];
+            mValue[1][j][i] = e0 * invD0;
+            e1 = n1 * mValue[1][j - 1][i - j + 1] - mValue[0][j - 1][i - j + 1];
+            mValue[1][j][i - j] = e1 * invD1;
+
+            if (order >= 2)
+            {
+                e0 = n0 * mValue[2][j - 1][i] + ((Real)2) * mValue[1][j - 1][i];
+                mValue[2][j][i] = e0 * invD0;
+                e1 = n1 * mValue[2][j - 1][i - j + 1] - ((Real)2) * mValue[1][j - 1][i - j + 1];
+                mValue[2][j][i - j] = e1 * invD1;
+
+                if (order >= 3)
+                {
+                    e0 = n0 * mValue[3][j - 1][i] + ((Real)3) * mValue[2][j - 1][i];
+                    mValue[3][j][i] = e0 * invD0;
+                    e1 = n1 * mValue[3][j - 1][i - j + 1] - ((Real)3) * mValue[2][j - 1][i - j + 1];
+                    mValue[3][j][i - j] = e1 * invD1;
+                }
+            }
+        }
+    }
+
+    for (j = 2; j <= mDegree; ++j)
+    {
+        for (int k = i - j + 1; k < i; ++k)
+        {
+            n0 = t - mKnots[k];
+            n1 = mKnots[k + j + 1] - t;
+            d0 = mKnots[k + j] - mKnots[k];
+            d1 = mKnots[k + j + 1] - mKnots[k + 1];
+            invD0 = (d0 >(Real)0 ? (Real)1 / d0 : (Real)0);
+            invD1 = (d1 > (Real)0 ? (Real)1 / d1 : (Real)0);
+
+            e0 = n0*mValue[0][j - 1][k];
+            e1 = n1*mValue[0][j - 1][k + 1];
+            mValue[0][j][k] = e0 * invD0 + e1 * invD1;
+
+            if (order >= 1)
+            {
+                e0 = n0 * mValue[1][j - 1][k] + mValue[0][j - 1][k];
+                e1 = n1 * mValue[1][j - 1][k + 1] - mValue[0][j - 1][k + 1];
+                mValue[1][j][k] = e0 * invD0 + e1 * invD1;
+
+                if (order >= 2)
+                {
+                    e0 = n0 * mValue[2][j - 1][k] + ((Real)2) * mValue[1][j - 1][k];
+                    e1 = n1 * mValue[2][j - 1][k + 1] - ((Real)2) * mValue[1][j - 1][k + 1];
+                    mValue[2][j][k] = e0 * invD0 + e1 * invD1;
+
+                    if (order >= 3)
+                    {
+                        e0 = n0 * mValue[3][j - 1][k] + ((Real)3) * mValue[2][j - 1][k];
+                        e1 = n1 * mValue[3][j - 1][k + 1] - ((Real)3) * mValue[2][j - 1][k + 1];
+                        mValue[3][j][k] = e0 * invD0 + e1 * invD1;
+                    }
+                }
+            }
+        }
+    }
+
+    minIndex = i - mDegree;
+    maxIndex = i;
+}
+
+template <typename Real>
+Real BasisFunction<Real>::GetValue(unsigned int order, int i) const
+{
+    if (!mConstructed)
+    {
+        // Errors were already generated during construction.  Return a value
+        // that leads to zero-valued positions and derivatives.
+        return (Real)0;
+    }
+
+    if (order < 4)
+    {
+        if (0 <= i && i < mNumControls + mDegree)
+        {
+            return mValue[order][mDegree][i];
+        }
+    }
+
+    LogError("Invalid input.");
+    return (Real)0;
+}
+
+template <typename Real>
+int BasisFunction<Real>::GetIndex(Real& t) const
+{
+    // Find the index i for which knot[i] <= t < knot[i+1].
+    if (mPeriodic)
+    {
+        // Wrap to [tmin,tmax].
+        Real r = std::fmod(t - mTMin, mTLength);
+        if (r < (Real)0)
+        {
+            r += mTLength;
+        }
+        t = mTMin + r;
+    }
+
+    // Clamp to [tmin,tmax].  For the periodic case, this handles small
+    // numerical rounding errors near the domain endpoints.
+    if (t <= mTMin)
+    {
+        t = mTMin;
+        return mDegree;
+    }
+    if (t >= mTMax)
+    {
+        t = mTMax;
+        return mNumControls - 1;
+    }
+
+    // At this point, tmin < t < tmax.
+    for (auto const& key : mKeys)
+    {
+        if (t < key.first)
+        {
+            return key.second;
+        }
+    }
+
+    // We should not reach this code.
+    LogError("Unexpected condition.");
+    t = mTMin;
+    return mDegree;
+}
+
+}
diff --git a/Engine/Plugins/Experimental/GeometryProcessing/Source/GeometryAlgorithms/Private/GTEngine/Mathematics/GteBezierCurve.h b/Engine/Plugins/Experimental/GeometryProcessing/Source/GeometryAlgorithms/Private/GTEngine/Mathematics/GteBezierCurve.h
new file mode 100644
--- /dev/null
+++ b/Engine/Plugins/Experimental/GeometryProcessing/Source/GeometryAlgorithms/Private/GTEngine/Mathematics/GteBezierCurve.h
@@ -0,0 +1,202 @@
+// David Eberly, Geometric Tools, Redmond WA 98052
+// Copyright (c) 1998-2019
+// Distributed under the Boost Software License, Version 1.0.
+// http://www.boost.org/LICENSE_1_0.txt
+// http://www.geometrictools.com/License/Boost/LICENSE_1_0.txt
+// File Version: 3.0.0 (2016/06/19)
+
+#pragma once
+
+#include <GTEngine/LowLevel/GteArray2.h>
+#include <GTEngine/LowLevel/GteLogger.h>
+#include <GTEngine/Mathematics/GteParametricCurve.h>
+
+namespace gte
+{
+
+template <int N, typename Real>
+class BezierCurve : public ParametricCurve<N, Real>
+{
+public:
+    // Construction and destruction.  The number of control points must be
+    // degree + 1.  This object copies the input array.  The domain is t
+    // in [0,1].  To validate construction, create an object as shown:
+    //     BezierCurve<N, Real> curve(parameters);
+    //     if (!curve) { <constructor failed, handle accordingly>; }
+    virtual ~BezierCurve();
+    BezierCurve(int degree, Vector<N, Real> const* controls);
+
+    // Member access.
+    inline int GetDegree() const;
+    inline int GetNumControls() const;
+    inline Vector<N, Real> const* GetControls() const;
+
+    // Evaluation of the curve.  The function supports derivative calculation
+    // through order 3; that is, maxOrder <= 3 is required.  If you want
+    // only the position, pass in maxOrder of 0.  If you want the position and
+    // first derivative, pass in maxOrder of 1, and so on.  The output
+    // 'values' are ordered as: position, first derivative, second derivative,
+    // third derivative.
+    virtual void Evaluate(Real t, unsigned int maxOrder,
+        Vector<N, Real> values[4]) const;
+
+protected:
+    // Support for Evaluate(...).
+    Vector<N, Real> Compute(Real t, Real omt, int order) const;
+
+    int mDegree, mNumControls;
+    std::vector<Vector<N, Real>> mControls[4];
+    Array2<Real> mChoose;
+};
+
+
+template <int N, typename Real>
+BezierCurve<N, Real>::~BezierCurve()
+{
+}
+
+template <int N, typename Real>
+BezierCurve<N, Real>::BezierCurve(int degree, Vector<N, Real> const* controls)
+    :
+    ParametricCurve<N, Real>((Real)0, (Real)1),
+    mDegree(degree),
+    mNumControls(degree + 1),
+    mChoose(mNumControls, mNumControls)
+{
+    if (degree < 2 || !controls)
+    {
+        LogError("Invalid input.");
+        return;
+    }
+
+    // Copy the controls.
+    mControls[0].resize(mNumControls);
+    std::copy(controls, controls + mNumControls, mControls[0].begin());
+
+    // Compute first-order differences.
+    mControls[1].resize(mNumControls - 1);
+    for (int i = 0; i < mNumControls - 1; ++i)
+    {
+        mControls[1][i] = mControls[0][i + 1] - mControls[0][i];
+    }
+
+    // Compute second-order differences.
+    mControls[2].resize(mNumControls - 2);
+    for (int i = 0; i < mNumControls - 2; ++i)
+    {
+        mControls[2][i] = mControls[1][i + 1] - mControls[1][i];
+    }
+
+    // Compute third-order differences.
+    if (degree >= 3)
+    {
+        mControls[3].resize(mNumControls - 3);
+        for (int i = 0; i < mNumControls - 3; ++i)
+        {
+            mControls[3][i] = mControls[2][i + 1] - mControls[2][i];
+        }
+    }
+
+    // Compute combinatorial values Choose(n,k) and store in mChoose[n][k].
+    // The values mChoose[r][c] are invalid for r < c; that is, we use only
+    // the entries for r >= c.
+    mChoose[0][0] = (Real)1;
+    mChoose[1][0] = (Real)1;
+    mChoose[1][1] = (Real)1;
+    for (int i = 2; i <= mDegree; ++i)
+    {
+        mChoose[i][0] = (Real)1;
+        mChoose[i][i] = (Real)1;
+        for (int j = 1; j < i; ++j)
+        {
+            mChoose[i][j] = mChoose[i - 1][j - 1] + mChoose[i - 1][j];
+        }
+    }
+
+    this->mConstructed = true;
+}
+
+template <int N, typename Real> inline
+int BezierCurve<N, Real>::GetDegree() const
+{
+    return mDegree;
+}
+
+template <int N, typename Real> inline
+int BezierCurve<N, Real>::GetNumControls() const
+{
+    return mNumControls;
+}
+
+template <int N, typename Real> inline
+Vector<N, Real> const* BezierCurve<N, Real>::GetControls() const
+{
+    return &mControls[0][0];
+}
+
+template <int N, typename Real>
+void BezierCurve<N, Real>::Evaluate(Real t, unsigned int maxOrder,
+    Vector<N, Real> values[4]) const
+{
+    if (!this->mConstructed)
+    {
+        // Errors were already generated during construction.
+        for (unsigned int order = 0; order < 4; ++order)
+        {
+            values[order].MakeZero();
+        }
+        return;
+    }
+
+    // Compute position.
+    Real omt = (Real)1 - t;
+    values[0] = Compute(t, omt, 0);
+    if (maxOrder >= 1)
+    {
+        // Compute first derivative.
+        values[1] = Compute(t, omt, 1);
+        if (maxOrder >= 2)
+        {
+            // Compute second derivative.
+            values[2] = Compute(t, omt, 2);
+            if (maxOrder >= 3 && mDegree >= 3)
+            {
+                // Compute third derivative.
+                values[3] = Compute(t, omt, 3);
+            }
+            else
+            {
+                values[3].MakeZero();
+            }
+        }
+    }
+}
+
+template <int N, typename Real>
+Vector<N, Real> BezierCurve<N, Real>::Compute(Real t, Real omt, int order)
+const
+{
+    Vector<N, Real> result = omt * mControls[order][0];
+
+    Real tpow = t;
+    int isup = mDegree - order;
+    for (int i = 1; i < isup; ++i)
+    {
+        Real c = mChoose[isup][i] * tpow;
+        result = (result + c * mControls[order][i]) * omt;
+        tpow *= t;
+    }
+    result = (result + tpow * mControls[order][isup]);
+
+    int multiplier = 1;
+    for (int i = 0; i < order; ++i)
+    {
+        multiplier *= mDegree - i;
+    }
+    result *= (Real)multiplier;
+
+    return result;
+}
+
+
+}
diff --git a/Engine/Plugins/Experimental/GeometryProcessing/Source/GeometryAlgorithms/Private/GTEngine/Mathematics/GteBitHacks.cpp b/Engine/Plugins/Experimental/GeometryProcessing/Source/GeometryAlgorithms/Private/GTEngine/Mathematics/GteBitHacks.cpp
new file mode 100644
--- /dev/null
+++ b/Engine/Plugins/Experimental/GeometryProcessing/Source/GeometryAlgorithms/Private/GTEngine/Mathematics/GteBitHacks.cpp
@@ -0,0 +1,191 @@
+// David Eberly, Geometric Tools, Redmond WA 98052
+// Copyright (c) 1998-2019
+// Distributed under the Boost Software License, Version 1.0.
+// http://www.boost.org/LICENSE_1_0.txt
+// http://www.geometrictools.com/License/Boost/LICENSE_1_0.txt
+// File Version: 3.0.3 (2017/10/22)
+
+#include <GTEngine/Mathematics/GteBitHacks.h>
+#include <GTEngine/GTEnginePCH.h>
+
+namespace gte
+{
+
+static int32_t const gsLeadingBitTable[32] = 
+{
+     0,  9,  1, 10, 13, 21,  2, 29,
+    11, 14, 16, 18, 22, 25,  3, 30,
+     8, 12, 20, 28, 15, 17, 24,  7,
+    19, 27, 23,  6, 26,  5,  4, 31
+};
+
+static int32_t const gsTrailingBitTable[32] = 
+{
+     0,  1, 28,  2, 29, 14, 24,  3,
+    30, 22, 20, 15, 25, 17,  4,  8, 
+    31, 27, 13, 23, 21, 19, 16,  7,
+    26, 12, 18,  6, 11,  5, 10,  9
+};
+
+
+bool IsPowerOfTwo(uint32_t value)
+{
+    return (value > 0) && ((value & (value - 1)) == 0);
+}
+
+bool IsPowerOfTwo(int32_t value)
+{
+    return (value > 0) && ((value & (value - 1)) == 0);
+}
+
+uint32_t Log2OfPowerOfTwo(uint32_t powerOfTwo)
+{
+    uint32_t log2 = (powerOfTwo & 0xAAAAAAAAu) != 0;
+    log2 |= ((powerOfTwo & 0xFFFF0000u) != 0) << 4;
+    log2 |= ((powerOfTwo & 0xFF00FF00u) != 0) << 3;
+    log2 |= ((powerOfTwo & 0xF0F0F0F0u) != 0) << 2;
+    log2 |= ((powerOfTwo & 0xCCCCCCCCu) != 0) << 1;
+    return log2;
+}
+
+int32_t Log2OfPowerOfTwo(int32_t powerOfTwo)
+{
+    uint32_t log2 = (powerOfTwo & 0xAAAAAAAAu) != 0;
+    log2 |= ((powerOfTwo & 0xFFFF0000u) != 0) << 4;
+    log2 |= ((powerOfTwo & 0xFF00FF00u) != 0) << 3;
+    log2 |= ((powerOfTwo & 0xF0F0F0F0u) != 0) << 2;
+    log2 |= ((powerOfTwo & 0xCCCCCCCCu) != 0) << 1;
+    return static_cast<int32_t>(log2);
+}
+
+int32_t GetLeadingBit(uint32_t value)
+{
+    value |= value >> 1;
+    value |= value >> 2;
+    value |= value >> 4;
+    value |= value >> 8;
+    value |= value >> 16;
+    uint32_t key = (value * 0x07C4ACDDu) >> 27;
+    return gsLeadingBitTable[key];
+}
+
+int32_t GetLeadingBit(int32_t value)
+{
+    value |= value >> 1;
+    value |= value >> 2;
+    value |= value >> 4;
+    value |= value >> 8;
+    value |= value >> 16;
+    uint32_t key = (value * 0x07C4ACDDu) >> 27;
+    return gsLeadingBitTable[key];
+}
+
+int32_t GetLeadingBit(uint64_t value)
+{
+    uint32_t v1 = (uint32_t)((value >> 32) & 0x00000000FFFFFFFFull);
+    if (v1 != 0)
+    {
+        return GetLeadingBit(v1) + 32;
+    }
+
+    uint32_t v0 = (uint32_t)(value & 0x00000000FFFFFFFFull);
+    return GetLeadingBit(v0);
+}
+
+int32_t GetLeadingBit(int64_t value)
+{
+    int32_t v1 = (int32_t)((value >> 32) & 0x00000000FFFFFFFFull);
+    if (v1 != 0)
+    {
+        return GetLeadingBit(v1) + 32;
+    }
+
+    int32_t v0 = (int32_t)(value & 0x00000000FFFFFFFFull);
+    return GetLeadingBit(v0);
+}
+
+int32_t GetTrailingBit(uint32_t value)
+{
+    // Avoid warning for negation of unsigned 'value'.
+    int32_t const& iValue = reinterpret_cast<int32_t const&>(value);
+    uint32_t key = ((uint32_t)((iValue & -iValue) * 0x077CB531u)) >> 27;
+    return gsTrailingBitTable[key];
+}
+
+int32_t GetTrailingBit(int32_t value)
+{
+    uint32_t key = ((uint32_t)((value & -value) * 0x077CB531u)) >> 27;
+    return gsTrailingBitTable[key];
+}
+
+int32_t GetTrailingBit(uint64_t value)
+{
+    uint32_t v0 = (uint32_t)(value & 0x00000000FFFFFFFFull);
+    if (v0 != 0)
+    {
+        return GetTrailingBit(v0);
+    }
+
+    uint32_t v1 = (uint32_t)((value >> 32) & 0x00000000FFFFFFFFull);
+    if (v1 != 0)
+    {
+        return GetTrailingBit(v1) + 32;
+    }
+    return 0;
+}
+
+int32_t GetTrailingBit(int64_t value)
+{
+    int32_t v0 = (int32_t)(value & 0x00000000FFFFFFFFull);
+    if (v0 != 0)
+    {
+        return GetTrailingBit(v0);
+    }
+
+    int32_t v1 = (int32_t)((value >> 32) & 0x00000000FFFFFFFFull);
+    if (v1 != 0)
+    {
+        return GetTrailingBit(v1) + 32;
+    }
+    return 0;
+}
+
+uint64_t RoundUpToPowerOfTwo(uint32_t value)
+{
+    if (value > 0)
+    {
+        int32_t leading = GetLeadingBit(value);
+        uint32_t mask = (1 << leading);
+        if ((value & ~mask) == 0)
+        {
+            // value is a power of two
+            return static_cast<uint64_t>(value);
+        }
+        else
+        {
+            // round up to a power of two
+            return (static_cast<uint64_t>(mask) << 1);
+        }
+
+    }
+    else
+    {
+        return 1ull;
+    }
+}
+
+uint32_t RoundDownToPowerOfTwo(uint32_t value)
+{
+    if (value > 0)
+    {
+        int32_t leading = GetLeadingBit(value);
+        uint32_t mask = (1 << leading);
+        return mask;
+    }
+    else
+    {
+        return 0;
+    }
+}
+
+}
diff --git a/Engine/Plugins/Experimental/GeometryProcessing/Source/GeometryAlgorithms/Private/GTEngine/Mathematics/GteBitHacks.h b/Engine/Plugins/Experimental/GeometryProcessing/Source/GeometryAlgorithms/Private/GTEngine/Mathematics/GteBitHacks.h
new file mode 100644
--- /dev/null
+++ b/Engine/Plugins/Experimental/GeometryProcessing/Source/GeometryAlgorithms/Private/GTEngine/Mathematics/GteBitHacks.h
@@ -0,0 +1,40 @@
+// David Eberly, Geometric Tools, Redmond WA 98052
+// Copyright (c) 1998-2019
+// Distributed under the Boost Software License, Version 1.0.
+// http://www.boost.org/LICENSE_1_0.txt
+// http://www.geometrictools.com/License/Boost/LICENSE_1_0.txt
+// File Version: 3.0.1 (2017/07/04)
+
+#pragma once
+
+#include <GTEngine/GTEngineDEF.h>
+#include <cstdint>
+
+namespace gte
+{
+
+GTE_IMPEXP bool IsPowerOfTwo(uint32_t value);
+GTE_IMPEXP bool IsPowerOfTwo(int32_t value);
+
+GTE_IMPEXP uint32_t Log2OfPowerOfTwo(uint32_t powerOfTwo);
+GTE_IMPEXP int32_t Log2OfPowerOfTwo(int32_t powerOfTwo);
+
+// Call these only for nonzero values.  If value is zero, then GetLeadingBit
+// and GetTrailingBit return zero.
+GTE_IMPEXP int32_t GetLeadingBit(uint32_t value);
+GTE_IMPEXP int32_t GetLeadingBit(int32_t value);
+GTE_IMPEXP int32_t GetLeadingBit(uint64_t value);
+GTE_IMPEXP int32_t GetLeadingBit(int64_t value);
+GTE_IMPEXP int32_t GetTrailingBit(uint32_t value);
+GTE_IMPEXP int32_t GetTrailingBit(int32_t value);
+GTE_IMPEXP int32_t GetTrailingBit(uint64_t value);
+GTE_IMPEXP int32_t GetTrailingBit(int64_t value);
+
+// Round up to a power of two.  If input is zero, the return is 1.  If input
+// is larger than 2^{31}, the return is 2^{32}.
+GTE_IMPEXP uint64_t RoundUpToPowerOfTwo(uint32_t value);
+
+// Round down to a power of two.  If input is zero, the return is 0.
+GTE_IMPEXP uint32_t RoundDownToPowerOfTwo(uint32_t value);
+
+}
diff --git a/Engine/Plugins/Experimental/GeometryProcessing/Source/GeometryAlgorithms/Private/GTEngine/Mathematics/GteCapsule.h b/Engine/Plugins/Experimental/GeometryProcessing/Source/GeometryAlgorithms/Private/GTEngine/Mathematics/GteCapsule.h
new file mode 100644
--- /dev/null
+++ b/Engine/Plugins/Experimental/GeometryProcessing/Source/GeometryAlgorithms/Private/GTEngine/Mathematics/GteCapsule.h
@@ -0,0 +1,109 @@
+// David Eberly, Geometric Tools, Redmond WA 98052
+// Copyright (c) 1998-2019
+// Distributed under the Boost Software License, Version 1.0.
+// http://www.boost.org/LICENSE_1_0.txt
+// http://www.geometrictools.com/License/Boost/LICENSE_1_0.txt
+// File Version: 3.0.0 (2016/06/19)
+
+#pragma once
+
+#include <GTEngine/Mathematics/GteSegment.h>
+
+namespace gte
+{
+
+// A capsule is the set of points that are equidistant from a segment, the
+// common distance called the radius.
+
+template <int N, typename Real>
+class Capsule
+{
+public:
+    // Construction and destruction.  The default constructor sets the segment
+    // to have endpoints p0 = (-1,0,...,0) and p1 = (1,0,...,0), and the
+    // radius is 1.
+    Capsule();
+    Capsule(Segment<N, Real> const& inSegment, Real inRadius);
+
+    // Public member access.
+    Segment<N, Real> segment;
+    Real radius;
+
+public:
+    // Comparisons to support sorted containers.
+    bool operator==(Capsule const& capsule) const;
+    bool operator!=(Capsule const& capsule) const;
+    bool operator< (Capsule const& capsule) const;
+    bool operator<=(Capsule const& capsule) const;
+    bool operator> (Capsule const& capsule) const;
+    bool operator>=(Capsule const& capsule) const;
+};
+
+// Template alias for convenience.
+template <typename Real>
+using Capsule3 = Capsule<3, Real>;
+
+
+template <int N, typename Real>
+Capsule<N, Real>::Capsule()
+    :
+    radius((Real)1)
+{
+}
+
+template <int N, typename Real>
+Capsule<N, Real>::Capsule(Segment<N, Real> const& inSegment, Real inRadius)
+    :
+    segment(inSegment),
+    radius(inRadius)
+{
+}
+
+template <int N, typename Real>
+bool Capsule<N, Real>::operator==(Capsule const& capsule) const
+{
+    return segment == capsule.segment && radius == capsule.radius;
+}
+
+template <int N, typename Real>
+bool Capsule<N, Real>::operator!=(Capsule const& capsule) const
+{
+    return !operator==(capsule);
+}
+
+template <int N, typename Real>
+bool Capsule<N, Real>::operator<(Capsule const& capsule) const
+{
+    if (segment < capsule.segment)
+    {
+        return true;
+    }
+
+    if (segment > capsule.segment)
+    {
+        return false;
+    }
+
+    return radius < capsule.radius;
+}
+
+template <int N, typename Real>
+bool Capsule<N, Real>::operator<=(Capsule const& capsule) const
+{
+    return operator<(capsule) || operator==(capsule);
+}
+
+template <int N, typename Real>
+bool Capsule<N, Real>::operator>(Capsule const& capsule) const
+{
+    return !operator<=(capsule);
+}
+
+template <int N, typename Real>
+bool Capsule<N, Real>::operator>=(Capsule const& capsule) const
+{
+    return !operator<(capsule);
+}
+
+
+}
diff --git a/Engine/Plugins/Experimental/GeometryProcessing/Source/GeometryAlgorithms/Private/GTEngine/Mathematics/GteChebyshevRatio.h b/Engine/Plugins/Experimental/GeometryProcessing/Source/GeometryAlgorithms/Private/GTEngine/Mathematics/GteChebyshevRatio.h
new file mode 100644
--- /dev/null
+++ b/Engine/Plugins/Experimental/GeometryProcessing/Source/GeometryAlgorithms/Private/GTEngine/Mathematics/GteChebyshevRatio.h
@@ -0,0 +1,192 @@
+// David Eberly, Geometric Tools, Redmond WA 98052
+// Copyright (c) 1998-2019
+// Distributed under the Boost Software License, Version 1.0.
+// http://www.boost.org/LICENSE_1_0.txt
+// http://www.geometrictools.com/License/Boost/LICENSE_1_0.txt
+// File Version: 3.0.1 (2018/10/05)
+
+#pragma once
+
+#include <GTEngine/GTEngineDEF.h>
+#include <cmath>
+
+// Let f(t,A) = sin(t*A)/sin(A).  The slerp of quaternions q0 and q1 is
+//   slerp(t,q0,q1) = f(1-t,A)*q0 + f(t,A)*q1.
+// Let y = 1-cos(A); we allow A in [0,pi], so y in [0,1].  As a function of y,
+// a series representation for f(t,y) is
+//   f(t,y) = sum_{i=0}^{infinity} c_{i}(t) y^{i}
+// where c_0(t) = t, c_{i}(t) = c_{i-1}(t)*(i^2 - t^2)/(i*(2*i+1)) for i >= 1.
+// The c_{i}(t) are polynomials in t of degree 2*i+1.  The paper
+//
+//   A Fast and Accurate Algorithm for Computing SLERP,
+//   David Eberly,
+//   Journal of Graphics, GPU, and Game Tools, 15 : 3, 161 - 176, 2011
+//   http://www.tandfonline.com/doi/abs/10.1080/2151237X.2011.610255#.VHRB7ouUd8E
+//
+// derives an approximation
+//   g(t,y) = sum_{i=0}^{n-1} c_{i}(t) y^{i} + (1+u_n) c_{n}(t) y^n
+// which has degree 2*n+1 in t and degree n in y.  The u_n were chosen to
+// balance the error at y = 1 (A = pi/2).  Unfortunately, the analysis for the
+// global error bounds (0 <= y <= 1) is flawed; the error bounds printed in
+// the paper are too small by about one order of magnitude (factor of 10).
+// Instead they are the following, verified by Mathematica using larger
+// precision than 'float' or 'double':
+//
+//   n | error      |  n | error      |  n | error      |  n | error
+//   --+------------+----+------------+----+------------+----+-------------
+//   1 | 2.60602e-2 |  5 | 3.63188e-4 |  9 | 1.12223e-5 | 13 | 4.37180e-7
+//   2 | 7.43321e-3 |  6 | 1.47056e-4 | 10 | 4.91138e-6 | 14 | 1.98230e-7
+//   3 | 2.51798e-3 |  7 | 6.11808e-5 | 11 | 2.17345e-6 | 15 | 9.04302e-8
+//   4 | 9.30819e-4 |  8 | 2.59880e-5 | 12 | 9.70876e-7 | 16 | 4.03665e-8
+//
+// The maximum errors all occur for an angle that is nearly pi/2.  The
+// approximation is much more accurate when for angles A in [0,pi/4].  With
+// this restriction, the global error bounds are
+//
+//   n | error      |  n | error      |  n | error       |  n | error
+//   --+------------+----+------------+----+-------------+----+-------------
+//   1 | 1.90648e-2 |  5 | 5.83197e-6 |  9 | 6.80465e-10 | 13 | 3.39728e-13
+//   2 | 2.43581e-3 |  6 | 8.00278e-6 | 10 | 9.12477e-11 | 14 | 4.70735e-14
+//   3 | 3.20235e-4 |  7 | 1.10649e-7 | 11 | 4.24608e-11 | 15 | 1.55431e-15
+//   4 | 4.29242e-5 |  8 | 1.53828e-7 | 12 | 5.93392e-12 | 16 | 1.11022e-16
+//
+// Given q0 and q1 such that cos(A) = dot(q0,q1) in [0,1], in which case
+// A in [0,pi/2], let qh = (q0+q1)/|q0 + q1| = slerp(1/2,q0,q1).  Note that
+// |q0 + q1| = 2*cos(A/2) because
+//   sin(A/2)/sin(A) = sin(A/2)/(2*sin(A/2)*cos(A/2)) = 1/(2*cos(A/2))
+// The angle between q0 and qh is the same as the angle between qh and q1,
+// namely, A/2 in [0,pi/4].  It may be shown that
+//                    +--
+//   slerp(t,q0,q1) = | slerp(2*t,q0,qh),     0 <= t <= 1/2
+//                    | slerp(2*t-1,qh,q1), 1/2 <= t <= 1
+//                    +--
+// The slerp functions on the right-hand side involve angles in [0,pi/4], so
+// the approximation is more accurate for those evaluations than using the
+// original.
+
+namespace gte
+{
+
+template <typename Real>
+class ChebyshevRatio
+{
+public:
+    // Compute f(t,A) = sin(t*A)/sin(A), where t in [0,1], A in [0,pi/2],
+    // cosA = cos(A), f0 = f(1-t,A), and f1 = f(t,A).
+    inline static void Get(Real t, Real cosA, Real& f0, Real& f1);
+
+    // Compute estimates for f(t,y) = sin(t*A)/sin(A), where t in [0,1],
+    // A in [0,pi/2], y = 1 - cos(A), f0 is the estimate for f(1-t,y), and
+    // f1 is the estimate for f(t,y).  The approximating function is a
+    // polynomial of two variables. The template parameter N must be in
+    // {1..16}.  The degree in t is 2*N+1 and the degree in Y is N.
+    template <int N>
+    inline static void GetEstimate(Real t, Real y, Real& f0, Real& f1);
+};
+
+
+template <typename Real> inline
+void ChebyshevRatio<Real>::Get(Real t, Real cosA, Real& f0, Real& f1)
+{
+    if (cosA < (Real)1)
+    {
+        // The angle A is in (0,pi/2].
+        Real A = std::acos(cosA);
+        Real invSinA = ((Real)1) / std::sin(A);
+        f0 = std::sin(((Real)1 - t) * A) * invSinA;
+        f1 = std::sin(t * A) * invSinA;
+    }
+    else
+    {
+        // The angle theta is 0.
+        f0 = (Real)1 - t;
+        f1 = (Real)t;
+    }
+}
+
+template <typename Real>
+template <int N> inline
+void ChebyshevRatio<Real>::GetEstimate(Real t, Real y, Real& f0, Real& f1)
+{
+    static_assert(1 <= N && N <= 16, "Invalid degree.");
+
+    // The ASM output of the MSVS 2013 Release build shows that the constants
+    // in these arrays are loaded to XMM registers as literal values, and only
+    // those constants required for the specified degree D are loaded.  That
+    // is, the compiler does a good job of optimizing the code.
+
+    Real const onePlusMu[16] =
+    {
+        (Real)1.62943436108234530,
+        (Real)1.73965850021313961,
+        (Real)1.79701067629566813,
+        (Real)1.83291820510335812,
+        (Real)1.85772477879039977,
+        (Real)1.87596835698904785,
+        (Real)1.88998444919711206,
+        (Real)1.90110745351730037,
+        (Real)1.91015881189952352,
+        (Real)1.91767344933047190,
+        (Real)1.92401541194159076,
+        (Real)1.92944142668012797,
+        (Real)1.93413793373091059,
+        (Real)1.93824371262559758,
+        (Real)1.94186426368404708,
+        (Real)1.94508125972497303
+    };
+
+    Real const a[16] =
+    {
+        (N != 1 ? (Real)1 : onePlusMu[0]) / ((Real)1 * (Real)3),
+        (N != 2 ? (Real)1 : onePlusMu[1]) / ((Real)2 * (Real)5),
+        (N != 3 ? (Real)1 : onePlusMu[2]) / ((Real)3 * (Real)7),
+        (N != 4 ? (Real)1 : onePlusMu[3]) / ((Real)4 * (Real)9),
+        (N != 5 ? (Real)1 : onePlusMu[4]) / ((Real)5 * (Real)11),
+        (N != 6 ? (Real)1 : onePlusMu[5]) / ((Real)6 * (Real)13),
+        (N != 7 ? (Real)1 : onePlusMu[6]) / ((Real)7 * (Real)15),
+        (N != 8 ? (Real)1 : onePlusMu[7]) / ((Real)8 * (Real)17),
+        (N != 9 ? (Real)1 : onePlusMu[8]) / ((Real)9 * (Real)19),
+        (N != 10 ? (Real)1 : onePlusMu[9]) / ((Real)10 * (Real)21),
+        (N != 11 ? (Real)1 : onePlusMu[10]) / ((Real)11 * (Real)23),
+        (N != 12 ? (Real)1 : onePlusMu[11]) / ((Real)12 * (Real)25),
+        (N != 13 ? (Real)1 : onePlusMu[12]) / ((Real)13 * (Real)27),
+        (N != 14 ? (Real)1 : onePlusMu[13]) / ((Real)14 * (Real)29),
+        (N != 15 ? (Real)1 : onePlusMu[14]) / ((Real)15 * (Real)31),
+        (N != 16 ? (Real)1 : onePlusMu[15]) / ((Real)16 * (Real)33)
+    };
+
+    Real const b[16] =
+    {
+        (N != 1 ? (Real)1 : onePlusMu[0]) * (Real)1 / (Real)3,
+        (N != 2 ? (Real)1 : onePlusMu[1]) * (Real)2 / (Real)5,
+        (N != 3 ? (Real)1 : onePlusMu[2]) * (Real)3 / (Real)7,
+        (N != 4 ? (Real)1 : onePlusMu[3]) * (Real)4 / (Real)9,
+        (N != 5 ? (Real)1 : onePlusMu[4]) * (Real)5 / (Real)11,
+        (N != 6 ? (Real)1 : onePlusMu[5]) * (Real)6 / (Real)13,
+        (N != 7 ? (Real)1 : onePlusMu[6]) * (Real)7 / (Real)15,
+        (N != 8 ? (Real)1 : onePlusMu[7]) * (Real)8 / (Real)17,
+        (N != 9 ? (Real)1 : onePlusMu[8]) * (Real)9 / (Real)19,
+        (N != 10 ? (Real)1 : onePlusMu[9]) * (Real)10 / (Real)21,
+        (N != 11 ? (Real)1 : onePlusMu[10]) * (Real)11 / (Real)23,
+        (N != 12 ? (Real)1 : onePlusMu[11]) * (Real)12 / (Real)25,
+        (N != 13 ? (Real)1 : onePlusMu[12]) * (Real)13 / (Real)27,
+        (N != 14 ? (Real)1 : onePlusMu[13]) * (Real)14 / (Real)29,
+        (N != 15 ? (Real)1 : onePlusMu[14]) * (Real)15 / (Real)31,
+        (N != 16 ? (Real)1 : onePlusMu[15]) * (Real)16 / (Real)33
+    };
+
+    Real term0 = (Real)1 - t, term1 = t;
+    Real sqr0 = term0 * term0, sqr1 = term1 * term1;
+    f0 = term0;
+    f1 = term1;
+    for (int i = 0; i < N; ++i)
+    {
+        term0 *= (b[i] - a[i] * sqr0) * y;
+        term1 *= (b[i] - a[i] * sqr1) * y;
+        f0 += term0;
+        f1 += term1;
+    }
+}
+
+
+}
