move GeometryProcessing plugin from NotForLicensees to Experimental

diff --git a/Engine/Plugins/Experimental/GeometryProcessing/Source/GeometryAlgorithms/Private/GTEngine/Mathematics/GteCholeskyDecomposition.h b/Engine/Plugins/Experimental/GeometryProcessing/Source/GeometryAlgorithms/Private/GTEngine/Mathematics/GteCholeskyDecomposition.h
new file mode 100644
--- /dev/null
+++ b/Engine/Plugins/Experimental/GeometryProcessing/Source/GeometryAlgorithms/Private/GTEngine/Mathematics/GteCholeskyDecomposition.h
@@ -0,0 +1,560 @@
+// David Eberly, Geometric Tools, Redmond WA 98052
+// Copyright (c) 1998-2019
+// Distributed under the Boost Software License, Version 1.0.
+// http://www.boost.org/LICENSE_1_0.txt
+// http://www.geometrictools.com/License/Boost/LICENSE_1_0.txt
+// File Version: 3.14.1 (2018/10/05)
+
+#pragma once
+
+#include <GTEngine/LowLevel/GteLogger.h>
+#include <GTEngine/Mathematics/GteMatrix.h>
+#include <GTEngine/Mathematics/GteGMatrix.h>
+
+namespace gte
+{
+    // Implementation for size known at compile time.
+    template <typename Real, int N = 0>
+    class CholeskyDecomposition
+    {
+    public:
+        // Ensure that N > 0 at compile time.
+        CholeskyDecomposition()
+        {
+            static_assert(N > 0, "Invalid size in CholeskyDecomposition constructor.");
+        }
+
+        // Disallow copies and moves.
+        CholeskyDecomposition(CholeskyDecomposition const&) = delete;
+        CholeskyDecomposition& operator=(CholeskyDecomposition const&) = delete;
+        CholeskyDecomposition(CholeskyDecomposition&&) = delete;
+        CholeskyDecomposition& operator=(CholeskyDecomposition&&) = delete;
+
+        // On input, A is symmetric.  Only the lower-triangular portion is
+        // modified.  On output, the lower-triangular portion is L where
+        // A = L * L^T.
+        bool Factor(Matrix<N, N, Real>& A)
+        {
+            for (int c = 0; c < N; ++c)
+            {
+                if (A(c, c) <= (Real)0)
+                {
+                    return false;
+                }
+                A(c, c) = std::sqrt(A(c, c));
+
+                for (int r = c + 1; r < N; ++r)
+                {
+                    A(r, c) /= A(c, c);
+                }
+
+                for (int k = c + 1; k < N; ++k)
+                {
+                    for (int r = k; r < N; ++r)
+                    {
+                        A(r, k) -= A(r, c) * A(k, c);
+                    }
+                }
+            }
+            return true;
+        }
+
+        // Solve L*Y = B, where L is lower triangular and invertible.  The
+        // input value of Y is B.  On output, Y is the solution.
+        void SolveLower(Matrix<N, N, Real> const& L, Vector<N, Real>& Y)
+        {
+            for (int r = 0; r < N; ++r)
+            {
+                for (int c = 0; c < r; ++c)
+                {
+                    Y[r] -= L(r, c) * Y[c];
+                }
+                Y[r] /= L(r, r);
+            }
+        }
+
+        // Solve L^T*X = Y, where L is lower triangular (L^T is upper
+        // triangular) and invertible.  The input value of X is Y.  On
+        // output, X is the solution.
+        void SolveUpper(Matrix<N, N, Real> const& L, Vector<N, Real>& X)
+        {
+            for (int r = N - 1; r >= 0; --r)
+            {
+                for (int c = r + 1; c < N; ++c)
+                {
+                    X[r] -= L(c, r) * X[c];
+                }
+                X[r] /= L(r, r);
+            }
+        }
+    };
+
+    // Implementation for size known only at run time.
+    template <typename Real>
+    class CholeskyDecomposition<Real, 0>
+    {
+    public:
+        int const N;
+
+        // Ensure that N > 0 at run time.
+        CholeskyDecomposition(int n)
+            :
+            N(n)
+        {
+        }
+
+        // Disallow copies and moves.  This is required to avoid compiler
+        // complaints about the 'int const N' member.
+        CholeskyDecomposition(CholeskyDecomposition const&) = delete;
+        CholeskyDecomposition& operator=(CholeskyDecomposition const&) = delete;
+        CholeskyDecomposition(CholeskyDecomposition&&) = delete;
+        CholeskyDecomposition& operator=(CholeskyDecomposition&&) = delete;
+
+        // On input, A is symmetric.  Only the lower-triangular portion is
+        // modified.  On output, the lower-triangular portion is L where
+        // A = L * L^T.
+        bool Factor(GMatrix<Real>& A)
+        {
+            if (A.GetNumRows() == N && A.GetNumCols() == N)
+            {
+                for (int c = 0; c < N; ++c)
+                {
+                    if (A(c, c) <= (Real)0)
+                    {
+                        return false;
+                    }
+                    A(c, c) = std::sqrt(A(c, c));
+
+                    for (int r = c + 1; r < N; ++r)
+                    {
+                        A(r, c) /= A(c, c);
+                    }
+
+                    for (int k = c + 1; k < N; ++k)
+                    {
+                        for (int r = k; r < N; ++r)
+                        {
+                            A(r, k) -= A(r, c) * A(k, c);
+                        }
+                    }
+                }
+                return true;
+            }
+            else
+            {
+                LogError("Matrix must be square in CholeskyDecomposition::Factor.");
+                return false;
+            }
+        }
+
+        // Solve L*Y = B, where L is lower triangular and invertible.  The
+        // input value of Y is B.  On output, Y is the solution.
+        void SolveLower(GMatrix<Real> const& L, GVector<Real>& Y)
+        {
+            if (L.GetNumRows() == N && L.GetNumCols() == N && Y.GetSize() == N)
+            {
+                for (int r = 0; r < N; ++r)
+                {
+                    for (int c = 0; c < r; ++c)
+                    {
+                        Y[r] -= L(r, c) * Y[c];
+                    }
+                    Y[r] /= L(r, r);
+                }
+            }
+            else
+            {
+                LogError("Invalid size in CholeskyDecomposition::SolveLower.");
+                Y.MakeZero();
+            }
+        }
+
+        // Solve L^T*X = Y, where L is lower triangular (L^T is upper
+        // triangular) and invertible.  The input value of X is Y.  On
+        // output, X is the solution.
+        void SolveUpper(GMatrix<Real> const& L, GVector<Real>& X)
+        {
+            if (L.GetNumRows() == N && L.GetNumCols() == N && X.GetSize() == N)
+            {
+                for (int r = N - 1; r >= 0; --r)
+                {
+                    for (int c = r + 1; c < N; ++c)
+                    {
+                        X[r] -= L(c, r) * X[c];
+                    }
+                    X[r] /= L(r, r);
+                }
+            }
+            else
+            {
+                LogError("Invalid size in CholeskyDecomposition::SolveLower.");
+                X.MakeZero();
+            }
+        }
+    };
+
+
+    // Implementation for sizes known at compile time.
+    template <typename Real, int BlockSize = 0, int NumBlocks = 0>
+    class BlockCholeskyDecomposition
+    {
+    public:
+        // Let B represent the block size and N represent the number of
+        // blocks.  The matrix A is (N*B)-by-(N*B) but partitioned into an
+        // N-by-N matrix of blocks, each block of size B-by-B.  The value
+        // N*B is NumDimensions.
+        enum
+        {
+            NumDimensions = NumBlocks * BlockSize
+        };
+
+        typedef std::array<Vector<BlockSize, Real>, NumBlocks> BlockVector;
+        typedef std::array<std::array<Matrix<BlockSize, BlockSize, Real>, NumBlocks>, NumBlocks> BlockMatrix;
+
+        // Ensure that BlockSize > 0 and NumBlocks > 0 at compile time.
+        BlockCholeskyDecomposition()
+        {
+            static_assert(BlockSize > 0 && NumBlocks > 0, "Invalid size in BlockCholeskyDecomposition constructor.");
+        }
+
+        // Disallow copies and moves.
+        BlockCholeskyDecomposition(BlockCholeskyDecomposition const&) = delete;
+        BlockCholeskyDecomposition& operator=(BlockCholeskyDecomposition const&) = delete;
+        BlockCholeskyDecomposition(BlockCholeskyDecomposition&&) = delete;
+        BlockCholeskyDecomposition& operator=(BlockCholeskyDecomposition&&) = delete;
+
+        // Treating the matrix as a 2D table of scalars with NUM_DIMENSIONS
+        // rows and NUM_DIMENSIONS columns, look up the correct block that
+        // stores the requested element and return a reference.
+        Real Get(BlockMatrix const& M, int row, int col)
+        {
+            int b0 = col / BlockSize, b1 = row / BlockSize;
+            int i0 = col % BlockSize, i1 = row % BlockSize;
+            auto const& block = M[b1][b0];
+            return block(i1, i0);
+        }
+
+        void Set(BlockMatrix& M, int row, int col, Real value)
+        {
+            int b0 = col / BlockSize, b1 = row / BlockSize;
+            int i0 = col % BlockSize, i1 = row % BlockSize;
+            auto& block = M[b1][b0];
+            block(i1, i0) = value;
+        }
+
+        bool Factor(BlockMatrix& A)
+        {
+            for (int c = 0; c < NumBlocks; ++c)
+            {
+                if (!mDecomposer.Factor(A[c][c]))
+                {
+                    return false;
+                }
+
+                for (int r = c + 1; r < NumBlocks; ++r)
+                {
+                    LowerTriangularSolver(r, c, A);
+                }
+
+                for (int k = c + 1; k < NumBlocks; ++k)
+                {
+                    for (int r = k; r < NumBlocks; ++r)
+                    {
+                        SubtractiveUpdate(r, k, c, A);
+                    }
+                }
+            }
+            return true;
+        }
+
+        // Solve L*Y = B, where L is an invertible lower-triangular block
+        // matrix whose diagonal blocks are lower-triangular matrices.
+        // The input B is a block vector of commensurate size.  The input
+        // value of Y is B.  On output, Y is the solution.
+        void SolveLower(BlockMatrix const& L, BlockVector& Y)
+        {
+            for (int r = 0; r < NumBlocks; ++r)
+            {
+                auto& Yr = Y[r];
+                for (int c = 0; c < r; ++c)
+                {
+                    auto const& Lrc = L[r][c];
+                    auto const& Yc = Y[c];
+                    for (int i = 0; i < BlockSize; ++i)
+                    {
+                        for (int j = 0; j < BlockSize; ++j)
+                        {
+                            Yr[i] -= Lrc(i, j) * Yc[j];
+                        }
+                    }
+                }
+                mDecomposer.SolveLower(L[r][r], Yr);
+            }
+        }
+
+        // Solve L^T*X = Y, where L is an invertible lower-triangular block
+        // matrix (L^T is an upper-triangular block matrix) whose diagonal
+        // blocks are lower-triangular matrices.  The input value of X is Y.
+        // On output, X is the solution.
+        void SolveUpper(BlockMatrix const& L, BlockVector& X)
+        {
+            for (int r = NumBlocks - 1; r >= 0; --r)
+            {
+                auto& Xr = X[r];
+                for (int c = r + 1; c < NumBlocks; ++c)
+                {
+                    auto const& Lcr = L[c][r];
+                    auto const& Xc = X[c];
+                    for (int i = 0; i < BlockSize; ++i)
+                    {
+                        for (int j = 0; j < BlockSize; ++j)
+                        {
+                            Xr[i] -= Lcr(j, i) * Xc[j];
+                        }
+                    }
+                }
+                mDecomposer.SolveUpper(L[r][r], Xr);
+            }
+        }
+
+    private:
+        // Solve G(c,c)*G(r,c)^T = A(r,c)^T for G(r,c).  The matrices
+        // G(c,c) and A(r,c) are known quantities, and G(c,c) occupies
+        // the lower triangular portion of A(c,c).  The solver stores
+        // its results in-place, so A(r,c) stores the G(r,c) result.
+        void LowerTriangularSolver(int r, int c, BlockMatrix& A)
+        {
+            auto const& Acc = A[c][c];
+            auto& Arc = A[r][c];
+            for (int j = 0; j < BlockSize; ++j)
+            {
+                for (int i = 0; i < j; ++i)
+                {
+                    Real Lji = Acc(j, i);
+                    for (int k = 0; k < BlockSize; ++k)
+                    {
+                        Arc(k, j) -= Lji * Arc(k, i);
+                    }
+                }
+
+                Real Ljj = Acc(j, j);
+                for (int k = 0; k < BlockSize; ++k)
+                {
+                    Arc(k, j) /= Ljj;
+                }
+            }
+        }
+
+        void SubtractiveUpdate(int r, int k, int c, BlockMatrix& A)
+        {
+            auto const& Arc = A[r][c];
+            auto const& Akc = A[k][c];
+            auto& Ark = A[r][k];
+            for (int j = 0; j < BlockSize; ++j)
+            {
+                for (int i = 0; i < BlockSize; ++i)
+                {
+                    for (int m = 0; m < BlockSize; ++m)
+                    {
+                        Ark(j, i) -= Arc(j, m) * Akc(i, m);
+                    }
+                }
+            }
+        }
+
+        CholeskyDecomposition<Real, BlockSize> mDecomposer;
+    };
+
+    // Implementation for sizes known only at run time.
+    template <typename Real>
+    class BlockCholeskyDecomposition<Real, 0, 0>
+    {
+    public:
+        // Let B represent the block size and N represent the number of
+        // blocks.  The matrix A is (N*B)-by-(N*B) but partitioned into an
+        // N-by-N matrix of blocks, each block of size B-by-B.  The value
+        // N*B is NumDimensions.
+        int const BlockSize;
+        int const NumBlocks;
+        int const NumDimensions;
+
+        // The number of elements in a BlockVector object must be NumBlocks
+        // and each GVector element has BlockSize components.
+        typedef std::vector<GVector<Real>> BlockVector;
+
+        // The BlockMatrix is an array of NumBlocks-by-NumBlocks matrices.
+        // Each block matrix is stored in row-major order.  The BlockMatrix
+        // elements themselves are stored in row-major order.  The block
+        // matrix element M = BlockMatrix[col + NumBlocks * row] is of size
+        // BlockSize-by-BlockSize (in row-major order) and is in the (row,col)
+        // location of the full matrix of blocks.
+        typedef std::vector<GMatrix<Real>> BlockMatrix;
+
+        // Ensure that BlockSize > 0 and NumDimensions > 0 at run time.
+        BlockCholeskyDecomposition(int blockSize, int numBlocks)
+            :
+            BlockSize(blockSize),
+            NumBlocks(numBlocks),
+            NumDimensions(numBlocks * blockSize),
+            mDecomposer(blockSize)
+        {
+            LogAssert(blockSize > 0 && numBlocks > 0, "Invalid size in BlockCholeskyDecomposition constructor.");
+        }
+
+        // Disallow copies and moves.  This is required to avoid compiler
+        // complaints about the 'int const' members.
+        BlockCholeskyDecomposition(BlockCholeskyDecomposition const&) = delete;
+        BlockCholeskyDecomposition& operator=(BlockCholeskyDecomposition const&) = delete;
+        BlockCholeskyDecomposition(BlockCholeskyDecomposition&&) = delete;
+        BlockCholeskyDecomposition& operator=(BlockCholeskyDecomposition&&) = delete;
+
+        // Treating the matrix as a 2D table of scalars with NumDimensions
+        // rows and NumDimensions columns, look up the correct block that
+        // stores the requested element and return a reference.
+        Real Get(BlockMatrix const& M, int row, int col)
+        {
+            int b0 = col / BlockSize, b1 = row / BlockSize;
+            int i0 = col % BlockSize, i1 = row % BlockSize;
+            auto const& block = M[GetIndex(b1, b0)];
+            return block(i1, i0);
+        }
+
+        void Set(BlockMatrix& M, int row, int col, Real value)
+        {
+            int b0 = col / BlockSize, b1 = row / BlockSize;
+            int i0 = col % BlockSize, i1 = row % BlockSize;
+            auto& block = M[GetIndex(b1, b0)];
+            block(i1, i0) = value;
+        }
+
+        bool Factor(BlockMatrix& A)
+        {
+            for (int c = 0; c < NumBlocks; ++c)
+            {
+                if (!mDecomposer.Factor(A[GetIndex(c, c)]))
+                {
+                    return false;
+                }
+
+                for (int r = c + 1; r < NumBlocks; ++r)
+                {
+                    LowerTriangularSolver(r, c, A);
+                }
+
+                for (int k = c + 1; k < NumBlocks; ++k)
+                {
+                    for (int r = k; r < NumBlocks; ++r)
+                    {
+                        SubtractiveUpdate(r, k, c, A);
+                    }
+                }
+            }
+            return true;
+        }
+
+        // Solve L*Y = B, where L is an invertible lower-triangular block
+        // matrix whose diagonal blocks are lower-triangular matrices.
+        // The input B is a block vector of commensurate size.  The input
+        // value of Y is B.  On output, Y is the solution.
+        void SolveLower(BlockMatrix const& L, BlockVector& Y)
+        {
+            for (int r = 0; r < NumBlocks; ++r)
+            {
+                auto& Yr = Y[r];
+                for (int c = 0; c < r; ++c)
+                {
+                    auto const& Lrc = L[GetIndex(r, c)];
+                    auto const& Yc = Y[c];
+                    for (int i = 0; i < NumBlocks; ++i)
+                    {
+                        for (int j = 0; j < NumBlocks; ++j)
+                        {
+                            Yr[i] -= Lrc[GetIndex(i, j)] * Yc[j];
+                        }
+                    }
+                }
+                mDecomposer.SolveLower(L[GetIndex(r, r)], Yr);
+            }
+        }
+
+        // Solve L^T*X = Y, where L is an invertible lower-triangular block
+        // matrix (L^T is an upper-triangular block matrix) whose diagonal
+        // blocks are lower-triangular matrices.  The input value of X is Y.
+        // On output, X is the solution.
+        void SolveUpper(BlockMatrix const& L, BlockVector& X)
+        {
+            for (int r = NumBlocks - 1; r >= 0; --r)
+            {
+                auto& Xr = X[r];
+                for (int c = r + 1; c < NumBlocks; ++c)
+                {
+                    auto const& Lcr = L[GetIndex(c, r)];
+                    auto const& Xc = X[c];
+                    for (int i = 0; i < BlockSize; ++i)
+                    {
+                        for (int j = 0; j < BlockSize; ++j)
+                        {
+                            Xr[i] -= Lcr[GetIndex(j, i)] * Xc[j];
+                        }
+                    }
+                }
+                mDecomposer.SolveUpper(L[GetIndex(r, r)], Xr);
+            }
+        }
+
+    private:
+        // Compute the 1-dimensional index of the block matrix in a
+        // 2-dimensional BlockMatrix object.
+        inline int GetIndex(int row, int col) const
+        {
+            return col + row * NumBlocks;
+        }
+
+        // Solve G(c,c)*G(r,c)^T = A(r,c)^T for G(r,c).  The matrices
+        // G(c,c) and A(r,c) are known quantities, and G(c,c) occupies
+        // the lower triangular portion of A(c,c).  The solver stores
+        // its results in-place, so A(r,c) stores the G(r,c) result.
+        void LowerTriangularSolver(int r, int c, BlockMatrix& A)
+        {
+            auto const& Acc = A[GetIndex(c, c)];
+            auto& Arc = A[GetIndex(r, c)];
+            for (int j = 0; j < BlockSize; ++j)
+            {
+                for (int i = 0; i < j; ++i)
+                {
+                    Real Lji = Acc[GetIndex(j, i)];
+                    for (int k = 0; k < BlockSize; ++k)
+                    {
+                        Arc[GetIndex(k, j)] -= Lji * Arc[GetIndex(k, i)];
+                    }
+                }
+
+                Real Ljj = Acc[GetIndex(j, j)];
+                for (int k = 0; k < BlockSize; ++k)
+                {
+                    Arc[GetIndex(k, j)] /= Ljj;
+                }
+            }
+        }
+
+        void SubtractiveUpdate(int r, int k, int c, BlockMatrix& A)
+        {
+            auto const& Arc = A[GetIndex(r, c)];
+            auto const& Akc = A[GetIndex(k, c)];
+            auto& Ark = A[GetIndex(r, k)];
+            for (int j = 0; j < BlockSize; ++j)
+            {
+                for (int i = 0; i < BlockSize; ++i)
+                {
+                    for (int m = 0; m < BlockSize; ++m)
+                    {
+                        Ark[GetIndex(j, i)] -= Arc[GetIndex(j, m)] * Akc[GetIndex(i, m)];
+                    }
+                }
+            }
+        }
+
+        // The decomposer has size BlockSize.
+        CholeskyDecomposition<Real> mDecomposer;
+    };
+}
diff --git a/Engine/Plugins/Experimental/GeometryProcessing/Source/GeometryAlgorithms/Private/GTEngine/Mathematics/GteCircle3.h b/Engine/Plugins/Experimental/GeometryProcessing/Source/GeometryAlgorithms/Private/GTEngine/Mathematics/GteCircle3.h
new file mode 100644
--- /dev/null
+++ b/Engine/Plugins/Experimental/GeometryProcessing/Source/GeometryAlgorithms/Private/GTEngine/Mathematics/GteCircle3.h
@@ -0,0 +1,123 @@
+// David Eberly, Geometric Tools, Redmond WA 98052
+// Copyright (c) 1998-2019
+// Distributed under the Boost Software License, Version 1.0.
+// http://www.boost.org/LICENSE_1_0.txt
+// http://www.geometrictools.com/License/Boost/LICENSE_1_0.txt
+// File Version: 3.0.0 (2016/06/19)
+
+#pragma once
+
+#include <GTEngine/Mathematics/GteVector3.h>
+
+// The circle is the intersection of the sphere |X-C|^2 = r^2 and the
+// plane Dot(N,X-C) = 0, where C is the circle center, r is the radius,
+// and N is a unit-length plane normal.
+
+namespace gte
+{
+
+template <typename Real>
+class Circle3
+{
+public:
+
+    // Construction and destruction.  The default constructor sets center to
+    // (0,0,0), normal to (0,0,1), and radius to 1.
+    Circle3();
+    Circle3(Vector3<Real> const& inCenter, Vector3<Real> const& inNormal,
+        Real inRadius);
+
+    // Public member access.
+    Vector3<Real> center, normal;
+    Real radius;
+
+public:
+    // Comparisons to support sorted containers.
+    bool operator==(Circle3 const& circle) const;
+    bool operator!=(Circle3 const& circle) const;
+    bool operator< (Circle3 const& circle) const;
+    bool operator<=(Circle3 const& circle) const;
+    bool operator> (Circle3 const& circle) const;
+    bool operator>=(Circle3 const& circle) const;
+};
+
+
+template <typename Real>
+Circle3<Real>::Circle3()
+    :
+    center(Vector3<Real>::Zero()),
+    normal(Vector3<Real>::Unit(2)),
+    radius((Real)1)
+{
+}
+
+template <typename Real>
+Circle3<Real>::Circle3(Vector3<Real> const& inCenter,
+    Vector3<Real> const& inNormal, Real inRadius)
+    :
+    center(inCenter),
+    normal(inNormal),
+    radius(inRadius)
+{
+}
+
+template <typename Real>
+bool Circle3<Real>::operator==(Circle3 const& circle) const
+{
+    return center == circle.center
+        && normal == circle.normal
+        && radius == circle.radius;
+}
+
+template <typename Real>
+bool Circle3<Real>::operator!=(Circle3 const& circle) const
+{
+    return !operator==(circle);
+}
+
+template <typename Real>
+bool Circle3<Real>::operator<(Circle3 const& circle) const
+{
+    if (center < circle.center)
+    {
+        return true;
+    }
+
+    if (center > circle.center)
+    {
+        return false;
+    }
+
+    if (normal < circle.normal)
+    {
+        return true;
+    }
+
+    if (normal > circle.normal)
+    {
+        return false;
+    }
+
+    return radius < circle.radius;
+}
+
+template <typename Real>
+bool Circle3<Real>::operator<=(Circle3 const& circle) const
+{
+    return operator<(circle) || operator==(circle);
+}
+
+template <typename Real>
+bool Circle3<Real>::operator>(Circle3 const& circle) const
+{
+    return !operator<=(circle);
+}
+
+template <typename Real>
+bool Circle3<Real>::operator>=(Circle3 const& circle) const
+{
+    return !operator<(circle);
+}
+
+
+}
diff --git a/Engine/Plugins/Experimental/GeometryProcessing/Source/GeometryAlgorithms/Private/GTEngine/Mathematics/GteCone.h b/Engine/Plugins/Experimental/GeometryProcessing/Source/GeometryAlgorithms/Private/GTEngine/Mathematics/GteCone.h
new file mode 100644
--- /dev/null
+++ b/Engine/Plugins/Experimental/GeometryProcessing/Source/GeometryAlgorithms/Private/GTEngine/Mathematics/GteCone.h
@@ -0,0 +1,175 @@
+// David Eberly, Geometric Tools, Redmond WA 98052
+// Copyright (c) 1998-2019
+// Distributed under the Boost Software License, Version 1.0.
+// http://www.boost.org/LICENSE_1_0.txt
+// http://www.geometrictools.com/License/Boost/LICENSE_1_0.txt
+// File Version: 3.0.3 (2019/02/15)
+
+#pragma once
+
+#include <GTEngine/LowLevel/GteLogger.h>
+#include <GTEngine/Mathematics/GteMath.h>
+#include <GTEngine/Mathematics/GteRay.h>
+
+// An acute cone is Dot(A,X-V) = |X-V| cos(t) where V is the vertex, A is the
+// unit-length direction of the axis of the cone, and T is the cone angle with
+// 0 < t < pi/2.  The cone interior is defined by the inequality
+// Dot(A,X-V) >= |X-V| cos(t).  Since cos(t) > 0, we can avoid computing
+// square roots.  The solid cone is defined by the inequality
+// Dot(A,X-V)^2 >= Dot(X-V,X-V) cos(t)^2.  I refer to this object as an
+// "infinite cone."  Cone axis points are V + h * A, where h is referred to
+// as height and 0 <= h < +infinity.
+//
+// The cone can be truncated by a plane perpendicular to its axis at a height
+// hmax with 0 < hmax < +infinity.  I refer to this object as a "finite cone."
+// The finite cone is capped by has a circular disk opposite the vertex; the
+// disk has radius hmax*tan(t).
+//
+// The finite cone can be additionally truncated at a height hmin with
+// 0 < hmin < hmax < +infinity.  I refer to this a a "cone frustum."
+
+namespace gte
+{
+    template <int N, typename Real>
+    class Cone
+    {
+    public:
+        // The default constructor creates an infinite cone with
+        //   vertex = (0,...,0)
+        //   axis = (0,...,0,1)
+        //   angle = pi/4
+        //   minimum height = 0
+        //   maximum height = std::numeric_limits<Real>::max()
+        Cone()
+            :
+            minHeight((Real)0),
+            maxHeight(std::numeric_limits<Real>::max())
+        {
+            ray.origin.MakeZero();
+            ray.direction.MakeUnit(N - 1);
+            SetAngle((Real)GTE_C_QUARTER_PI);
+        }
+
+        // This constructor creates an infinite cone with the specified
+        // vertex, axis direction and angle, and with heights
+        //   minimum height = 0
+        //   maximum height = std::numeric_limits<Real>::max()
+        Cone(Ray<N, Real> const& inRay, Real inAngle)
+            :
+            ray(inRay),
+            minHeight((Real)0),
+            maxHeight(std::numeric_limits<Real>::max())
+        {
+            SetAngle(inAngle);
+        }
+
+        // This constructor creates a cone with all parameters specified.
+        Cone(Ray<N, Real> const& inRay, Real inAngle, Real inMinHeight, Real inMaxHeight)
+            :
+            ray(inRay),
+            minHeight(inMinHeight),
+            maxHeight(inMaxHeight)
+        {
+            LogAssert((Real)0 <= minHeight && minHeight < maxHeight, "Invalid height interval.");
+            SetAngle(inAngle);
+        }
+
+        // The angle must be in (0,pi/2).  The function sets 'angle' and
+        // computes 'cosAngle', 'sinAngle', 'tanAngle', 'cosAngleSqr',
+        // 'sinAngleSqr' and 'invSinAngle'.
+        void SetAngle(Real inAngle)
+        {
+            LogAssert((Real)0 < inAngle && inAngle < (Real)GTE_C_HALF_PI, "Invalid angle.");
+            angle = inAngle;
+            cosAngle = std::cos(angle);
+            sinAngle = std::sin(angle);
+            tanAngle = std::tan(angle);
+            cosAngleSqr = cosAngle * cosAngle;
+            sinAngleSqr = sinAngle * sinAngle;
+            invSinAngle = (Real)1 / sinAngle;
+        }
+
+        // The cone axis direction must be unit length.  The angle must
+        // be in (0,pi/2).  The heights must satisfy
+        // 0 <= minHeight < maxHeight <= std::numeric_limits<Real>::max().
+        Ray<N, Real> ray;
+        Real angle;
+        Real minHeight, maxHeight;
+
+        // Members derived from 'angle', to avoid calling trigonometric
+        // functions in geometric queries (for speed).  You may set 'angle'
+        // and compute these by calling SetAngle(inAngle).
+        Real cosAngle, sinAngle, tanAngle;
+        Real cosAngleSqr, sinAngleSqr, invSinAngle;
+
+    public:
+        // Comparisons to support sorted containers.  These based only on
+        // 'ray', 'angle', 'minHeight' and 'maxHeight'.
+        bool operator==(Cone const& cone) const
+        {
+            return ray == cone.ray
+                && angle == cone.angle
+                && minHeight == cone.minHeight
+                && maxHeight == cone.maxHeight;
+        }
+
+        bool operator!=(Cone const& cone) const
+        {
+            return !operator==(cone);
+        }
+
+        bool operator< (Cone const& cone) const
+        {
+            if (ray < cone.ray)
+            {
+                return true;
+            }
+
+            if (ray > cone.ray)
+            {
+                return false;
+            }
+
+            if (angle < cone.angle)
+            {
+                return true;
+            }
+
+            if (angle > cone.angle)
+            {
+                return false;
+            }
+
+            if (minHeight < cone.minHeight)
+            {
+                return true;
+            }
+
+            if (minHeight > cone.minHeight)
+            {
+                return false;
+            }
+
+            return maxHeight < cone.maxHeight;
+        }
+
+        bool operator<=(Cone const& cone) const
+        {
+            return !cone.operator<(*this);
+        }
+
+        bool operator> (Cone const& cone) const
+        {
+            return cone.operator<(*this);
+        }
+
+        bool operator>=(Cone const& cone) const
+        {
+            return !operator<(cone);
+        }
+    };
+
+    // Template alias for convenience.
+    template <typename Real>
+    using Cone3 = Cone<3, Real>;
+}
diff --git a/Engine/Plugins/Experimental/GeometryProcessing/Source/GeometryAlgorithms/Private/GTEngine/Mathematics/GteConstrainedDelaunay2.h b/Engine/Plugins/Experimental/GeometryProcessing/Source/GeometryAlgorithms/Private/GTEngine/Mathematics/GteConstrainedDelaunay2.h
new file mode 100644
--- /dev/null
+++ b/Engine/Plugins/Experimental/GeometryProcessing/Source/GeometryAlgorithms/Private/GTEngine/Mathematics/GteConstrainedDelaunay2.h
@@ -0,0 +1,766 @@
+// David Eberly, Geometric Tools, Redmond WA 98052
+// Copyright (c) 1998-2019
+// Distributed under the Boost Software License, Version 1.0.
+// http://www.boost.org/LICENSE_1_0.txt
+// http://www.geometrictools.com/License/Boost/LICENSE_1_0.txt
+// File Version: 3.0.0 (2016/06/19)
+
+#pragma once
+
+#include <GTEngine/Mathematics/GteDelaunay2.h>
+#include <list>
+
+// Various parts of the code have this trap for error conditions.  With a
+// correct algorithm and exact arithmetic, we do not expect to trigger the
+// error.  However, with floating-point arithmetic, it is possible that the
+// triangulation becomes malformed.  The constrained Delaunay triangulation
+// implementation is designed to return gracefully with such a failure.  The
+// following macros are used to make the code more readable.  Do NOT disable
+// them, because they have necessary side effects.
+#define GTE_CDT_REQUIRE(c) { if (!(c)) { Trap(); return false; } }
+#define GTE_CDT_FAILURE { Trap(); return false; }
+#define GTE_CDT_REQUIRE_RET(c, r)  { if (!(c)) { Trap(); return r; } }
+#define GTE_CDT_FAILURE_RET(r) { Trap(); return r; }
+
+namespace gte
+{
+
+template <typename InputType, typename ComputeType>
+class ConstrainedDelaunay2 : public Delaunay2<InputType, ComputeType>
+{
+public:
+    // The class is a functor to support computing the constrained Delaunay
+    // triangulation of multiple data sets using the same class object.
+    virtual ~ConstrainedDelaunay2();
+    ConstrainedDelaunay2();
+
+    // This operator computes the Delaunay triangulation only.  Read the
+    // Delaunay2 constructor comments about 'vertices' and 'epsilon'.  The
+    // 'edges' array has indices into the 'vertices' array.  No two edges
+    // should intersect except at endpoints.
+    bool operator()(int numVertices, Vector2<InputType> const* vertices,
+        InputType epsilon);
+
+    // Insert required edges into the triangulation.  For correctness of the
+    // algorithm, if two edges passed to this function intersect, they must
+    // do so only at vertices passed to operator().  If you have two edges
+    // that intersect at a point not in the vertices, compute that point of
+    // intersection and subdivide the edges at that intersection (to form
+    // more edges), and add the point to the vertices before calling
+    // operator().  This function has an output array that contains the
+    // input edge when the only vertices on the edge are its endpoints.  If
+    // the input edge passes through more vertices, the edge is subdivided
+    // in this function.  The output edge is that subdivision with first
+    // vertex edge[0] and last vertex edge[1], and the other vertices are
+    // correctly ordered along the edge.
+    bool Insert(std::array<int, 2> const& edge, std::vector<int>& outEdge);
+
+private:
+    // The top-level entry point for inserting an edge in the triangulation.
+    bool Insert(std::array<int, 2> const& edge, int v0Triangle,
+        std::vector<int>& outEdge);
+
+    // Process the coincident edge.
+    bool ProcessCoincident(int tri, int v0, int v1, int vOther,
+        std::vector<int>& outEdge);
+
+    // Process the triangle strip originating at the first endpoint of the
+    // edge.
+    bool ProcessInterior(int tri, int v0, int v1, int vNext, int vPrev,
+        std::vector<int>& outEdge);
+
+    // Remove the triangles in the triangle strip and retriangulate the
+    // left and right polygons using the empty circumcircle condition.
+    bool Retriangulate(std::vector<int>& polygon,
+        std::vector<std::array<int, 2>> const& lBoundary,
+        std::vector<std::array<int, 2>> const& rBoundary);
+
+    int RetriangulateLRecurse(
+        std::vector<std::array<int, 2>> const& lBoundary,
+        int i0, int i1, int a0, std::vector<int>& polygon);
+
+    int RetriangulateRRecurse(
+        std::vector<std::array<int, 2>> const& rBoundary,
+        int i0, int i1, int a0, std::vector<int>& polygon);
+
+    int SelectSplit(std::vector<std::array<int, 2>> const& boundary, int i0,
+        int i1) const;
+
+    // Compute a pseudosquared distance from the vertex at v2 to the edge
+    // <v0,v1>.
+    ComputeType ComputePSD(int v0, int v1, int v2) const;
+
+    // Search the triangulation for a triangle that contains the specified
+    // vertex.
+    int GetLinkTriangle(int v) const;
+
+    // Determine the index in {0,1,2} for the triangle 'tri' that contains
+    // the vertex 'v'.
+    int GetIndexOfVertex(int tri, int v) const;
+
+    // Given a triangle 'tri' with CCW-edge <v0,v1>, return <adj,v2> where
+    // 'adj' is the index of the triangle adjacent to 'tri' that shares the
+    // edge and 'v2' is the vertex of the adjacent triangle opposite the
+    // edge.  This function supports traversing a triangle strip that contains
+    // a constraint edge, so it is called only when an adjacent triangle
+    // actually exists.
+    std::array<int, 2> GetAdjInterior(int tri, int v0, int v1) const;
+
+    // Given a triangle 'tri' of the triangle strip, the boundary edge must
+    // contain the vertex with index 'needBndVertex'.  The input
+    // 'needAdjVIndex' specifies where to look for the index of the triangle
+    // outside the strip but adjacent to the boundary edge.  The return
+    // value is <needBndVertex, adj> and is used to connect 'tri' and 'adj'
+    // across a triangle strip boundary.
+    std::array<int, 2> GetAdjBoundary(int tri, int needBndVertex,
+        int needAdjVIndex) const;
+
+    // Set the indices and adjacencies arrays so that 'tri' and 'adj' share
+    // the common edge; 'tri' has CCW-edge <v0,v1> and 'adj' has CCW-edge
+    // <v1,v0>.
+    bool Connect(int tri, int adj, int v0, int v1);
+
+    // Create an ordered list of triangles forming the link of a vertex.  The
+    // pair of the list is <triangle,GetIndexOfV(triangle)>.  This allows us
+    // to cache the index of v relative to each triangle in the link.  The
+    // vertex v might be a boundary vertex, in which case the neighborhood is
+    // open; otherwise, v is an interior vertex and the neighborhood is
+    // closed.  The 'isOpen' parameter specifies the case.
+    bool BuildLink(int v, int vTriangle, std::list<std::pair<int,int>>& link,
+        bool& isOpen) const;
+
+    // Support for return-false-on-error.  This allows us to investigate the
+    // error by setting a break point, trigger an assert when the logger is
+    // active, and get a call stack.
+    void Trap() const;
+};
+
+
+template <typename InputType, typename ComputeType>
+ConstrainedDelaunay2<InputType, ComputeType>::~ConstrainedDelaunay2()
+{
+}
+
+template <typename InputType, typename ComputeType>
+ConstrainedDelaunay2<InputType, ComputeType>::ConstrainedDelaunay2()
+    :
+    Delaunay2<InputType, ComputeType>()
+{
+}
+
+template <typename InputType, typename ComputeType>
+bool ConstrainedDelaunay2<InputType, ComputeType>::operator()(
+    int numVertices, Vector2<InputType> const* vertices, InputType epsilon)
+{
+    return Delaunay2<InputType, ComputeType>::operator()(numVertices,
+        vertices, epsilon);
+}
+
+template <typename InputType, typename ComputeType>
+bool ConstrainedDelaunay2<InputType, ComputeType>::Insert(
+    std::array<int, 2> const& edge, std::vector<int>& outEdge)
+{
+    int v0 = edge[0], v1 = edge[1];
+    if (0 <= v0 && v0 < this->mNumVertices
+        && 0 <= v1 && v1 < this->mNumVertices)
+    {
+        int v0Triangle = GetLinkTriangle(v0);
+        if (v0Triangle >= 0)
+        {
+            // Once an edge is inserted, the base-class mGraph no longer
+            // represents the triangulation.  Clear it in case the user tries
+            // to access it.
+            this->mGraph.Clear();
+
+            outEdge.clear();
+            return Insert(edge, v0Triangle, outEdge);
+        }
+    }
+    return false;
+}
+
+template <typename InputType, typename ComputeType>
+bool ConstrainedDelaunay2<InputType, ComputeType>::Insert(
+    std::array<int, 2> const& edge, int v0Triangle, std::vector<int>& outEdge)
+{
+    // Create the neighborhood of triangles that share the vertex v0.  On
+    // entry we already know one such triangle (v0Triangle).
+    int v0 = edge[0], v1 = edge[1];
+    std::list<std::pair<int, int>> link;
+    bool isOpen = true;
+    GTE_CDT_REQUIRE(BuildLink(v0, v0Triangle, link, isOpen));
+
+    // Determine which triangle contains the edge.  Process the edge according
+    // to whether it is strictly between two triangle edges or is coincident
+    // with a triangle edge.
+    auto item = link.begin();
+    std::array<int, 3> indices;
+    GTE_CDT_REQUIRE(this->GetIndices(item->first, indices));
+    int vNext = indices[(item->second + 1) % 3];
+    int qr0 = this->mQuery.ToLine(v1, v0, vNext);
+    while (item != link.end())
+    {
+        if (qr0 == 0)
+        {
+            // We have to be careful about parallel edges that point in
+            // the opposite direction of <v0,v1>.
+            Vector2<ComputeType> const& ctv0 = this->mComputeVertices[v0];
+            Vector2<ComputeType> const& ctv1 = this->mComputeVertices[v1];
+            Vector2<ComputeType> const& ctvnext =
+                this->mComputeVertices[vNext];
+            if (Dot(ctv1 - ctv0, ctvnext - ctv0) > (ComputeType)0)
+            {
+                // <v0,v1> is coincident to triangle edge0.
+                return ProcessCoincident(item->first, v0, v1, vNext,
+                    outEdge);
+            }
+
+            // Make sure we enter the next "if" statement to continue
+            // traversing the link.
+            qr0 = 1;
+        }
+
+        if (qr0 > 0)
+        {
+            // <v0,v1> is not in triangle.  Visit the next triangle.
+            if (++item == link.end())
+            {
+                return false;
+            }
+            GTE_CDT_REQUIRE(this->GetIndices(item->first, indices));
+            vNext = indices[(item->second + 1) % 3];
+            qr0 = this->mQuery.ToLine(v1, v0, vNext);
+            continue;
+        }
+
+        int vPrev = indices[(item->second + 2) % 3];
+        int qr1 = this->mQuery.ToLine(v1, v0, vPrev);
+        while (item != link.end())
+        {
+            if (qr1 == 0)
+            {
+                // We have to be careful about parallel edges that point in
+                // the opposite direction of <v0,v1>.
+                Vector2<ComputeType> const& ctv0 = this->mComputeVertices[v0];
+                Vector2<ComputeType> const& ctv1 = this->mComputeVertices[v1];
+                Vector2<ComputeType> const& ctvprev =
+                    this->mComputeVertices[vPrev];
+                if (Dot(ctv1 - ctv0, ctvprev - ctv0) > (ComputeType)0)
+                {
+                    // <v0,v1> is coincident to triangle edge1.
+                    return ProcessCoincident(item->first, v0, v1, vPrev,
+                        outEdge);
+                }
+
+                // Make sure we enter the next "if" statement to continue
+                // traversing the link.
+                qr1 = -1;
+            }
+
+            if (qr1 < 0)
+            {
+                // <v0,v1> is not in triangle.  Visit the next triangle.
+                if (++item == link.end())
+                {
+                    return false;
+                }
+                this->GetIndices(item->first, indices);
+                vNext = vPrev;
+                vPrev = indices[(item->second + 2) % 3];
+                qr1 = this->mQuery.ToLine(v1, v0, vPrev);
+                continue;
+            }
+
+            // <v0,v1> is interior to triangle <v0,vNext,vPrev>.
+            return ProcessInterior(item->first, v0, v1, vNext, vPrev,
+                outEdge);
+        }
+        break;
+    }
+
+    // The edge must be contained in some link triangle.
+    GTE_CDT_FAILURE;
+}
+
+template <typename InputType, typename ComputeType>
+bool ConstrainedDelaunay2<InputType, ComputeType>::ProcessCoincident(int tri,
+    int v0, int v1, int vOther, std::vector<int>& outEdge)
+{
+    outEdge.push_back(v0);
+    if (v1 != vOther)
+    {
+        // Decompose the edge and process the right-most subedge.
+        return Insert({ vOther, v1 }, tri, outEdge);
+    }
+    else
+    {
+        // <v0,v1> is already in the triangulation.
+        outEdge.push_back(v1);
+        return true;
+    }
+}
+
+template <typename InputType, typename ComputeType>
+bool ConstrainedDelaunay2<InputType, ComputeType>::ProcessInterior(int tri,
+    int v0, int v1, int vNext, int vPrev, std::vector<int>& outEdge)
+{
+    // The triangles of the strip are stored in 'polygon'.  The
+    // retriangulation leads to the same number of triangles, so we can reuse
+    // the mIndices[] and mAdjacencies[] locations implied by the 'polygons'
+    // indices.
+    std::vector<int> polygon;
+
+    // The sBoundary[i] (s in {l,r}) array element is <v0,adj>; see the
+    // header comments for GetAdjBoundary about what these mean.  The
+    // boundary vertex is 'v0', the adjacent triangle 'adj' is outside the
+    // strip and shares the edge <sBoundary[i-1][0], sBoundary[i][0]> with a
+    // triangle in 'polygon'.  This information allows us to connect the
+    // adjacent triangles outside the strip to new triangles inserted by
+    // the retriangulation.  The value sBoundary[0][1,2] values are set to
+    // -1 but they are not used in the construction.
+    std::vector<std::array<int, 2>> lBoundary, rBoundary;
+    std::array<int, 2> binfo;
+
+    polygon.push_back(tri);
+
+    lBoundary.push_back({ v0, -1 });
+    binfo = GetAdjBoundary(tri, vPrev, vPrev);
+    GTE_CDT_REQUIRE(binfo[0] != -2);
+    lBoundary.push_back(binfo);
+
+    rBoundary.push_back({ v0, -1 });
+    binfo = GetAdjBoundary(tri, vNext, v0);
+    GTE_CDT_REQUIRE(binfo[0] != -2);
+    rBoundary.push_back(binfo);
+
+    // Visit the triangles in the strip.  Guard against an infinite loop.
+    for (int i = 0; i < this->mNumTriangles; ++i)
+    {
+        // Find the vertex of the adjacent triangle that is opposite the
+        // edge <vNext,vPrev> shared with the current triangle.
+        auto iinfo = GetAdjInterior(tri, vNext, vPrev);
+        int adj = iinfo[0], vOpposite = iinfo[1];
+        GTE_CDT_REQUIRE(vOpposite >= 0);
+
+        // Visit the adjacent triangle and insert it into the polygon.
+        tri = adj;
+        polygon.push_back(tri);
+
+        int qr = this->mQuery.ToLine(vOpposite, v0, v1);
+        if (qr == 0)
+        {
+            // We have encountered a vertex that terminates the triangle
+            // strip.  Retriangulate the polygon.  If the edge continues
+            // through vOpposite, decompose the edge and insert the
+            // right-most subedge.
+            binfo = GetAdjBoundary(tri, vOpposite, vOpposite);
+            GTE_CDT_REQUIRE(binfo[0] != -2);
+            lBoundary.push_back(binfo);
+
+            binfo = GetAdjBoundary(tri, vOpposite, vNext);
+            GTE_CDT_REQUIRE(binfo[0] != -2);
+            rBoundary.push_back(binfo);
+
+            Retriangulate(polygon, lBoundary, rBoundary);
+            if (vOpposite != v1)
+            {
+                outEdge.push_back(v0);
+                return Insert({ vOpposite, v1 }, tri, outEdge);
+            }
+            else
+            {
+				outEdge.push_back(v0);
+				outEdge.push_back(v1);
+                return true;
+            }
+        }
+
+        if (qr < 0)
+        {
+            binfo = GetAdjBoundary(tri, vOpposite, vOpposite);
+            GTE_CDT_REQUIRE(binfo[0] != -2);
+            lBoundary.push_back(binfo);
+            vPrev = vOpposite;
+        }
+        else  // qr > 0
+        {
+            binfo = GetAdjBoundary(tri, vOpposite, vNext);
+            GTE_CDT_REQUIRE(binfo[0] != -2);
+            rBoundary.push_back(binfo);
+            vNext = vOpposite;
+        }
+    }
+
+    // The triangle strip should have been located in the loop.
+    GTE_CDT_FAILURE;
+}
+
+template <typename InputType, typename ComputeType>
+bool ConstrainedDelaunay2<InputType, ComputeType>::Retriangulate(
+    std::vector<int>& polygon,
+    std::vector<std::array<int, 2>> const& lBoundary,
+    std::vector<std::array<int, 2>> const& rBoundary)
+{
+    int t0 = RetriangulateLRecurse(lBoundary, 0,
+        static_cast<int>(lBoundary.size()) - 1, -1, polygon);
+
+    int t1 = RetriangulateRRecurse(rBoundary, 0,
+        static_cast<int>(rBoundary.size()) - 1, -1, polygon);
+
+    int v0 = lBoundary.front()[0];
+    int v1 = lBoundary.back()[0];
+    GTE_CDT_REQUIRE(Connect(t0, t1, v0, v1));
+    return true;
+}
+
+template <typename InputType, typename ComputeType>
+int ConstrainedDelaunay2<InputType, ComputeType>::RetriangulateLRecurse(
+    std::vector<std::array<int, 2>> const& lBoundary, int i0, int i1, int a0,
+    std::vector<int>& polygon)
+{
+    // Create triangles when recursing down, connect adjacent triangles when
+    // returning.
+
+    int v0 = lBoundary[i0][0];
+    int v1 = lBoundary[i1][0];
+
+    if (i1 - i0 == 1)
+    {
+        GTE_CDT_REQUIRE_RET(Connect(a0, lBoundary[i1][1], v1, v0), -2);
+        return -1;  // No triangle created.
+    }
+    else
+    {
+        // Select i2 in [i0+1,i1-1] for minimum distance to edge <i0,i1>.
+        int i2 = SelectSplit(lBoundary, i0, i1);
+        int v2 = lBoundary[i2][0];
+
+        // Reuse a triangle and fill in its new vertices.
+        int tri = polygon.back();
+        polygon.pop_back();
+        this->mIndices[3 * tri + 0] = v0;
+        this->mIndices[3 * tri + 1] = v1;
+        this->mIndices[3 * tri + 2] = v2;
+
+        // Recurse downward and create triangles.
+        int ret0 = RetriangulateLRecurse(lBoundary, i0, i2, tri, polygon);
+        GTE_CDT_REQUIRE_RET(ret0 >= -1, -2);
+        int ret1 = RetriangulateLRecurse(lBoundary, i2, i1, tri, polygon);
+        GTE_CDT_REQUIRE_RET(ret1 >= -1, -2);
+
+        // Return and connect triangles.
+        GTE_CDT_REQUIRE_RET(Connect(a0, tri, v1, v0), -2);
+        return tri;
+    }
+}
+
+template <typename InputType, typename ComputeType>
+int ConstrainedDelaunay2<InputType, ComputeType>::RetriangulateRRecurse(
+    std::vector<std::array<int, 2>> const& rBoundary, int i0, int i1, int a0,
+    std::vector<int>& polygon)
+{
+    // Create triangles when recursing down, connect adjacent triangles when
+    // returning.
+
+    int v0 = rBoundary[i0][0];
+    int v1 = rBoundary[i1][0];
+
+    if (i1 - i0 == 1)
+    {
+        GTE_CDT_REQUIRE_RET(Connect(a0, rBoundary[i1][1], v0, v1), -2);
+        return -1;  // No triangle created.
+    }
+    else
+    {
+        // Select i2 in [i0+1,i1-1] for minimum distance to edge <i0,i1>.
+        int i2 = SelectSplit(rBoundary, i0, i1);
+        int v2 = rBoundary[i2][0];
+
+        // Reuse a triangle and fill in its new vertices.
+        int tri = polygon.back();
+        polygon.pop_back();
+        this->mIndices[3 * tri + 0] = v1;
+        this->mIndices[3 * tri + 1] = v0;
+        this->mIndices[3 * tri + 2] = v2;
+
+        // Recurse downward and create triangles.
+        int ret0 = RetriangulateRRecurse(rBoundary, i0, i2, tri, polygon);
+        GTE_CDT_REQUIRE_RET(ret0 >= -1, -2);
+        int ret1 = RetriangulateRRecurse(rBoundary, i2, i1, tri, polygon);
+        GTE_CDT_REQUIRE_RET(ret1 >= -1, -2);
+
+        // Return and connect triangles.
+        GTE_CDT_REQUIRE_RET(Connect(a0, tri, v0, v1), -2);
+        return tri;
+    }
+}
+
+template <typename InputType, typename ComputeType>
+int ConstrainedDelaunay2<InputType, ComputeType>::SelectSplit(
+    std::vector<std::array<int, 2>> const& boundary, int i0, int i1) const
+{
+    int i2;
+    if (i1 - i0 == 2)
+    {
+        // This is the only candidate.
+        i2 = i0 + 1;
+    }
+    else  // i1 - i0 > 2
+    {
+        // Select the index i2 in [i0+1,i1-1] for which the distance from the
+        // vertex v2 at i2 to the edge <v0,v1> is minimized.  To allow exact
+        // arithmetic, use a pseudosquared distance that avoids divisions and
+        // square roots.
+        i2 = i0 + 1;
+        int v0 = boundary[i0][0];
+        int v1 = boundary[i1][0];
+        int v2 = boundary[i2][0];
+        ComputeType minpsd = ComputePSD(v0, v1, v2);
+        for (int i = i2 + 1; i < i1; ++i)
+        {
+            v2 = boundary[i][0];
+            ComputeType psd = ComputePSD(v0, v1, v2);
+            if (psd < minpsd)
+            {
+                minpsd = psd;
+                i2 = i;
+            }
+        }
+    }
+    return i2;
+}
+
+template <typename InputType, typename ComputeType>
+ComputeType ConstrainedDelaunay2<InputType, ComputeType>::ComputePSD(int v0,
+    int v1, int v2) const
+{
+    ComputeType psd;
+
+    Vector2<ComputeType> const& ctv0 = this->mComputeVertices[v0];
+    Vector2<ComputeType> const& ctv1 = this->mComputeVertices[v1];
+    Vector2<ComputeType> const& ctv2 = this->mComputeVertices[v2];
+
+    Vector2<ComputeType> V1mV0 = ctv1 - ctv0;
+    Vector2<ComputeType> V2mV0 = ctv2 - ctv0;
+    ComputeType sqrlen10 = Dot(V1mV0, V1mV0);
+    ComputeType dot = Dot(V1mV0, V2mV0);
+    ComputeType zero(0);
+
+    if (dot <= zero)
+    {
+        ComputeType sqrlen20 = Dot(V2mV0, V2mV0);
+        psd = sqrlen10 * sqrlen20;
+    }
+    else
+    {
+        Vector2<ComputeType> V2mV1 = ctv2 - ctv1;
+        dot = Dot(V1mV0, V2mV1);
+        if (dot >= zero)
+        {
+            ComputeType sqrlen21 = Dot(V2mV1, V2mV1);
+            psd = sqrlen10 * sqrlen21;
+        }
+        else
+        {
+            dot = DotPerp(V2mV0, V1mV0);
+            psd = sqrlen10 * dot * dot;
+        }
+    }
+
+    return psd;
+}
+
+template <typename InputType, typename ComputeType>
+int ConstrainedDelaunay2<InputType, ComputeType>::GetLinkTriangle(int v) const
+{
+    // Remap in case an edge vertex was specified that is a duplicate.
+    v = this->mDuplicates[v];
+
+    int tri = 0;
+    for (int i = 0; i < this->mNumTriangles; ++i)
+    {
+        // Test whether v is a vertex of the triangle.
+        std::array<int, 3> indices;
+        GTE_CDT_REQUIRE(this->GetIndices(tri, indices));
+        for (int j = 0; j < 3; ++j)
+        {
+            if (v == indices[j])
+            {
+                return tri;
+            }
+        }
+
+        // v must be outside the triangle.
+        for (int j0 = 2, j1 = 0; j1 < 3; j0 = j1++)
+        {
+            if (this->mQuery.ToLine(v, indices[j0], indices[j1]) > 0)
+            {
+                // Vertex v sees the edge from outside, so traverse to the
+                // triangle sharing the edge.
+                std::array<int, 3> adjacencies;
+                GTE_CDT_REQUIRE(this->GetAdjacencies(tri, adjacencies));
+                int adj = adjacencies[j0];
+                GTE_CDT_REQUIRE_RET(adj >= 0, -1);
+                tri = adj;
+                break;
+            }
+        }
+    }
+
+    // The vertex must be in the triangulation.
+    GTE_CDT_FAILURE_RET(-1);
+}
+
+template <typename InputType, typename ComputeType>
+int ConstrainedDelaunay2<InputType, ComputeType>::GetIndexOfVertex(int tri,
+    int v) const
+{
+    std::array<int, 3> indices;
+    GTE_CDT_REQUIRE_RET(this->GetIndices(tri, indices), -1);
+    int indexOfV;
+    for (indexOfV = 0; indexOfV < 3; ++indexOfV)
+    {
+        if (v == indices[indexOfV])
+        {
+            return indexOfV;
+        }
+    }
+
+    // Unexpected failure.
+    GTE_CDT_FAILURE_RET(-1);
+}
+
+template <typename InputType, typename ComputeType>
+std::array<int, 2> ConstrainedDelaunay2<InputType, ComputeType>::
+GetAdjInterior(int tri, int v0, int v1) const
+{
+    std::array<int, 2> error = { -2, -2 };
+    int vIndex = GetIndexOfVertex(tri, v0);
+    GTE_CDT_REQUIRE_RET(vIndex >= 0, error);
+    int adj = this->mAdjacencies[3 * tri + vIndex];
+    if (adj >= 0)
+    {
+        for (int v2Index = 0; v2Index < 3; ++v2Index)
+        {
+            int v2 = this->mIndices[3 * adj + v2Index];
+            if (v2 != v0 && v2 != v1)
+            {
+                return{ adj, v2 };
+            }
+        }
+    }
+    else
+    {
+        return{ -1, -1 };
+    }
+
+    GTE_CDT_FAILURE_RET(error);
+}
+
+template <typename InputType, typename ComputeType>
+std::array<int, 2> ConstrainedDelaunay2<InputType, ComputeType>::
+GetAdjBoundary(int tri, int needBndVertex, int needAdjVIndex) const
+{
+    std::array<int, 2> error = { -2, -2 };
+    int vIndex = GetIndexOfVertex(tri, needAdjVIndex);
+    GTE_CDT_REQUIRE_RET(vIndex >= 0, error);
+    int adj = this->mAdjacencies[3 * tri + vIndex];
+    return{ needBndVertex, adj };
+}
+
+template <typename InputType, typename ComputeType>
+bool ConstrainedDelaunay2<InputType, ComputeType>::Connect(int tri, int adj,
+    int v0, int v1)
+{
+    if (tri >= 0)
+    {
+        int v0Index = GetIndexOfVertex(tri, v0);
+        GTE_CDT_REQUIRE(v0Index >= 0);
+        if (adj >= 0)
+        {
+            int v1Index = GetIndexOfVertex(adj, v1);
+            GTE_CDT_REQUIRE(v1Index >= 0);
+            this->mAdjacencies[3 * adj + v1Index] = tri;
+        }
+
+        this->mAdjacencies[3 * tri + v0Index] = adj;
+    }
+    // else: tri = -1, which occurs in the top-level call to retriangulate
+    return true;
+}
+
+template <typename InputType, typename ComputeType>
+bool ConstrainedDelaunay2<InputType, ComputeType>::BuildLink(int v,
+    int vTriangle, std::list<std::pair<int, int>>& link, bool& isOpen) const
+{
+    // The list starts with a known triangle in the link of v.
+    int vStartIndex = GetIndexOfVertex(vTriangle, v);
+    GTE_CDT_REQUIRE(vStartIndex >= 0);
+    link.push_front(std::make_pair(vTriangle, vStartIndex));
+
+    // Traverse adjacent triangles to the "left" of v.  Guard against an
+    // infinite loop.
+    int tri = vTriangle, vIndex = vStartIndex;
+    std::array<int, 3> adjacencies;
+    for (int i = 0; i < this->mNumTriangles; ++i)
+    {
+        GTE_CDT_REQUIRE(this->GetAdjacencies(tri, adjacencies));
+        int adjPrev = adjacencies[(vIndex + 2) % 3];
+        if (adjPrev >= 0)
+        {
+            if (adjPrev != vTriangle)
+            {
+                tri = adjPrev;
+                vIndex = GetIndexOfVertex(tri, v);
+                GTE_CDT_REQUIRE(vIndex >= 0);
+                link.push_back(std::make_pair(tri, vIndex));
+            }
+            else
+            {
+                // We have reached the starting triangle, so v is an interior
+                // vertex.
+                isOpen = false;
+                return true;
+            }
+        }
+        else
+        {
+            // We have reached a triangle with boundary edge, so v is a
+            // boundary vertex.  We mush find more triangles by searching
+            // to the "right" of v.  Guard against an infinite loop.
+            isOpen = true;
+            tri = vTriangle;
+            vIndex = vStartIndex;
+            for (int j = 0; j < this->mNumTriangles; ++j)
+            {
+                this->GetAdjacencies(tri, adjacencies);
+                int adjNext = adjacencies[vIndex];
+                if (adjNext >= 0)
+                {
+                    tri = adjNext;
+                    vIndex = GetIndexOfVertex(tri, v);
+                    GTE_CDT_REQUIRE(vIndex >= 0);
+                    link.push_front(std::make_pair(tri, vIndex));
+                }
+                else
+                {
+                    // We have reached the other boundary edge that shares v.
+                    return true;
+                }
+            }
+            break;
+        }
+    }
+
+    // Unexpected failure.
+    GTE_CDT_FAILURE;
+}
+
+template <typename InputType, typename ComputeType>
+void ConstrainedDelaunay2<InputType, ComputeType>::Trap() const
+{
+    LogError("CDT failure.");
+}
+
+
+}
diff --git a/Engine/Plugins/Experimental/GeometryProcessing/Source/GeometryAlgorithms/Private/GTEngine/Mathematics/GteContAlignedBox.h b/Engine/Plugins/Experimental/GeometryProcessing/Source/GeometryAlgorithms/Private/GTEngine/Mathematics/GteContAlignedBox.h
new file mode 100644
--- /dev/null
+++ b/Engine/Plugins/Experimental/GeometryProcessing/Source/GeometryAlgorithms/Private/GTEngine/Mathematics/GteContAlignedBox.h
@@ -0,0 +1,65 @@
+// David Eberly, Geometric Tools, Redmond WA 98052
+// Copyright (c) 1998-2019
+// Distributed under the Boost Software License, Version 1.0.
+// http://www.boost.org/LICENSE_1_0.txt
+// http://www.geometrictools.com/License/Boost/LICENSE_1_0.txt
+// File Version: 3.13.0 (2018/04/15)
+
+#pragma once
+
+#include <GTEngine/Mathematics/GteAlignedBox.h>
+
+namespace gte
+{
+
+// Compute the minimum size aligned bounding box of the points.  The extreme
+// values are the minima and maxima of the point coordinates.
+template <int N, typename Real>
+bool GetContainer(int numPoints, Vector<N, Real> const* points,
+     AlignedBox<N, Real>& box);
+
+// Test for containment.
+template <int N, typename Real>
+bool InContainer(Vector<N, Real> const& point, AlignedBox<N, Real> const& box);
+
+// Construct an aligned box that contains two other aligned boxes.  The
+// result is the minimum size box containing the input boxes.
+template <int N, typename Real>
+bool MergeContainers(AlignedBox<N, Real> const& box0,
+    AlignedBox<N,Real> const& box1, AlignedBox<N, Real>& merge);
+
+
+template <int N, typename Real>
+bool GetContainer(int numPoints, Vector<N, Real> const* points,
+    AlignedBox<N, Real>& box)
+{
+    return ComputeExtremes(numPoints, points, box.min, box.max);
+}
+
+template <int N, typename Real>
+bool InContainer(Vector<N, Real> const& point, AlignedBox<N, Real> const& box)
+{
+    for (int i = 0; i < N; ++i)
+    {
+        Real value = point[i];
+        if (value < box.min[i] || value > box.max[i])
+        {
+            return false;
+        }
+    }
+    return true;
+}
+
+template <int N, typename Real>
+bool MergeContainers(AlignedBox<N, Real> const& box0,
+    AlignedBox<N, Real> const& box1, AlignedBox<N, Real>& merge)
+{
+    for (int i = 0; i < N; ++i)
+    {
+        merge.min[i] = std::min(box0.min[i], box1.min[i]);
+        merge.max[i] = std::max(box0.max[i], box1.max[i]);
+    }
+    return true;
+}
+
+}
diff --git a/Engine/Plugins/Experimental/GeometryProcessing/Source/GeometryAlgorithms/Private/GTEngine/Mathematics/GteContCapsule3.h b/Engine/Plugins/Experimental/GeometryProcessing/Source/GeometryAlgorithms/Private/GTEngine/Mathematics/GteContCapsule3.h
new file mode 100644
--- /dev/null
+++ b/Engine/Plugins/Experimental/GeometryProcessing/Source/GeometryAlgorithms/Private/GTEngine/Mathematics/GteContCapsule3.h
@@ -0,0 +1,290 @@
+// David Eberly, Geometric Tools, Redmond WA 98052
+// Copyright (c) 1998-2019
+// Distributed under the Boost Software License, Version 1.0.
+// http://www.boost.org/LICENSE_1_0.txt
+// http://www.geometrictools.com/License/Boost/LICENSE_1_0.txt
+// File Version: 3.0.1 (2018/10/05)
+
+#pragma once
+
+#include <GTEngine/Mathematics/GteApprOrthogonalLine3.h>
+#include <GTEngine/Mathematics/GteDistPointLine.h>
+#include <GTEngine/Mathematics/GteDistPointSegment.h>
+#include <GTEngine/Mathematics/GteCapsule.h>
+#include <GTEngine/Mathematics/GteHypersphere.h>
+
+namespace gte
+{
+
+// Compute the axis of the capsule segment using least-squares fitting.  The
+// radius is the maximum distance from the points to the axis.  Hemispherical
+// caps are chosen as close together as possible.
+template <typename Real>
+bool GetContainer(int numPoints, Vector3<Real> const* points,
+    Capsule3<Real>& capsule);
+
+// Test for containment of a point by a capsule.
+template <typename Real>
+bool InContainer(Vector3<Real> const& point, Capsule3<Real> const& capsule);
+
+// Test for containment of a sphere by a capsule.
+template <typename Real>
+bool InContainer(Sphere3<Real> const& sphere, Capsule3<Real> const& capsule);
+
+// Test for containment of a capsule by a capsule.
+template <typename Real>
+bool InContainer(Capsule3<Real> const& testCapsule,
+    Capsule3<Real> const& capsule);
+
+// Compute a capsule that contains the input capsules.  The returned capsule
+// is not necessarily the one of smallest volume that contains the inputs.
+template <typename Real>
+bool MergeContainers(Capsule3<Real> const& capsule0,
+     Capsule3<Real> const& capsule1, Capsule3<Real>& merge);
+
+
+template <typename Real>
+bool GetContainer(int numPoints, Vector3<Real> const* points,
+    Capsule3<Real>& capsule)
+{
+    ApprOrthogonalLine3<Real> fitter;
+    fitter.Fit(numPoints, points);
+    Line3<Real> line = fitter.GetParameters();
+
+    DCPQuery<Real, Vector3<Real>, Line3<Real>> plQuery;
+    Real maxRadiusSqr = (Real)0;
+    for (int i = 0; i < numPoints; ++i)
+    {
+        auto result = plQuery(points[i], line);
+        if (result.sqrDistance > maxRadiusSqr)
+        {
+            maxRadiusSqr = result.sqrDistance;
+        }
+    }
+
+    Vector3<Real> basis[3];
+    basis[0] = line.direction;
+    ComputeOrthogonalComplement(1, basis);
+
+    Real minValue = std::numeric_limits<Real>::max();
+    Real maxValue = -std::numeric_limits<Real>::max();
+    for (int i = 0; i < numPoints; ++i)
+    {
+        Vector3<Real> diff = points[i] - line.origin;
+        Real uDotDiff = Dot(diff, basis[1]);
+        Real vDotDiff = Dot(diff, basis[2]);
+        Real wDotDiff = Dot(diff, basis[0]);
+        Real discr = maxRadiusSqr - (uDotDiff*uDotDiff + vDotDiff*vDotDiff);
+        Real radical = std::sqrt(std::max(discr, (Real)0));
+
+        Real test = wDotDiff + radical;
+        if (test < minValue)
+        {
+            minValue = test;
+        }
+
+        test = wDotDiff - radical;
+        if (test > maxValue)
+        {
+            maxValue = test;
+        }
+    }
+
+    Vector3<Real> center = line.origin +
+        (((Real)0.5)*(minValue + maxValue))*line.direction;
+
+    Real extent;
+    if (maxValue > minValue)
+    {
+        // Container is a capsule.
+        extent = ((Real)0.5)*(maxValue - minValue);
+    }
+    else
+    {
+        // Container is a sphere.
+        extent = (Real)0;
+    }
+
+    capsule.segment = Segment3<Real>(center, line.direction, extent);
+    capsule.radius = std::sqrt(maxRadiusSqr);
+    return true;
+}
+
+template <typename Real>
+bool InContainer(Vector3<Real> const& point, Capsule3<Real> const& capsule)
+{
+    DCPQuery<Real, Vector3<Real>, Segment3<Real>> psQuery;
+    auto result = psQuery(point, capsule.segment);
+    return result.distance <= capsule.radius;
+}
+
+template <typename Real>
+bool InContainer(Sphere3<Real> const& sphere, Capsule3<Real> const& capsule)
+{
+    Real rDiff = capsule.radius - sphere.radius;
+    if (rDiff >= (Real)0)
+    {
+        DCPQuery<Real, Vector3<Real>, Segment3<Real>> psQuery;
+        auto result = psQuery(sphere.center, capsule.segment);
+        return result.distance <= rDiff;
+    }
+    return false;
+}
+
+template <typename Real>
+bool InContainer(Capsule3<Real> const& testCapsule,
+    Capsule3<Real> const& capsule)
+{
+    Sphere3<Real> spherePosEnd(testCapsule.segment.p[1], testCapsule.radius);
+    Sphere3<Real> sphereNegEnd(testCapsule.segment.p[0], testCapsule.radius);
+    return InContainer<Real>(spherePosEnd, capsule)
+        && InContainer<Real>(sphereNegEnd, capsule);
+}
+
+template <typename Real>
+bool MergeContainers(Capsule3<Real> const& capsule0,
+    Capsule3<Real> const& capsule1, Capsule3<Real>& merge)
+{
+    if (InContainer<Real>(capsule0, capsule1))
+    {
+        merge = capsule1;
+        return true;
+    }
+
+    if (InContainer<Real>(capsule1, capsule0))
+    {
+        merge = capsule0;
+        return true;
+    }
+
+    Vector3<Real> P0, P1, D0, D1;
+    Real extent0, extent1;
+    capsule0.segment.GetCenteredForm(P0, D0, extent0);
+    capsule1.segment.GetCenteredForm(P1, D1, extent1);
+
+    // Axis of final capsule.
+    Line3<Real> line;
+
+    // Axis center is average of input axis centers.
+    line.origin = ((Real)0.5)*(P0 + P1);
+
+    // Axis unit direction is average of input axis unit directions.
+    if (Dot(D0, D1) >= (Real)0)
+    {
+        line.direction = D0 + D1;
+    }
+    else
+    {
+        line.direction = D0 - D1;
+    }
+    Normalize(line.direction);
+
+    // Cylinder with axis 'line' must contain the spheres centered at the
+    // endpoints of the input capsules.
+    DCPQuery<Real, Vector3<Real>, Line3<Real>> plQuery;
+    Vector3<Real> posEnd0 = capsule0.segment.p[1];
+    Real radius = plQuery(posEnd0, line).distance + capsule0.radius;
+
+    Vector3<Real> negEnd0 = capsule0.segment.p[0];
+    Real tmp = plQuery(negEnd0, line).distance + capsule0.radius;
+
+    Vector3<Real> posEnd1 = capsule1.segment.p[1];
+    tmp = plQuery(posEnd1, line).distance + capsule1.radius;
+    if (tmp > radius)
+    {
+        radius = tmp;
+    }
+
+    Vector3<Real> negEnd1 = capsule1.segment.p[0];
+    tmp = plQuery(negEnd1, line).distance + capsule1.radius;
+    if (tmp > radius)
+    {
+        radius = tmp;
+    }
+
+    // Process sphere <posEnd0,r0>.
+    Real rDiff = radius - capsule0.radius;
+    Real rDiffSqr = rDiff*rDiff;
+    Vector3<Real> diff = line.origin - posEnd0;
+    Real k0 = Dot(diff, diff) - rDiffSqr;
+    Real k1 = Dot(diff, line.direction);
+    Real discr = k1*k1 - k0;  // assert:  k1*k1-k0 >= 0, guard against anyway
+    Real root = std::sqrt(std::max(discr, (Real)0));
+    Real tPos = -k1 - root;
+    Real tNeg = -k1 + root;
+
+    // Process sphere <negEnd0,r0>.
+    diff = line.origin - negEnd0;
+    k0 = Dot(diff, diff) - rDiffSqr;
+    k1 = Dot(diff, line.direction);
+    discr = k1*k1 - k0;  // assert:  k1*k1-k0 >= 0, guard against anyway
+    root = std::sqrt(std::max(discr, (Real)0));
+    tmp = -k1 - root;
+    if (tmp > tPos)
+    {
+        tPos = tmp;
+    }
+    tmp = -k1 + root;
+    if (tmp < tNeg)
+    {
+        tNeg = tmp;
+    }
+
+    // Process sphere <posEnd1,r1>.
+    rDiff = radius - capsule1.radius;
+    rDiffSqr = rDiff*rDiff;
+    diff = line.origin - posEnd1;
+    k0 = Dot(diff, diff) - rDiffSqr;
+    k1 = Dot(diff, line.direction);
+    discr = k1*k1 - k0;  // assert:  k1*k1-k0 >= 0, guard against anyway
+    root = std::sqrt(std::max(discr, (Real)0));
+    tmp = -k1 - root;
+    if (tmp > tPos)
+    {
+        tPos = tmp;
+    }
+    tmp = -k1 + root;
+    if (tmp < tNeg)
+    {
+        tNeg = tmp;
+    }
+
+    // Process sphere <negEnd1,r1>.
+    diff = line.origin - negEnd1;
+    k0 = Dot(diff, diff) - rDiffSqr;
+    k1 = Dot(diff, line.direction);
+    discr = k1*k1 - k0;  // assert:  k1*k1-k0 >= 0, guard against anyway
+    root = std::sqrt(std::max(discr, (Real)0));
+    tmp = -k1 - root;
+    if (tmp > tPos)
+    {
+        tPos = tmp;
+    }
+    tmp = -k1 + root;
+    if (tmp < tNeg)
+    {
+        tNeg = tmp;
+    }
+
+    Vector3<Real> center = line.origin +
+        ((Real)0.5)*(tPos + tNeg)*line.direction;
+
+    Real extent;
+    if (tPos > tNeg)
+    {
+        // Container is a capsule.
+        extent = ((Real)0.5)*(tPos - tNeg);
+    }
+    else
+    {
+        // Container is a sphere.
+        extent = (Real)0;
+    }
+
+    merge.segment = Segment3<Real>(center, line.direction, extent);
+    merge.radius = radius;
+    return true;
+}
+
+
+}
diff --git a/Engine/Plugins/Experimental/GeometryProcessing/Source/GeometryAlgorithms/Private/GTEngine/Mathematics/GteContCircle2.h b/Engine/Plugins/Experimental/GeometryProcessing/Source/GeometryAlgorithms/Private/GTEngine/Mathematics/GteContCircle2.h
new file mode 100644
--- /dev/null
+++ b/Engine/Plugins/Experimental/GeometryProcessing/Source/GeometryAlgorithms/Private/GTEngine/Mathematics/GteContCircle2.h
@@ -0,0 +1,81 @@
+// David Eberly, Geometric Tools, Redmond WA 98052
+// Copyright (c) 1998-2019
+// Distributed under the Boost Software License, Version 1.0.
+// http://www.boost.org/LICENSE_1_0.txt
+// http://www.geometrictools.com/License/Boost/LICENSE_1_0.txt
+// File Version: 3.0.2 (2018/10/16)
+
+#pragma once
+
+#include <GTEngine/Mathematics/GteVector2.h>
+#include <GTEngine/Mathematics/GteHypersphere.h>
+
+namespace gte
+{
+    // Compute the smallest bounding circle whose center is the average of
+    // the input points.
+    template <typename Real>
+    bool GetContainer(int numPoints, Vector2<Real> const* points, Circle2<Real>& circle)
+    {
+        circle.center = points[0];
+        for (int i = 1; i < numPoints; ++i)
+        {
+            circle.center += points[i];
+        }
+        circle.center /= (Real)numPoints;
+
+        circle.radius = (Real)0;
+        for (int i = 0; i < numPoints; ++i)
+        {
+            Vector2<Real> diff = points[i] - circle.center;
+            Real radiusSqr = Dot(diff, diff);
+            if (radiusSqr > circle.radius)
+            {
+                circle.radius = radiusSqr;
+            }
+        }
+
+        circle.radius = std::sqrt(circle.radius);
+        return true;
+    }
+
+    // Test for containment of a point inside a circle.
+    template <typename Real>
+    bool InContainer(Vector2<Real> const& point, Circle2<Real> const& circle)
+    {
+        Vector2<Real> diff = point - circle.center;
+        return Length(diff) <= circle.radius;
+    }
+
+    // Compute the smallest bounding circle that contains the input circles.
+    template <typename Real>
+    bool MergeContainers(Circle2<Real> const& circle0, Circle2<Real> const& circle1, Circle2<Real>& merge)
+    {
+        Vector2<Real> cenDiff = circle1.center - circle0.center;
+        Real lenSqr = Dot(cenDiff, cenDiff);
+        Real rDiff = circle1.radius - circle0.radius;
+        Real rDiffSqr = rDiff * rDiff;
+
+        if (rDiffSqr >= lenSqr)
+        {
+            merge = (rDiff >= (Real)0 ? circle1 : circle0);
+        }
+        else
+        {
+            Real length = std::sqrt(lenSqr);
+            if (length > (Real)0)
+            {
+                Real coeff = (length + rDiff) / (((Real)2)*length);
+                merge.center = circle0.center + coeff * cenDiff;
+            }
+            else
+            {
+                merge.center = circle0.center;
+            }
+
+            merge.radius = ((Real)0.5)*(length + circle0.radius + circle1.radius);
+        }
+
+        return true;
+    }
+}
diff --git a/Engine/Plugins/Experimental/GeometryProcessing/Source/GeometryAlgorithms/Private/GTEngine/Mathematics/GteContCone.h b/Engine/Plugins/Experimental/GeometryProcessing/Source/GeometryAlgorithms/Private/GTEngine/Mathematics/GteContCone.h
new file mode 100644
--- /dev/null
+++ b/Engine/Plugins/Experimental/GeometryProcessing/Source/GeometryAlgorithms/Private/GTEngine/Mathematics/GteContCone.h
@@ -0,0 +1,22 @@
+// David Eberly, Geometric Tools, Redmond WA 98052
+// Copyright (c) 1998-2019
+// Distributed under the Boost Software License, Version 1.0.
+// http://www.boost.org/LICENSE_1_0.txt
+// http://www.geometrictools.com/License/Boost/LICENSE_1_0.txt
+// File Version: 3.19.0 (2018/11/30)
+
+#pragma once
+
+#include <GTEngine/Mathematics/GteCone.h>
+
+namespace gte
+{
+    // Test for containment of a point by a cone.
+    template <int N, typename Real>
+    bool InContainer(Vector<N, Real> const& point, Cone<N, Real> const& cone)
+    {
+        Vector<N, Real> diff = point - cone.ray.origin;
+        Real h = Dot(cone.ray.direction, diff);
+        return cone.minHeight <= h && h <= cone.maxHeight && h >= cone.cosAngle * Length(diff);
+    }
+}
diff --git a/Engine/Plugins/Experimental/GeometryProcessing/Source/GeometryAlgorithms/Private/GTEngine/Mathematics/GteContCylinder3.h b/Engine/Plugins/Experimental/GeometryProcessing/Source/GeometryAlgorithms/Private/GTEngine/Mathematics/GteContCylinder3.h
new file mode 100644
--- /dev/null
+++ b/Engine/Plugins/Experimental/GeometryProcessing/Source/GeometryAlgorithms/Private/GTEngine/Mathematics/GteContCylinder3.h
@@ -0,0 +1,88 @@
+// David Eberly, Geometric Tools, Redmond WA 98052
+// Copyright (c) 1998-2019
+// Distributed under the Boost Software License, Version 1.0.
+// http://www.boost.org/LICENSE_1_0.txt
+// http://www.geometrictools.com/License/Boost/LICENSE_1_0.txt
+// File Version: 3.0.1 (2018/10/05)
+
+#pragma once
+
+#include <GTEngine/Mathematics/GteApprOrthogonalLine3.h>
+#include <GTEngine/Mathematics/GteDistPointLine.h>
+#include <GTEngine/Mathematics/GteCylinder3.h>
+
+namespace gte
+{
+
+// Compute the cylinder axis segment using least-squares fit.  The radius is
+// the maximum distance from points to the axis.  The height is determined by
+// projection of points onto the axis and determining the containing interval.
+template <typename Real>
+bool GetContainer(int numPoints, Vector3<Real> const* points,
+     Cylinder3<Real>& cylinder);
+
+// Test for containment of a point by a cylinder.
+template <typename Real>
+bool InContainer(Vector3<Real> const& point, Cylinder3<Real> const& cylinder);
+
+
+template <typename Real>
+bool GetContainer(int numPoints, Vector3<Real> const* points,
+    Cylinder3<Real>& cylinder)
+{
+    ApprOrthogonalLine3<Real> fitter;
+    fitter.Fit(numPoints, points);
+    Line3<Real> line = fitter.GetParameters();
+
+    DCPQuery<Real, Vector3<Real>, Line3<Real>> plQuery;
+    Real maxRadiusSqr = (Real)0;
+    for (int i = 0; i < numPoints; ++i)
+    {
+        auto result = plQuery(points[i], line);
+        if (result.sqrDistance > maxRadiusSqr)
+        {
+            maxRadiusSqr = result.sqrDistance;
+        }
+    }
+
+    Vector3<Real> diff = points[0] - line.origin;
+    Real wMin = Dot(line.direction, diff);
+    Real wMax = wMin;
+    for (int i = 1; i < numPoints; ++i)
+    {
+        diff = points[i] - line.origin;
+        Real w = Dot(line.direction, diff);
+        if (w < wMin)
+        {
+            wMin = w;
+        }
+        else if (w > wMax)
+        {
+            wMax = w;
+        }
+    }
+
+    cylinder.axis.origin = line.origin +
+        (((Real)0.5)*(wMax + wMin))*line.direction;
+    cylinder.axis.direction = line.direction;
+    cylinder.radius = std::sqrt(maxRadiusSqr);
+    cylinder.height = wMax - wMin;
+    return true;
+}
+
+template <typename Real>
+bool InContainer(Vector3<Real> const& point, Cylinder3<Real> const& cylinder)
+{
+    Vector3<Real> diff = point - cylinder.axis.origin;
+    Real zProj = Dot(diff, cylinder.axis.direction);
+    if (std::abs(zProj)*(Real)2 > cylinder.height)
+    {
+        return false;
+    }
+
+    Vector3<Real> xyProj = diff - zProj*cylinder.axis.direction;
+    return Length(xyProj) <= cylinder.radius;
+}
+
+
+}
diff --git a/Engine/Plugins/Experimental/GeometryProcessing/Source/GeometryAlgorithms/Private/GTEngine/Mathematics/GteContEllipse2.h b/Engine/Plugins/Experimental/GeometryProcessing/Source/GeometryAlgorithms/Private/GTEngine/Mathematics/GteContEllipse2.h
new file mode 100644
--- /dev/null
+++ b/Engine/Plugins/Experimental/GeometryProcessing/Source/GeometryAlgorithms/Private/GTEngine/Mathematics/GteContEllipse2.h
@@ -0,0 +1,164 @@
+// David Eberly, Geometric Tools, Redmond WA 98052
+// Copyright (c) 1998-2019
+// Distributed under the Boost Software License, Version 1.0.
+// http://www.boost.org/LICENSE_1_0.txt
+// http://www.geometrictools.com/License/Boost/LICENSE_1_0.txt
+// File Version: 3.0.1 (2018/10/05)
+
+#pragma once
+
+#include <GTEngine/Mathematics/GteApprGaussian2.h>
+#include <GTEngine/Mathematics/GteHyperellipsoid.h>
+#include <GTEngine/Mathematics/GteProjection.h>
+
+namespace gte
+{
+
+// The input points are fit with a Gaussian distribution.  The center C of the
+// ellipsoid is chosen to be the mean of the distribution.  The axes of the
+// ellipsoid are chosen to be the eigenvectors of the covariance matrix M.
+// The shape of the ellipsoid is determined by the absolute values of the
+// eigenvalues.  NOTE: The construction is ill-conditioned if the points are
+// (nearly) collinear.  In this case M has a (nearly) zero eigenvalue, so
+// inverting M can be a problem numerically.
+template <typename Real>
+bool GetContainer(int numPoints, Vector2<Real> const* points,
+    Ellipse2<Real>& ellipse);
+
+// Test for containment of a point inside an ellipse.
+template <typename Real>
+bool InContainer(Vector2<Real> const& point, Ellipse2<Real> const& ellipse);
+
+// Construct a bounding ellipse for the two input ellipses.  The result is
+// not necessarily the minimum-area ellipse containing the two ellipses.
+template <typename Real>
+bool MergeContainers(Ellipse2<Real> const& ellipse0,
+    Ellipse2<Real> const& ellipse1, Ellipse2<Real>& merge);
+
+
+template <typename Real>
+bool GetContainer(int numPoints, Vector2<Real> const* points,
+    Ellipse2<Real>& ellipse)
+{
+    // Fit the points with a Gaussian distribution.  The covariance matrix
+    // is M = sum_j D[j]*U[j]*U[j]^T, where D[j] are the eigenvalues and U[j]
+    // are corresponding unit-length eigenvectors.
+    ApprGaussian2<Real> fitter;
+    if (fitter.Fit(numPoints, points))
+    {
+        OrientedBox2<Real> box = fitter.GetParameters();
+
+        // If either eigenvalue is nonpositive, adjust the D[] values so that
+        // we actually build an ellipse.
+        for (int j = 0; j < 2; ++j)
+        {
+            if (box.extent[j] < (Real)0)
+            {
+                box.extent[j] = -box.extent[j];
+            }
+        }
+
+        // Grow the ellipse, while retaining its shape determined by the
+        // covariance matrix, to enclose all the input points.  The quadratic
+        // form that is used for the ellipse construction is
+        //   Q(X) = (X-C)^T*M*(X-C)
+        //        = (X-C)^T*(sum_j D[j]*U[j]*U[j]^T)*(X-C)
+        //        = sum_j D[j]*Dot(U[j],X-C)^2
+        // If the maximum value of Q(X[i]) for all input points is V^2, then a
+        // bounding ellipse is Q(X) = V^2, because Q(X[i]) <= V^2 for all i.
+
+        Real maxValue = (Real)0;
+        for (int i = 0; i < numPoints; ++i)
+        {
+            Vector2<Real> diff = points[i] - box.center;
+            Real dot[2] =
+            {
+                Dot(box.axis[0], diff),
+                Dot(box.axis[1], diff)
+            };
+
+            Real value =
+                box.extent[0] * dot[0] * dot[0] +
+                box.extent[1] * dot[1] * dot[1];
+
+            if (value > maxValue)
+            {
+                maxValue = value;
+            }
+        }
+
+        // Arrange for the quadratic to satisfy Q(X) <= 1.
+        ellipse.center = box.center;
+        for (int j = 0; j < 2; ++j)
+        {
+            ellipse.axis[j] = box.axis[j];
+            ellipse.extent[j] = std::sqrt(maxValue / box.extent[j]);
+        }
+        return true;
+
+    }
+
+    return false;
+}
+
+template <typename Real>
+bool InContainer(Vector2<Real> const& point, Ellipse2<Real> const& ellipse)
+{
+    Vector2<Real> diff = point - ellipse.center;
+    Vector2<Real> standardized{
+        Dot(diff, ellipse.axis[0]) / ellipse.extent[0],
+        Dot(diff, ellipse.axis[1]) / ellipse.extent[1] };
+    return Length(standardized) <= (Real)1;
+}
+
+template <typename Real>
+bool MergeContainers(Ellipse2<Real> const& ellipse0,
+    Ellipse2<Real> const& ellipse1, Ellipse2<Real>& merge)
+{
+    // Compute the average of the input centers.
+    merge.center = ((Real)0.5)*(ellipse0.center + ellipse1.center);
+
+    // The bounding ellipse orientation is the average of the input
+    // orientations.
+    if (Dot(ellipse0.axis[0], ellipse1.axis[0]) >= (Real)0)
+    {
+        merge.axis[0] = ((Real)0.5)*(ellipse0.axis[0] + ellipse1.axis[0]);
+    }
+    else
+    {
+        merge.axis[0] = ((Real)0.5)*(ellipse0.axis[0] - ellipse1.axis[0]);
+    }
+    Normalize(merge.axis[0]);
+    merge.axis[1] = -Perp(merge.axis[0]);
+
+    // Project the input ellipses onto the axes obtained by the average of the
+    // orientations and that go through the center obtained by the average of
+    // the centers.
+    for (int j = 0; j < 2; ++j)
+    {
+        // Projection axis.
+        Line2<Real> line(merge.center, merge.axis[j]);
+
+        // Project ellipsoids onto the axis.
+        Real min0, max0, min1, max1;
+        Project(ellipse0, line, min0, max0);
+        Project(ellipse1, line, min1, max1);
+
+        // Determine the smallest interval containing the projected
+        // intervals.
+        Real maxIntr = (max0 >= max1 ? max0 : max1);
+        Real minIntr = (min0 <= min1 ? min0 : min1);
+
+        // Update the average center to be the center of the bounding box
+        // defined by the projected intervals.
+        merge.center += line.direction*(((Real)0.5)*(minIntr + maxIntr));
+
+        // Compute the extents of the box based on the new center.
+        merge.extent[j] = ((Real)0.5)*(maxIntr - minIntr);
+    }
+
+    return true;
+}
+
+
+}
diff --git a/Engine/Plugins/Experimental/GeometryProcessing/Source/GeometryAlgorithms/Private/GTEngine/Mathematics/GteContEllipse2MinCR.h b/Engine/Plugins/Experimental/GeometryProcessing/Source/GeometryAlgorithms/Private/GTEngine/Mathematics/GteContEllipse2MinCR.h
new file mode 100644
--- /dev/null
+++ b/Engine/Plugins/Experimental/GeometryProcessing/Source/GeometryAlgorithms/Private/GTEngine/Mathematics/GteContEllipse2MinCR.h
@@ -0,0 +1,233 @@
+// David Eberly, Geometric Tools, Redmond WA 98052
+// Copyright (c) 1998-2019
+// Distributed under the Boost Software License, Version 1.0.
+// http://www.boost.org/LICENSE_1_0.txt
+// http://www.geometrictools.com/License/Boost/LICENSE_1_0.txt
+// File Version: 3.0.0 (2016/06/19)
+
+#pragma once
+
+#include <GTEngine/LowLevel/GteLogger.h>
+#include <GTEngine/Mathematics/GteMatrix2x2.h>
+#include <algorithm>
+
+// Compute the minimum-area ellipse, (X-C)^T R D R^T (X-C) = 1, given the
+// center C and the orientation matrix R.  The columns of R are the axes of
+// the ellipse.  The algorithm computes the diagonal matrix D.  The minimum
+// area is pi/sqrt(D[0]*D[1]), where D = diag(D[0],D[1]).  The problem is
+// equivalent to maximizing the product D[0]*D[1] for a given C and R, and
+// subject to the constraints
+//   (P[i]-C)^T R D R^T (P[i]-C) <= 1
+// for all input points P[i] with 0 <= i < N.  Each constraint has the form
+//   A[0]*D[0] + A[1]*D[1] <= 1
+// where A[0] >= 0 and A[1] >= 0.
+
+namespace gte
+{
+
+template <typename Real>
+class ContEllipse2MinCR
+{
+public:
+    void operator()(int numPoints, Vector2<Real> const* points,
+        Vector2<Real> const& C, Matrix2x2<Real> const& R, Real D[2]) const;
+
+private:
+    static void MaxProduct(std::vector<Vector2<Real>>& A, Real D[2]);
+};
+
+
+template <typename Real>
+void ContEllipse2MinCR<Real>::operator()(int numPoints,
+    Vector2<Real> const* points, Vector2<Real> const& C,
+    Matrix2x2<Real>const & R, Real D[2]) const
+{
+    // Compute the constraint coefficients, of the form (A[0],A[1]) for
+    // each i.
+    std::vector<Vector2<Real>> A(numPoints);
+    for (int i = 0; i < numPoints; ++i)
+    {
+        Vector2<Real> diff = points[i] - C;  // P[i] - C
+        Vector2<Real> prod = diff*R;  // R^T*(P[i] - C) = (u,v)
+        A[i] = prod * prod;  // (u^2, v^2)
+    }
+
+    // Sort to eliminate redundant constraints.  Use a lexicographical sort,
+    // (x0,y0) > (x1,y1) if x0 > x1 or if x0 = x1 and y0 > y1.  Remove all
+    // but the first entry in blocks with x0 = x1 because the corresponding
+    // constraint lines for the first entry "hides" all the others from the
+    // origin.
+    std::sort(A.begin(), A.end(),
+        [](Vector2<Real> const& P0, Vector2<Real> const& P1)
+    {
+        if (P0[0] > P1[0]) { return true; }
+        if (P0[0] < P1[0]) { return false; }
+        return P0[1] > P1[1];
+    }
+    );
+    auto end = std::unique(A.begin(), A.end(),
+        [](Vector2<Real> const& P0, Vector2<Real> const& P1)
+    {
+        return P0[0] == P1[0];
+    }
+    );
+    A.erase(end, A.end());
+
+    // Lexicographical sort, (x0,y0) > (x1,y1) if y0 > y1 or if y0 = y1 and
+    // x0 > x1.  Remove all but first entry in blocks with y0 = y1 since the
+    // corresponding constraint lines for the first entry "hides" all the
+    // others from the origin.
+    std::sort(A.begin(), A.end(),
+        [](Vector2<Real> const& P0, Vector2<Real> const& P1)
+    {
+        if (P0[1] > P1[1])
+        {
+            return true;
+        }
+
+        if (P0[1] < P1[1])
+        {
+            return false;
+        }
+
+        return P0[0] > P1[0];
+    }
+    );
+    end = std::unique(A.begin(), A.end(),
+        [](Vector2<Real> const& P0, Vector2<Real> const& P1)
+    {
+        return P0[1] == P1[1];
+    }
+    );
+    A.erase(end, A.end());
+
+    MaxProduct(A, D);
+}
+
+template <typename Real>
+void ContEllipse2MinCR<Real>::MaxProduct(std::vector<Vector2<Real>>& A,
+    Real D[2])
+{
+    // Keep track of which constraint lines have already been used in the
+    // search.
+    int numConstraints = static_cast<int>(A.size());
+    std::vector<bool> used(A.size());
+    std::fill(used.begin(), used.end(), false);
+
+    // Find the constraint line whose y-intercept (0,ymin) is closest to the
+    // origin.  This line contributes to the convex hull of the constraints
+    // and the search for the maximum starts here.  Also find the constraint
+    // line whose x-intercept (xmin,0) is closest to the origin.  This line
+    // contributes to the convex hull of the constraints and the search for
+    // the maximum terminates before or at this line.
+    int i, iYMin = -1;
+    int iXMin = -1;
+    Real axMax = (Real)0, ayMax = (Real)0;  // A[i] >= (0,0) by design
+    for (i = 0; i < numConstraints; ++i)
+    {
+        // The minimum x-intercept is 1/A[iXMin][0] for A[iXMin][0] the
+        // maximum of the A[i][0].
+        if (A[i][0] > axMax)
+        {
+            axMax = A[i][0];
+            iXMin = i;
+        }
+
+        // The minimum y-intercept is 1/A[iYMin][1] for A[iYMin][1] the
+        // maximum of the A[i][1].
+        if (A[i][1] > ayMax)
+        {
+            ayMax = A[i][1];
+            iYMin = i;
+        }
+    }
+    LogAssert(iXMin != -1 && iYMin != -1, "Unexpected condition.");
+    used[iYMin] = true;
+
+    // The convex hull is searched in a clockwise manner starting with the
+    // constraint line constructed above.  The next vertex of the hull occurs
+    // as the closest point to the first vertex on the current constraint
+    // line.  The following loop finds each consecutive vertex.
+    Real x0 = (Real)0, xMax = ((Real)1) / axMax;
+    int j;
+    for (j = 0; j < numConstraints; ++j)
+    {
+        // Find the line whose intersection with the current line is closest
+        // to the last hull vertex.  The last vertex is at (x0,y0) on the
+        // current line.
+        Real x1 = xMax;
+        int line = -1;
+        for (i = 0; i < numConstraints; ++i)
+        {
+            if (!used[i])
+            {
+                // This line not yet visited, process it.  Given current
+                // constraint line a0*x+b0*y =1 and candidate line
+                // a1*x+b1*y = 1, find the point of intersection.  The
+                // determinant of the system is d = a0*b1-a1*b0.  We only
+                // care about lines that have more negative slope than the
+                // previous one, that is, -a1/b1 < -a0/b0, in which case we
+                // process only lines for which d < 0.
+                Real det = DotPerp(A[iYMin], A[i]);
+                if (det < (Real)0)
+                {
+                    // Compute the x-value for the point of intersection,
+                    // (x1,y1).  There may be floating point error issues in
+                    // the comparision 'D[0] <= fX1'.  Consider modifying to
+                    // 'D[0] <= fX1+epsilon'.
+                    D[0] = (A[i][1] - A[iYMin][1]) / det;
+                    if (x0 < D[0] && D[0] <= x1)
+                    {
+                        line = i;
+                        x1 = D[0];
+                    }
+                }
+            }
+        }
+
+        // Next vertex is at (x1,y1) whose x-value was computed above.  First
+        // check for the maximum of x*y on the current line for x in [x0,x1].
+        // On this interval the function is f(x) = x*(1-a0*x)/b0.  The
+        // derivative is f'(x) = (1-2*a0*x)/b0 and f'(r) = 0 when
+        // r = 1/(2*a0).  The three candidates for the maximum are f(x0),
+        // f(r), and f(x1).  Comparisons are made between r and the end points
+        // x0 and x1.  Since a0 = 0 is possible (constraint line is horizontal
+        // and f is increasing on line), the division in r is not performed
+        // and the comparisons are made between 1/2 = a0*r and a0*x0 or a0*x1.
+
+        // Compare r < x0.
+        if ((Real)0.5 < A[iYMin][0] * x0)
+        {
+            // The maximum is f(x0) since the quadratic f decreases for
+            // x > r.
+            D[0] = x0;
+            D[1] = ((Real)1 - A[iYMin][0] * D[0]) / A[iYMin][1];  // = f(x0)
+            break;
+        }
+
+        // Compare r < x1.
+        if ((Real)0.5 < A[iYMin][0] * x1)
+        {
+            // The maximum is f(r).  The search ends here because the
+            // current line is tangent to the level curve of f(x)=f(r)
+            // and x*y can therefore only decrease as we traverse further
+            // around the hull in the clockwise direction.
+            D[0] = ((Real)0.5) / A[iYMin][0];
+            D[1] = ((Real)0.5) / A[iYMin][1];  // = f(r)
+            break;
+        }
+
+        // The maximum is f(x1).  The function x*y is potentially larger
+        // on the next line, so continue the search.
+        LogAssert(line != -1, "Unexpected condition.");
+        x0 = x1;
+        x1 = xMax;
+        used[line] = true;
+        iYMin = line;
+    }
+
+    LogAssert(j < numConstraints, "Unexpected condition.");
+}
+
+
+}
diff --git a/Engine/Plugins/Experimental/GeometryProcessing/Source/GeometryAlgorithms/Private/GTEngine/Mathematics/GteContEllipsoid3.h b/Engine/Plugins/Experimental/GeometryProcessing/Source/GeometryAlgorithms/Private/GTEngine/Mathematics/GteContEllipsoid3.h
new file mode 100644
--- /dev/null
+++ b/Engine/Plugins/Experimental/GeometryProcessing/Source/GeometryAlgorithms/Private/GTEngine/Mathematics/GteContEllipsoid3.h
@@ -0,0 +1,181 @@
+// David Eberly, Geometric Tools, Redmond WA 98052
+// Copyright (c) 1998-2019
+// Distributed under the Boost Software License, Version 1.0.
+// http://www.boost.org/LICENSE_1_0.txt
+// http://www.geometrictools.com/License/Boost/LICENSE_1_0.txt
+// File Version: 3.0.1 (2018/10/05)
+
+#pragma once
+
+#include <GTEngine/Mathematics/GteMatrix3x3.h>
+#include <GTEngine/Mathematics/GteRotation.h>
+#include <GTEngine/Mathematics/GteApprGaussian3.h>
+#include <GTEngine/Mathematics/GteHyperellipsoid.h>
+#include <GTEngine/Mathematics/GteProjection.h>
+
+namespace gte
+{
+
+// The input points are fit with a Gaussian distribution.  The center C of the
+// ellipsoid is chosen to be the mean of the distribution.  The axes of the
+// ellipsoid are chosen to be the eigenvectors of the covariance matrix M.
+// The shape of the ellipsoid is determined by the absolute values of the
+// eigenvalues.  NOTE: The construction is ill-conditioned if the points are
+// (nearly) collinear or (nearly) planar.  In this case M has a (nearly) zero
+// eigenvalue, so inverting M is problematic.
+template <typename Real>
+bool GetContainer(int numPoints, Vector3<Real> const* points,
+    Ellipsoid3<Real>& ellipsoid);
+
+// Test for containment of a point inside an ellipsoid.
+template <typename Real>
+bool InContainer(Vector3<Real> const& point,
+    Ellipsoid3<Real> const& ellipsoid);
+
+// Construct a bounding ellipsoid for the two input ellipsoids.  The result is
+// not necessarily the minimum-volume ellipsoid containing the two ellipsoids.
+template <typename Real>
+bool MergeContainers(Ellipsoid3<Real> const& ellipsoid0,
+    Ellipsoid3<Real> const& ellipsoid1, Ellipsoid3<Real>& merge);
+
+
+template <typename Real>
+bool GetContainer(int numPoints, Vector3<Real> const* points,
+    Ellipsoid3<Real>& ellipsoid)
+{
+    // Fit the points with a Gaussian distribution.  The covariance matrix
+    // is M = sum_j D[j]*U[j]*U[j]^T, where D[j] are the eigenvalues and U[j]
+    // are corresponding unit-length eigenvectors.
+    ApprGaussian3<Real> fitter;
+    if (fitter.Fit(numPoints, points))
+    {
+        OrientedBox3<Real> box = fitter.GetParameters();
+
+        // If either eigenvalue is nonpositive, adjust the D[] values so that
+        // we actually build an ellipse.
+        for (int j = 0; j < 3; ++j)
+        {
+            if (box.extent[j] < (Real)0)
+            {
+                box.extent[j] = -box.extent[j];
+            }
+        }
+
+        // Grow the ellipsoid, while retaining its shape determined by the
+        // covariance matrix, to enclose all the input points.  The quadratic
+        // form that is used for the ellipsoid construction is
+        //   Q(X) = (X-C)^T*M*(X-C)
+        //        = (X-C)^T*(sum_j D[j]*U[j]*U[j]^T)*(X-C)
+        //        = sum_j D[j]*Dot(U[j],X-C)^2
+        // If the maximum value of Q(X[i]) for all input points is V^2, then a
+        // bounding ellipsoid is Q(X) = V^2 since Q(X[i]) <= V^2 for all i.
+
+        Real maxValue = (Real)0;
+        for (int i = 0; i < numPoints; ++i)
+        {
+            Vector3<Real> diff = points[i] - box.center;
+            Real dot[3] =
+            {
+                Dot(box.axis[0], diff),
+                Dot(box.axis[1], diff),
+                Dot(box.axis[2], diff)
+            };
+
+            Real value =
+                box.extent[0] * dot[0] * dot[0] +
+                box.extent[1] * dot[1] * dot[1] +
+                box.extent[2] * dot[2] * dot[2];
+
+            if (value > maxValue)
+            {
+                maxValue = value;
+            }
+        }
+
+        // Arrange for the quadratic to satisfy Q(X) <= 1.
+        ellipsoid.center = box.center;
+        for (int j = 0; j < 3; ++j)
+        {
+            ellipsoid.axis[j] = box.axis[j];
+            ellipsoid.extent[j] = std::sqrt(maxValue / box.extent[j]);
+        }
+        return true;
+    }
+
+    return false;
+}
+
+template <typename Real>
+bool InContainer(Vector3<Real> const& point,
+    Ellipsoid3<Real> const& ellipsoid)
+{
+    Vector3<Real> diff = point - ellipsoid.center;
+    Vector3<Real> standardized{
+        Dot(diff, ellipsoid.axis[0]) / ellipsoid.extent[0],
+        Dot(diff, ellipsoid.axis[1]) / ellipsoid.extent[1],
+        Dot(diff, ellipsoid.axis[2]) / ellipsoid.extent[2] };
+    return Length(standardized) <= (Real)1;
+}
+
+template <typename Real>
+bool MergeContainers(Ellipsoid3<Real> const& ellipsoid0,
+    Ellipsoid3<Real> const& ellipsoid1, Ellipsoid3<Real>& merge)
+{
+    // Compute the average of the input centers
+    merge.center = ((Real)0.5)*(ellipsoid0.center + ellipsoid1.center);
+
+    // The bounding ellipsoid orientation is the average of the input
+    // orientations.
+    Matrix3x3<Real> rot0, rot1;
+    rot0.SetCol(0, ellipsoid0.axis[0]);
+    rot0.SetCol(1, ellipsoid0.axis[1]);
+    rot0.SetCol(2, ellipsoid0.axis[2]);
+    rot1.SetCol(0, ellipsoid1.axis[0]);
+    rot1.SetCol(1, ellipsoid1.axis[1]);
+    rot1.SetCol(2, ellipsoid1.axis[2]);
+    Quaternion<Real> q0 = Rotation<3, Real>(rot0);
+    Quaternion<Real> q1 = Rotation<3, Real>(rot1);
+    if (Dot(q0, q1) < (Real)0)
+    {
+        q1 = -q1;
+    }
+
+    Quaternion<Real> q = q0 + q1;
+    Normalize(q);
+    Matrix3x3<Real> rot = Rotation<3, Real>(q);
+    for (int j = 0; j < 3; ++j)
+    {
+        merge.axis[j] = rot.GetCol(j);
+    }
+
+    // Project the input ellipsoids onto the axes obtained by the average
+    // of the orientations and that go through the center obtained by the
+    // average of the centers.
+    for (int i = 0; i < 3; ++i)
+    {
+        // Projection axis.
+        Line3<Real> line(merge.center, merge.axis[i]);
+
+        // Project ellipsoids onto the axis.
+        Real min0, max0, min1, max1;
+        Project(ellipsoid0, line, min0, max0);
+        Project(ellipsoid1, line, min1, max1);
+
+        // Determine the smallest interval containing the projected
+        // intervals.
+        Real maxIntr = (max0 >= max1 ? max0 : max1);
+        Real minIntr = (min0 <= min1 ? min0 : min1);
+
+        // Update the average center to be the center of the bounding box
+        // defined by the projected intervals.
+        merge.center += line.direction*(((Real)0.5)*(minIntr + maxIntr));
+
+        // Compute the extents of the box based on the new center.
+        merge.extent[i] = ((Real)0.5)*(maxIntr - minIntr);
+    }
+
+    return true;
+}
+
+
+}
