move GeometryProcessing plugin from NotForLicensees to Experimental

diff --git a/Engine/Plugins/Experimental/GeometryProcessing/Source/GeometryAlgorithms/Private/GTEngine/Mathematics/GteContEllipsoid3MinCR.h b/Engine/Plugins/Experimental/GeometryProcessing/Source/GeometryAlgorithms/Private/GTEngine/Mathematics/GteContEllipsoid3MinCR.h
new file mode 100644
--- /dev/null
+++ b/Engine/Plugins/Experimental/GeometryProcessing/Source/GeometryAlgorithms/Private/GTEngine/Mathematics/GteContEllipsoid3MinCR.h
@@ -0,0 +1,331 @@
+// David Eberly, Geometric Tools, Redmond WA 98052
+// Copyright (c) 1998-2019
+// Distributed under the Boost Software License, Version 1.0.
+// http://www.boost.org/LICENSE_1_0.txt
+// http://www.geometrictools.com/License/Boost/LICENSE_1_0.txt
+// File Version: 3.0.1 (2018/10/05)
+
+#pragma once
+
+#include <GTEngine/LowLevel/GteLogger.h>
+#include <GTEngine/Mathematics/GteMatrix3x3.h>
+#include <random>
+
+// Compute the minimum-volume ellipsoid, (X-C)^T R D R^T (X-C) = 1, given the
+// center C and orientation matrix R.  The columns of R are the axes of the
+// ellipsoid.  The algorithm computes the diagonal matrix D.  The minimum
+// volume is (4*pi/3)/sqrt(D[0]*D[1]*D[2]), where D = diag(D[0],D[1],D[2]).
+// The problem is equivalent to maximizing the product D[0]*D[1]*D[2] for a
+// given C and R, and subject to the constraints
+//   (P[i]-C)^T R D R^T (P[i]-C) <= 1
+// for all input points P[i] with 0 <= i < N.  Each constraint has the form
+//   A[0]*D[0] + A[1]*D[1] + A[2]*D[2] <= 1
+// where A[0] >= 0, A[1] >= 0, and A[2] >= 0.
+
+namespace gte
+{
+
+template <typename Real>
+class ContEllipsoid3MinCR
+{
+public:
+    void operator()(int numPoints, Vector3<Real> const* points,
+        Vector3<Real> const& C, Matrix3x3<Real> const& R, Real D[3]);
+
+private:
+    void FindEdgeMax(std::vector<Vector3<Real>>& A, int& plane0,
+        int& plane1, Real D[3]);
+
+    void FindFacetMax(std::vector<Vector3<Real>>& A, int& plane0,
+        Real D[3]);
+
+    void MaxProduct(std::vector<Vector3<Real>>& A, Real D[3]);
+};
+
+
+template <typename Real>
+void ContEllipsoid3MinCR<Real>::operator()(int numPoints,
+    Vector3<Real> const* points, Vector3<Real> const& C,
+    Matrix3x3<Real> const& R, Real D[3])
+{
+    // Compute the constraint coefficients, of the form (A[0],A[1]) for
+    // each i.
+    std::vector<Vector3<Real>> A(numPoints);
+    for (int i = 0; i < numPoints; ++i)
+    {
+        Vector3<Real> diff = points[i] - C;  // P[i] - C
+        Vector3<Real> prod = diff*R;  // R^T*(P[i] - C) = (u,v,w)
+        A[i] = prod * prod;  // (u^2, v^2, w^2)
+    }
+
+    // TODO:  Sort the constraints to eliminate redundant ones.  It is clear
+    // how to do this in ContEllipse2MinCR.  How to do this in 3D?
+
+    MaxProduct(A, D);
+}
+
+template <typename Real>
+void ContEllipsoid3MinCR<Real>::FindEdgeMax(std::vector<Vector3<Real>>& A,
+    int& plane0, int& plane1, Real D[3])
+{
+    // Compute direction to local maximum point on line of intersection.
+    Real xDir = A[plane0][1] * A[plane1][2] - A[plane1][1] * A[plane0][2];
+    Real yDir = A[plane0][2] * A[plane1][0] - A[plane1][2] * A[plane0][0];
+    Real zDir = A[plane0][0] * A[plane1][1] - A[plane1][0] * A[plane0][1];
+
+    // Build quadratic Q'(t) = (d/dt)(x(t)y(t)z(t)) = a0+a1*t+a2*t^2.
+    Real a0 = D[0] * D[1] * zDir + D[0] * D[2] * yDir + D[1] * D[2] * xDir;
+    Real a1 = ((Real)2)*(D[2] * xDir*yDir + D[1] * xDir*zDir + D[0] * yDir*zDir);
+    Real a2 = ((Real)3)*(xDir*yDir*zDir);
+
+    // Find root to Q'(t) = 0 corresponding to maximum.
+    Real tFinal;
+    if (a2 != (Real)0)
+    {
+        Real invA2 = ((Real)1) / a2;
+        Real discr = a1*a1 - ((Real)4)*a0*a2;
+        discr = std::sqrt(std::max(discr, (Real)0));
+        tFinal = -((Real)0.5)*(a1 + discr)*invA2;
+        if (a1 + ((Real)2)*a2*tFinal > (Real)0)
+        {
+            tFinal = ((Real)0.5)*(-a1 + discr)*invA2;
+        }
+    }
+    else if (a1 != (Real)0)
+    {
+        tFinal = -a0 / a1;
+    }
+    else if (a0 != (Real)0)
+    {
+        Real fmax = std::numeric_limits<Real>::max();
+        tFinal = (a0 >= (Real)0 ? fmax : -fmax);
+    }
+    else
+    {
+        return;
+    }
+
+    if (tFinal < (Real)0)
+    {
+        // Make (xDir,yDir,zDir) point in direction of increase of Q.
+        tFinal = -tFinal;
+        xDir = -xDir;
+        yDir = -yDir;
+        zDir = -zDir;
+    }
+
+    // Sort remaining planes along line from current point to local maximum.
+    Real tMax = tFinal;
+    int plane2 = -1;
+    int numPoints = static_cast<int>(A.size());
+    for (int i = 0; i < numPoints; ++i)
+    {
+        if (i == plane0 || i == plane1)
+        {
+            continue;
+        }
+
+        Real norDotDir = A[i][0] * xDir + A[i][1] * yDir + A[i][2] * zDir;
+        if (norDotDir <= (Real)0)
+        {
+            continue;
+        }
+
+        // Theoretically the numerator must be nonnegative since an invariant
+        // in the algorithm is that (x0,y0,z0) is on the convex hull of the
+        // constraints.  However, some numerical error may make this a small
+        // negative number.  In that case set tmax = 0 (no change in
+        // position).
+        Real numer = (Real)1 - A[i][0] * D[0] - A[i][1] * D[1] - A[i][2] * D[2];
+        if (numer < (Real)0)
+        {
+            LogError("Unexpected condition.");
+            plane2 = i;
+            tMax = (Real)0;
+            break;
+        }
+
+        Real t = numer / norDotDir;
+        if (0 <= t && t < tMax)
+        {
+            plane2 = i;
+            tMax = t;
+        }
+    }
+
+    D[0] += tMax*xDir;
+    D[1] += tMax*yDir;
+    D[2] += tMax*zDir;
+
+    if (tMax == tFinal)
+    {
+        return;
+    }
+
+    if (tMax > (Real)0)
+    {
+        plane0 = plane2;
+        FindFacetMax(A, plane0, D);
+        return;
+    }
+
+    // tmax == 0, so return with D[0], D[1], and D[2] unchanged.
+}
+
+template <typename Real>
+void ContEllipsoid3MinCR<Real>::FindFacetMax(std::vector<Vector3<Real>>& A,
+    int& plane0, Real D[3])
+{
+    Real tFinal, xDir, yDir, zDir;
+
+    if (A[plane0][0] > (Real)0
+        && A[plane0][1] > (Real)0
+        && A[plane0][2] > (Real)0)
+    {
+        // Compute local maximum point on plane.
+        Real oneThird = ((Real)1) / (Real)3;
+        Real xMax = oneThird / A[plane0][0];
+        Real yMax = oneThird / A[plane0][1];
+        Real zMax = oneThird / A[plane0][2];
+
+        // Compute direction to local maximum point on plane.
+        tFinal = (Real)1;
+        xDir = xMax - D[0];
+        yDir = yMax - D[1];
+        zDir = zMax - D[2];
+    }
+    else
+    {
+        tFinal = std::numeric_limits<Real>::max();
+
+        if (A[plane0][0] > (Real)0)
+        {
+            xDir = (Real)0;
+        }
+        else
+        {
+            xDir = (Real)1;
+        }
+
+        if (A[plane0][1] > (Real)0)
+        {
+            yDir = (Real)0;
+        }
+        else
+        {
+            yDir = (Real)1;
+        }
+
+        if (A[plane0][2] > (Real)0)
+        {
+            zDir = (Real)0;
+        }
+        else
+        {
+            zDir = (Real)1;
+        }
+    }
+
+    // Sort remaining planes along line from current point.
+    Real tMax = tFinal;
+    int plane1 = -1;
+    int numPoints = static_cast<int>(A.size());
+    for (int i = 0; i < numPoints; ++i)
+    {
+        if (i == plane0)
+        {
+            continue;
+        }
+
+        Real norDotDir = A[i][0] * xDir + A[i][1] * yDir + A[i][2] * zDir;
+        if (norDotDir <= (Real)0)
+        {
+            continue;
+        }
+
+        // Theoretically the numerator must be nonnegative since an invariant
+        // in the algorithm is that (x0,y0,z0) is on the convex hull of the
+        // constraints.  However, some numerical error may make this a small
+        // negative number.  In that case, set tmax = 0 (no change in
+        // position).
+        Real numer = (Real)1 - A[i][0] * D[0] - A[i][1] * D[1] - A[i][2] * D[2];
+        if (numer < (Real)0)
+        {
+            LogError("Unexpected condition.");
+            plane1 = i;
+            tMax = (Real)0;
+            break;
+        }
+
+        Real t = numer / norDotDir;
+        if (0 <= t && t < tMax)
+        {
+            plane1 = i;
+            tMax = t;
+        }
+    }
+
+    D[0] += tMax*xDir;
+    D[1] += tMax*yDir;
+    D[2] += tMax*zDir;
+
+    if (tMax == (Real)1)
+    {
+        return;
+    }
+
+    if (tMax > (Real)0)
+    {
+        plane0 = plane1;
+        FindFacetMax(A, plane0, D);
+        return;
+    }
+
+    FindEdgeMax(A, plane0, plane1, D);
+}
+
+template <typename Real>
+void ContEllipsoid3MinCR<Real>::MaxProduct(std::vector<Vector3<Real>>& A,
+    Real D[3])
+{
+    // Maximize x*y*z subject to x >= 0, y >= 0, z >= 0, and
+    // A[i]*x+B[i]*y+C[i]*z <= 1 for 0 <= i < N where A[i] >= 0,
+    // B[i] >= 0, and C[i] >= 0.
+
+    // Jitter the lines to avoid cases where more than three planes
+    // intersect at the same point.  Should also break parallelism
+    // and planes parallel to the coordinate planes.
+    std::mt19937 mte;
+    std::uniform_real_distribution<Real> rnd((Real)0, (Real)1);
+    Real maxJitter = (Real)1e-12;
+    int numPoints = static_cast<int>(A.size());
+    int i;
+    for (i = 0; i < numPoints; ++i)
+    {
+        A[i][0] += maxJitter*rnd(mte);
+        A[i][1] += maxJitter*rnd(mte);
+        A[i][2] += maxJitter*rnd(mte);
+    }
+
+    // Sort lines along the z-axis (x = 0 and y = 0).
+    int plane = -1;
+    Real zmax = (Real)0;
+    for (i = 0; i < numPoints; ++i)
+    {
+        if (A[i][2] > zmax)
+        {
+            zmax = A[i][2];
+            plane = i;
+        }
+    }
+    LogAssert(plane != -1, "Unexpected condition.");
+
+    // Walk along convex hull searching for maximum.
+    D[0] = (Real)0;
+    D[1] = (Real)0;
+    D[2] = ((Real)1) / zmax;
+    FindFacetMax(A, plane, D);
+}
+
+
+}
diff --git a/Engine/Plugins/Experimental/GeometryProcessing/Source/GeometryAlgorithms/Private/GTEngine/Mathematics/GteContOrientedBox2.h b/Engine/Plugins/Experimental/GeometryProcessing/Source/GeometryAlgorithms/Private/GTEngine/Mathematics/GteContOrientedBox2.h
new file mode 100644
--- /dev/null
+++ b/Engine/Plugins/Experimental/GeometryProcessing/Source/GeometryAlgorithms/Private/GTEngine/Mathematics/GteContOrientedBox2.h
@@ -0,0 +1,186 @@
+// David Eberly, Geometric Tools, Redmond WA 98052
+// Copyright (c) 1998-2019
+// Distributed under the Boost Software License, Version 1.0.
+// http://www.boost.org/LICENSE_1_0.txt
+// http://www.geometrictools.com/License/Boost/LICENSE_1_0.txt
+// File Version: 3.0.0 (2016/06/19)
+
+#pragma once
+
+#include <GTEngine/Mathematics/GteApprGaussian2.h>
+
+namespace gte
+{
+
+// Compute an oriented bounding box of the points.  The box center is the
+// average of the points.  The box axes are the eigenvectors of the covariance
+// matrix.
+template <typename Real>
+bool GetContainer(int numPoints, Vector2<Real> const* points,
+     OrientedBox2<Real>& box);
+
+// Test for containment.  Let X = C + y0*U0 + y1*U1 where C is the box center
+// and U0 and U1 are the orthonormal axes of the box.  X is in the box when
+// |y_i| <= E_i for all i, where E_i are the extents of the box.
+template <typename Real>
+bool InContainer(Vector2<Real> const& point, OrientedBox2<Real> const& box);
+
+// Construct an oriented box that contains two other oriented boxes.  The
+// result is not guaranteed to be the minimum area box containing the input
+// boxes.
+template <typename Real>
+bool MergeContainers(OrientedBox2<Real> const& box0,
+    OrientedBox2<Real> const& box1, OrientedBox2<Real>& merge);
+
+
+template <typename Real>
+bool GetContainer(int numPoints, Vector2<Real> const* points,
+    OrientedBox2<Real>& box)
+{
+    // Fit the points with a Gaussian distribution.
+    ApprGaussian2<Real> fitter;
+    if (fitter.Fit(numPoints, points))
+    {
+        box = fitter.GetParameters();
+
+        // Let C be the box center and let U0 and U1 be the box axes.  Each
+        // input point is of the form X = C + y0*U0 + y1*U1.  The following
+        // code computes min(y0), max(y0), min(y1), and max(y1).  The box
+        // center is then adjusted to be
+        //   C' = C + 0.5*(min(y0)+max(y0))*U0 + 0.5*(min(y1)+max(y1))*U1
+
+        Vector2<Real> diff = points[0] - box.center;
+        Vector2<Real> pmin{ Dot(diff, box.axis[0]), Dot(diff, box.axis[1]) };
+        Vector2<Real> pmax = pmin;
+        for (int i = 1; i < numPoints; ++i)
+        {
+            diff = points[i] - box.center;
+            for (int j = 0; j < 2; ++j)
+            {
+                Real dot = Dot(diff, box.axis[j]);
+                if (dot < pmin[j])
+                {
+                    pmin[j] = dot;
+                }
+                else if (dot > pmax[j])
+                {
+                    pmax[j] = dot;
+                }
+            }
+        }
+
+        for (int j = 0; j < 2; ++j)
+        {
+            box.center += (((Real)0.5)*(pmin[j] + pmax[j]))*box.axis[j];
+            box.extent[j] = ((Real)0.5)*(pmax[j] - pmin[j]);
+        }
+        return true;
+    }
+
+    return false;
+}
+
+template <typename Real>
+bool InContainer(Vector2<Real> const& point, OrientedBox2<Real> const& box)
+{
+    Vector2<Real> diff = point - box.center;
+    for (int i = 0; i < 2; ++i)
+    {
+        Real coeff = Dot(diff, box.axis[i]);
+        if (std::abs(coeff) > box.extent[i])
+        {
+            return false;
+        }
+    }
+    return true;
+}
+
+template <typename Real>
+bool MergeContainers(OrientedBox2<Real> const& box0,
+    OrientedBox2<Real> const& box1, OrientedBox2<Real>& merge)
+{
+    // The first guess at the box center.  This value will be updated later
+    // after the input box vertices are projected onto axes determined by an
+    // average of box axes.
+    merge.center = ((Real)0.5)*(box0.center + box1.center);
+
+    // The merged box axes are the averages of the input box axes.  The axes
+    // of the second box are negated, if necessary, so they form acute angles
+    // with the axes of the first box.
+    if (Dot(box0.axis[0], box1.axis[0]) >= (Real)0)
+    {
+        merge.axis[0] = ((Real)0.5)*(box0.axis[0] + box1.axis[0]);
+    }
+    else
+    {
+        merge.axis[0] = ((Real)0.5)*(box0.axis[0] - box1.axis[0]);
+    }
+    Normalize(merge.axis[0]);
+    merge.axis[1] = -Perp(merge.axis[0]);
+
+    // Project the input box vertices onto the merged-box axes.  Each axis
+    // D[i] containing the current center C has a minimum projected value
+    // min[i] and a maximum projected value max[i].  The corresponding
+    // endpoints on the axes are C+min[i]*D[i] and C+max[i]*D[i].  The point
+    // C is not necessarily the midpoint for any of the intervals.  The actual
+    // box center will be adjusted from C to a point C' that is the midpoint
+    // of each interval,
+    //   C' = C + sum_{i=0}^1 0.5*(min[i]+max[i])*D[i]
+    // The box extents are
+    //   e[i] = 0.5*(max[i]-min[i])
+
+    std::array<Vector2<Real>, 4> vertex;
+    Vector2<Real> pmin{ (Real)0, (Real)0 };
+    Vector2<Real> pmax{ (Real)0, (Real)0 };
+
+    box0.GetVertices(vertex);
+    for (int i = 0; i < 4; ++i)
+    {
+        Vector2<Real> diff = vertex[i] - merge.center;
+        for (int j = 0; j < 2; ++j)
+        {
+            Real dot = Dot(diff, merge.axis[j]);
+            if (dot > pmax[j])
+            {
+                pmax[j] = dot;
+            }
+            else if (dot < pmin[j])
+            {
+                pmin[j] = dot;
+            }
+        }
+    }
+
+    box1.GetVertices(vertex);
+    for (int i = 0; i < 4; ++i)
+    {
+        Vector2<Real> diff = vertex[i] - merge.center;
+        for (int j = 0; j < 2; ++j)
+        {
+            Real dot = Dot(diff, merge.axis[j]);
+            if (dot > pmax[j])
+            {
+                pmax[j] = dot;
+            }
+            else if (dot < pmin[j])
+            {
+                pmin[j] = dot;
+            }
+        }
+    }
+
+    // [min,max] is the axis-aligned box in the coordinate system of the
+    // merged box axes.  Update the current box center to be the center of
+    // the new box.  Compute the extents based on the new center.
+    Real const half = (Real)0.5;
+    for (int j = 0; j < 2; ++j)
+    {
+        merge.center += half*(pmax[j] + pmin[j])*merge.axis[j];
+        merge.extent[j] = half*(pmax[j] - pmin[j]);
+    }
+
+    return true;
+}
+
+
+}
diff --git a/Engine/Plugins/Experimental/GeometryProcessing/Source/GeometryAlgorithms/Private/GTEngine/Mathematics/GteContOrientedBox3.h b/Engine/Plugins/Experimental/GeometryProcessing/Source/GeometryAlgorithms/Private/GTEngine/Mathematics/GteContOrientedBox3.h
new file mode 100644
--- /dev/null
+++ b/Engine/Plugins/Experimental/GeometryProcessing/Source/GeometryAlgorithms/Private/GTEngine/Mathematics/GteContOrientedBox3.h
@@ -0,0 +1,204 @@
+// David Eberly, Geometric Tools, Redmond WA 98052
+// Copyright (c) 1998-2019
+// Distributed under the Boost Software License, Version 1.0.
+// http://www.boost.org/LICENSE_1_0.txt
+// http://www.geometrictools.com/License/Boost/LICENSE_1_0.txt
+// File Version: 3.0.0 (2016/06/19)
+
+#pragma once
+
+#include <GTEngine/Mathematics/GteApprGaussian3.h>
+#include <GTEngine/Mathematics/GteMatrix3x3.h>
+#include <GTEngine/Mathematics/GteRotation.h>
+
+namespace gte
+{
+
+// Compute an oriented bounding box of the points.  The box center is the
+// average of the points.  The box axes are the eigenvectors of the covariance
+// matrix.
+template <typename Real>
+bool GetContainer(int numPoints, Vector3<Real> const* points,
+    OrientedBox3<Real>& box);
+
+// Test for containment.  Let X = C + y0*U0 + y1*U1 + y2*U2 where C is the
+// box center and U0, U1, U2 are the orthonormal axes of the box.  X is in
+// the box if |y_i| <= E_i for all i where E_i are the extents of the box.
+template <typename Real>
+bool InContainer(Vector3<Real> const& point, OrientedBox3<Real> const& box);
+
+// Construct an oriented box that contains two other oriented boxes.  The
+// result is not guaranteed to be the minimum volume box containing the
+// input boxes.
+template <typename Real>
+bool MergeContainers(OrientedBox3<Real> const& box0,
+    OrientedBox3<Real> const& box1, OrientedBox3<Real>& merge);
+
+
+template <typename Real>
+bool GetContainer(int numPoints, Vector3<Real> const* points,
+    OrientedBox3<Real>& box)
+{
+    // Fit the points with a Gaussian distribution.
+    ApprGaussian3<Real> fitter;
+    if (fitter.Fit(numPoints, points))
+    {
+        box = fitter.GetParameters();
+
+        // Let C be the box center and let U0, U1, and U2 be the box axes.
+        // Each input point is of the form X = C + y0*U0 + y1*U1 + y2*U2.
+        // The following code computes min(y0), max(y0), min(y1), max(y1),
+        // min(y2), and max(y2).  The box center is then adjusted to be
+        //   C' = C + 0.5*(min(y0)+max(y0))*U0 + 0.5*(min(y1)+max(y1))*U1 +
+        //        0.5*(min(y2)+max(y2))*U2
+
+        Vector3<Real> diff = points[0] - box.center;
+        Vector3<Real> pmin{ Dot(diff, box.axis[0]), Dot(diff, box.axis[1]),
+            Dot(diff, box.axis[2]) };
+        Vector3<Real> pmax = pmin;
+        for (int i = 1; i < numPoints; ++i)
+        {
+            diff = points[i] - box.center;
+            for (int j = 0; j < 3; ++j)
+            {
+                Real dot = Dot(diff, box.axis[j]);
+                if (dot < pmin[j])
+                {
+                    pmin[j] = dot;
+                }
+                else if (dot > pmax[j])
+                {
+                    pmax[j] = dot;
+                }
+            }
+        }
+
+        for (int j = 0; j < 3; ++j)
+        {
+            box.center += (((Real)0.5)*(pmin[j] + pmax[j]))*box.axis[j];
+            box.extent[j] = ((Real)0.5)*(pmax[j] - pmin[j]);
+        }
+        return true;
+    }
+
+    return false;
+}
+
+template <typename Real>
+bool InContainer(Vector3<Real> const& point, OrientedBox3<Real> const& box)
+{
+    Vector3<Real> diff = point - box.center;
+    for (int i = 0; i < 3; ++i)
+    {
+        Real coeff = Dot(diff, box.axis[i]);
+        if (std::abs(coeff) > box.extent[i])
+        {
+            return false;
+        }
+    }
+    return true;
+}
+
+template <typename Real>
+bool MergeContainers(OrientedBox3<Real> const& box0,
+    OrientedBox3<Real> const& box1, OrientedBox3<Real>& merge)
+{
+    // The first guess at the box center.  This value will be updated later
+    // after the input box vertices are projected onto axes determined by an
+    // average of box axes.
+    merge.center = ((Real)0.5)*(box0.center + box1.center);
+
+    // A box's axes, when viewed as the columns of a matrix, form a rotation
+    // matrix.  The input box axes are converted to quaternions.  The average
+    // quaternion is computed, then normalized to unit length.  The result is
+    // the slerp of the two input quaternions with t-value of 1/2.  The result
+    // is converted back to a rotation matrix and its columns are selected as
+    // the merged box axes.
+    Matrix3x3<Real> rot0, rot1;
+    rot0.SetCol(0, box0.axis[0]);
+    rot0.SetCol(1, box0.axis[1]);
+    rot0.SetCol(2, box0.axis[2]);
+    rot1.SetCol(0, box1.axis[0]);
+    rot1.SetCol(1, box1.axis[1]);
+    rot1.SetCol(2, box1.axis[2]);
+    Quaternion<Real> q0 = Rotation<3, Real>(rot0);
+    Quaternion<Real> q1 = Rotation<3, Real>(rot1);
+    if (Dot(q0, q1) < (Real)0)
+    {
+        q1 = -q1;
+    }
+
+    Quaternion<Real> q = q0 + q1;
+    Normalize(q);
+    Matrix3x3<Real> rot = Rotation<3, Real>(q);
+    for (int j = 0; j < 3; ++j)
+    {
+        merge.axis[j] = rot.GetCol(j);
+    }
+
+    // Project the input box vertices onto the merged-box axes.  Each axis
+    // D[i] containing the current center C has a minimum projected value
+    // min[i] and a maximum projected value max[i].  The corresponding end
+    // points on the axes are C+min[i]*D[i] and C+max[i]*D[i].  The point C
+    // is not necessarily the midpoint for any of the intervals.  The actual
+    // box center will be adjusted from C to a point C' that is the midpoint
+    // of each interval,
+    //   C' = C + sum_{i=0}^2 0.5*(min[i]+max[i])*D[i]
+    // The box extents are
+    //   e[i] = 0.5*(max[i]-min[i])
+
+    std::array<Vector3<Real>, 8> vertex;
+    Vector3<Real> pmin{ (Real)0, (Real)0, (Real)0 };
+    Vector3<Real> pmax{ (Real)0, (Real)0, (Real)0 };
+
+    box0.GetVertices(vertex);
+    for (int i = 0; i < 8; ++i)
+    {
+        Vector3<Real> diff = vertex[i] - merge.center;
+        for (int j = 0; j < 3; ++j)
+        {
+            Real dot = Dot(diff, merge.axis[j]);
+            if (dot > pmax[j])
+            {
+                pmax[j] = dot;
+            }
+            else if (dot < pmin[j])
+            {
+                pmin[j] = dot;
+            }
+        }
+    }
+
+    box1.GetVertices(vertex);
+    for (int i = 0; i < 8; ++i)
+    {
+        Vector3<Real> diff = vertex[i] - merge.center;
+        for (int j = 0; j < 3; ++j)
+        {
+            Real dot = Dot(diff, merge.axis[j]);
+            if (dot > pmax[j])
+            {
+                pmax[j] = dot;
+            }
+            else if (dot < pmin[j])
+            {
+                pmin[j] = dot;
+            }
+        }
+    }
+
+    // [min,max] is the axis-aligned box in the coordinate system of the
+    // merged box axes.  Update the current box center to be the center of
+    // the new box.  Compute the extents based on the new center.
+    Real const half = (Real)0.5;
+    for (int j = 0; j < 3; ++j)
+    {
+        merge.center += half*(pmax[j] + pmin[j])*merge.axis[j];
+        merge.extent[j] = half*(pmax[j] - pmin[j]);
+    }
+
+    return true;
+}
+
+
+}
diff --git a/Engine/Plugins/Experimental/GeometryProcessing/Source/GeometryAlgorithms/Private/GTEngine/Mathematics/GteContPointInPolygon2.h b/Engine/Plugins/Experimental/GeometryProcessing/Source/GeometryAlgorithms/Private/GTEngine/Mathematics/GteContPointInPolygon2.h
new file mode 100644
--- /dev/null
+++ b/Engine/Plugins/Experimental/GeometryProcessing/Source/GeometryAlgorithms/Private/GTEngine/Mathematics/GteContPointInPolygon2.h
@@ -0,0 +1,227 @@
+// David Eberly, Geometric Tools, Redmond WA 98052
+// Copyright (c) 1998-2019
+// Distributed under the Boost Software License, Version 1.0.
+// http://www.boost.org/LICENSE_1_0.txt
+// http://www.geometrictools.com/License/Boost/LICENSE_1_0.txt
+// File Version: 3.0.0 (2016/06/19)
+
+#pragma once
+
+#include <GTEngine/Mathematics/GteVector2.h>
+
+// Given a polygon as an order list of vertices (x[i],y[i]) for 0 <= i < N
+// and a test point (xt,yt), return 'true' if (xt,yt) is in the polygon and
+// 'false' if it is not.  All queries require that the number of vertices
+// satisfies N >= 3.
+
+namespace gte
+{
+
+template <typename Real>
+class PointInPolygon2
+{
+public:
+    // The class object stores a copy of 'points', so be careful about the
+    // persistence of 'points' when you have created a PointInPolygon2 object.
+    PointInPolygon2(int numPoints, Vector2<Real> const* points);
+
+    // Simple polygons (ray-intersection counting).
+    bool Contains(Vector2<Real> const& p) const;
+
+    // Algorithms for convex polygons.  The input polygons must have vertices
+    // in counterclockwise order.
+
+    // O(N) algorithm (which-side-of-edge tests)
+    bool ContainsConvexOrderN(Vector2<Real> const& p) const;
+
+    // O(log N) algorithm (bisection and recursion, like BSP tree)
+    bool ContainsConvexOrderLogN(Vector2<Real> const& p) const;
+
+    // The polygon must have exactly four vertices.  This method is like the
+    // O(log N) and uses three which-side-of-segment test instead of four
+    // which-side-of-edge tests.  If the polygon does not have four vertices,
+    // the function returns false.
+    bool ContainsQuadrilateral(Vector2<Real> const& p) const;
+
+private:
+    // For recursion in ContainsConvexOrderLogN.
+    bool SubContainsPoint(Vector2<Real> const& p, int i0, int i1) const;
+
+    int mNumPoints;
+    Vector2<Real> const* mPoints;
+};
+
+
+template <typename Real>
+PointInPolygon2<Real>::PointInPolygon2(int numPoints,
+    Vector2<Real> const* points)
+    :
+    mNumPoints(numPoints),
+    mPoints(points)
+{
+}
+
+template <typename Real>
+bool PointInPolygon2<Real>::Contains(Vector2<Real> const& p) const
+{
+    bool inside = false;
+    for (int i = 0, j = mNumPoints - 1; i < mNumPoints; j = i++)
+    {
+        Vector2<Real> const& U0 = mPoints[i];
+        Vector2<Real> const& U1 = mPoints[j];
+        Real rhs, lhs;
+
+        if (p[1] < U1[1])  // U1 above ray
+        {
+            if (U0[1] <= p[1])  // U0 on or below ray
+            {
+                lhs = (p[1] - U0[1])*(U1[0] - U0[0]);
+                rhs = (p[0] - U0[0])*(U1[1] - U0[1]);
+                if (lhs > rhs)
+                {
+                    inside = !inside;
+                }
+            }
+        }
+        else if (p[1] < U0[1])  // U1 on or below ray, U0 above ray
+        {
+            lhs = (p[1] - U0[1])*(U1[0] - U0[0]);
+            rhs = (p[0] - U0[0])*(U1[1] - U0[1]);
+            if (lhs < rhs)
+            {
+                inside = !inside;
+            }
+        }
+    }
+    return inside;
+}
+
+template <typename Real>
+bool PointInPolygon2<Real>::ContainsConvexOrderN(Vector2<Real> const& p) const
+{
+    for (int i1 = 0, i0 = mNumPoints - 1; i1 < mNumPoints; i0 = i1++)
+    {
+        Real nx = mPoints[i1][1] - mPoints[i0][1];
+        Real ny = mPoints[i0][0] - mPoints[i1][0];
+        Real dx = p[0] - mPoints[i0][0];
+        Real dy = p[1] - mPoints[i0][1];
+        if (nx*dx + ny*dy >(Real)0)
+        {
+            return false;
+        }
+    }
+    return true;
+}
+
+template <typename Real>
+bool PointInPolygon2<Real>::ContainsConvexOrderLogN(Vector2<Real> const& p)
+const
+{
+    return SubContainsPoint(p, 0, 0);
+}
+
+template <typename Real>
+bool PointInPolygon2<Real>::ContainsQuadrilateral(Vector2<Real> const& p)
+const
+{
+    if (mNumPoints != 4)
+    {
+        return false;
+    }
+
+    Real nx = mPoints[2][1] - mPoints[0][1];
+    Real ny = mPoints[0][0] - mPoints[2][0];
+    Real dx = p[0] - mPoints[0][0];
+    Real dy = p[1] - mPoints[0][1];
+
+    if (nx*dx + ny*dy > (Real)0)
+    {
+        // P potentially in <V0,V1,V2>
+        nx = mPoints[1][1] - mPoints[0][1];
+        ny = mPoints[0][0] - mPoints[1][0];
+        if (nx*dx + ny*dy > (Real)0.0)
+        {
+            return false;
+        }
+
+        nx = mPoints[2][1] - mPoints[1][1];
+        ny = mPoints[1][0] - mPoints[2][0];
+        dx = p[0] - mPoints[1][0];
+        dy = p[1] - mPoints[1][1];
+        if (nx*dx + ny*dy > (Real)0)
+        {
+            return false;
+        }
+    }
+    else
+    {
+        // P potentially in <V0,V2,V3>
+        nx = mPoints[0][1] - mPoints[3][1];
+        ny = mPoints[3][0] - mPoints[0][0];
+        if (nx*dx + ny*dy > (Real)0)
+        {
+            return false;
+        }
+
+        nx = mPoints[3][1] - mPoints[2][1];
+        ny = mPoints[2][0] - mPoints[3][0];
+        dx = p[0] - mPoints[3][0];
+        dy = p[1] - mPoints[3][1];
+        if (nx*dx + ny*dy > (Real)0)
+        {
+            return false;
+        }
+    }
+    return true;
+}
+
+template <typename Real>
+bool PointInPolygon2<Real>::SubContainsPoint(Vector2<Real> const& p, int i0,
+    int i1) const
+{
+    Real nx, ny, dx, dy;
+
+    int diff = i1 - i0;
+    if (diff == 1 || (diff < 0 && diff + mNumPoints == 1))
+    {
+        nx = mPoints[i1][1] - mPoints[i0][1];
+        ny = mPoints[i0][0] - mPoints[i1][0];
+        dx = p[0] - mPoints[i0][0];
+        dy = p[1] - mPoints[i0][1];
+        return nx*dx + ny*dy <= (Real)0;
+    }
+
+    // Bisect the index range.
+    int mid;
+    if (i0 < i1)
+    {
+        mid = (i0 + i1) >> 1;
+    }
+    else
+    {
+        mid = ((i0 + i1 + mNumPoints) >> 1);
+        if (mid >= mNumPoints)
+        {
+            mid -= mNumPoints;
+        }
+    }
+
+    // Determine which side of the splitting line contains the point.
+    nx = mPoints[mid][1] - mPoints[i0][1];
+    ny = mPoints[i0][0] - mPoints[mid][0];
+    dx = p[0] - mPoints[i0][0];
+    dy = p[1] - mPoints[i0][1];
+    if (nx*dx + ny*dy > (Real)0)
+    {
+        // P potentially in <V(i0),V(i0+1),...,V(mid-1),V(mid)>
+        return SubContainsPoint(p, i0, mid);
+    }
+    else
+    {
+        // P potentially in <V(mid),V(mid+1),...,V(i1-1),V(i1)>
+        return SubContainsPoint(p, mid, i1);
+    }
+}
+
+
+}
diff --git a/Engine/Plugins/Experimental/GeometryProcessing/Source/GeometryAlgorithms/Private/GTEngine/Mathematics/GteContPointInPolyhedron3.h b/Engine/Plugins/Experimental/GeometryProcessing/Source/GeometryAlgorithms/Private/GTEngine/Mathematics/GteContPointInPolyhedron3.h
new file mode 100644
--- /dev/null
+++ b/Engine/Plugins/Experimental/GeometryProcessing/Source/GeometryAlgorithms/Private/GTEngine/Mathematics/GteContPointInPolyhedron3.h
@@ -0,0 +1,613 @@
+// David Eberly, Geometric Tools, Redmond WA 98052
+// Copyright (c) 1998-2019
+// Distributed under the Boost Software License, Version 1.0.
+// http://www.boost.org/LICENSE_1_0.txt
+// http://www.geometrictools.com/License/Boost/LICENSE_1_0.txt
+// File Version: 3.0.1 (2016/11/25)
+
+#pragma once
+
+#include <GTEngine/LowLevel/GteLogger.h>
+#include <GTEngine/Mathematics/GteContPointInPolygon2.h>
+#include <GTEngine/Mathematics/GteIntrRay3Plane3.h>
+#include <GTEngine/Mathematics/GteIntrRay3Triangle3.h>
+#include <vector>
+
+// This class contains various implementations for point-in-polyhedron
+// queries.  The planes stored with the faces are used in all cases to
+// reject ray-face intersection tests, a quick culling operation.
+//
+// The algorithm is to cast a ray from the input point P and test for
+// intersection against each face of the polyhedron.  If the ray only
+// intersects faces at interior points (not vertices, not edge points),
+// then the point is inside when the number of intersections is odd and
+// the point is outside when the number of intersections is even.  If the
+// ray intersects an edge or a vertex, then the counting must be handled
+// differently.  The details are tedious.  As an alternative, the approach
+// here is to allow you to specify 2*N+1 rays, where N >= 0.  You should
+// choose these rays randomly.  Each ray reports "inside" or "outside".
+// Whichever result occurs N+1 or more times is the "winner".  The input
+// rayQuantity is 2*N+1.  The input array Direction must have rayQuantity
+// elements.  If you are feeling lucky, choose rayQuantity to be 1.
+
+namespace gte
+{
+
+template <typename Real>
+class PointInPolyhedron3
+{
+public:
+    // For simple polyhedra with triangle faces.
+    class TriangleFace
+    {
+    public:
+        // When you view the face from outside, the vertices are
+        // counterclockwise ordered.  The Indices array stores the indices
+        // into the vertex array.
+        int indices[3];
+
+        // The normal vector is unit length and points to the outside of the
+        // polyhedron.
+        Plane3<Real> plane;
+    };
+
+    // The Contains query will use ray-triangle intersection queries.
+    PointInPolyhedron3(int numPoints, Vector3<Real> const* points,
+        int numFaces, TriangleFace const* faces, int numRays,
+        Vector3<Real> const* directions);
+
+    // For simple polyhedra with convex polygon faces.
+    class ConvexFace
+    {
+    public:
+        // When you view the face from outside, the vertices are
+        // counterclockwise ordered.  The Indices array stores the indices
+        // into the vertex array.
+        std::vector<int> indices;
+
+        // The normal vector is unit length and points to the outside of the
+        // polyhedron.
+        Plane3<Real> plane;
+    };
+
+    // The Contains query will use ray-convexpolygon intersection queries.  A
+    // ray-convexpolygon intersection query can be implemented in many ways.
+    // In this context, uiMethod is one of three value:
+    //   0 : Use a triangle fan and perform a ray-triangle intersection query
+    //       for each triangle.
+    //   1 : Find the point of intersection of ray and plane of polygon.  Test
+    //       whether that point is inside the convex polygon using an O(N)
+    //       test.
+    //   2 : Find the point of intersection of ray and plane of polygon.  Test
+    //       whether that point is inside the convex polygon using an O(log N)
+    //       test.
+    PointInPolyhedron3(int numPoints, Vector3<Real> const* points,
+        int numFaces, ConvexFace const* faces, int numRays,
+        Vector3<Real> const* directions, unsigned int method);
+
+    // For simple polyhedra with simple polygon faces that are generally
+    // not all convex.
+    class SimpleFace
+    {
+    public:
+        // When you view the face from outside, the vertices are
+        // counterclockwise ordered.  The Indices array stores the indices
+        // into the vertex array.
+        std::vector<int> indices;
+
+        // The normal vector is unit length and points to the outside of the
+        // polyhedron.
+        Plane3<Real> plane;
+
+        // Each simple face may be triangulated.  The indices are relative to
+        // the vertex array.  Each triple of indices represents a triangle in
+        // the triangulation.
+        std::vector<int> triangles;
+    };
+
+    // The Contains query will use ray-simplepolygon intersection queries.  A
+    // ray-simplepolygon intersection query can be implemented in a couple of
+    // ways.  In this context, uiMethod is one of two value:
+    //   0 : Iterate over the triangles of each face and perform a
+    //       ray-triangle intersection query for each triangle.  This requires
+    //       that the SimpleFace::Triangles array be initialized for each
+    //       face.
+    //   1 : Find the point of intersection of ray and plane of polygon.  Test
+    //       whether that point is inside the polygon using an O(N) test.  The
+    //       SimpleFace::Triangles array is not used for this method, so it
+    //       does not have to be initialized for each face.
+    PointInPolyhedron3(int numPoints, Vector3<Real> const* points,
+        int numFaces, SimpleFace const* faces, int numRays,
+        Vector3<Real> const* directions, unsigned int method);
+
+    // This function will select the actual algorithm based on which
+    // constructor you used for this class.
+    bool Contains(Vector3<Real> const& p) const;
+
+private:
+    // For all types of faces.  The ray origin is the test point.  The ray
+    // direction is one of those passed to the constructors.  The plane origin
+    // is a point on the plane of the face.  The plane normal is a unit-length
+    // normal to the face and that points outside the polyhedron.
+    static bool FastNoIntersect(Ray3<Real> const& ray,
+        Plane3<Real> const& plane);
+
+    // For triangle faces.
+    bool ContainsT0(Vector3<Real> const& p) const;
+
+    // For convex faces.
+    bool ContainsC0(Vector3<Real> const& p) const;
+    bool ContainsC1C2(Vector3<Real> const& p, unsigned int method) const;
+
+    // For simple faces.
+    bool ContainsS0(Vector3<Real> const& p) const;
+    bool ContainsS1(Vector3<Real> const& p) const;
+
+    int mNumPoints;
+    Vector3<Real> const* mPoints;
+
+    int mNumFaces;
+    TriangleFace const* mTFaces;
+    ConvexFace const* mCFaces;
+    SimpleFace const* mSFaces;
+
+    unsigned int mMethod;
+    int mNumRays;
+    Vector3<Real> const* mDirections;
+
+    // Temporary storage for those methods that reduce the problem to 2D
+    // point-in-polygon queries.  The array stores the projections of
+    // face vertices onto the plane of the face.  It is resized as needed.
+    mutable std::vector<Vector2<Real>> mProjVertices;
+};
+
+
+template <typename Real>
+PointInPolyhedron3<Real>::PointInPolyhedron3(int numPoints,
+    Vector3<Real> const* points, int numFaces, TriangleFace const* faces,
+    int numRays, Vector3<Real> const* directions)
+    :
+    mNumPoints(numPoints),
+    mPoints(points),
+    mNumFaces(numFaces),
+    mTFaces(faces),
+    mCFaces(nullptr),
+    mSFaces(nullptr),
+    mMethod(0),
+    mNumRays(numRays),
+    mDirections(directions)
+{
+}
+
+template <typename Real>
+PointInPolyhedron3<Real>::PointInPolyhedron3(int numPoints,
+    Vector3<Real> const* points, int numFaces, ConvexFace const* faces,
+    int numRays, Vector3<Real> const* directions, unsigned int method)
+    :
+    mNumPoints(numPoints),
+    mPoints(points),
+    mNumFaces(numFaces),
+    mTFaces(nullptr),
+    mCFaces(faces),
+    mSFaces(nullptr),
+    mMethod(method),
+    mNumRays(numRays),
+    mDirections(directions)
+{
+}
+
+template <typename Real>
+PointInPolyhedron3<Real>::PointInPolyhedron3(int numPoints,
+    Vector3<Real> const* points, int numFaces, SimpleFace const* faces,
+    int numRays, Vector3<Real> const* directions, unsigned int method)
+    :
+    mNumPoints(numPoints),
+    mPoints(points),
+    mNumFaces(numFaces),
+    mTFaces(nullptr),
+    mCFaces(nullptr),
+    mSFaces(faces),
+    mMethod(method),
+    mNumRays(numRays),
+    mDirections(directions)
+{
+}
+
+template <typename Real>
+bool PointInPolyhedron3<Real>::Contains(Vector3<Real> const& p) const
+{
+    if (mTFaces)
+    {
+        return ContainsT0(p);
+    }
+
+    if (mCFaces)
+    {
+        if (mMethod == 0)
+        {
+            return ContainsC0(p);
+        }
+
+        return ContainsC1C2(p, mMethod);
+    }
+
+    if (mSFaces)
+    {
+        if (mMethod == 0)
+        {
+            return ContainsS0(p);
+        }
+
+        if (mMethod == 1)
+        {
+            return ContainsS1(p);
+        }
+    }
+
+    return false;
+}
+
+template <typename Real>
+bool PointInPolyhedron3<Real>::FastNoIntersect(Ray3<Real> const& ray,
+    Plane3<Real> const& plane)
+{
+    Real planeDistance = Dot(plane.normal, ray.origin) - plane.constant;
+    Real planeAngle = Dot(plane.normal, ray.direction);
+
+    if (planeDistance < (Real)0)
+    {
+        // The ray origin is on the negative side of the plane.
+        if (planeAngle <= (Real)0)
+        {
+            // The ray points away from the plane.
+            return true;
+        }
+    }
+
+    if (planeDistance >(Real)0)
+    {
+        // The ray origin is on the positive side of the plane.
+        if (planeAngle >= (Real)0)
+        {
+            // The ray points away from the plane.
+            return true;
+        }
+    }
+
+    return false;
+}
+
+template <typename Real>
+bool PointInPolyhedron3<Real>::ContainsT0(Vector3<Real> const& p) const
+{
+    int insideCount = 0;
+
+    TIQuery<Real, Ray3<Real>, Triangle3<Real>> rtQuery;
+    Triangle3<Real> triangle;
+    Ray3<Real> ray;
+    ray.origin = p;
+
+    for (int j = 0; j < mNumRays; ++j)
+    {
+        ray.direction = mDirections[j];
+
+        // Zero intersections to start with.
+        bool odd = false;
+
+        TriangleFace const* face = mTFaces;
+        for (int i = 0; i < mNumFaces; ++i, ++face)
+        {
+            // Attempt to quickly cull the triangle.
+            if (FastNoIntersect(ray, face->plane))
+            {
+                continue;
+            }
+
+            // Get the triangle vertices.
+            for (int k = 0; k < 3; ++k)
+            {
+                triangle.v[k] = mPoints[face->indices[k]];
+            }
+
+            // Test for intersection.
+            if (rtQuery(ray, triangle).intersect)
+            {
+                // The ray intersects the triangle.
+                odd = !odd;
+            }
+        }
+
+        if (odd)
+        {
+            insideCount++;
+        }
+    }
+
+    return insideCount > mNumRays / 2;
+}
+
+template <typename Real>
+bool PointInPolyhedron3<Real>::ContainsC0(Vector3<Real> const& p) const
+{
+    int insideCount = 0;
+
+    TIQuery<Real, Ray3<Real>, Triangle3<Real>> rtQuery;
+    Triangle3<Real> triangle;
+    Ray3<Real> ray;
+    ray.origin = p;
+
+    for (int j = 0; j < mNumRays; ++j)
+    {
+        ray.direction = mDirections[j];
+
+        // Zero intersections to start with.
+        bool odd = false;
+
+        ConvexFace const* face = mCFaces;
+        for (int i = 0; i < mNumFaces; ++i, ++face)
+        {
+            // Attempt to quickly cull the triangle.
+            if (FastNoIntersect(ray, face->plane))
+            {
+                continue;
+            }
+
+            // Process the triangles in a trifan of the face.
+            size_t numVerticesM1 = face->indices.size() - 1;
+            triangle.v[0] = mPoints[face->indices[0]];
+            for (size_t k = 1; k < numVerticesM1; ++k)
+            {
+                triangle.v[1] = mPoints[face->indices[k]];
+                triangle.v[2] = mPoints[face->indices[k + 1]];
+
+                if (rtQuery(ray, triangle).intersect)
+                {
+                    // The ray intersects the triangle.
+                    odd = !odd;
+                }
+            }
+        }
+
+        if (odd)
+        {
+            insideCount++;
+        }
+    }
+
+    return insideCount > mNumRays / 2;
+}
+
+template <typename Real>
+bool PointInPolyhedron3<Real>::ContainsS0(Vector3<Real> const& p) const
+{
+    int insideCount = 0;
+
+    TIQuery<Real, Ray3<Real>, Triangle3<Real>> rtQuery;
+    Triangle3<Real> triangle;
+    Ray3<Real> ray;
+    ray.origin = p;
+
+    for (int j = 0; j < mNumRays; ++j)
+    {
+        ray.direction = mDirections[j];
+
+        // Zero intersections to start with.
+        bool odd = false;
+
+        SimpleFace const* face = mSFaces;
+        for (int i = 0; i < mNumFaces; ++i, ++face)
+        {
+            // Attempt to quickly cull the triangle.
+            if (FastNoIntersect(ray, face->plane))
+            {
+                continue;
+            }
+
+            // The triangulation must exist to use it.
+            size_t numTriangles = face->triangles.size() / 3;
+            LogAssert(numTriangles > 0, "Triangulation must exist.");
+
+            // Process the triangles in a triangulation of the face.
+            int const* currIndex = &face->triangles[0];
+            for (size_t t = 0; t < numTriangles; ++t)
+            {
+                // Get the triangle vertices.
+                for (int k = 0; k < 3; ++k)
+                {
+                    triangle.v[k] = mPoints[*currIndex++];
+                }
+
+                // Test for intersection.
+                if (rtQuery(ray, triangle).intersect)
+                {
+                    // The ray intersects the triangle.
+                    odd = !odd;
+                }
+            }
+        }
+
+        if (odd)
+        {
+            insideCount++;
+        }
+    }
+
+    return insideCount > mNumRays / 2;
+}
+
+template <typename Real>
+bool PointInPolyhedron3<Real>::ContainsC1C2(Vector3<Real> const& p,
+    unsigned int method) const
+{
+    int insideCount = 0;
+
+    FIQuery<Real, Ray3<Real>, Plane3<Real>> rpQuery;
+    Ray3<Real> ray;
+    ray.origin = p;
+
+    for (int j = 0; j < mNumRays; ++j)
+    {
+        ray.direction = mDirections[j];
+
+        // Zero intersections to start with.
+        bool odd = false;
+
+        ConvexFace const* face = mCFaces;
+        for (int i = 0; i < mNumFaces; ++i, ++face)
+        {
+            // Attempt to quickly cull the triangle.
+            if (FastNoIntersect(ray, face->plane))
+            {
+                continue;
+            }
+
+            // Compute the ray-plane intersection.
+            auto result = rpQuery(ray, face->plane);
+
+            // If you trigger this assertion, numerical round-off errors have
+            // led to a discrepancy between FastNoIntersect and the Find()
+            // result.
+            LogAssert(result.intersect, "Unexpected condition.");
+
+            // Get a coordinate system for the plane.  Use vertex 0 as the
+            // origin.
+            Vector3<Real> const& V0 = mPoints[face->indices[0]];
+            Vector3<Real> basis[3];
+            basis[0] = face->plane.normal;
+            ComputeOrthogonalComplement(1, basis);
+
+            // Project the intersection onto the plane.
+            Vector3<Real> diff = result.point - V0;
+            Vector2<Real> projIntersect{
+                Dot(basis[1], diff), Dot(basis[2], diff) };
+
+            // Project the face vertices onto the plane of the face.
+            if (face->indices.size() > mProjVertices.size())
+            {
+                mProjVertices.resize(face->indices.size());
+            }
+
+            // Project the remaining vertices.  Vertex 0 is always the origin.
+            size_t numIndices = face->indices.size();
+            mProjVertices[0] = Vector2<Real>::Zero();
+            for (size_t k = 1; k < numIndices; ++k)
+            {
+                diff = mPoints[face->indices[k]] - V0;
+                mProjVertices[k][0] = Dot(basis[1], diff);
+                mProjVertices[k][1] = Dot(basis[2], diff);
+            }
+
+            // Test whether the intersection point is in the convex polygon.
+            PointInPolygon2<Real> PIP(static_cast<int>(mProjVertices.size()),
+                &mProjVertices[0]);
+
+            if (method == 1)
+            {
+                if (PIP.ContainsConvexOrderN(projIntersect))
+                {
+                    // The ray intersects the triangle.
+                    odd = !odd;
+                }
+            }
+            else
+            {
+                if (PIP.ContainsConvexOrderLogN(projIntersect))
+                {
+                    // The ray intersects the triangle.
+                    odd = !odd;
+                }
+            }
+        }
+
+        if (odd)
+        {
+            insideCount++;
+        }
+    }
+
+    return insideCount > mNumRays / 2;
+}
+
+template <typename Real>
+bool PointInPolyhedron3<Real>::ContainsS1(Vector3<Real> const& p) const
+{
+    int insideCount = 0;
+
+    FIQuery<Real, Ray3<Real>, Plane3<Real>> rpQuery;
+    Ray3<Real> ray;
+    ray.origin = p;
+
+    for (int j = 0; j < mNumRays; ++j)
+    {
+        ray.direction = mDirections[j];
+
+        // Zero intersections to start with.
+        bool odd = false;
+
+        SimpleFace const* face = mSFaces;
+        for (int i = 0; i < mNumFaces; ++i, ++face)
+        {
+            // Attempt to quickly cull the triangle.
+            if (FastNoIntersect(ray, face->plane))
+            {
+                continue;
+            }
+
+            // Compute the ray-plane intersection.
+            auto result = rpQuery(ray, face->plane);
+
+            // If you trigger this assertion, numerical round-off errors have
+            // led to a discrepancy between FastNoIntersect and the Find()
+            // result.
+            LogAssert(result.intersect, "Unexpected condition.");
+
+            // Get a coordinate system for the plane.  Use vertex 0 as the
+            // origin.
+            Vector3<Real> const& V0 = mPoints[face->indices[0]];
+            Vector3<Real> basis[3];
+            basis[0] = face->plane.normal;
+            ComputeOrthogonalComplement(1, basis);
+
+            // Project the intersection onto the plane.
+            Vector3<Real> diff = result.point - V0;
+            Vector2<Real> projIntersect{
+                Dot(basis[1], diff), Dot(basis[2], diff) };
+
+            // Project the face vertices onto the plane of the face.
+            if (face->indices.size() > mProjVertices.size())
+            {
+                mProjVertices.resize(face->indices.size());
+            }
+
+            // Project the remaining vertices.  Vertex 0 is always the origin.
+            size_t numIndices = face->indices.size();
+            mProjVertices[0] = Vector2<Real>::Zero();
+            for (size_t k = 1; k < numIndices; ++k)
+            {
+                diff = mPoints[face->indices[k]] - V0;
+                mProjVertices[k][0] = Dot(basis[1], diff);
+                mProjVertices[k][1] = Dot(basis[2], diff);
+            }
+
+            // Test whether the intersection point is in the convex polygon.
+            PointInPolygon2<Real> PIP(static_cast<int>(mProjVertices.size()),
+                &mProjVertices[0]);
+
+            if (PIP.Contains(projIntersect))
+            {
+                // The ray intersects the triangle.
+                odd = !odd;
+            }
+        }
+
+        if (odd)
+        {
+            insideCount++;
+        }
+    }
+
+    return insideCount > mNumRays / 2;
+}
+
+
+}
diff --git a/Engine/Plugins/Experimental/GeometryProcessing/Source/GeometryAlgorithms/Private/GTEngine/Mathematics/GteContScribeCircle2.h b/Engine/Plugins/Experimental/GeometryProcessing/Source/GeometryAlgorithms/Private/GTEngine/Mathematics/GteContScribeCircle2.h
new file mode 100644
--- /dev/null
+++ b/Engine/Plugins/Experimental/GeometryProcessing/Source/GeometryAlgorithms/Private/GTEngine/Mathematics/GteContScribeCircle2.h
@@ -0,0 +1,73 @@
+// David Eberly, Geometric Tools, Redmond WA 98052
+// Copyright (c) 1998-2019
+// Distributed under the Boost Software License, Version 1.0.
+// http://www.boost.org/LICENSE_1_0.txt
+// http://www.geometrictools.com/License/Boost/LICENSE_1_0.txt
+// File Version: 3.0.0 (2016/06/19)
+
+#pragma once
+
+#include <GTEngine/Mathematics/GteHypersphere.h>
+#include <GTEngine/Mathematics/GteLinearSystem.h>
+
+namespace gte
+{
+
+// All functions return 'true' if the circle has been constructed, 'false'
+// otherwise (input points are linearly dependent).
+
+// Circle circumscribing triangle.
+template <typename Real>
+bool Circumscribe(Vector2<Real> const& v0, Vector2<Real> const& v1,
+    Vector2<Real> const& v2, Circle2<Real>& circle);
+
+// Circle inscribing triangle.
+template <typename Real>
+bool Inscribe(Vector2<Real> const& v0, Vector2<Real> const& v1,
+    Vector2<Real> const& v2, Circle2<Real>& circle);
+
+
+template <typename Real>
+bool Circumscribe(Vector2<Real> const& v0, Vector2<Real> const& v1,
+    Vector2<Real> const& v2, Circle2<Real>& circle)
+{
+    Vector2<Real> e10 = v1 - v0;
+    Vector2<Real> e20 = v2 - v0;
+    Matrix2x2<Real> A{ e10[0], e10[1], e20[0], e20[1] };
+    Vector2<Real> B{ ((Real)0.5)*Dot(e10, e10), ((Real)0.5)*Dot(e20, e20) };
+    Vector2<Real> solution;
+    if (LinearSystem<Real>::Solve(A, B, solution))
+    {
+        circle.center = v0 + solution;
+        circle.radius = Length(solution);
+        return true;
+    }
+    return false;
+}
+
+template <typename Real>
+bool Inscribe(Vector2<Real> const& v0, Vector2<Real> const& v1,
+    Vector2<Real> const& v2, Circle2<Real>& circle)
+{
+    Vector2<Real> d10 = v1 - v0;
+    Vector2<Real> d20 = v2 - v0;
+    Vector2<Real> d21 = v2 - v1;
+    Real len10 = Length(d10);
+    Real len20 = Length(d20);
+    Real len21 = Length(d21);
+    Real perimeter = len10 + len20 + len21;
+    if (perimeter > (Real)0)
+    {
+        Real inv = ((Real)1) / perimeter;
+        len10 *= inv;
+        len20 *= inv;
+        len21 *= inv;
+        circle.center = len21*v0 + len20*v1 + len10*v2;
+        circle.radius = inv*std::abs(DotPerp(d10, d20));
+        return circle.radius > (Real)0;
+    }
+    return false;
+}
+
+
+}
diff --git a/Engine/Plugins/Experimental/GeometryProcessing/Source/GeometryAlgorithms/Private/GTEngine/Mathematics/GteContScribeCircle3Sphere3.h b/Engine/Plugins/Experimental/GeometryProcessing/Source/GeometryAlgorithms/Private/GTEngine/Mathematics/GteContScribeCircle3Sphere3.h
new file mode 100644
--- /dev/null
+++ b/Engine/Plugins/Experimental/GeometryProcessing/Source/GeometryAlgorithms/Private/GTEngine/Mathematics/GteContScribeCircle3Sphere3.h
@@ -0,0 +1,189 @@
+// David Eberly, Geometric Tools, Redmond WA 98052
+// Copyright (c) 1998-2019
+// Distributed under the Boost Software License, Version 1.0.
+// http://www.boost.org/LICENSE_1_0.txt
+// http://www.geometrictools.com/License/Boost/LICENSE_1_0.txt
+// File Version: 3.0.0 (2016/06/19)
+
+#pragma once
+
+#include <GTEngine/Mathematics/GteCircle3.h>
+#include <GTEngine/Mathematics/GteHypersphere.h>
+#include <GTEngine/Mathematics/GteLinearSystem.h>
+
+namespace gte
+{
+
+// All functions return 'true' if circle/sphere has been constructed,
+// 'false' otherwise (input points are linearly dependent).
+
+// Circle circumscribing a triangle in 3D.
+template <typename Real>
+bool Circumscribe(Vector3<Real> const& v0, Vector3<Real> const& v1,
+    Vector3<Real> const& v2, Circle3<Real>& circle);
+
+// Sphere circumscribing a tetrahedron.
+template <typename Real>
+bool Circumscribe(Vector3<Real> const& v0, Vector3<Real> const& v1,
+    Vector3<Real> const& v2, Vector3<Real> const& v3, Sphere3<Real>& sphere);
+
+// Circle inscribing a triangle in 3D.
+template <typename Real>
+bool Inscribe(Vector3<Real> const& v0, Vector3<Real> const& v1,
+    Vector3<Real> const& v2, Circle3<Real>& circle);
+
+// Sphere inscribing tetrahedron.
+template <typename Real>
+bool Inscribe(Vector3<Real> const& v0, Vector3<Real> const& v1,
+    Vector3<Real> const& v2, Vector3<Real> const& v3, Sphere3<Real>& sphere);
+
+
+template <typename Real>
+bool Circumscribe(Vector3<Real> const& v0, Vector3<Real> const& v1,
+    Vector3<Real> const& v2, Circle3<Real>& circle)
+{
+    Vector3<Real> E02 = v0 - v2;
+    Vector3<Real> E12 = v1 - v2;
+    Real e02e02 = Dot(E02, E02);
+    Real e02e12 = Dot(E02, E12);
+    Real e12e12 = Dot(E12, E12);
+    Real det = e02e02*e12e12 - e02e12*e02e12;
+    if (det != (Real)0)
+    {
+        Real halfInvDet = ((Real)0.5) / det;
+        Real u0 = halfInvDet*e12e12*(e02e02 - e02e12);
+        Real u1 = halfInvDet*e02e02*(e12e12 - e02e12);
+        Vector3<Real> tmp = u0*E02 + u1*E12;
+        circle.center = v2 + tmp;
+        circle.normal = UnitCross(E02, E12);
+        circle.radius = Length(tmp);
+        return true;
+    }
+    return false;
+}
+
+template <typename Real>
+bool Circumscribe(Vector3<Real> const& v0, Vector3<Real> const& v1,
+    Vector3<Real> const& v2, Vector3<Real> const& v3, Sphere3<Real>& sphere)
+{
+    Vector3<Real> E10 = v1 - v0;
+    Vector3<Real> E20 = v2 - v0;
+    Vector3<Real> E30 = v3 - v0;
+
+    Matrix3x3<Real> A;
+    A.SetRow(0, E10);
+    A.SetRow(1, E20);
+    A.SetRow(2, E30);
+
+    Vector3<Real> B{
+        ((Real)0.5)*Dot(E10, E10),
+        ((Real)0.5)*Dot(E20, E20),
+        ((Real)0.5)*Dot(E30, E30) };
+
+    Vector3<Real> solution;
+    if (LinearSystem<Real>::Solve(A, B, solution))
+    {
+        sphere.center = v0 + solution;
+        sphere.radius = Length(solution);
+        return true;
+    }
+    return false;
+}
+
+template <typename Real>
+bool Inscribe(Vector3<Real> const& v0, Vector3<Real> const& v1,
+    Vector3<Real> const& v2, Circle3<Real>& circle)
+{
+    // Edges.
+    Vector3<Real> E0 = v1 - v0;
+    Vector3<Real> E1 = v2 - v1;
+    Vector3<Real> E2 = v0 - v2;
+
+    // Plane normal.
+    circle.normal = Cross(E1, E0);
+
+    // Edge normals within the plane.
+    Vector3<Real> N0 = UnitCross(circle.normal, E0);
+    Vector3<Real> N1 = UnitCross(circle.normal, E1);
+    Vector3<Real> N2 = UnitCross(circle.normal, E2);
+
+    Real a0 = Dot(N1, E0);
+    if (a0 == (Real)0)
+    {
+        return false;
+    }
+
+    Real a1 = Dot(N2, E1);
+    if (a1 == (Real)0)
+    {
+        return false;
+    }
+
+    Real a2 = Dot(N0, E2);
+    if (a2 == (Real)0)
+    {
+        return false;
+    }
+
+    Real invA0 = ((Real)1) / a0;
+    Real invA1 = ((Real)1) / a1;
+    Real invA2 = ((Real)1) / a2;
+
+    circle.radius = ((Real)1) / (invA0 + invA1 + invA2);
+    circle.center = circle.radius*(invA0*v0 + invA1*v1 + invA2*v2);
+    Normalize(circle.normal);
+    return true;
+}
+
+template <typename Real>
+bool Inscribe(Vector3<Real> const& v0, Vector3<Real> const& v1,
+    Vector3<Real> const& v2, Vector3<Real> const& v3, Sphere3<Real>& sphere)
+{
+    // Edges.
+    Vector3<Real> E10 = v1 - v0;
+    Vector3<Real> E20 = v2 - v0;
+    Vector3<Real> E30 = v3 - v0;
+    Vector3<Real> E21 = v2 - v1;
+    Vector3<Real> E31 = v3 - v1;
+
+    // Normals.
+    Vector3<Real> N0 = Cross(E31, E21);
+    Vector3<Real> N1 = Cross(E20, E30);
+    Vector3<Real> N2 = Cross(E30, E10);
+    Vector3<Real> N3 = Cross(E10, E20);
+
+    // Normalize the normals.
+    if (Normalize(N0) == (Real)0)
+    {
+        return false;
+    }
+    if (Normalize(N1) == (Real)0)
+    {
+        return false;
+    }
+    if (Normalize(N2) == (Real)0)
+    {
+        return false;
+    }
+    if (Normalize(N3) == (Real)0)
+    {
+        return false;
+    }
+
+    Matrix3x3<Real> A;
+    A.SetRow(0, N1 - N0);
+    A.SetRow(1, N2 - N0);
+    A.SetRow(2, N3 - N0);
+    Vector3<Real> B{ (Real)0, (Real)0, -Dot(N3, E30) };
+    Vector3<Real> solution;
+    if (LinearSystem<Real>::Solve(A, B, solution))
+    {
+        sphere.center = v3 + solution;
+        sphere.radius = std::abs(Dot(N0, solution));
+        return true;
+    }
+    return false;
+}
+
+
+}
diff --git a/Engine/Plugins/Experimental/GeometryProcessing/Source/GeometryAlgorithms/Private/GTEngine/Mathematics/GteContSphere3.h b/Engine/Plugins/Experimental/GeometryProcessing/Source/GeometryAlgorithms/Private/GTEngine/Mathematics/GteContSphere3.h
new file mode 100644
--- /dev/null
+++ b/Engine/Plugins/Experimental/GeometryProcessing/Source/GeometryAlgorithms/Private/GTEngine/Mathematics/GteContSphere3.h
@@ -0,0 +1,81 @@
+// David Eberly, Geometric Tools, Redmond WA 98052
+// Copyright (c) 1998-2019
+// Distributed under the Boost Software License, Version 1.0.
+// http://www.boost.org/LICENSE_1_0.txt
+// http://www.geometrictools.com/License/Boost/LICENSE_1_0.txt
+// File Version: 3.0.2 (2018/10/16)
+
+#pragma once
+
+#include <GTEngine/Mathematics/GteVector3.h>
+#include <GTEngine/Mathematics/GteHypersphere.h>
+
+namespace gte
+{
+    // Compute the smallest bounding sphere whose center is the average of
+    // the input points.
+    template <typename Real>
+    bool GetContainer(int numPoints, Vector3<Real> const* points, Sphere3<Real>& sphere)
+    {
+        sphere.center = points[0];
+        for (int i = 1; i < numPoints; ++i)
+        {
+            sphere.center += points[i];
+        }
+        sphere.center /= (Real)numPoints;
+
+        sphere.radius = (Real)0;
+        for (int i = 0; i < numPoints; ++i)
+        {
+            Vector3<Real> diff = points[i] - sphere.center;
+            Real radiusSqr = Dot(diff, diff);
+            if (radiusSqr > sphere.radius)
+            {
+                sphere.radius = radiusSqr;
+            }
+        }
+
+        sphere.radius = std::sqrt(sphere.radius);
+        return true;
+    }
+
+    // Test for containment of a point inside a sphere.
+    template <typename Real>
+    bool InContainer(Vector3<Real> const& point, Sphere3<Real> const& sphere)
+    {
+        Vector3<Real> diff = point - sphere.center;
+        return Length(diff) <= sphere.radius;
+    }
+
+    // Compute the smallest bounding sphere that contains the input spheres.
+    template <typename Real>
+    bool MergeContainers(Sphere3<Real> const& sphere0, Sphere3<Real> const& sphere1, Sphere3<Real>& merge)
+    {
+        Vector3<Real> cenDiff = sphere1.center - sphere0.center;
+        Real lenSqr = Dot(cenDiff, cenDiff);
+        Real rDiff = sphere1.radius - sphere0.radius;
+        Real rDiffSqr = rDiff * rDiff;
+
+        if (rDiffSqr >= lenSqr)
+        {
+            merge = (rDiff >= (Real)0 ? sphere1 : sphere0);
+        }
+        else
+        {
+            Real length = std::sqrt(lenSqr);
+            if (length > (Real)0)
+            {
+                Real coeff = (length + rDiff) / (((Real)2)*length);
+                merge.center = sphere0.center + coeff * cenDiff;
+            }
+            else
+            {
+                merge.center = sphere0.center;
+            }
+
+            merge.radius = ((Real)0.5)*(length + sphere0.radius + sphere1.radius);
+        }
+
+        return true;
+    }
+}
diff --git a/Engine/Plugins/Experimental/GeometryProcessing/Source/GeometryAlgorithms/Private/GTEngine/Mathematics/GteConvertCoordinates.h b/Engine/Plugins/Experimental/GeometryProcessing/Source/GeometryAlgorithms/Private/GTEngine/Mathematics/GteConvertCoordinates.h
new file mode 100644
--- /dev/null
+++ b/Engine/Plugins/Experimental/GeometryProcessing/Source/GeometryAlgorithms/Private/GTEngine/Mathematics/GteConvertCoordinates.h
@@ -0,0 +1,353 @@
+// David Eberly, Geometric Tools, Redmond WA 98052
+// Copyright (c) 1998-2019
+// Distributed under the Boost Software License, Version 1.0.
+// http://www.boost.org/LICENSE_1_0.txt
+// http://www.geometrictools.com/License/Boost/LICENSE_1_0.txt
+// File Version: 3.0.1 (2018/10/16)
+
+#pragma once
+
+#include <GTEngine/Mathematics/GteMatrix.h>
+#include <GTEngine/Mathematics/GteGaussianElimination.h>
+
+// Convert points and transformations between two coordinate systems.
+// The mathematics involves a change of basis.  See the document
+//   http://www.geometrictools.com/Documentation/ConvertingBetweenCoordinateSystems.pdf
+// for the details.  Typical usage for 3D conversion is shown next.
+//
+// // Linear change of basis.  The columns of U are the basis vectors for the
+// // source coordinate system.  A vector X = { x0, x1, x2 } in the source
+// // coordinate system is represented by
+// //   X = x0*(1,0,0) + x1*(0,1,0) + x2*(0,0,1)
+// // The Cartesian coordinates for the point are the combination of these
+// // terms,
+// //   X = (x0, x1, x2)
+// // The columns of V are the basis vectors for the target coordinate system.
+// // A vector Y = { y0, y1, y2 } in the target coordinate system is
+// // represented by
+// //   Y = y0*(1,0,0,0) + y1*(0,0,1) + y2*(0,1,0)
+// // The Cartesian coordinates for the vector are the combination of these
+// // terms,
+// //   Y = (y0, y2, y1)
+// // The call Y = convert.UToV(X) computes y0, y1 and y2 so that the Cartesian
+// // coordinates for X and for Y are the same.  For example,
+// //   X = { 1.0, 2.0, 3.0 }
+// //     = 1.0*(1,0,0) + 2.0*(0,1,0) + 3.0*(0,0,1)
+// //     = (1, 2, 3)
+// //   Y = { 1.0, 3.0, 2.0 }
+// //     = 1.0*(1,0,0) + 3.0*(0,0,1) + 2.0*(0,1,0)
+// //     = (1, 2, 3)
+// // X and Y represent the same vector (equal Cartesian coordinates) but have
+// // different representations in the source and target coordinates.
+//
+// ConvertCoordinates<3, double> convert;
+// Vector<3, double> X, Y, P0, P1, diff;
+// Matrix<3, 3, double> U, V, A, B;
+// bool isRHU, isRHV;
+// U.SetCol(0, Vector3<double>{1.0, 0.0, 0.0});
+// U.SetCol(1, Vector3<double>{0.0, 1.0, 0.0});
+// U.SetCol(2, Vector3<double>{0.0, 0.0, 1.0});
+// V.SetCol(0, Vector3<double>{1.0, 0.0, 0.0});
+// V.SetCol(1, Vector3<double>{0.0, 0.0, 1.0});
+// V.SetCol(2, Vector3<double>{0.0, 1.0, 0.0});
+// convert(U, true, V, true);
+// isRHU = convert.IsRightHandedU();  // true
+// isRHV = convert.IsRightHandedV();  // false
+// X = { 1.0, 2.0, 3.0 };
+// Y = convert.UToV(X);  // { 1.0, 3.0, 2.0 }
+// P0 = U*X;
+// P1 = V*Y;
+// diff = P0 - P1;  // { 0, 0, 0 }
+// Y = { 0.0, 1.0, 2.0 };
+// X = convert.VToU(Y);  // { 0.0, 2.0, 1.0 }
+// P0 = U*X;
+// P1 = V*Y;
+// diff = P0 - P1;  // { 0, 0, 0 }
+// double cs = 0.6, sn = 0.8;  // cs*cs + sn*sn = 1
+// A.SetCol(0, Vector3<double>{  c,   s, 0.0});
+// A.SetCol(1, Vector3<double>{ -s,   c, 0.0});
+// A.SetCol(2, Vector3<double>{0.0, 0.0, 1.0});
+// B = convert.UToV(A);
+//   // B.GetCol(0) = { c, 0, s}
+//   // B.GetCol(1) = { 0, 1, 0}
+//   // B.GetCol(2) = {-s, 0, c}
+// X = A*X;  // U is VOR
+// Y = B*Y;  // V is VOR
+// P0 = U*X;
+// P1 = V*Y;
+// diff = P0 - P1;  // { 0, 0, 0 }
+//
+// // Affine change of basis.  The first three columns of U are the basis
+// // vectors for the source coordinate system and must have last components
+// // set to 0.  The last column is the origin for that system and must have
+// // last component set to 1.  A point X = { x0, x1, x2, 1 } in the source
+// // coordinate system is represented by
+// //   X = x0*(-1,0,0,0) + x1*(0,0,1,0) + x2*(0,-1,0,0) + 1*(1,2,3,1)
+// // The Cartesian coordinates for the point are the combination of these
+// // terms,
+// //   X = (-x0 + 1, -x2 + 2, x1 + 3, 1)
+// // The first three columns of V are the basis vectors for the target
+// // coordinate system and must have last components set to 0.  The last
+// // column is the origin for that system and must have last component set
+// // to 1.  A point Y = { y0, y1, y2, 1 } in the target coordinate system is
+// // represented by
+// //   Y = y0*(0,1,0,0) + y1*(-1,0,0,0) + y2*(0,0,1,0) + 1*(4,5,6,1)
+// // The Cartesian coordinates for the point are the combination of these
+// // terms,
+// //   Y = (-y1 + 4, y0 + 5, y2 + 6, 1)
+// // The call Y = convert.UToV(X) computes y0, y1 and y2 so that the Cartesian
+// // coordinates for X and for Y are the same.  For example,
+// //   X = { -1.0, 4.0, -3.0, 1.0 }
+// //     = -1.0*(-1,0,0,0) + 4.0*(0,0,1,0) - 3.0*(0,-1,0,0) + 1.0*(1,2,3,1)
+// //     = (2, 5, 7, 1)
+// //   Y = { 0.0, 2.0, 1.0, 1.0 }
+// //     = 0.0*(0,1,0,0) + 2.0*(-1,0,0,0) + 1.0*(0,0,1,0) + 1.0*(4,5,6,1)
+// //     = (2, 5, 7, 1)
+// // X and Y represent the same point (equal Cartesian coordinates) but have
+// // different representations in the source and target affine coordinates.
+//
+// ConvertCoordinates<4, double> convert;
+// Vector<4, double> X, Y, P0, P1, diff;
+// Matrix<4, 4, double> U, V, A, B;
+// bool isRHU, isRHV;
+// U.SetCol(0, Vector4<double>{-1.0, 0.0, 0.0, 0.0});
+// U.SetCol(1, Vector4<double>{0.0, 0.0, 1.0, 0.0});
+// U.SetCol(2, Vector4<double>{0.0, -1.0, 0.0, 0.0});
+// U.SetCol(3, Vector4<double>{1.0, 2.0, 3.0, 1.0});
+// V.SetCol(0, Vector4<double>{0.0, 1.0, 0.0, 0.0});
+// V.SetCol(1, Vector4<double>{-1.0, 0.0, 0.0, 0.0});
+// V.SetCol(2, Vector4<double>{0.0, 0.0, 1.0, 0.0});
+// V.SetCol(3, Vector4<double>{4.0, 5.0, 6.0, 1.0});
+// convert(U, true, V, false);
+// isRHU = convert.IsRightHandedU();  // false
+// isRHV = convert.IsRightHandedV();  // true
+// X = { -1.0, 4.0, -3.0, 1.0 };
+// Y = convert.UToV(X);  // { 0.0, 2.0, 1.0, 1.0 }
+// P0 = U*X;
+// P1 = V*Y;
+// diff = P0 - P1;  // { 0, 0, 0, 0 }
+// Y = { 1.0, 2.0, 3.0, 1.0 };
+// X = convert.VToU(Y);  // { -1.0, 6.0, -4.0, 1.0 }
+// P0 = U*X;
+// P1 = V*Y;
+// diff = P0 - P1;  // { 0, 0, 0, 0 }
+// double c = 0.6, s = 0.8;  // c*c + s*s = 1
+// A.SetCol(0, Vector4<double>{  c,  s,   0.0, 0.0});
+// A.SetCol(1, Vector4<double>{ -s,  c,   0.0, 0.0});
+// A.SetCol(2, Vector4<double>{0.0, 0.0,  1.0, 0.0});
+// A.SetCol(3, Vector4<double>{0.3, 1.0, -2.0, 1.0});
+// B = convert.UToV(A);
+// // B.GetCol(0) = {   1,    0,    0, 0 }
+// // B.GetCol(1) = {   0,    c,    s, 0 }
+// // B.GetCol(2) = {   0,   -s,    c, 0 }
+// // B.GetCol(3) = { 2.0, -0.9, -2.6, 1 }
+// X = A*X;  // U is VOR
+// Y = Y*B;  // V is VOL (not VOR)
+// P0 = U*X;
+// P1 = V*Y;
+// diff = P0 - P1;  // { 0, 0, 0, 0 }
+
+namespace gte
+{
+    template <int N, typename Real>
+    class ConvertCoordinates
+    {
+    public:
+        // Construction of the change of basis matrix.  The implementation
+        // supports both linear change of basis and affine change of basis.
+        ConvertCoordinates()
+            :
+            mIsVectorOnRightU(true),
+            mIsVectorOnRightV(true),
+            mIsRightHandedU(true),
+            mIsRightHandedV(true)
+        {
+            mC.MakeIdentity();
+            mInverseC.MakeIdentity();
+        }
+
+        // Compute a change of basis between two coordinate systems.  The return
+        // value is 'true' iff U and V are invertible.  The matrix-vector
+        // multiplication conventions affect the conversion of matrix
+        // transformations.  The Boolean inputs indicate how you want the matrices
+        // to be interpreted when applied as transformations of a vector.
+        bool operator()(
+            Matrix<N, N, Real> const& U, bool vectorOnRightU,
+            Matrix<N, N, Real> const& V, bool vectorOnRightV)
+        {
+            // Initialize in case of early exit.
+            mC.MakeIdentity();
+            mInverseC.MakeIdentity();
+            mIsVectorOnRightU = true;
+            mIsVectorOnRightV = true;
+            mIsRightHandedU = true;
+            mIsRightHandedV = true;
+
+            Matrix<N, N, Real> inverseU;
+            Real determinantU;
+            bool invertibleU = GaussianElimination<Real>()(N, &U[0], &inverseU[0],
+                determinantU, nullptr, nullptr, nullptr, 0, nullptr);
+            if (!invertibleU)
+            {
+                return false;
+            }
+
+            Matrix<N, N, Real> inverseV;
+            Real determinantV;
+            bool invertibleV = GaussianElimination<Real>()(N, &V[0], &inverseV[0],
+                determinantV, nullptr, nullptr, nullptr, 0, nullptr);
+            if (!invertibleV)
+            {
+                return false;
+            }
+
+            mC = inverseU * V;
+            mInverseC = inverseV * U;
+            mIsVectorOnRightU = vectorOnRightU;
+            mIsVectorOnRightV = vectorOnRightV;
+            mIsRightHandedU = (determinantU > (Real)0);
+            mIsRightHandedV = (determinantV > (Real)0);
+            return true;
+        }
+
+        // Member access.
+        inline Matrix<N, N, Real> const& GetC() const
+        {
+            return mC;
+        }
+
+        inline Matrix<N, N, Real> const& GetInverseC() const
+        {
+            return mInverseC;
+        }
+
+        inline bool IsVectorOnRightU() const
+        {
+            return mIsVectorOnRightU;
+        }
+
+        inline bool IsVectorOnRightV() const
+        {
+            return mIsVectorOnRightV;
+        }
+
+        inline bool IsRightHandedU() const
+        {
+            return mIsRightHandedU;
+        }
+
+        inline bool IsRightHandedV() const
+        {
+            return mIsRightHandedV;
+        }
+
+        // Convert points between coordinate systems.  The names of the systems
+        // are U and V to make it clear which inputs of operator() they are
+        // associated with.  The X vector stores coordinates for the U-system and
+        // the Y vector stores coordinates for the V-system.
+
+        // Y = C^{-1}*X
+        inline Vector<N, Real> UToV(Vector<N, Real> const& X) const
+        {
+            return mInverseC * X;
+        }
+
+        // X = C*Y
+        inline Vector<N, Real> VToU(Vector<N, Real> const& Y) const
+        {
+            return mC * Y;
+        }
+
+        // Convert transformations between coordinate systems.  The outputs are
+        // computed according to the tables shown before the function
+        // declarations. The superscript T denotes the transpose operator.
+        // vectorOnRightU = true:  transformation is X' = A*X
+        // vectorOnRightU = false: transformation is (X')^T = X^T*A
+        // vectorOnRightV = true:  transformation is Y' = B*Y
+        // vectorOnRightV = false: transformation is (Y')^T = Y^T*B
+
+        // vectorOnRightU  | vectorOnRightV  | output
+        // ----------------+-----------------+---------------------
+        // true            | true            | C^{-1} * A * C
+        // true            | false           | (C^{-1} * A * C)^T 
+        // false           | true            | C^{-1} * A^T * C
+        // false           | false           | (C^{-1} * A^T * C)^T
+        Matrix<N, N, Real> UToV(Matrix<N, N, Real> const& A) const
+        {
+            Matrix<N, N, Real> product;
+
+            if (mIsVectorOnRightU)
+            {
+                product = mInverseC * A * mC;
+                if (mIsVectorOnRightV)
+                {
+                    return product;
+                }
+                else
+                {
+                    return Transpose(product);
+                }
+            }
+            else
+            {
+                product = mInverseC * MultiplyATB(A, mC);
+                if (mIsVectorOnRightV)
+                {
+                    return product;
+                }
+                else
+                {
+                    return Transpose(product);
+                }
+            }
+        }
+
+        // vectorOnRightU  | vectorOnRightV  | output
+        // ----------------+-----------------+---------------------
+        // true            | true            | C * B * C^{-1}
+        // true            | false           | C * B^T * C^{-1}
+        // false           | true            | (C * B * C^{-1})^T
+        // false           | false           | (C * B^T * C^{-1})^T
+        Matrix<N, N, Real> VToU(Matrix<N, N, Real> const& B) const
+        {
+            // vectorOnRightU  | vectorOnRightV  | output
+            // ----------------+-----------------+---------------------
+            // true            | true            | C * B * C^{-1}
+            // true            | false           | C * B^T * C^{-1}
+            // false           | true            | (C * B * C^{-1})^T
+            // false           | false           | (C * B^T * C^{-1})^T
+            Matrix<N, N, Real> product;
+
+            if (mIsVectorOnRightV)
+            {
+                product = mC * B * mInverseC;
+                if (mIsVectorOnRightU)
+                {
+                    return product;
+                }
+                else
+                {
+                    return Transpose(product);
+                }
+            }
+            else
+            {
+                product = mC * MultiplyATB(B, mInverseC);
+                if (mIsVectorOnRightU)
+                {
+                    return product;
+                }
+                else
+                {
+                    return Transpose(product);
+                }
+            }
+        }
+
+    private:
+        // C = U^{-1}*V, C^{-1} = V^{-1}*U
+        Matrix<N, N, Real> mC, mInverseC;
+        bool mIsVectorOnRightU, mIsVectorOnRightV;
+        bool mIsRightHandedU, mIsRightHandedV;
+    };
+}
diff --git a/Engine/Plugins/Experimental/GeometryProcessing/Source/GeometryAlgorithms/Private/GTEngine/Mathematics/GteConvexHull2.h b/Engine/Plugins/Experimental/GeometryProcessing/Source/GeometryAlgorithms/Private/GTEngine/Mathematics/GteConvexHull2.h
new file mode 100644
--- /dev/null
+++ b/Engine/Plugins/Experimental/GeometryProcessing/Source/GeometryAlgorithms/Private/GTEngine/Mathematics/GteConvexHull2.h
@@ -0,0 +1,409 @@
+// David Eberly, Geometric Tools, Redmond WA 98052
+// Copyright (c) 1998-2019
+// Distributed under the Boost Software License, Version 1.0.
+// http://www.boost.org/LICENSE_1_0.txt
+// http://www.geometrictools.com/License/Boost/LICENSE_1_0.txt
+// File Version: 3.0.0 (2016/06/19)
+
+#pragma once
+
+// Compute the convex hull of 2D points using a divide-and-conquer algorithm.
+// This is an O(N log N) algorithm for N input points.  The only way to ensure
+// a correct result for the input vertices (assumed to be exact) is to choose
+// ComputeType for exact rational arithmetic.  You may use BSNumber.  No
+// divisions are performed in this computation, so you do not have to use
+// BSRational.
+//
+// The worst-case choices of N for Real of type BSNumber or BSRational with
+// integer storage UIntegerFP32<N> are listed in the next table.  The numerical
+// computations are encapsulated in PrimalQuery2<Real>::ToLineExtended.  We
+// recommend using only BSNumber, because no divisions are performed in the
+// convex-hull computations.
+//
+//    input type | compute type | N
+//    -----------+--------------+------
+//    float      | BSNumber     |   18
+//    double     | BSNumber     |  132
+//    float      | BSRational   |  214
+//    double     | BSRational   | 1587
+
+#include <GTEngine/LowLevel/GteLogger.h>
+#include <GTEngine/Mathematics/GtePrimalQuery2.h>
+#include <GTEngine/Mathematics/GteLine.h>
+#include <algorithm>
+#include <vector>
+
+namespace gte
+{
+
+template <typename InputType, typename ComputeType>
+class ConvexHull2
+{
+public:
+    // The class is a functor to support computing the convex hull of multiple
+    // data sets using the same class object.
+    ConvexHull2();
+
+    // The input is the array of points whose convex hull is required.  The
+    // epsilon value is used to determine the intrinsic dimensionality of the
+    // vertices (d = 0, 1, or 2).  When epsilon is positive, the determination
+    // is fuzzy--points approximately the same point, approximately on a
+    // line, or planar.  The return value is 'true' if and only if the hull
+    // construction is successful.
+    bool operator()(int numPoints, Vector2<InputType> const* points, InputType epsilon);
+
+    // Dimensional information.  If GetDimension() returns 1, the points lie
+    // on a line P+t*D (fuzzy comparison when epsilon > 0).  You can sort
+    // these if you need a polyline output by projecting onto the line each
+    // vertex X = P+t*D, where t = Dot(D,X-P).
+    inline InputType GetEpsilon() const;
+    inline int GetDimension() const;
+    inline Line2<InputType> const& GetLine() const;
+
+    // Member access.
+    inline int GetNumPoints() const;
+    inline int GetNumUniquePoints() const;
+    inline Vector2<InputType> const* GetPoints() const;
+    inline PrimalQuery2<ComputeType> const& GetQuery() const;
+
+    // The convex hull is a convex polygon whose vertices are listed in
+    // counterclockwise order.
+    inline std::vector<int> const& GetHull() const;
+
+private:
+    // Support for divide-and-conquer.
+    void GetHull(int& i0, int& i1);
+    void Merge(int j0, int j1, int j2, int j3, int& i0, int& i1);
+    void GetTangent(int j0, int j1, int j2, int j3, int& i0, int& i1);
+
+    // The epsilon value is used for fuzzy determination of intrinsic
+    // dimensionality.  If the dimension is 0 or 1, the constructor returns
+    // early.  The caller is responsible for retrieving the dimension and
+    // taking an alternate path should the dimension be smaller than 2.  If
+    // the dimension is 0, the caller may as well treat all points[] as a
+    // single point, say, points[0].  If the dimension is 1, the caller can
+    // query for the approximating line and project points[] onto it for
+    // further processing.
+    InputType mEpsilon;
+    int mDimension;
+    Line2<InputType> mLine;
+
+    // The array of points used for geometric queries.  If you want to be
+    // certain of a correct result, choose ComputeType to be BSNumber.
+    std::vector<Vector2<ComputeType>> mComputePoints;
+    PrimalQuery2<ComputeType> mQuery;
+
+    int mNumPoints;
+    int mNumUniquePoints;
+    Vector2<InputType> const* mPoints;
+    std::vector<int> mMerged, mHull;
+};
+
+
+template <typename InputType, typename ComputeType>
+ConvexHull2<InputType, ComputeType>::ConvexHull2()
+    :
+    mEpsilon((InputType)0),
+    mDimension(0),
+    mLine(Vector2<InputType>::Zero(), Vector2<InputType>::Zero()),
+    mNumPoints(0),
+    mNumUniquePoints(0),
+    mPoints(nullptr)
+{
+}
+
+template <typename InputType, typename ComputeType>
+bool ConvexHull2<InputType, ComputeType>::operator()(int numPoints,
+    Vector2<InputType> const* points, InputType epsilon)
+{
+    mEpsilon = std::max(epsilon, (InputType)0);
+    mDimension = 0;
+    mLine.origin = Vector2<InputType>::Zero();
+    mLine.direction = Vector2<InputType>::Zero();
+    mNumPoints = numPoints;
+    mNumUniquePoints = 0;
+    mPoints = points;
+    mMerged.clear();
+    mHull.clear();
+
+    int i, j;
+    if (mNumPoints < 3)
+    {
+        // ConvexHull2 should be called with at least three points.
+        return false;
+    }
+
+    IntrinsicsVector2<InputType> info(mNumPoints, mPoints, mEpsilon);
+    if (info.dimension == 0)
+    {
+        // mDimension is 0
+        return false;
+    }
+
+    if (info.dimension == 1)
+    {
+        // The set is (nearly) collinear.
+        mDimension = 1;
+        mLine = Line2<InputType>(info.origin, info.direction[0]);
+        return false;
+    }
+
+    mDimension = 2;
+
+    // Compute the points for the queries.
+    mComputePoints.resize(mNumPoints);
+    mQuery.Set(mNumPoints, &mComputePoints[0]);
+    for (i = 0; i < mNumPoints; ++i)
+    {
+        for (j = 0; j < 2; ++j)
+        {
+            mComputePoints[i][j] = points[i][j];
+        }
+    }
+
+    // Sort the points.
+    mHull.resize(mNumPoints);
+    for (i = 0; i < mNumPoints; ++i)
+    {
+        mHull[i] = i;
+    }
+    std::sort(mHull.begin(), mHull.end(),
+        [points](int i0, int i1)
+        {
+            if (points[i0][0] < points[i1][0]) { return true; }
+            if (points[i0][0] > points[i1][0]) { return false; }
+            return points[i0][1] < points[i1][1];
+        }
+    );
+
+    // Remove duplicates.
+    auto newEnd = std::unique(mHull.begin(), mHull.end(),
+        [points](int i0, int i1)
+        {
+            return points[i0] == points[i1];
+        }
+    );
+    mHull.erase(newEnd, mHull.end());
+    mNumUniquePoints = static_cast<int>(mHull.size());
+
+    // Use a divide-and-conquer algorithm.  The merge step computes the
+    // convex hull of two convex polygons.
+    mMerged.resize(mNumUniquePoints);
+    int i0 = 0, i1 = mNumUniquePoints - 1;
+    GetHull(i0, i1);
+    mHull.resize(i1 - i0 + 1);
+    return true;
+}
+
+template <typename InputType, typename ComputeType> inline
+InputType ConvexHull2<InputType, ComputeType>::GetEpsilon() const
+{
+    return mEpsilon;
+}
+
+template <typename InputType, typename ComputeType> inline
+int ConvexHull2<InputType, ComputeType>::GetDimension() const
+{
+    return mDimension;
+}
+
+template <typename InputType, typename ComputeType> inline
+Line2<InputType> const& ConvexHull2<InputType, ComputeType>::GetLine() const
+{
+    return mLine;
+}
+
+template <typename InputType, typename ComputeType> inline
+int ConvexHull2<InputType, ComputeType>::GetNumPoints() const
+{
+    return mNumPoints;
+}
+
+template <typename InputType, typename ComputeType> inline
+int ConvexHull2<InputType, ComputeType>::GetNumUniquePoints() const
+{
+    return mNumUniquePoints;
+}
+
+template <typename InputType, typename ComputeType> inline
+Vector2<InputType> const* ConvexHull2<InputType, ComputeType>::GetPoints() const
+{
+    return mPoints;
+}
+
+template <typename InputType, typename ComputeType> inline
+PrimalQuery2<ComputeType> const& ConvexHull2<InputType, ComputeType>::GetQuery() const
+{
+    return mQuery;
+}
+
+template <typename InputType, typename ComputeType> inline
+std::vector<int> const& ConvexHull2<InputType, ComputeType>::GetHull() const
+{
+    return mHull;
+}
+
+template <typename InputType, typename ComputeType> inline
+void ConvexHull2<InputType, ComputeType>::GetHull(int& i0, int& i1)
+{
+    int numVertices = i1 - i0 + 1;
+    if (numVertices > 1)
+    {
+        // Compute the middle index of input range.
+        int mid = (i0 + i1) / 2;
+
+        // Compute the hull of subsets (mid-i0+1 >= i1-mid).
+        int j0 = i0, j1 = mid, j2 = mid + 1, j3 = i1;
+        GetHull(j0, j1);
+        GetHull(j2, j3);
+
+        // Merge the convex hulls into a single convex hull.
+        Merge(j0, j1, j2, j3, i0, i1);
+    }
+    // else: The convex hull is a single point.
+}
+
+template <typename InputType, typename ComputeType> inline
+void ConvexHull2<InputType, ComputeType>::Merge(int j0, int j1, int j2, int j3,
+    int& i0, int& i1)
+{
+    // Subhull0 is to the left of subhull1 because of the initial sorting of
+    // the points by x-components.  We need to find two mutually visible
+    // points, one on the left subhull and one on the right subhull.
+    int size0 = j1 - j0 + 1;
+    int size1 = j3 - j2 + 1;
+
+    int i;
+    Vector2<ComputeType> p;
+
+    // Find the right-most point of the left subhull.
+    Vector2<ComputeType> pmax0 = mComputePoints[mHull[j0]];
+    int imax0 = j0;
+    for (i = j0 + 1; i <= j1; ++i)
+    {
+        p = mComputePoints[mHull[i]];
+        if (pmax0 < p)
+        {
+            pmax0 = p;
+            imax0 = i;
+        }
+    }
+
+    // Find the left-most point of the right subhull.
+    Vector2<ComputeType> pmin1 = mComputePoints[mHull[j2]];
+    int imin1 = j2;
+    for (i = j2 + 1; i <= j3; ++i)
+    {
+        p = mComputePoints[mHull[i]];
+        if (p < pmin1)
+        {
+            pmin1 = p;
+            imin1 = i;
+        }
+    }
+
+    // Get the lower tangent to hulls (LL = lower-left, LR = lower-right).
+    int iLL = imax0, iLR = imin1;
+    GetTangent(j0, j1, j2, j3, iLL, iLR);
+
+    // Get the upper tangent to hulls (UL = upper-left, UR = upper-right).
+    int iUL = imax0, iUR = imin1;
+    GetTangent(j2, j3, j0, j1, iUR, iUL);
+
+    // Construct the counterclockwise-ordered merged-hull vertices.
+    int k;
+    int numMerged = 0;
+
+    i = iUL;
+    for (k = 0; k < size0; ++k)
+    {
+        mMerged[numMerged++] = mHull[i];
+        if (i == iLL)
+        {
+            break;
+        }
+        i = (i < j1 ? i + 1 : j0);
+    }
+    LogAssert(k < size0, "Unexpected condition.");
+
+    i = iLR;
+    for (k = 0; k < size1; ++k)
+    {
+        mMerged[numMerged++] = mHull[i];
+        if (i == iUR)
+        {
+            break;
+        }
+        i = (i < j3 ? i + 1 : j2);
+    }
+    LogAssert(k < size1, "Unexpected condition.");
+
+    int next = j0;
+    for (k = 0; k < numMerged; ++k)
+    {
+        mHull[next] = mMerged[k];
+        ++next;
+    }
+
+    i0 = j0;
+    i1 = next - 1;
+}
+
+template <typename InputType, typename ComputeType> inline
+void ConvexHull2<InputType, ComputeType>::GetTangent(int j0, int j1, int j2, int j3,
+    int& i0, int& i1)
+{
+    // In theory the loop terminates in a finite number of steps, but the
+    // upper bound for the loop variable is used to trap problems caused by
+    // floating point roundoff errors that might lead to an infinite loop.
+
+    int size0 = j1 - j0 + 1;
+    int size1 = j3 - j2 + 1;
+    int const imax = size0 + size1;
+    int i, iLm1, iRp1;
+    Vector2<ComputeType> L0, L1, R0, R1;
+
+    for (i = 0; i < imax; i++)
+    {
+        // Get the endpoints of the potential tangent.
+        L1 = mComputePoints[mHull[i0]];
+        R0 = mComputePoints[mHull[i1]];
+
+        // Walk along the left hull to find the point of tangency.
+        if (size0 > 1)
+        {
+            iLm1 = (i0 > j0 ? i0 - 1 : j1);
+            L0 = mComputePoints[mHull[iLm1]];
+            auto order = mQuery.ToLineExtended(R0, L0, L1);
+            if (order == PrimalQuery2<ComputeType>::ORDER_NEGATIVE
+                || order == PrimalQuery2<ComputeType>::ORDER_COLLINEAR_RIGHT)
+            {
+                i0 = iLm1;
+                continue;
+            }
+        }
+
+        // Walk along right hull to find the point of tangency.
+        if (size1 > 1)
+        {
+            iRp1 = (i1 < j3 ? i1 + 1 : j2);
+            R1 = mComputePoints[mHull[iRp1]];
+            auto order = mQuery.ToLineExtended(L1, R0, R1);
+            if (order == PrimalQuery2<ComputeType>::ORDER_NEGATIVE
+                || order == PrimalQuery2<ComputeType>::ORDER_COLLINEAR_LEFT)
+            {
+                i1 = iRp1;
+                continue;
+            }
+        }
+
+        // The tangent segment has been found.
+        break;
+    }
+
+    // Detect an "infinite loop" caused by floating point round-off errors.
+    LogAssert(i < imax, "Unexpected condition.");
+}
+
+
+}
