move GeometryProcessing plugin from NotForLicensees to Experimental

diff --git a/Engine/Plugins/Experimental/GeometryProcessing/Source/GeometryAlgorithms/Private/GTEngine/Mathematics/GteConvexHull3.h b/Engine/Plugins/Experimental/GeometryProcessing/Source/GeometryAlgorithms/Private/GTEngine/Mathematics/GteConvexHull3.h
new file mode 100644
--- /dev/null
+++ b/Engine/Plugins/Experimental/GeometryProcessing/Source/GeometryAlgorithms/Private/GTEngine/Mathematics/GteConvexHull3.h
@@ -0,0 +1,430 @@
+// David Eberly, Geometric Tools, Redmond WA 98052
+// Copyright (c) 1998-2019
+// Distributed under the Boost Software License, Version 1.0.
+// http://www.boost.org/LICENSE_1_0.txt
+// http://www.geometrictools.com/License/Boost/LICENSE_1_0.txt
+// File Version: 3.0.0 (2016/06/19)
+
+#pragma once
+
+// Compute the convex hull of 3D points using incremental insertion.  The only
+// way to ensure a correct result for the input vertices (assumed to be exact)
+// is to choose ComputeType for exact rational arithmetic.  You may use
+// BSNumber.  No divisions are performed in this computation, so you do not
+// have to use BSRational.
+//
+// The worst-case choices of N for Real of type BSNumber or BSRational with
+// integer storage UIntegerFP32<N> are listed in the next table.  The numerical
+// computations are encapsulated in PrimalQuery3<Real>::ToPlane.  We recommend
+// using only BSNumber, because no divisions are performed in the convex-hull
+// computations.
+//
+//    input type | compute type | N
+//    -----------+--------------+------
+//    float      | BSNumber     |    27
+//    double     | BSNumber     |   197
+//    float      | BSRational   |  2882
+//    double     | BSRational   | 21688
+
+#include <GTEngine/Mathematics/GteETManifoldMesh.h>
+#include <GTEngine/Mathematics/GtePrimalQuery3.h>
+#include <GTEngine/Mathematics/GteLine.h>
+#include <GTEngine/Mathematics/GteHyperplane.h>
+#include <functional>
+#include <set>
+#include <thread>
+#include <vector>
+
+namespace gte
+{
+
+template <typename InputType, typename ComputeType>
+class ConvexHull3
+{
+public:
+    // The class is a functor to support computing the convex hull of multiple
+    // data sets using the same class object.  For multithreading in Update,
+    // choose 'numThreads' subject to the constraints
+    //     1 <= numThreads <= std::thread::hardware_concurrency().
+    ConvexHull3(unsigned int numThreads = 1);
+
+    // The input is the array of points whose convex hull is required.  The
+    // epsilon value is used to determine the intrinsic dimensionality of the
+    // vertices (d = 0, 1, 2, or 3).  When epsilon is positive, the
+    // determination is fuzzy--points approximately the same point,
+    // approximately on a line, approximately planar, or volumetric.
+    bool operator()(int numPoints, Vector3<InputType> const* points,
+        InputType epsilon);
+
+    // Dimensional information.  If GetDimension() returns 1, the points lie
+    // on a line P+t*D (fuzzy comparison when epsilon > 0).  You can sort
+    // these if you need a polyline output by projecting onto the line each
+    // vertex X = P+t*D, where t = Dot(D,X-P).  If GetDimension() returns 2,
+    // the points line on a plane P+s*U+t*V (fuzzy comparison when
+    // epsilon > 0).  You can project each point X = P+s*U+t*V, where
+    // s = Dot(U,X-P) and t = Dot(V,X-P), then apply ConvexHull2 to the (s,t)
+    // tuples.
+    inline InputType GetEpsilon() const;
+    inline int GetDimension() const;
+    inline Line3<InputType> const& GetLine() const;
+    inline Plane3<InputType> const& GetPlane() const;
+
+    // Member access.
+    inline int GetNumPoints() const;
+    inline int GetNumUniquePoints() const;
+    inline Vector3<InputType> const* GetPoints() const;
+    inline PrimalQuery3<ComputeType> const& GetQuery() const;
+
+    // The convex hull is a convex polyhedron with triangular faces.
+    inline std::vector<TriangleKey<true>> const& GetHullUnordered() const;
+    ETManifoldMesh const& GetHullMesh() const;
+
+private:
+    // Support for incremental insertion.
+    void Update(int i);
+
+    // The epsilon value is used for fuzzy determination of intrinsic
+    // dimensionality.  If the dimension is 0, 1, or 2, the constructor
+    // returns early.  The caller is responsible for retrieving the dimension
+    // and taking an alternate path should the dimension be smaller than 3.
+    // If the dimension is 0, the caller may as well treat all points[] as a
+    // single point, say, points[0].  If the dimension is 1, the caller can
+    // query for the approximating line and project points[] onto it for
+    // further processing.  If the dimension is 2, the caller can query for
+    // the approximating plane and project points[] onto it for further
+    // processing.
+    InputType mEpsilon;
+    int mDimension;
+    Line3<InputType> mLine;
+    Plane3<InputType> mPlane;
+
+    // The array of points used for geometric queries.  If you want to be
+    // certain of a correct result, choose ComputeType to be BSNumber.
+    std::vector<Vector3<ComputeType>> mComputePoints;
+    PrimalQuery3<ComputeType> mQuery;
+
+    int mNumPoints;
+    int mNumUniquePoints;
+    Vector3<InputType> const* mPoints;
+    std::vector<TriangleKey<true>> mHullUnordered;
+    mutable ETManifoldMesh mHullMesh;
+    unsigned int mNumThreads;
+};
+
+
+template <typename InputType, typename ComputeType>
+ConvexHull3<InputType, ComputeType>::ConvexHull3(unsigned int numThreads)
+    :
+    mEpsilon((InputType)0),
+    mDimension(0),
+    mLine(Vector3<InputType>::Zero(), Vector3<InputType>::Zero()),
+    mPlane(Vector3<InputType>::Zero(), (InputType)0),
+    mNumPoints(0),
+    mNumUniquePoints(0),
+    mPoints(nullptr),
+    mNumThreads(numThreads)
+{
+}
+
+template <typename InputType, typename ComputeType>
+bool ConvexHull3<InputType, ComputeType>::operator()(int numPoints,
+    Vector3<InputType> const* points, InputType epsilon)
+{
+    mEpsilon = std::max(epsilon, (InputType)0);
+    mDimension = 0;
+    mLine.origin = Vector3<InputType>::Zero();
+    mLine.direction = Vector3<InputType>::Zero();
+    mPlane.normal = Vector3<InputType>::Zero();
+    mPlane.constant = (InputType)0;
+    mNumPoints = numPoints;
+    mNumUniquePoints = 0;
+    mPoints = points;
+    mHullUnordered.clear();
+    mHullMesh.Clear();
+
+    int i, j;
+    if (mNumPoints < 4)
+    {
+        // ConvexHull3 should be called with at least four points.
+        return false;
+    }
+
+    IntrinsicsVector3<InputType> info(mNumPoints, mPoints, mEpsilon);
+    if (info.dimension == 0)
+    {
+        // The set is (nearly) a point.
+        return false;
+    }
+
+    if (info.dimension == 1)
+    {
+        // The set is (nearly) collinear.
+        mDimension = 1;
+        mLine = Line3<InputType>(info.origin, info.direction[0]);
+        return false;
+    }
+
+    if (info.dimension == 2)
+    {
+        // The set is (nearly) coplanar.
+        mDimension = 2;
+        mPlane = Plane3<InputType>(UnitCross(info.direction[0],
+            info.direction[1]), info.origin);
+        return false;
+    }
+
+    mDimension = 3;
+
+    // Compute the vertices for the queries.
+    mComputePoints.resize(mNumPoints);
+    mQuery.Set(mNumPoints, &mComputePoints[0]);
+    for (i = 0; i < mNumPoints; ++i)
+    {
+        for (j = 0; j < 3; ++j)
+        {
+            mComputePoints[i][j] = points[i][j];
+        }
+    }
+
+    // Insert the faces of the (nondegenerate) tetrahedron constructed by the
+    // call to GetInformation.
+    if (!info.extremeCCW)
+    {
+        std::swap(info.extreme[2], info.extreme[3]);
+    }
+
+    mHullUnordered.push_back(TriangleKey<true>(info.extreme[1],
+        info.extreme[2], info.extreme[3]));
+    mHullUnordered.push_back(TriangleKey<true>(info.extreme[0],
+        info.extreme[3], info.extreme[2]));
+    mHullUnordered.push_back(TriangleKey<true>(info.extreme[0],
+        info.extreme[1], info.extreme[3]));
+    mHullUnordered.push_back(TriangleKey<true>(info.extreme[0],
+        info.extreme[2], info.extreme[1]));
+
+    // Incrementally update the hull.  The set of processed points is
+    // maintained to eliminate duplicates, either in the original input
+    // points or in the points obtained by snap rounding.
+    std::set<Vector3<InputType>> processed;
+    for (i = 0; i < 4; ++i)
+    {
+        processed.insert(points[info.extreme[i]]);
+    }
+    for (i = 0; i < mNumPoints; ++i)
+    {
+        if (processed.find(points[i]) == processed.end())
+        {
+            Update(i);
+            processed.insert(points[i]);
+        }
+    }
+    mNumUniquePoints = static_cast<int>(processed.size());
+    return true;
+}
+
+template <typename InputType, typename ComputeType> inline
+InputType ConvexHull3<InputType, ComputeType>::GetEpsilon() const
+{
+    return mEpsilon;
+}
+
+template <typename InputType, typename ComputeType> inline
+int ConvexHull3<InputType, ComputeType>::GetDimension() const
+{
+    return mDimension;
+}
+
+template <typename InputType, typename ComputeType> inline
+Line3<InputType> const& ConvexHull3<InputType, ComputeType>::GetLine() const
+{
+    return mLine;
+}
+
+template <typename InputType, typename ComputeType> inline
+Plane3<InputType> const& ConvexHull3<InputType, ComputeType>::GetPlane() const
+{
+    return mPlane;
+}
+
+template <typename InputType, typename ComputeType> inline
+int ConvexHull3<InputType, ComputeType>::GetNumPoints() const
+{
+    return mNumPoints;
+}
+
+template <typename InputType, typename ComputeType> inline
+int ConvexHull3<InputType, ComputeType>::GetNumUniquePoints() const
+{
+    return mNumUniquePoints;
+}
+
+template <typename InputType, typename ComputeType> inline
+Vector3<InputType> const* ConvexHull3<InputType, ComputeType>::GetPoints() const
+{
+    return mPoints;
+}
+
+template <typename InputType, typename ComputeType> inline
+PrimalQuery3<ComputeType> const& ConvexHull3<InputType, ComputeType>::GetQuery() const
+{
+    return mQuery;
+}
+
+template <typename InputType, typename ComputeType> inline
+std::vector<TriangleKey<true>> const& ConvexHull3<InputType, ComputeType>::GetHullUnordered() const
+{
+    return mHullUnordered;
+}
+
+template <typename InputType, typename ComputeType>
+ETManifoldMesh const& ConvexHull3<InputType, ComputeType>::GetHullMesh() const
+{
+    // Create the mesh only on demand.
+    if (mHullMesh.GetTriangles().size() == 0)
+    {
+        for (auto const& tri : mHullUnordered)
+        {
+            mHullMesh.Insert(tri.V[0], tri.V[1], tri.V[2]);
+        }
+    }
+
+    return mHullMesh;
+}
+
+template <typename InputType, typename ComputeType>
+void ConvexHull3<InputType, ComputeType>::Update(int i)
+{
+    // The terminator that separates visible faces from nonvisible faces is
+    // constructed by this code.  Visible faces for the incoming hull are
+    // removed, and the boundary of that set of triangles is the terminator.
+    // New visible faces are added using the incoming point and the edges of
+    // the terminator.
+    //
+    // A simple algorithm for computing terminator edges is the following.
+    // Back-facing triangles are located and the three edges are processed.
+    // The first time an edge is visited, insert it into the terminator.  If
+    // it is visited a second time, the edge is removed because it is shared
+    // by another back-facing triangle and, therefore, cannot be a terminator
+    // edge.  After visiting all back-facing triangles, the only remaining
+    // edges in the map are the terminator edges.
+    //
+    // The order of vertices of an edge is important for adding new faces with
+    // the correct vertex winding.  However, the edge "toggle" (insert edge,
+    // remove edge) should use edges with unordered vertices, because the
+    // edge shared by one triangle has opposite ordering relative to that of
+    // the other triangle.  The map uses unordered edges as the keys but
+    // stores the ordered edge as the value.  This avoids having to look up
+    // an edge twice in a map with ordered edge keys.
+
+    unsigned int numFaces = static_cast<unsigned int>(mHullUnordered.size());
+    std::vector<int> queryResult(numFaces);
+    if (mNumThreads > 1 && numFaces >= mNumThreads)
+    {
+        // Partition the data for multiple threads.
+        unsigned int numFacesPerThread = numFaces / mNumThreads;
+        std::vector<unsigned int> jmin(mNumThreads), jmax(mNumThreads);
+        for (unsigned int t = 0; t < mNumThreads; ++t)
+        {
+            jmin[t] = t * numFacesPerThread;
+            jmax[t] = jmin[t] + numFacesPerThread - 1;
+        }
+        jmax[mNumThreads - 1] = numFaces - 1;
+
+        // Execute the point-plane queries in multiple threads.
+        std::vector<std::thread> process(mNumThreads);
+        for (unsigned int t = 0; t < mNumThreads; ++t)
+        {
+            process[t] = std::thread([this, i, t, &jmin, &jmax, 
+                &queryResult]()
+            {
+                for (unsigned int j = jmin[t]; j <= jmax[t]; ++j)
+                {
+                    TriangleKey<true> const& tri = mHullUnordered[j];
+                    queryResult[j] = mQuery.ToPlane(i, tri.V[0], tri.V[1], tri.V[2]);
+                }
+            });
+        }
+
+        // Wait for all threads to finish.
+        for (unsigned int t = 0; t < mNumThreads; ++t)
+        {
+            process[t].join();
+        }
+    }
+    else
+    {
+        unsigned int j = 0;
+        for (auto const& tri : mHullUnordered)
+        {
+            queryResult[j++] = mQuery.ToPlane(i, tri.V[0], tri.V[1], tri.V[2]);
+        }
+    }
+
+    std::map<EdgeKey<false>, std::pair<int, int>> terminator;
+    std::vector<TriangleKey<true>> backFaces;
+    bool existsFrontFacingTriangle = false;
+    unsigned int j = 0;
+    for (auto const& tri : mHullUnordered)
+    {
+        if (queryResult[j++] <= 0)
+        {
+            // The triangle is back facing.  These include triangles that
+            // are coplanar with the incoming point.
+            backFaces.push_back(tri);
+
+            // The current hull is a 2-manifold watertight mesh.  The
+            // terminator edges are those shared with a front-facing triangle.
+            // The first time an edge of a back-facing triangle is visited,
+            // insert it into the terminator.  If it is visited a second time,
+            // the edge is removed because it is shared by another back-facing
+            // triangle.  After all back-facing triangles are visited, the
+            // only remaining edges are shared by a single back-facing
+            // triangle, which makes them terminator edges.
+            for (int j0 = 2, j1 = 0; j1 < 3; j0 = j1++)
+            {
+                int v0 = tri.V[j0], v1 = tri.V[j1];
+                EdgeKey<false> edge(v0, v1);
+                auto iter = terminator.find(edge);
+                if (iter == terminator.end())
+                {
+                    // The edge is visited for the first time.
+                    terminator.insert(std::make_pair(edge, std::make_pair(v0, v1)));
+                }
+                else
+                {
+                    // The edge is visited for the second time.
+                    terminator.erase(edge);
+                }
+            }
+        }
+        else
+        {
+            // If there are no strictly front-facing triangles, then the
+            // incoming point is inside or on the convex hull.  If we get
+            // to this code, then the point is truly outside and we can
+            // update the hull.
+            existsFrontFacingTriangle = true;
+        }
+    }
+
+    if (!existsFrontFacingTriangle)
+    {
+        // The incoming point is inside or on the current hull, so no
+        // update of the hull is necessary.
+        return;
+    }
+
+    // The updated hull contains the triangles not visible to the incoming
+    // point.
+    mHullUnordered = backFaces;
+
+    // Insert the triangles formed by the incoming point and the terminator
+    // edges.
+    for (auto const& edge : terminator)
+    {
+        mHullUnordered.push_back(TriangleKey<true>(i, edge.second.second, edge.second.first));
+    }
+}
+
+
+}
diff --git a/Engine/Plugins/Experimental/GeometryProcessing/Source/GeometryAlgorithms/Private/GTEngine/Mathematics/GteConvexPolyhedron3.h b/Engine/Plugins/Experimental/GeometryProcessing/Source/GeometryAlgorithms/Private/GTEngine/Mathematics/GteConvexPolyhedron3.h
new file mode 100644
--- /dev/null
+++ b/Engine/Plugins/Experimental/GeometryProcessing/Source/GeometryAlgorithms/Private/GTEngine/Mathematics/GteConvexPolyhedron3.h
@@ -0,0 +1,113 @@
+// David Eberly, Geometric Tools, Redmond WA 98052
+// Copyright (c) 1998-2019
+// Distributed under the Boost Software License, Version 1.0.
+// http://www.boost.org/LICENSE_1_0.txt
+// http://www.geometrictools.com/License/Boost/LICENSE_1_0.txt
+// File Version: 3.5.0 (2016/11/25)
+
+#pragma once
+
+#include <GTEngine/Mathematics/GteVector4.h>
+#include <GTEngine/Mathematics/GteAlignedBox.h>
+#include <cstdint>
+#include <vector>
+
+namespace gte
+{
+
+template <typename Real>
+class ConvexPolyhedron3
+{
+public:
+    // Construction.  The convex polyhedra represented by this class has
+    // triangle faces that are counterclockwise ordered when viewed from
+    // outside the polyhedron.  No attempt is made to verify that the
+    // polyhedron is convex; the caller is responsible for enforcing this.
+    // The constructors move (not copy!) the input arrays.  The constructor
+    // succeeds when the number of vertices is at least 4 and the number of
+    // indices is at least 12.  If the constructor fails, no move occurs
+    // and the member arrays have no elements.
+    //
+    // To support geometric algorithms that are formulated using convex
+    // quadratic programming (such as computing the distance from a point to
+    // a convex polyhedron), it is necessary to know the planes of the faces
+    // and an axis-aligned bounding box.  If you want either the faces or the
+    // box, pass 'true' to the appropriate parameters.  When planes are
+    // generated, the normals are not created to be unit length in order to
+    // support queries using exact rational arithmetic.  If a normal to a
+    // face is N = (n0,n1,n2) and V is a vertex of the face, the plane is
+    // Dot(N,X-V) = 0 and is stored as Dot(n0,n1,n2,-Dot(N,V)).  The normals
+    // are computed to be outer pointing.
+    ConvexPolyhedron3();
+    ConvexPolyhedron3(std::vector<Vector3<Real>>&& inVertices, std::vector<int>&& inIndices,
+        bool wantPlanes, bool wantAlignedBox);
+
+    // If you modifty the vertices or indices and you want the new face
+    // planes or aligned box computed, call these functions.
+    void GeneratePlanes();
+    void GenerateAlignedBox();
+
+    std::vector<Vector3<Real>> vertices;
+    std::vector<int> indices;
+    std::vector<Vector4<Real>> planes;
+    AlignedBox3<Real> alignedBox;
+};
+
+
+template <typename Real>
+ConvexPolyhedron3<Real>::ConvexPolyhedron3()
+{
+}
+
+template <typename Real>
+ConvexPolyhedron3<Real>::ConvexPolyhedron3(std::vector<Vector3<Real>>&& inVertices,
+    std::vector<int>&& inIndices, bool wantPlanes, bool wantAlignedBox)
+{
+    if (inVertices.size() >= 4 && inIndices.size() >= 12)
+    {
+        vertices = std::move(inVertices);
+        indices = std::move(inIndices);
+
+        if (wantPlanes)
+        {
+            GeneratePlanes();
+        }
+
+        if (wantAlignedBox)
+        {
+            GenerateAlignedBox();
+        }
+    }
+}
+
+template <typename Real>
+void ConvexPolyhedron3<Real>::GeneratePlanes()
+{
+    if (vertices.size() > 0 && indices.size() > 0)
+    {
+        uint32_t const numTriangles = static_cast<uint32_t>(indices.size()) / 3;
+        planes.resize(numTriangles);
+        for (uint32_t t = 0, i = 0; t < numTriangles; ++t)
+        {
+            Vector3<Real> V0 = vertices[indices[i++]];
+            Vector3<Real> V1 = vertices[indices[i++]];
+            Vector3<Real> V2 = vertices[indices[i++]];
+            Vector3<Real> E1 = V1 - V0;
+            Vector3<Real> E2 = V2 - V0;
+            Vector3<Real> N = Cross(E1, E2);
+            planes[t] = HLift(N, -Dot(N, V0));
+        }
+    }
+}
+
+template <typename Real>
+void ConvexPolyhedron3<Real>::GenerateAlignedBox()
+{
+    if (vertices.size() > 0 && indices.size() > 0)
+    {
+        ComputeExtremes(static_cast<int>(vertices.size()), vertices.data(),
+            alignedBox.min, alignedBox.max);
+    }
+}
+
+}
diff --git a/Engine/Plugins/Experimental/GeometryProcessing/Source/GeometryAlgorithms/Private/GTEngine/Mathematics/GteCosEstimate.h b/Engine/Plugins/Experimental/GeometryProcessing/Source/GeometryAlgorithms/Private/GTEngine/Mathematics/GteCosEstimate.h
new file mode 100644
--- /dev/null
+++ b/Engine/Plugins/Experimental/GeometryProcessing/Source/GeometryAlgorithms/Private/GTEngine/Mathematics/GteCosEstimate.h
@@ -0,0 +1,164 @@
+// David Eberly, Geometric Tools, Redmond WA 98052
+// Copyright (c) 1998-2019
+// Distributed under the Boost Software License, Version 1.0.
+// http://www.boost.org/LICENSE_1_0.txt
+// http://www.geometrictools.com/License/Boost/LICENSE_1_0.txt
+// File Version: 3.0.1 (2018/10/05)
+
+#pragma once
+
+#include <GTEngine/Mathematics/GteMath.h>
+
+// Minimax polynomial approximations to cos(x).  The polynomial p(x) of
+// degree D has only even-power terms, is required to have constant term 1,
+// and p(pi/2) = cos(pi/2) = 0.  It minimizes the quantity
+// maximum{|cos(x) - p(x)| : x in [-pi/2,pi/2]} over all polynomials of
+// degree D subject to the constraints mentioned.
+
+namespace gte
+{
+
+template <typename Real>
+class CosEstimate
+{
+public:
+    // The input constraint is x in [-pi/2,pi/2].  For example,
+    //   float x; // in [-pi/2,pi/2]
+    //   float result = CosEstimate<float>::Degree<4>(x);
+    template <int D>
+    inline static Real Degree(Real x);
+
+    // The input x can be any real number.  Range reduction is used to
+    // generate a value y in [-pi/2,pi/2] and a sign s for which
+    // cos(y) = s*cos(x).  For example,
+    //   float x;  // x any real number
+    //   float result = CosEstimate<float>::DegreeRR<3>(x);
+    template <int D>
+    inline static Real DegreeRR(Real x);
+
+private:
+    // Metaprogramming and private implementation to allow specialization of
+    // a template member function.
+    template <int D> struct degree {};
+    inline static Real Evaluate(degree<2>, Real x);
+    inline static Real Evaluate(degree<4>, Real x);
+    inline static Real Evaluate(degree<6>, Real x);
+    inline static Real Evaluate(degree<8>, Real x);
+    inline static Real Evaluate(degree<10>, Real x);
+
+    // Support for range reduction.
+    inline static void Reduce(Real x, Real& y, Real& sign);
+};
+
+
+template <typename Real>
+template <int D>
+inline Real CosEstimate<Real>::Degree(Real x)
+{
+    return Evaluate(degree<D>(), x);
+}
+
+template <typename Real>
+template <int D>
+inline Real CosEstimate<Real>::DegreeRR(Real x)
+{
+    Real y, sign;
+    Reduce(x, y, sign);
+    Real poly = sign * Degree<D>(y);
+    return poly;
+}
+
+template <typename Real>
+inline Real CosEstimate<Real>::Evaluate(degree<2>, Real x)
+{
+    Real xsqr = x * x;
+    Real poly;
+    poly = (Real)GTE_C_COS_DEG2_C1;
+    poly = (Real)GTE_C_COS_DEG2_C0 + poly * xsqr;
+    return poly;
+}
+
+template <typename Real>
+inline Real CosEstimate<Real>::Evaluate(degree<4>, Real x)
+{
+    Real xsqr = x * x;
+    Real poly;
+    poly = (Real)GTE_C_COS_DEG4_C2;
+    poly = (Real)GTE_C_COS_DEG4_C1 + poly * xsqr;
+    poly = (Real)GTE_C_COS_DEG4_C0 + poly * xsqr;
+    return poly;
+}
+
+template <typename Real>
+inline Real CosEstimate<Real>::Evaluate(degree<6>, Real x)
+{
+    Real xsqr = x * x;
+    Real poly;
+    poly = (Real)GTE_C_COS_DEG6_C3;
+    poly = (Real)GTE_C_COS_DEG6_C2 + poly * xsqr;
+    poly = (Real)GTE_C_COS_DEG6_C1 + poly * xsqr;
+    poly = (Real)GTE_C_COS_DEG6_C0 + poly * xsqr;
+    return poly;
+}
+
+template <typename Real>
+inline Real CosEstimate<Real>::Evaluate(degree<8>, Real x)
+{
+    Real xsqr = x * x;
+    Real poly;
+    poly = (Real)GTE_C_COS_DEG8_C4;
+    poly = (Real)GTE_C_COS_DEG8_C3 + poly * xsqr;
+    poly = (Real)GTE_C_COS_DEG8_C2 + poly * xsqr;
+    poly = (Real)GTE_C_COS_DEG8_C1 + poly * xsqr;
+    poly = (Real)GTE_C_COS_DEG8_C0 + poly * xsqr;
+    return poly;
+}
+
+template <typename Real>
+inline Real CosEstimate<Real>::Evaluate(degree<10>, Real x)
+{
+    Real xsqr = x * x;
+    Real poly;
+    poly = (Real)GTE_C_COS_DEG10_C5;
+    poly = (Real)GTE_C_COS_DEG10_C4 + poly * xsqr;
+    poly = (Real)GTE_C_COS_DEG10_C3 + poly * xsqr;
+    poly = (Real)GTE_C_COS_DEG10_C2 + poly * xsqr;
+    poly = (Real)GTE_C_COS_DEG10_C1 + poly * xsqr;
+    poly = (Real)GTE_C_COS_DEG10_C0 + poly * xsqr;
+    return poly;
+}
+
+template <typename Real>
+inline void CosEstimate<Real>::Reduce(Real x, Real& y, Real& sign)
+{
+    // Map x to y in [-pi,pi], x = 2*pi*quotient + remainder.
+    Real quotient = (Real)GTE_C_INV_TWO_PI * x;
+    if (x >= (Real)0)
+    {
+        quotient = (Real)((int)(quotient + (Real)0.5));
+    }
+    else
+    {
+        quotient = (Real)((int)(quotient - (Real)0.5));
+    }
+    y = x - (Real)GTE_C_TWO_PI * quotient;
+
+    // Map y to [-pi/2,pi/2] with cos(y) = sign*cos(x).
+    if (y > (Real)GTE_C_HALF_PI)
+    {
+        y = (Real)GTE_C_PI - y;
+        sign = (Real)-1;
+    }
+    else if (y < (Real)-GTE_C_HALF_PI)
+    {
+        y = (Real)-GTE_C_PI - y;
+        sign = (Real)-1;
+    }
+    else
+    {
+        sign = (Real)1;
+    }
+}
+
+
+}
diff --git a/Engine/Plugins/Experimental/GeometryProcessing/Source/GeometryAlgorithms/Private/GTEngine/Mathematics/GteCubicRootsQR.h b/Engine/Plugins/Experimental/GeometryProcessing/Source/GeometryAlgorithms/Private/GTEngine/Mathematics/GteCubicRootsQR.h
new file mode 100644
--- /dev/null
+++ b/Engine/Plugins/Experimental/GeometryProcessing/Source/GeometryAlgorithms/Private/GTEngine/Mathematics/GteCubicRootsQR.h
@@ -0,0 +1,241 @@
+// David Eberly, Geometric Tools, Redmond WA 98052
+// Copyright (c) 1998-2019
+// Distributed under the Boost Software License, Version 1.0.
+// http://www.boost.org/LICENSE_1_0.txt
+// http://www.geometrictools.com/License/Boost/LICENSE_1_0.txt
+// File Version: 3.0.1 (2018/10/05)
+
+#pragma once
+
+#include <array>
+#include <cmath>
+#include <cstdint>
+
+// An implementation of the QR algorithm described in "Matrix Computations,
+// 2nd edition" by G. H. Golub and C. F. Van Loan, The Johns Hopkins
+// University Press, Baltimore MD, Fourth Printing 1993.  In particular,
+// the implementation is based on Chapter 7 (The Unsymmetric Eigenvalue
+// Problem), Section 7.5 (The Practical QR Algorithm).  The algorithm is
+// specialized for the companion matrix associated with a cubic polynomial.
+
+namespace gte
+{
+
+template <typename Real>
+class CubicRootsQR
+{
+public:
+    typedef std::array<std::array<Real, 3>, 3> Matrix;
+
+    // Solve p(x) = c0 + c1 * x + c2 * x^2 + x^3 = 0.
+    uint32_t operator() (uint32_t maxIterations, Real c0, Real c1, Real c2,
+        uint32_t& numRoots, std::array<Real, 3>& roots);
+
+    // Compute the real eigenvalues of the upper Hessenberg matrix A.  The
+    // matrix is modified by in-place operations, so if you need to remember
+    // A, you must make your own copy before calling this function.
+    uint32_t operator() (uint32_t maxIterations, Matrix& A,
+        uint32_t& numRoots, std::array<Real, 3>& roots);
+
+private:
+    void DoIteration(std::array<Real, 3> const& V, Matrix& A);
+
+    template <int N>
+    std::array<Real, N> House(std::array<Real, N> const& X);
+
+    template <int N>
+    void RowHouse(int rmin, int rmax, int cmin, int cmax,
+        std::array<Real, N> const& V, std::array<Real, N> const& MV, Matrix& A);
+
+    template <int N>
+    void ColHouse(int rmin, int rmax, int cmin, int cmax,
+        std::array<Real, N> const& V, std::array<Real, N> const& MV, Matrix& A);
+
+    void GetQuadraticRoots(int i0, int i1, Matrix const & A,
+        uint32_t& numRoots, std::array<Real, 3>& roots);
+};
+
+
+template <typename Real>
+uint32_t CubicRootsQR<Real>::operator() (uint32_t maxIterations, Real c0, Real c1, Real c2,
+    uint32_t& numRoots, std::array<Real, 3>& roots)
+{
+    // Create the companion matrix for the polynomial.  The matrix is in upper
+    // Hessenberg form.
+    Matrix A;
+    A[0][0] = (Real)0;
+    A[0][1] = (Real)0;
+    A[0][2] = -c0;
+    A[1][0] = (Real)1;
+    A[1][1] = (Real)0;
+    A[1][2] = -c1;
+    A[2][0] = (Real)0;
+    A[2][1] = (Real)1;
+    A[2][2] = -c2;
+
+    // Avoid the QR-cycle when c1 = c2 = 0 and avoid the slow convergence
+    // when c1 and c2 are nearly zero.
+    std::array<Real, 3> V{
+        (Real)1,
+        (Real)0.36602540378443865,
+        (Real)0.36602540378443865 };
+    DoIteration(V, A);
+
+    return operator()(maxIterations, A, numRoots, roots);
+}
+
+template <typename Real>
+uint32_t CubicRootsQR<Real>::operator() (uint32_t maxIterations, Matrix& A,
+    uint32_t& numRoots, std::array<Real, 3>& roots)
+{
+    numRoots = 0;
+    std::fill(roots.begin(), roots.end(), (Real)0);
+
+    for (uint32_t numIterations = 0; numIterations < maxIterations; ++numIterations)
+    {
+        // Apply a Francis QR iteration.
+        Real tr = A[1][1] + A[2][2];
+        Real det = A[1][1] * A[2][2] - A[1][2] * A[2][1];
+        std::array<Real, 3> X{
+            A[0][0] * A[0][0] + A[0][1] * A[1][0] - tr * A[0][0] + det,
+            A[1][0] * (A[0][0] + A[1][1] - tr),
+            A[1][0] * A[2][1] };
+        std::array<Real, 3> V = House<3>(X);
+        DoIteration(V, A);
+
+        // Test for uncoupling of A.
+        Real tr01 = A[0][0] + A[1][1];
+        if (tr01 + A[1][0] == tr01)
+        {
+            numRoots = 1;
+            roots[0] = A[0][0];
+            GetQuadraticRoots(1, 2, A, numRoots, roots);
+            return numIterations;
+        }
+
+        Real tr12 = A[1][1] + A[2][2];
+        if (tr12 + A[2][1] == tr12)
+        {
+            numRoots = 1;
+            roots[0] = A[2][2];
+            GetQuadraticRoots(0, 1, A, numRoots, roots);
+            return numIterations;
+        }
+    }
+    return maxIterations;
+}
+
+template <typename Real>
+void CubicRootsQR<Real>::DoIteration(std::array<Real, 3> const& V, Matrix& A)
+{
+    Real multV = ((Real)-2) / (V[0] * V[0] + V[1] * V[1] + V[2] * V[2]);
+    std::array<Real, 3> MV{ multV * V[0], multV * V[1], multV * V[2] };
+    RowHouse<3>(0, 2, 0, 2, V, MV, A);
+    ColHouse<3>(0, 2, 0, 2, V, MV, A);
+
+    std::array<Real, 2> Y{ A[1][0], A[2][0] };
+    std::array<Real, 2> W = House<2>(Y);
+    Real multW = ((Real)-2) / (W[0] * W[0] + W[1] * W[1]);
+    std::array<Real, 2> MW{ multW * W[0], multW * W[1] };
+    RowHouse<2>(1, 2, 0, 2, W, MW, A);
+    ColHouse<2>(0, 2, 1, 2, W, MW, A);
+}
+
+template <typename Real>
+template <int N>
+std::array<Real, N> CubicRootsQR<Real>::House(std::array<Real, N> const & X)
+{
+    std::array<Real, N> V;
+    Real length = (Real)0;
+    for (int i = 0; i < N; ++i)
+    {
+        length += X[i] * X[i];
+    }
+    length = std::sqrt(length);
+    if (length != (Real)0)
+    {
+        Real sign = (X[0] >= (Real)0 ? (Real)1 : (Real)-1);
+        Real denom = X[0] + sign * length;
+        for (int i = 1; i < N; ++i)
+        {
+            V[i] = X[i] / denom;
+        }
+    }
+    V[0] = (Real)1;
+    return V;
+}
+
+template <typename Real>
+template <int N>
+void CubicRootsQR<Real>::RowHouse(int rmin, int rmax, int cmin, int cmax,
+    std::array<Real, N> const& V, std::array<Real, N> const& MV, Matrix& A)
+{
+    std::array<Real, 3> W;  // only elements cmin through cmax are used
+
+    for (int c = cmin; c <= cmax; ++c)
+    {
+        W[c] = (Real)0;
+        for (int r = rmin, k = 0; r <= rmax; ++r, ++k)
+        {
+            W[c] += V[k] * A[r][c];
+        }
+    }
+
+    for (int r = rmin, k = 0; r <= rmax; ++r, ++k)
+    {
+        for (int c = cmin; c <= cmax; ++c)
+        {
+            A[r][c] += MV[k] * W[c];
+        }
+    }
+}
+
+template <typename Real>
+template <int N>
+void CubicRootsQR<Real>::ColHouse(int rmin, int rmax, int cmin, int cmax,
+    std::array<Real, N> const& V, std::array<Real, N> const& MV, Matrix& A)
+{
+    std::array<Real, 3> W;  // only elements rmin through rmax are used
+
+    for (int r = rmin; r <= rmax; ++r)
+    {
+        W[r] = (Real)0;
+        for (int c = cmin, k = 0; c <= cmax; ++c, ++k)
+        {
+            W[r] += V[k] * A[r][c];
+        }
+    }
+
+    for (int r = rmin; r <= rmax; ++r)
+    {
+        for (int c = cmin, k = 0; c <= cmax; ++c, ++k)
+        {
+            A[r][c] += W[r] * MV[k];
+        }
+    }
+}
+
+template <typename Real>
+void CubicRootsQR<Real>::GetQuadraticRoots(int i0, int i1, Matrix const& A,
+    uint32_t& numRoots, std::array<Real, 3>& roots)
+{
+    // Solve x^2 - t * x + d = 0, where t is the trace and d is the
+    // determinant of the 2x2 matrix defined by indices i0 and i1.  The
+    // discriminant is D = (t/2)^2 - d.  When D >= 0, the roots are real
+    // values t/2 - sqrt(D) and t/2 + sqrt(D).  To avoid potential numerical
+    // issues with subtractive cancellation, the roots are computed as
+    //   r0 = t/2 + sign(t/2)*sqrt(D), r1 = trace - r0.
+    Real trace = A[i0][i0] + A[i1][i1];
+    Real halfTrace = trace * (Real)0.5;
+    Real determinant = A[i0][i0] * A[i1][i1] - A[i0][i1] * A[i1][i0];
+    Real discriminant = halfTrace * halfTrace - determinant;
+    if (discriminant >= (Real)0)
+    {
+        Real sign = (trace >= (Real)0 ? (Real)1 : (Real)-1);
+        Real root = halfTrace + sign * std::sqrt(discriminant);
+        roots[numRoots++] = root;
+        roots[numRoots++] = trace - root;
+    }
+}
+
+}
diff --git a/Engine/Plugins/Experimental/GeometryProcessing/Source/GeometryAlgorithms/Private/GTEngine/Mathematics/GteCylinder3.h b/Engine/Plugins/Experimental/GeometryProcessing/Source/GeometryAlgorithms/Private/GTEngine/Mathematics/GteCylinder3.h
new file mode 100644
--- /dev/null
+++ b/Engine/Plugins/Experimental/GeometryProcessing/Source/GeometryAlgorithms/Private/GTEngine/Mathematics/GteCylinder3.h
@@ -0,0 +1,123 @@
+// David Eberly, Geometric Tools, Redmond WA 98052
+// Copyright (c) 1998-2019
+// Distributed under the Boost Software License, Version 1.0.
+// http://www.boost.org/LICENSE_1_0.txt
+// http://www.geometrictools.com/License/Boost/LICENSE_1_0.txt
+// File Version: 3.0.0 (2016/06/19)
+
+#pragma once
+
+#include <GTEngine/Mathematics/GteLine.h>
+
+// The cylinder axis is a line.  The origin of the cylinder is chosen to be
+// the line origin.  The cylinder wall is at a distance R units from the axis.
+// An infinite cylinder has infinite height.  A finite cylinder has center C
+// at the line origin and has a finite height H.  The segment for the finite
+// cylinder has endpoints C-(H/2)*D and C+(H/2)*D where D is a unit-length
+// direction of the line.
+
+namespace gte
+{
+
+template <typename Real>
+class Cylinder3
+{
+public:
+    // Construction and destruction.  The default constructor sets axis to
+    // (0,0,1), radius to 1, and height to 1.
+    Cylinder3();
+    Cylinder3(Line3<Real> const& inAxis, Real inRadius, Real inHeight);
+
+    Line3<Real> axis;
+    Real radius, height;
+
+public:
+    // Comparisons to support sorted containers.
+    bool operator==(Cylinder3 const& cylinder) const;
+    bool operator!=(Cylinder3 const& cylinder) const;
+    bool operator< (Cylinder3 const& cylinder) const;
+    bool operator<=(Cylinder3 const& cylinder) const;
+    bool operator> (Cylinder3 const& cylinder) const;
+    bool operator>=(Cylinder3 const& cylinder) const;
+};
+
+
+template <typename Real>
+Cylinder3<Real>::Cylinder3()
+    :
+    axis(Line3<Real>()),
+    radius((Real)1),
+    height((Real)1)
+{
+}
+
+template <typename Real>
+Cylinder3<Real>::Cylinder3(Line3<Real> const& inAxis, Real inRadius,
+    Real inHeight)
+    :
+    axis(inAxis),
+    radius(inRadius),
+    height(inHeight)
+{
+}
+
+template <typename Real>
+bool Cylinder3<Real>::operator==(Cylinder3 const& cylinder) const
+{
+    return axis == cylinder.axis
+        && radius == cylinder.radius
+        && height == cylinder.height;
+}
+
+template <typename Real>
+bool Cylinder3<Real>::operator!=(Cylinder3 const& cylinder) const
+{
+    return !operator==(cylinder);
+}
+
+template <typename Real>
+bool Cylinder3<Real>::operator<(Cylinder3 const& cylinder) const
+{
+    if (axis < cylinder.axis)
+    {
+        return true;
+    }
+
+    if (axis > cylinder.axis)
+    {
+        return false;
+    }
+
+    if (radius < cylinder.radius)
+    {
+        return true;
+    }
+
+    if (radius > cylinder.radius)
+    {
+        return false;
+    }
+
+    return height < cylinder.height;
+}
+
+template <typename Real>
+bool Cylinder3<Real>::operator<=(Cylinder3 const& cylinder) const
+{
+    return operator<(cylinder) || operator==(cylinder);
+}
+
+template <typename Real>
+bool Cylinder3<Real>::operator>(Cylinder3 const& cylinder) const
+{
+    return !operator<=(cylinder);
+}
+
+template <typename Real>
+bool Cylinder3<Real>::operator>=(Cylinder3 const& cylinder) const
+{
+    return !operator<(cylinder);
+}
+
+
+}
diff --git a/Engine/Plugins/Experimental/GeometryProcessing/Source/GeometryAlgorithms/Private/GTEngine/Mathematics/GteDCPQuery.h b/Engine/Plugins/Experimental/GeometryProcessing/Source/GeometryAlgorithms/Private/GTEngine/Mathematics/GteDCPQuery.h
new file mode 100644
--- /dev/null
+++ b/Engine/Plugins/Experimental/GeometryProcessing/Source/GeometryAlgorithms/Private/GTEngine/Mathematics/GteDCPQuery.h
@@ -0,0 +1,32 @@
+// David Eberly, Geometric Tools, Redmond WA 98052
+// Copyright (c) 1998-2019
+// Distributed under the Boost Software License, Version 1.0.
+// http://www.boost.org/LICENSE_1_0.txt
+// http://www.geometrictools.com/License/Boost/LICENSE_1_0.txt
+// File Version: 3.0.1 (2018/10/05)
+
+#pragma once
+
+#include <GTEngine/Mathematics/GteMath.h>
+
+namespace gte
+{
+
+// Distance and closest-point queries.
+template <typename Real, typename Type0, typename Type1>
+class DCPQuery
+{
+public:
+    struct Result
+    {
+        // A DCPQuery-base class B must define a B::Result struct with member
+        // 'Real distance'.  A DCPQuery-derived class D must also derive a
+        // D::Result from B:Result but may have no members.  The idea is to
+        // allow Result to store closest-point information in addition to the
+        // distance.  The operator() is non-const to allow DCPQuery to store
+        // and modify private state that supports the query.
+    };
+    Result operator()(Type0 const& primitive0, Type1 const& primitive1);
+};
+
+}
diff --git a/Engine/Plugins/Experimental/GeometryProcessing/Source/GeometryAlgorithms/Private/GTEngine/Mathematics/GteDarbouxFrame.h b/Engine/Plugins/Experimental/GeometryProcessing/Source/GeometryAlgorithms/Private/GTEngine/Mathematics/GteDarbouxFrame.h
new file mode 100644
--- /dev/null
+++ b/Engine/Plugins/Experimental/GeometryProcessing/Source/GeometryAlgorithms/Private/GTEngine/Mathematics/GteDarbouxFrame.h
@@ -0,0 +1,160 @@
+// David Eberly, Geometric Tools, Redmond WA 98052
+// Copyright (c) 1998-2019
+// Distributed under the Boost Software License, Version 1.0.
+// http://www.boost.org/LICENSE_1_0.txt
+// http://www.geometrictools.com/License/Boost/LICENSE_1_0.txt
+// File Version: 3.0.2 (2018/10/05)
+
+#pragma once
+
+#include <GTEngine/Mathematics/GteMatrix2x2.h>
+#include <GTEngine/Mathematics/GteVector3.h>
+#include <GTEngine/Mathematics/GteParametricSurface.h>
+#include <memory>
+
+namespace gte
+{
+template <typename Real>
+class DarbouxFrame3
+{
+public:
+    // Construction.  The curve must persist as long as the DarbouxFrame3
+    // object does.
+    DarbouxFrame3(std::shared_ptr<ParametricSurface<3, Real>> const& surface);
+
+    // Get a coordinate frame, {T0, T1, N}.  At a nondegenerate surface
+    // points, dX/du and dX/dv are linearly independent tangent vectors.
+    // The frame is constructed as
+    //   T0 = (dX/du)/|dX/du|
+    //   N  = Cross(dX/du,dX/dv)/|Cross(dX/du,dX/dv)|
+    //   T1 = Cross(N, T0)
+    // so that {T0, T1, N} is a right-handed orthonormal set.
+    void operator()(Real u, Real v, Vector3<Real>& position,
+        Vector3<Real>& tangent0, Vector3<Real>& tangent1, Vector3<Real>& normal) const;
+
+    // Compute the principal curvatures and principal directions.
+    void GetPrincipalInformation(Real u, Real v, Real& curvature0, Real& curvature1,
+        Vector3<Real>& direction0, Vector3<Real>& direction1) const;
+
+private:
+    std::shared_ptr<ParametricSurface<3, Real>> mSurface;
+};
+
+
+template <typename Real>
+DarbouxFrame3<Real>::DarbouxFrame3(
+    std::shared_ptr<ParametricSurface<3, Real>> const& surface)
+    :
+    mSurface(surface)
+{
+}
+
+template <typename Real>
+void DarbouxFrame3<Real>::operator()(Real u, Real v, Vector3<Real>& position,
+    Vector3<Real>& tangent0, Vector3<Real>& tangent1, Vector3<Real>& normal) const
+{
+    Vector3<Real> values[6];
+    mSurface->Evaluate(u, v, 1, values);
+    position = values[0];
+    tangent0 = values[1];
+    Normalize(tangent0);
+    tangent1 = values[2];
+    Normalize(tangent1);
+    normal = UnitCross(tangent0, tangent1);
+    tangent1 = Cross(normal, tangent0);
+}
+
+template <typename Real>
+void DarbouxFrame3<Real>::GetPrincipalInformation(Real u, Real v, Real& curvature0,
+    Real& curvature1, Vector3<Real>& direction0, Vector3<Real>& direction1) const
+{
+    // Tangents:  T0 = (x_u,y_u,z_u), T1 = (x_v,y_v,z_v)
+    // Normal:    N = Cross(T0,T1)/Length(Cross(T0,T1))
+    // Metric Tensor:    G = +-                      -+
+    //                       | Dot(T0,T0)  Dot(T0,T1) |
+    //                       | Dot(T1,T0)  Dot(T1,T1) |
+    //                       +-                      -+
+    //
+    // Curvature Tensor:  B = +-                          -+
+    //                        | -Dot(N,T0_u)  -Dot(N,T0_v) |
+    //                        | -Dot(N,T1_u)  -Dot(N,T1_v) |
+    //                        +-                          -+
+    //
+    // Principal curvatures k are the generalized eigenvalues of
+    //
+    //     Bw = kGw
+    //
+    // If k is a curvature and w=(a,b) is the corresponding solution to
+    // Bw = kGw, then the principal direction as a 3D vector is d = a*U+b*V.
+    //
+    // Let k1 and k2 be the principal curvatures.  The mean curvature
+    // is (k1+k2)/2 and the Gaussian curvature is k1*k2.
+
+    // Compute derivatives.
+    Vector3<Real> values[6];
+    mSurface->Evaluate(u, v, 2, values);
+    Vector3<Real> derU = values[1];
+    Vector3<Real> derV = values[2];
+    Vector3<Real> derUU = values[3];
+    Vector3<Real> derUV = values[4];
+    Vector3<Real> derVV = values[5];
+
+    // Compute the metric tensor.
+    Matrix2x2<Real> metricTensor;
+    metricTensor(0, 0) = Dot(values[1], values[1]);
+    metricTensor(0, 1) = Dot(values[1], values[2]);
+    metricTensor(1, 0) = metricTensor(0, 1);
+    metricTensor(1, 1) = Dot(values[2], values[2]);
+
+    // Compute the curvature tensor.
+    Vector3<Real> normal = UnitCross(values[1], values[2]);
+    Matrix2x2<Real> curvatureTensor;
+    curvatureTensor(0, 0) = -Dot(normal, derUU);
+    curvatureTensor(0, 1) = -Dot(normal, derUV);
+    curvatureTensor(1, 0) = curvatureTensor(0, 1);
+    curvatureTensor(1, 1) = -Dot(normal, derVV);
+
+    // Characteristic polynomial is 0 = det(B-kG) = c2*k^2+c1*k+c0.
+    Real c0 = Determinant(curvatureTensor);
+    Real c1 = ((Real)2) * curvatureTensor(0, 1) * metricTensor(0, 1) -
+        curvatureTensor(0, 0) * metricTensor(1, 1) -
+        curvatureTensor(1, 1) * metricTensor(0, 0);
+    Real c2 = Determinant(metricTensor);
+
+    // Principal curvatures are roots of characteristic polynomial.
+    Real temp = std::sqrt(std::abs(c1 * c1 - ((Real)4) * c0 * c2));
+    Real mult = ((Real)0.5) / c2;
+    curvature0 = -mult * (c1 + temp);
+    curvature1 = -mult * (c1 - temp);
+
+    // Principal directions are solutions to (B-kG)w = 0,
+    // w1 = (b12-k1*g12,-(b11-k1*g11)) OR (b22-k1*g22,-(b12-k1*g12)).
+    Real a0 = curvatureTensor(0, 1) - curvature0 * metricTensor(0, 1);
+    Real a1 = curvature0 * metricTensor(0, 0) - curvatureTensor(0, 0);
+    Real length = std::sqrt(a0 * a0 + a1 * a1);
+    if (length > (Real)0)
+    {
+        direction0 = a0 * derU + a1 * derV;
+    }
+    else
+    {
+        a0 = curvatureTensor(1, 1) - curvature0 * metricTensor(1, 1);
+        a1 = curvature0 * metricTensor(0, 1) - curvatureTensor(0, 1);
+        length = std::sqrt(a0 * a0 + a1 * a1);
+        if (length > (Real)0)
+        {
+            direction0 = a0*derU + a1*derV;
+        }
+        else
+        {
+            // Umbilic (surface is locally sphere, any direction principal).
+            direction0 = derU;
+        }
+    }
+    Normalize(direction0);
+
+    // Second tangent is cross product of first tangent and normal.
+    direction1 = Cross(direction0, normal);
+}
+
+}
diff --git a/Engine/Plugins/Experimental/GeometryProcessing/Source/GeometryAlgorithms/Private/GTEngine/Mathematics/GteDelaunay2.h b/Engine/Plugins/Experimental/GeometryProcessing/Source/GeometryAlgorithms/Private/GTEngine/Mathematics/GteDelaunay2.h
new file mode 100644
--- /dev/null
+++ b/Engine/Plugins/Experimental/GeometryProcessing/Source/GeometryAlgorithms/Private/GTEngine/Mathematics/GteDelaunay2.h
@@ -0,0 +1,862 @@
+// David Eberly, Geometric Tools, Redmond WA 98052
+// Copyright (c) 1998-2019
+// Distributed under the Boost Software License, Version 1.0.
+// http://www.boost.org/LICENSE_1_0.txt
+// http://www.geometrictools.com/License/Boost/LICENSE_1_0.txt
+// File Version: 3.0.0 (2016/06/19)
+
+#pragma once
+
+#include <GTEngine/LowLevel/GteLogger.h>
+#include <GTEngine/Mathematics/GteETManifoldMesh.h>
+#include <GTEngine/Mathematics/GtePrimalQuery2.h>
+#include <GTEngine/Mathematics/GteLine.h>
+#include <vector>
+
+// Delaunay triangulation of points (intrinsic dimensionality 2).
+//   VQ = number of vertices
+//   V  = array of vertices
+//   TQ = number of triangles
+//   I  = Array of 3-tuples of indices into V that represent the triangles
+//        (3*TQ total elements).  Access via GetIndices(*).
+//   A  = Array of 3-tuples of indices into I that represent the adjacent
+//        triangles (3*TQ total elements).  Access via GetAdjacencies(*).
+// The i-th triangle has vertices
+//   vertex[0] = V[I[3*i+0]]
+//   vertex[1] = V[I[3*i+1]]
+//   vertex[2] = V[I[3*i+2]]
+// and edge index pairs
+//   edge[0] = <I[3*i+0],I[3*i+1]>
+//   edge[1] = <I[3*i+1],I[3*i+2]>
+//   edge[2] = <I[3*i+2],I[3*i+0]>
+// The triangles adjacent to these edges have indices
+//   adjacent[0] = A[3*i+0] is the triangle sharing edge[0]
+//   adjacent[1] = A[3*i+1] is the triangle sharing edge[1]
+//   adjacent[2] = A[3*i+2] is the triangle sharing edge[2]
+// If there is no adjacent triangle, the A[*] value is set to -1.  The
+// triangle adjacent to edge[j] has vertices
+//   adjvertex[0] = V[I[3*adjacent[j]+0]]
+//   adjvertex[1] = V[I[3*adjacent[j]+1]]
+//   adjvertex[2] = V[I[3*adjacent[j]+2]]
+// The only way to ensure a correct result for the input vertices (assumed to
+// be exact) is to choose ComputeType for exact rational arithmetic.  You may
+// use BSNumber.  No divisions are performed in this computation, so you do
+// not have to use BSRational.
+//
+// The worst-case choices of N for Real of type BSNumber or BSRational with
+// integer storage UIntegerFP32<N> are listed in the next table.  The numerical
+// computations are encapsulated in PrimalQuery2<Real>::ToLine and
+// PrimalQuery2<Real>::ToCircumcircle, the latter query the dominant one in
+// determining N.  We recommend using only BSNumber, because no divisions are
+// performed in the convex-hull computations.
+//
+//    input type | compute type | N
+//    -----------+--------------+------
+//    float      | BSNumber     |    35
+//    double     | BSNumber     |   263
+//    float      | BSRational   | 12302
+//    double     | BSRational   | 92827
+
+namespace gte
+{
+
+template <typename InputType, typename ComputeType>
+class Delaunay2
+{
+public:
+    // The class is a functor to support computing the Delaunay triangulation
+    // of multiple data sets using the same class object.
+    virtual ~Delaunay2();
+    Delaunay2();
+
+    // The input is the array of vertices whose Delaunay triangulation is
+    // required.  The epsilon value is used to determine the intrinsic
+    // dimensionality of the vertices (d = 0, 1, or 2).  When epsilon is
+    // positive, the determination is fuzzy--vertices approximately the same
+    // point, approximately on a line, or planar.  The return value is 'true'
+    // if and only if the hull construction is successful.
+    bool operator()(int numVertices, Vector2<InputType> const* vertices, InputType epsilon);
+
+    // Dimensional information.  If GetDimension() returns 1, the points lie
+    // on a line P+t*D (fuzzy comparison when epsilon > 0).  You can sort
+    // these if you need a polyline output by projecting onto the line each
+    // vertex X = P+t*D, where t = Dot(D,X-P).
+    inline InputType GetEpsilon() const;
+    inline int GetDimension() const;
+    inline Line2<InputType> const& GetLine() const;
+
+    // Member access.
+    inline int GetNumVertices() const;
+    inline int GetNumUniqueVertices() const;
+    inline int GetNumTriangles() const;
+    inline Vector2<InputType> const* GetVertices() const;
+    inline PrimalQuery2<ComputeType> const& GetQuery() const;
+    inline ETManifoldMesh const& GetGraph() const;
+    inline std::vector<int> const& GetIndices() const;
+    inline std::vector<int> const& GetAdjacencies() const;
+
+    // If 'vertices' has no duplicates, GetDuplicates()[i] = i for all i.
+    // If vertices[i] is the first occurrence of a vertex and if vertices[j]
+    // is found later, then GetDuplicates()[j] = i.
+    inline std::vector<int> const& GetDuplicates() const;
+
+    // Locate those triangle edges that do not share other triangles.  The
+    // returned array has hull.size() = 2*numEdges, each pair representing an
+    // edge.  The edges are not ordered, but the pair of vertices for an edge
+    // is ordered so that they conform to a counterclockwise traversal of the
+    // hull.  The return value is 'true' iff the dimension is 2.
+    bool GetHull(std::vector<int>& hull) const;
+
+    // Get the vertex indices for triangle i.  The function returns 'true'
+    // when the dimension is 2 and i is a valid triangle index, in which case
+    // the vertices are valid; otherwise, the function returns 'false' and the
+    // vertices are invalid.
+    bool GetIndices(int i, std::array<int, 3>& indices) const;
+
+    // Get the indices for triangles adjacent to triangle i.  The function
+    // returns 'true' when the dimension is 2 and if i is a valid triangle
+    // index, in which case the adjacencies are valid; otherwise, the function
+    // returns 'false' and the adjacencies are invalid.
+    bool GetAdjacencies(int i, std::array<int, 3>& adjacencies) const;
+
+    // Support for searching the triangulation for a triangle that contains
+    // a point.  If there is a containing triangle, the returned value is a
+    // triangle index i with 0 <= i < GetNumTriangles().  If there is not a
+    // containing triangle, -1 is returned.  The computations are performed
+    // using exact rational arithmetic.
+    //
+    // The SearchInfo input stores information about the triangle search when
+    // looking for the triangle (if any) that contains p.  The first triangle
+    // searched is 'initialTriangle'.  On return 'path' stores those (ordered)
+    // triangle indices visited during the search.  The last visited triangle
+    // has index 'finalTriangle and vertex indices 'finalV[0,1,2]', stored in
+    // counterclockwise order.  The last edge of the search is
+    // <finalV[0],finalV[1]>.  For spatially coherent inputs p for numerous
+    // calls to this function, you will want to specify 'finalTriangle' from
+    // the previous call as 'initialTriangle' for the next call, which should
+    // reduce search times.
+    struct SearchInfo
+    {
+        int initialTriangle;
+        int numPath;
+        std::vector<int> path;
+        int finalTriangle;
+        std::array<int, 3> finalV;
+    };
+    int GetContainingTriangle(Vector2<InputType> const& p, SearchInfo& info) const;
+
+protected:
+    // Support for incremental Delaunay triangulation.
+    typedef ETManifoldMesh::Triangle Triangle;
+    bool GetContainingTriangle(int i, std::shared_ptr<Triangle>& tri) const;
+    bool GetAndRemoveInsertionPolygon(int i, std::set<std::shared_ptr<Triangle>>& candidates,
+        std::set<EdgeKey<true>>& boundary);
+    bool Update(int i);
+
+    // The epsilon value is used for fuzzy determination of intrinsic
+    // dimensionality.  If the dimension is 0 or 1, the constructor returns
+    // early.  The caller is responsible for retrieving the dimension and
+    // taking an alternate path should the dimension be smaller than 2.
+    // If the dimension is 0, the caller may as well treat all vertices[]
+    // as a single point, say, vertices[0].  If the dimension is 1, the
+    // caller can query for the approximating line and project vertices[]
+    // onto it for further processing.
+    InputType mEpsilon;
+    int mDimension;
+    Line2<InputType> mLine;
+
+    // The array of vertices used for geometric queries.  If you want to be
+    // certain of a correct result, choose ComputeType to be BSNumber.
+    std::vector<Vector2<ComputeType>> mComputeVertices;
+    PrimalQuery2<ComputeType> mQuery;
+
+    // The graph information.
+    int mNumVertices;
+    int mNumUniqueVertices;
+    int mNumTriangles;
+    Vector2<InputType> const* mVertices;
+    ETManifoldMesh mGraph;
+    std::vector<int> mIndices;
+    std::vector<int> mAdjacencies;
+
+    // If a vertex occurs multiple times in the 'vertices' input to the
+    // constructor, the first processed occurrence of that vertex has an
+    // index stored in this array.  If there are no duplicates, then
+    // mDuplicates[i] = i for all i.
+
+    struct ProcessedVertex
+    {
+        ProcessedVertex();
+        ProcessedVertex(Vector2<InputType> const& inVertex, int inLocation);
+        bool operator<(ProcessedVertex const& v) const;
+
+        Vector2<InputType> vertex;
+        int location;
+    };
+
+    std::vector<int> mDuplicates;
+
+    // Indexing for the vertices of the triangle adjacent to a vertex.  The
+    // edge adjacent to vertex j is <mIndex[j][0], mIndex[j][1]> and is listed
+    // so that the triangle interior is to your left as you walk around the
+    // edges.  TODO: Use the "opposite edge" to be consistent with that of
+    // TetrahedronKey.  The "opposite" idea extends easily to higher
+    // dimensions.
+    std::array<std::array<int, 2>, 3> mIndex;
+};
+
+
+template <typename InputType, typename ComputeType>
+Delaunay2<InputType, ComputeType>::~Delaunay2()
+{
+}
+
+template <typename InputType, typename ComputeType>
+Delaunay2<InputType, ComputeType>::Delaunay2()
+    :
+    mEpsilon((InputType)0),
+    mDimension(0),
+    mLine(Vector2<InputType>::Zero(), Vector2<InputType>::Zero()),
+    mNumVertices(0),
+    mNumUniqueVertices(0),
+    mNumTriangles(0),
+    mVertices(nullptr)
+{
+    // INVESTIGATE.  If the initialization of mIndex is placed in the
+    // constructor initializer list, MSVS 2012 generates an internal
+    // compiler error.
+    mIndex = { { { 0, 1 }, { 1, 2 }, { 2, 0 } } };
+}
+
+template <typename InputType, typename ComputeType>
+bool Delaunay2<InputType, ComputeType>::operator()(int numVertices,
+    Vector2<InputType> const* vertices, InputType epsilon)
+{
+    mEpsilon = std::max(epsilon, (InputType)0);
+    mDimension = 0;
+    mLine.origin = Vector2<InputType>::Zero();
+    mLine.direction = Vector2<InputType>::Zero();
+    mNumVertices = numVertices;
+    mNumUniqueVertices = 0;
+    mNumTriangles = 0;
+    mVertices = vertices;
+    mGraph.Clear();
+    mIndices.clear();
+    mAdjacencies.clear();
+    mDuplicates.resize(std::max(numVertices, 3));
+
+    int i, j;
+    if (mNumVertices < 3)
+    {
+        // Delaunay2 should be called with at least three points.
+        return false;
+    }
+
+    IntrinsicsVector2<InputType> info(mNumVertices, vertices, mEpsilon);
+    if (info.dimension == 0)
+    {
+        // mDimension is 0; mGraph, mIndices, and mAdjacencies are empty
+        return false;
+    }
+
+    if (info.dimension == 1)
+    {
+        // The set is (nearly) collinear.
+        mDimension = 1;
+        mLine = Line2<InputType>(info.origin, info.direction[0]);
+        return false;
+    }
+
+    mDimension = 2;
+
+    // Compute the vertices for the queries.
+    mComputeVertices.resize(mNumVertices);
+    mQuery.Set(mNumVertices, &mComputeVertices[0]);
+    for (i = 0; i < mNumVertices; ++i)
+    {
+        for (j = 0; j < 2; ++j)
+        {
+            mComputeVertices[i][j] = vertices[i][j];
+        }
+    }
+
+    // Insert the (nondegenerate) triangle constructed by the call to
+    // GetInformation.  This is necessary for the circumcircle-visibility
+    // algorithm to work correctly.
+    if (!info.extremeCCW)
+    {
+        std::swap(info.extreme[1], info.extreme[2]);
+    }
+    if (!mGraph.Insert(info.extreme[0], info.extreme[1], info.extreme[2]))
+    {
+        return false;
+    }
+
+    // Incrementally update the triangulation.  The set of processed points
+    // is maintained to eliminate duplicates, either in the original input
+    // points or in the points obtained by snap rounding.
+    std::set<ProcessedVertex> processed;
+    for (i = 0; i < 3; ++i)
+    {
+        j = info.extreme[i];
+        processed.insert(ProcessedVertex(vertices[j], j));
+        mDuplicates[j] = j;
+    }
+    for (i = 0; i < mNumVertices; ++i)
+    {
+        ProcessedVertex v(vertices[i], i);
+        auto iter = processed.find(v);
+        if (iter == processed.end())
+        {
+            if (!Update(i))
+            {
+                // A failure can occur if ComputeType is not an exact
+                // arithmetic type.
+                return false;
+            }
+            processed.insert(v);
+            mDuplicates[i] = i;
+        }
+        else
+        {
+            mDuplicates[i] = iter->location;
+        }
+    }
+    mNumUniqueVertices = static_cast<int>(processed.size());
+
+    // Assign integer values to the triangles for use by the caller.
+    std::map<std::shared_ptr<Triangle>, int> permute;
+    i = -1;
+    permute[nullptr] = i++;
+    for (auto const& element : mGraph.GetTriangles())
+    {
+        permute[element.second] = i++;
+    }
+
+    // Put Delaunay triangles into an array (vertices and adjacency info).
+    mNumTriangles = static_cast<int>(mGraph.GetTriangles().size());
+    int numindices = 3 * mNumTriangles;
+    if (numindices > 0)
+    {
+        mIndices.resize(numindices);
+        mAdjacencies.resize(numindices);
+        i = 0;
+        for (auto const& element : mGraph.GetTriangles())
+        {
+            std::shared_ptr<Triangle> tri = element.second;
+            for (j = 0; j < 3; ++j, ++i)
+            {
+                mIndices[i] = tri->V[j];
+                mAdjacencies[i] = permute[tri->T[j].lock()];
+            }
+        }
+    }
+
+    return true;
+}
+
+template <typename InputType, typename ComputeType> inline
+InputType Delaunay2<InputType, ComputeType>::GetEpsilon() const
+{
+    return mEpsilon;
+}
+
+template <typename InputType, typename ComputeType> inline
+int Delaunay2<InputType, ComputeType>::GetDimension() const
+{
+    return mDimension;
+}
+
+template <typename InputType, typename ComputeType> inline
+Line2<InputType> const& Delaunay2<InputType, ComputeType>::GetLine() const
+{
+    return mLine;
+}
+
+template <typename InputType, typename ComputeType> inline
+int Delaunay2<InputType, ComputeType>::GetNumVertices() const
+{
+    return mNumVertices;
+}
+
+template <typename InputType, typename ComputeType> inline
+int Delaunay2<InputType, ComputeType>::GetNumUniqueVertices() const
+{
+    return mNumUniqueVertices;
+}
+
+template <typename InputType, typename ComputeType> inline
+int Delaunay2<InputType, ComputeType>::GetNumTriangles() const
+{
+    return mNumTriangles;
+}
+
+template <typename InputType, typename ComputeType> inline
+Vector2<InputType> const* Delaunay2<InputType, ComputeType>::GetVertices() const
+{
+    return mVertices;
+}
+
+template <typename InputType, typename ComputeType> inline
+PrimalQuery2<ComputeType> const& Delaunay2<InputType, ComputeType>::GetQuery() const
+{
+    return mQuery;
+}
+
+template <typename InputType, typename ComputeType> inline
+ETManifoldMesh const& Delaunay2<InputType, ComputeType>::GetGraph() const
+{
+    return mGraph;
+}
+
+template <typename InputType, typename ComputeType> inline
+std::vector<int> const& Delaunay2<InputType, ComputeType>::GetIndices() const
+{
+    return mIndices;
+}
+
+template <typename InputType, typename ComputeType> inline
+std::vector<int> const& Delaunay2<InputType, ComputeType>::GetAdjacencies() const
+{
+    return mAdjacencies;
+}
+
+template <typename InputType, typename ComputeType> inline
+std::vector<int> const& Delaunay2<InputType, ComputeType>::GetDuplicates() const
+{
+    return mDuplicates;
+}
+
+template <typename InputType, typename ComputeType>
+bool Delaunay2<InputType, ComputeType>::GetHull(std::vector<int>& hull) const
+{
+    if (mDimension == 2)
+    {
+        // Count the number of edges that are not shared by two triangles.
+        int numEdges = 0;
+        for (auto adj : mAdjacencies)
+        {
+            if (adj == -1)
+            {
+                ++numEdges;
+            }
+        }
+
+        if (numEdges > 0)
+        {
+            // Enumerate the edges.
+            hull.resize(2 * numEdges);
+            int current = 0, i = 0;
+            for (auto adj : mAdjacencies)
+            {
+                if (adj == -1)
+                {
+                    int tri = i / 3, j = i % 3;
+                    hull[current++] = mIndices[3 * tri + j];
+                    hull[current++] = mIndices[3 * tri + ((j + 1) % 3)];
+                }
+                ++i;
+            }
+            return true;
+        }
+        else
+        {
+            LogError("Unexpected.  There must be at least one triangle.");
+        }
+    }
+    else
+    {
+        LogError("The dimension must be 2.");
+    }
+    return false;
+}
+
+template <typename InputType, typename ComputeType>
+bool Delaunay2<InputType, ComputeType>::GetIndices(int i, std::array<int, 3>& indices) const
+{
+    if (mDimension == 2)
+    {
+        int numTriangles = static_cast<int>(mIndices.size() / 3);
+        if (0 <= i && i < numTriangles)
+        {
+            indices[0] = mIndices[3 * i];
+            indices[1] = mIndices[3 * i + 1];
+            indices[2] = mIndices[3 * i + 2];
+            return true;
+        }
+    }
+    else
+    {
+        LogError("The dimension must be 2.");
+    }
+    return false;
+}
+
+template <typename InputType, typename ComputeType>
+bool Delaunay2<InputType, ComputeType>::GetAdjacencies(int i, std::array<int, 3>& adjacencies) const
+{
+    if (mDimension == 2)
+    {
+        int numTriangles = static_cast<int>(mIndices.size() / 3);
+        if (0 <= i && i < numTriangles)
+        {
+            adjacencies[0] = mAdjacencies[3 * i];
+            adjacencies[1] = mAdjacencies[3 * i + 1];
+            adjacencies[2] = mAdjacencies[3 * i + 2];
+            return true;
+        }
+    }
+    else
+    {
+        LogError("The dimension must be 2.");
+    }
+    return false;
+}
+
+template <typename InputType, typename ComputeType>
+int Delaunay2<InputType, ComputeType>::GetContainingTriangle(
+    Vector2<InputType> const& p, SearchInfo& info) const
+{
+    if (mDimension == 2)
+    {
+        Vector2<ComputeType> test{ p[0], p[1] };
+
+        int numTriangles = static_cast<int>(mIndices.size() / 3);
+        info.path.resize(numTriangles);
+        info.numPath = 0;
+        int triangle;
+        if (0 <= info.initialTriangle && info.initialTriangle < numTriangles)
+        {
+            triangle = info.initialTriangle;
+        }
+        else
+        {
+            info.initialTriangle = 0;
+            triangle = 0;
+        }
+
+        // Use triangle edges as binary separating lines.
+        for (int i = 0; i < numTriangles; ++i)
+        {
+            int ibase = 3 * triangle;
+            int const* v = &mIndices[ibase];
+
+            info.path[info.numPath++] = triangle;
+            info.finalTriangle = triangle;
+            info.finalV[0] = v[0];
+            info.finalV[1] = v[1];
+            info.finalV[2] = v[2];
+
+            if (mQuery.ToLine(test, v[0], v[1]) > 0)
+            {
+                triangle = mAdjacencies[ibase];
+                if (triangle == -1)
+                {
+                    info.finalV[0] = v[0];
+                    info.finalV[1] = v[1];
+                    info.finalV[2] = v[2];
+                    return -1;
+                }
+                continue;
+            }
+
+            if (mQuery.ToLine(test, v[1], v[2]) > 0)
+            {
+                triangle = mAdjacencies[ibase + 1];
+                if (triangle == -1)
+                {
+                    info.finalV[0] = v[1];
+                    info.finalV[1] = v[2];
+                    info.finalV[2] = v[0];
+                    return -1;
+                }
+                continue;
+            }
+
+            if (mQuery.ToLine(test, v[2], v[0]) > 0)
+            {
+                triangle = mAdjacencies[ibase + 2];
+                if (triangle == -1)
+                {
+                    info.finalV[0] = v[2];
+                    info.finalV[1] = v[0];
+                    info.finalV[2] = v[1];
+                    return -1;
+                }
+                continue;
+            }
+
+            return triangle;
+        }
+    }
+    else
+    {
+        LogError("The dimension must be 2.");
+    }
+    return -1;
+}
+
+template <typename InputType, typename ComputeType>
+bool Delaunay2<InputType, ComputeType>::GetContainingTriangle(int i, std::shared_ptr<Triangle>& tri) const
+{
+    int numTriangles = static_cast<int>(mGraph.GetTriangles().size());
+    for (int t = 0; t < numTriangles; ++t)
+    {
+        int j;
+        for (j = 0; j < 3; ++j)
+        {
+            int v0 = tri->V[mIndex[j][0]];
+            int v1 = tri->V[mIndex[j][1]];
+            if (mQuery.ToLine(i, v0, v1) > 0)
+            {
+                // Point i sees edge <v0,v1> from outside the triangle.
+                auto adjTri = tri->T[j].lock();
+                if (adjTri)
+                {
+                    // Traverse to the triangle sharing the face.
+                    tri = adjTri;
+                    break;
+                }
+                else
+                {
+                    // We reached a hull edge, so the point is outside the
+                    // hull.  TODO:  Once a hull data structure is in place,
+                    // return tri->T[j] as the candidate for starting a search
+                    // for visible hull edges.
+                    return false;
+                }
+            }
+
+        }
+
+        if (j == 3)
+        {
+            // The point is inside all four edges, so the point is inside
+            // a triangle.
+            return true;
+        }
+    }
+
+    LogError("Unexpected termination of GetContainingTriangle.");
+    return false;
+}
+
+template <typename InputType, typename ComputeType>
+bool Delaunay2<InputType, ComputeType>::GetAndRemoveInsertionPolygon(int i,
+    std::set<std::shared_ptr<Triangle>>& candidates, std::set<EdgeKey<true>>& boundary)
+{
+    // Locate the triangles that make up the insertion polygon.
+    ETManifoldMesh polygon;
+    while (candidates.size() > 0)
+    {
+        std::shared_ptr<Triangle> tri = *candidates.begin();
+        candidates.erase(candidates.begin());
+
+        for (int j = 0; j < 3; ++j)
+        {
+            auto adj = tri->T[j].lock();
+            if (adj && candidates.find(adj) == candidates.end())
+            {
+                int a0 = adj->V[0];
+                int a1 = adj->V[1];
+                int a2 = adj->V[2];
+                if (mQuery.ToCircumcircle(i, a0, a1, a2) <= 0)
+                {
+                    // Point i is in the circumcircle.
+                    candidates.insert(adj);
+                }
+            }
+        }
+
+        if (!polygon.Insert(tri->V[0], tri->V[1], tri->V[2]))
+        {
+            return false;
+        }
+        if (!mGraph.Remove(tri->V[0], tri->V[1], tri->V[2]))
+        {
+            return false;
+        }
+    }
+
+    // Get the boundary edges of the insertion polygon.
+    for (auto const& element : polygon.GetTriangles())
+    {
+        std::shared_ptr<Triangle> tri = element.second;
+        for (int j = 0; j < 3; ++j)
+        {
+            if (!tri->T[j].lock())
+            {
+                boundary.insert(EdgeKey<true>(tri->V[mIndex[j][0]], tri->V[mIndex[j][1]]));
+            }
+        }
+    }
+    return true;
+}
+
+template <typename InputType, typename ComputeType>
+bool Delaunay2<InputType, ComputeType>::Update(int i)
+{
+    // The return value of mGraph.Insert(...) is nullptr if there was a
+    // failure to insert.  The Update function will return 'false' when
+    // the insertion fails.
+
+    auto const& tmap = mGraph.GetTriangles();
+    std::shared_ptr<Triangle> tri = tmap.begin()->second;
+    if (GetContainingTriangle(i, tri))
+    {
+        // The point is inside the convex hull.  The insertion polygon
+        // contains only triangles in the current triangulation; the
+        // hull does not change.
+
+        // Use a depth-first search for those triangles whose circumcircles
+        // contain point i.
+        std::set<std::shared_ptr<Triangle>> candidates;
+        candidates.insert(tri);
+
+        // Get the boundary of the insertion polygon C that contains the
+        // triangles whose circumcircles contain point i.  C contains the
+        // point i.
+        std::set<EdgeKey<true>> boundary;
+        if (!GetAndRemoveInsertionPolygon(i, candidates, boundary))
+        {
+            return false;
+        }
+
+        // The insertion polygon consists of the triangles formed by
+        // point i and the faces of C.
+        for (auto const& key : boundary)
+        {
+            int v0 = key.V[0];
+            int v1 = key.V[1];
+            if (mQuery.ToLine(i, v0, v1) < 0)
+            {
+                if (!mGraph.Insert(i, v0, v1))
+                {
+                    return false;
+                }
+            }
+            // else:  Point i is on an edge of 'tri', so the
+            // subdivision has degenerate triangles.  Ignore these.
+        }
+    }
+    else
+    {
+        // The point is outside the convex hull.  The insertion polygon
+        // is formed by point i and any triangles in the current
+        // triangulation whose circumcircles contain point i.
+
+        // Locate the convex hull of the triangles.  TODO:  Maintain a hull
+        // data structure that is updated incrementally.
+        std::set<EdgeKey<true>> hull;
+        for (auto const& element : tmap)
+        {
+            std::shared_ptr<Triangle> t = element.second;
+            for (int j = 0; j < 3; ++j)
+            {
+                if (!t->T[j].lock())
+                {
+                    hull.insert(EdgeKey<true>(t->V[mIndex[j][0]], t->V[mIndex[j][1]]));
+                }
+            }
+        }
+
+        // TODO:  Until the hull change, for now just iterate over all the
+        // hull edges and use the ones visible to point i to locate the
+        // insertion polygon.
+        auto const& emap = mGraph.GetEdges();
+        std::set<std::shared_ptr<Triangle>> candidates;
+        std::set<EdgeKey<true>> visible;
+        for (auto const& key : hull)
+        {
+            int v0 = key.V[0];
+            int v1 = key.V[1];
+            if (mQuery.ToLine(i, v0, v1) > 0)
+            {
+                auto iter = emap.find(EdgeKey<false>(v0, v1));
+                if (iter != emap.end() && iter->second->T[1].lock() == nullptr)
+                {
+                    auto adj = iter->second->T[0].lock();
+                    if (adj && candidates.find(adj) == candidates.end())
+                    {
+                        int a0 = adj->V[0];
+                        int a1 = adj->V[1];
+                        int a2 = adj->V[2];
+                        if (mQuery.ToCircumcircle(i, a0, a1, a2) <= 0)
+                        {
+                            // Point i is in the circumcircle.
+                            candidates.insert(adj);
+                        }
+                        else
+                        {
+                            // Point i is not in the circumcircle but the hull
+                            // edge is visible.
+                            visible.insert(key);
+                        }
+                    }
+                }
+                else
+                {
+                    LogError("Unexpected condition (ComputeType not exact?)");
+                    return false;
+                }
+            }
+        }
+
+        // Get the boundary of the insertion subpolygon C that contains the
+        // triangles whose circumcircles contain point i.
+        std::set<EdgeKey<true>> boundary;
+        if (!GetAndRemoveInsertionPolygon(i, candidates, boundary))
+        {
+            return false;
+        }
+
+        // The insertion polygon P consists of the triangles formed by point i
+        // and the back edges of C *and* the visible edges of mGraph-C.
+        for (auto const& key : boundary)
+        {
+            int v0 = key.V[0];
+            int v1 = key.V[1];
+            if (mQuery.ToLine(i, v0, v1) < 0)
+            {
+                // This is a back edge of the boundary.
+                if (!mGraph.Insert(i, v0, v1))
+                {
+                    return false;
+                }
+            }
+        }
+        for (auto const& key : visible)
+        {
+            if (!mGraph.Insert(i, key.V[1], key.V[0]))
+            {
+                return false;
+            }
+        }
+    }
+
+    return true;
+}
+
+template <typename InputType, typename ComputeType>
+Delaunay2<InputType, ComputeType>::ProcessedVertex::ProcessedVertex()
+{
+}
+
+template <typename InputType, typename ComputeType>
+Delaunay2<InputType, ComputeType>::ProcessedVertex::ProcessedVertex(
+    Vector2<InputType> const& inVertex, int inLocation)
+    :
+    vertex(inVertex),
+    location(inLocation)
+{
+}
+
+template <typename InputType, typename ComputeType>
+bool Delaunay2<InputType, ComputeType>::ProcessedVertex::operator<(
+    ProcessedVertex const& v) const
+{
+    return vertex < v.vertex;
+}
+
+
+}
diff --git a/Engine/Plugins/Experimental/GeometryProcessing/Source/GeometryAlgorithms/Private/GTEngine/Mathematics/GteDelaunay2Mesh.h b/Engine/Plugins/Experimental/GeometryProcessing/Source/GeometryAlgorithms/Private/GTEngine/Mathematics/GteDelaunay2Mesh.h
new file mode 100644
--- /dev/null
+++ b/Engine/Plugins/Experimental/GeometryProcessing/Source/GeometryAlgorithms/Private/GTEngine/Mathematics/GteDelaunay2Mesh.h
@@ -0,0 +1,169 @@
+// David Eberly, Geometric Tools, Redmond WA 98052
+// Copyright (c) 1998-2019
+// Distributed under the Boost Software License, Version 1.0.
+// http://www.boost.org/LICENSE_1_0.txt
+// http://www.geometrictools.com/License/Boost/LICENSE_1_0.txt
+// File Version: 3.0.0 (2016/06/19)
+
+#pragma once
+
+#include <GTEngine/Mathematics/GteDelaunay2.h>
+
+namespace gte
+{
+
+template <typename InputType, typename ComputeType, typename RationalType>
+class Delaunay2Mesh
+{
+public:
+    // Construction.
+    Delaunay2Mesh(Delaunay2<InputType, ComputeType> const& delaunay);
+
+    // Mesh information.
+    inline int GetNumVertices() const;
+    inline int GetNumTriangles() const;
+    inline Vector2<InputType> const* GetVertices() const;
+    inline int const* GetIndices() const;
+    inline int const* GetAdjacencies() const;
+
+    // Containment queries.
+    int GetContainingTriangle(Vector2<InputType> const& P) const;
+    bool GetVertices(int t, std::array<Vector2<InputType>, 3>& vertices)
+        const;
+    bool GetIndices(int t, std::array<int, 3>& indices) const;
+    bool GetAdjacencies(int t, std::array<int, 3>& adjacencies) const;
+    bool GetBarycentrics(int t, Vector2<InputType> const& P,
+        std::array<InputType, 3>& bary) const;
+
+private:
+    Delaunay2<InputType, ComputeType> const* mDelaunay;
+};
+
+
+template <typename InputType, typename ComputeType, typename RationalType>
+Delaunay2Mesh<InputType, ComputeType, RationalType>::Delaunay2Mesh(
+    Delaunay2<InputType, ComputeType> const& delaunay)
+    :
+    mDelaunay(&delaunay)
+{
+}
+
+template <typename InputType, typename ComputeType, typename RationalType>
+inline int Delaunay2Mesh<InputType, ComputeType, RationalType>::
+GetNumVertices() const
+{
+    return mDelaunay->GetNumVertices();
+}
+
+template <typename InputType, typename ComputeType, typename RationalType>
+inline int Delaunay2Mesh<InputType, ComputeType, RationalType>::
+GetNumTriangles() const
+{
+    return mDelaunay->GetNumTriangles();
+}
+
+template <typename InputType, typename ComputeType, typename RationalType>
+inline Vector2<InputType> const*
+Delaunay2Mesh<InputType, ComputeType, RationalType>::
+GetVertices() const
+{
+    return mDelaunay->GetVertices();
+}
+
+template <typename InputType, typename ComputeType, typename RationalType>
+inline int const* Delaunay2Mesh<InputType, ComputeType, RationalType>::
+GetIndices() const
+{
+    return &mDelaunay->GetIndices()[0];
+}
+
+template <typename InputType, typename ComputeType, typename RationalType>
+inline int const* Delaunay2Mesh<InputType, ComputeType, RationalType>::
+GetAdjacencies() const
+{
+    return &mDelaunay->GetAdjacencies()[0];
+}
+
+template <typename InputType, typename ComputeType, typename RationalType>
+int Delaunay2Mesh<InputType, ComputeType, RationalType>::
+GetContainingTriangle(Vector2<InputType> const& P) const
+{
+    typename Delaunay2<InputType, ComputeType>::SearchInfo info;
+    return mDelaunay->GetContainingTriangle(P, info);
+}
+
+template <typename InputType, typename ComputeType, typename RationalType>
+bool Delaunay2Mesh<InputType, ComputeType, RationalType>::
+GetVertices(int t, std::array<Vector2<InputType>, 3>& vertices) const
+{
+    if (mDelaunay->GetDimension() == 2)
+    {
+        std::array<int, 3> indices;
+        if (mDelaunay->GetIndices(t, indices))
+        {
+            PrimalQuery2<ComputeType> const& query = mDelaunay->GetQuery();
+            Vector2<ComputeType> const* ctVertices = query.GetVertices();
+            for (int i = 0; i < 3; ++i)
+            {
+                Vector2<ComputeType> const& V = ctVertices[indices[i]];
+                for (int j = 0; j < 2; ++j)
+                {
+                    vertices[i][j] = (InputType)V[j];
+                }
+            }
+            return true;
+        }
+    }
+    return false;
+}
+
+template <typename InputType, typename ComputeType, typename RationalType>
+bool Delaunay2Mesh<InputType, ComputeType, RationalType>::
+GetIndices(int t, std::array<int, 3>& indices) const
+{
+    return mDelaunay->GetIndices(t, indices);
+}
+
+template <typename InputType, typename ComputeType, typename RationalType>
+bool Delaunay2Mesh<InputType, ComputeType, RationalType>::
+GetAdjacencies(int t, std::array<int, 3>& indices) const
+{
+    return mDelaunay->GetAdjacencies(t, indices);
+}
+
+template <typename InputType, typename ComputeType, typename RationalType>
+bool Delaunay2Mesh<InputType, ComputeType, RationalType>::
+GetBarycentrics(int t, Vector2<InputType> const& P,
+std::array<InputType, 3>& bary) const
+{
+    std::array<int, 3> indices;
+    if (mDelaunay->GetIndices(t, indices))
+    {
+        PrimalQuery2<ComputeType> const& query = mDelaunay->GetQuery();
+        Vector2<ComputeType> const* vertices = query.GetVertices();
+        Vector2<RationalType> rtP{ P[0], P[1] };
+        std::array<Vector2<RationalType>, 3> rtV;
+        for (int i = 0; i < 3; ++i)
+        {
+            Vector2<ComputeType> const& V = vertices[indices[i]];
+            for (int j = 0; j < 2; ++j)
+            {
+                rtV[i][j] = (RationalType)V[j];
+            }
+        };
+
+        RationalType rtBary[3];
+        if (ComputeBarycentrics(rtP, rtV[0], rtV[1], rtV[2], rtBary))
+        {
+            for (int i = 0; i < 3; ++i)
+            {
+                bary[i] = (InputType)rtBary[i];
+            }
+            return true;
+        }
+    }
+    return false;
+}
+
+
+}
