move GeometryProcessing plugin from NotForLicensees to Experimental

diff --git a/Engine/Plugins/Experimental/GeometryProcessing/Source/GeometryAlgorithms/Private/GTEngine/Mathematics/GteDelaunay3.h b/Engine/Plugins/Experimental/GeometryProcessing/Source/GeometryAlgorithms/Private/GTEngine/Mathematics/GteDelaunay3.h
new file mode 100644
--- /dev/null
+++ b/Engine/Plugins/Experimental/GeometryProcessing/Source/GeometryAlgorithms/Private/GTEngine/Mathematics/GteDelaunay3.h
@@ -0,0 +1,884 @@
+// David Eberly, Geometric Tools, Redmond WA 98052
+// Copyright (c) 1998-2019
+// Distributed under the Boost Software License, Version 1.0.
+// http://www.boost.org/LICENSE_1_0.txt
+// http://www.geometrictools.com/License/Boost/LICENSE_1_0.txt
+// File Version: 3.0.0 (2016/06/19)
+
+#pragma once
+
+#include <GTEngine/LowLevel/GteLogger.h>
+#include <GTEngine/Mathematics/GtePrimalQuery3.h>
+#include <GTEngine/Mathematics/GteTSManifoldMesh.h>
+#include <GTEngine/Mathematics/GteLine.h>
+#include <GTEngine/Mathematics/GteHyperplane.h>
+#include <vector>
+
+// Delaunay tetrahedralization of points (intrinsic dimensionality 3).
+//   VQ = number of vertices
+//   V  = array of vertices
+//   TQ = number of tetrahedra
+//   I  = Array of 4-tuples of indices into V that represent the tetrahedra
+//        (4*TQ total elements).  Access via GetIndices(*).
+//   A  = Array of 4-tuples of indices into I that represent the adjacent
+//        tetrahedra (4*TQ total elements).  Access via GetAdjacencies(*).
+// The i-th tetrahedron has vertices
+//   vertex[0] = V[I[4*i+0]]
+//   vertex[1] = V[I[4*i+1]]
+//   vertex[2] = V[I[4*i+2]]
+//   vertex[3] = V[I[4*i+3]]
+// and face index triples listed below.  The face vertex ordering when
+// viewed from outside the tetrahedron is counterclockwise.
+//   face[0] = <I[4*i+1],I[4*i+2],I[4*i+3]>
+//   face[1] = <I[4*i+0],I[4*i+3],I[4*i+2]>
+//   face[2] = <I[4*i+0],I[4*i+1],I[4*i+3]>
+//   face[3] = <I[4*i+0],I[4*i+2],I[4*i+1]>
+// The tetrahedra adjacent to these faces have indices
+//   adjacent[0] = A[4*i+0] is the tetrahedron opposite vertex[0], so it
+//                 is the tetrahedron sharing face[0].
+//   adjacent[1] = A[4*i+1] is the tetrahedron opposite vertex[1], so it
+//                 is the tetrahedron sharing face[1].
+//   adjacent[2] = A[4*i+2] is the tetrahedron opposite vertex[2], so it
+//                 is the tetrahedron sharing face[2].
+//   adjacent[3] = A[4*i+3] is the tetrahedron opposite vertex[3], so it
+//                 is the tetrahedron sharing face[3].
+// If there is no adjacent tetrahedron, the A[*] value is set to -1.  The
+// tetrahedron adjacent to face[j] has vertices
+//   adjvertex[0] = V[I[4*adjacent[j]+0]]
+//   adjvertex[1] = V[I[4*adjacent[j]+1]]
+//   adjvertex[2] = V[I[4*adjacent[j]+2]]
+//   adjvertex[3] = V[I[4*adjacent[j]+3]]
+// The only way to ensure a correct result for the input vertices (assumed to
+// be exact) is to choose ComputeType for exact rational arithmetic.  You may
+// use BSNumber.  No divisions are performed in this computation, so you do
+// not have to use BSRational.
+//
+// The worst-case choices of N for Real of type BSNumber or BSRational with
+// integer storage UIntegerFP32<N> are listed in the next table.  The numerical
+// computations are encapsulated in PrimalQuery3<Real>::ToPlane and
+// PrimalQuery3<Real>::ToCircumsphere, the latter query the dominant one in
+// determining N.  We recommend using only BSNumber, because no divisions are
+// performed in the convex-hull computations.
+//
+//    input type | compute type | N
+//    -----------+--------------+--------
+//    float      | BSNumber     |      44
+//    float      | BSRational   |     329
+//    double     | BSNumber     |  298037
+//    double     | BSRational   | 2254442
+
+namespace gte
+{
+
+template <typename InputType, typename ComputeType>
+class Delaunay3
+{
+public:
+    // The class is a functor to support computing the Delaunay
+    // tetrahedralization of multiple data sets using the same class object.
+    Delaunay3();
+
+    // The input is the array of vertices whose Delaunay tetrahedralization
+    // is required.  The epsilon value is used to determine the intrinsic
+    // dimensionality of the vertices (d = 0, 1, 2, or 3).  When epsilon is
+    // positive, the determination is fuzzy--vertices approximately the same
+    // point, approximately on a line, approximately planar, or volumetric.
+    bool operator()(int numVertices, Vector3<InputType> const* vertices, InputType epsilon);
+
+    // Dimensional information.  If GetDimension() returns 1, the points lie
+    // on a line P+t*D (fuzzy comparison when epsilon > 0).  You can sort
+    // these if you need a polyline output by projecting onto the line each
+    // vertex X = P+t*D, where t = Dot(D,X-P).  If GetDimension() returns 2,
+    // the points line on a plane P+s*U+t*V (fuzzy comparison when
+    // epsilon > 0).  You can project each vertex X = P+s*U+t*V, where
+    // s = Dot(U,X-P) and t = Dot(V,X-P), then apply Delaunay2 to the (s,t)
+    // tuples.
+    inline InputType GetEpsilon() const;
+    inline int GetDimension() const;
+    inline Line3<InputType> const& GetLine() const;
+    inline Plane3<InputType> const& GetPlane() const;
+
+    // Member access.
+    inline int GetNumVertices() const;
+    inline int GetNumUniqueVertices() const;
+    inline int GetNumTetrahedra() const;
+    inline Vector3<InputType> const* GetVertices() const;
+    inline PrimalQuery3<ComputeType> const& GetQuery() const;
+    inline TSManifoldMesh const& GetGraph() const;
+    inline std::vector<int> const& GetIndices() const;
+    inline std::vector<int> const& GetAdjacencies() const;
+
+    // Locate those tetrahedra faces that do not share other tetrahedra.  The
+    // returned array has hull.size() = 3*numFaces indices, each triple
+    // representing a triangle.  The triangles are counterclockwise ordered
+    // when viewed from outside the hull.  The return value is 'true' iff the
+    // dimension is 3.
+    bool GetHull(std::vector<int>& hull) const;
+
+    // Get the vertex indices for tetrahedron i.  The function returns 'true'
+    // when the dimension is 3 and i is a valid tetrahedron index, in which
+    // case the vertices are valid; otherwise, the function returns 'false'
+    // and the vertices are invalid.
+    bool GetIndices(int i, std::array<int, 4>& indices) const;
+
+    // Get the indices for tetrahedra adjacent to tetrahedron i.  The function
+    // returns 'true' when the dimension is 3 and i is a valid tetrahedron
+    // index, in which case the adjacencies are valid; otherwise, the function
+    // returns 'false' and the adjacencies are invalid.
+    bool GetAdjacencies(int i, std::array<int, 4>& adjacencies) const;
+
+    // Support for searching the tetrahedralization for a tetrahedron that
+    // contains a point.  If there is a containing tetrahedron, the returned
+    // value is a tetrahedron index i with 0 <= i < GetNumTetrahedra().  If
+    // there is not a containing tetrahedron, -1 is returned.  The
+    // computations are performed using exact rational arithmetic.
+    //
+    // The SearchInfo input stores information about the tetrahedron search
+    // when looking for the tetrahedron (if any) that contains p.  The first
+    // tetrahedron searched is 'initialTetrahedron'.  On return 'path' stores
+    // those (ordered) tetrahedron indices visited during the search.  The
+    // last visited tetrahedron has index 'finalTetrahedron and vertex indices
+    // 'finalV[0,1,2,3]', stored in volumetric counterclockwise order.  The
+    // last face of the search is <finalV[0],finalV[1],finalV[2]>.  For
+    // spatially coherent inputs p for numerous calls to this function, you
+    // will want to specify 'finalTetrahedron' from the previous call as
+    // 'initialTetrahedron' for the next call, which should reduce search
+    // times.
+    struct SearchInfo
+    {
+        int initialTetrahedron;
+        int numPath;
+        std::vector<int> path;
+        int finalTetrahedron;
+        std::array<int, 4> finalV;
+    };
+    int GetContainingTetrahedron(Vector3<InputType> const& p, SearchInfo& info) const;
+
+private:
+    // Support for incremental Delaunay tetrahedralization.
+    typedef TSManifoldMesh::Tetrahedron Tetrahedron;
+    bool GetContainingTetrahedron(int i, std::shared_ptr<Tetrahedron>& tetra) const;
+    bool GetAndRemoveInsertionPolyhedron(int i, std::set<std::shared_ptr<Tetrahedron>>& candidates,
+        std::set<TriangleKey<true>>& boundary);
+    bool Update(int i);
+
+    // The epsilon value is used for fuzzy determination of intrinsic
+    // dimensionality.  If the dimension is 0, 1, or 2, the constructor
+    // returns early.  The caller is responsible for retrieving the dimension
+    // and taking an alternate path should the dimension be smaller than 3.
+    // If the dimension is 0, the caller may as well treat all vertices[]
+    // as a single point, say, vertices[0].  If the dimension is 1, the
+    // caller can query for the approximating line and project vertices[]
+    // onto it for further processing.  If the dimension is 2, the caller can
+    // query for the approximating plane and project vertices[] onto it for
+    // further processing.
+    InputType mEpsilon;
+    int mDimension;
+    Line3<InputType> mLine;
+    Plane3<InputType> mPlane;
+
+    // The array of vertices used for geometric queries.  If you want to be
+    // certain of a correct result, choose ComputeType to be BSNumber.
+    std::vector<Vector3<ComputeType>> mComputeVertices;
+    PrimalQuery3<ComputeType> mQuery;
+
+    // The graph information.
+    int mNumVertices;
+    int mNumUniqueVertices;
+    int mNumTetrahedra;
+    Vector3<InputType> const* mVertices;
+    TSManifoldMesh mGraph;
+    std::vector<int> mIndices;
+    std::vector<int> mAdjacencies;
+};
+
+
+template <typename InputType, typename ComputeType>
+Delaunay3<InputType, ComputeType>::Delaunay3()
+    :
+    mEpsilon((InputType)0),
+    mDimension(0),
+    mLine(Vector3<InputType>::Zero(), Vector3<InputType>::Zero()),
+    mPlane(Vector3<InputType>::Zero(), (InputType)0),
+    mNumVertices(0),
+    mNumUniqueVertices(0),
+    mNumTetrahedra(0),
+    mVertices(nullptr)
+{
+}
+
+template <typename InputType, typename ComputeType>
+bool Delaunay3<InputType, ComputeType>::operator()(int numVertices,
+    Vector3<InputType> const* vertices, InputType epsilon)
+{
+    mEpsilon = std::max(epsilon, (InputType)0);
+    mDimension = 0;
+    mLine.origin = Vector3<InputType>::Zero();
+    mLine.direction = Vector3<InputType>::Zero();
+    mPlane.normal = Vector3<InputType>::Zero();
+    mPlane.constant = (InputType)0;
+    mNumVertices = numVertices;
+    mNumUniqueVertices = 0;
+    mNumTetrahedra = 0;
+    mVertices = vertices;
+    mGraph = TSManifoldMesh();
+    mIndices.clear();
+    mAdjacencies.clear();
+
+    int i, j;
+    if (mNumVertices < 4)
+    {
+        // Delaunay3 should be called with at least four points.
+        return false;
+    }
+
+    IntrinsicsVector3<InputType> info(mNumVertices, vertices, mEpsilon);
+    if (info.dimension == 0)
+    {
+        // mDimension is 0; mGraph, mIndices, and mAdjacencies are empty
+        return false;
+    }
+
+    if (info.dimension == 1)
+    {
+        // The set is (nearly) collinear.
+        mDimension = 1;
+        mLine = Line3<InputType>(info.origin, info.direction[0]);
+        return false;
+    }
+
+    if (info.dimension == 2)
+    {
+        // The set is (nearly) coplanar.
+        mDimension = 2;
+        mPlane = Plane3<InputType>(UnitCross(info.direction[0],
+            info.direction[1]), info.origin);
+        return false;
+    }
+
+    mDimension = 3;
+
+    // Compute the vertices for the queries.
+    mComputeVertices.resize(mNumVertices);
+    mQuery.Set(mNumVertices, &mComputeVertices[0]);
+    for (i = 0; i < mNumVertices; ++i)
+    {
+        for (j = 0; j < 3; ++j)
+        {
+            mComputeVertices[i][j] = vertices[i][j];
+        }
+    }
+
+    // Insert the (nondegenerate) tetrahedron constructed by the call to
+    // GetInformation. This is necessary for the circumsphere-visibility
+    // algorithm to work correctly.
+    if (!info.extremeCCW)
+    {
+        std::swap(info.extreme[2], info.extreme[3]);
+    }
+    if (!mGraph.Insert(info.extreme[0], info.extreme[1], info.extreme[2], info.extreme[3]))
+    {
+        return false;
+    }
+
+    // Incrementally update the tetrahedralization.  The set of processed
+    // points is maintained to eliminate duplicates, either in the original
+    // input points or in the points obtained by snap rounding.
+    std::set<Vector3<InputType>> processed;
+    for (i = 0; i < 4; ++i)
+    {
+        processed.insert(vertices[info.extreme[i]]);
+    }
+    for (i = 0; i < mNumVertices; ++i)
+    {
+        if (processed.find(vertices[i]) == processed.end())
+        {
+            if (!Update(i))
+            {
+                // A failure can occur if ComputeType is not an exact
+                // arithmetic type.
+                return false;
+            }
+            processed.insert(vertices[i]);
+        }
+    }
+    mNumUniqueVertices = static_cast<int>(processed.size());
+
+    // Assign integer values to the tetrahedra for use by the caller.
+    std::map<std::shared_ptr<Tetrahedron>, int> permute;
+    i = -1;
+    permute[nullptr] = i++;
+    for (auto const& element : mGraph.GetTetrahedra())
+    {
+        permute[element.second] = i++;
+    }
+
+    // Put Delaunay tetrahedra into an array (vertices and adjacency info).
+    mNumTetrahedra = static_cast<int>(mGraph.GetTetrahedra().size());
+    int numIndices = 4 * mNumTetrahedra;
+    if (mNumTetrahedra > 0)
+    {
+        mIndices.resize(numIndices);
+        mAdjacencies.resize(numIndices);
+        i = 0;
+        for (auto const& element : mGraph.GetTetrahedra())
+        {
+            std::shared_ptr<Tetrahedron> tetra = element.second;
+            for (j = 0; j < 4; ++j, ++i)
+            {
+                mIndices[i] = tetra->V[j];
+                mAdjacencies[i] = permute[tetra->S[j].lock()];
+            }
+        }
+    }
+
+    return true;
+}
+
+template <typename InputType, typename ComputeType> inline
+InputType Delaunay3<InputType, ComputeType>::GetEpsilon() const
+{
+    return mEpsilon;
+}
+
+template <typename InputType, typename ComputeType> inline
+int Delaunay3<InputType, ComputeType>::GetDimension() const
+{
+    return mDimension;
+}
+
+template <typename InputType, typename ComputeType> inline
+Line3<InputType> const& Delaunay3<InputType, ComputeType>::GetLine() const
+{
+    return mLine;
+}
+
+template <typename InputType, typename ComputeType> inline
+Plane3<InputType> const& Delaunay3<InputType, ComputeType>::GetPlane() const
+{
+    return mPlane;
+}
+
+template <typename InputType, typename ComputeType> inline
+int Delaunay3<InputType, ComputeType>::GetNumVertices() const
+{
+    return mNumVertices;
+}
+
+template <typename InputType, typename ComputeType> inline
+int Delaunay3<InputType, ComputeType>::GetNumUniqueVertices() const
+{
+    return mNumUniqueVertices;
+}
+
+template <typename InputType, typename ComputeType> inline
+int Delaunay3<InputType, ComputeType>::GetNumTetrahedra() const
+{
+    return mNumTetrahedra;
+}
+
+template <typename InputType, typename ComputeType> inline
+Vector3<InputType> const* Delaunay3<InputType, ComputeType>::GetVertices() const
+{
+    return mVertices;
+}
+
+template <typename InputType, typename ComputeType> inline
+PrimalQuery3<ComputeType> const& Delaunay3<InputType, ComputeType>::GetQuery() const
+{
+    return mQuery;
+}
+
+template <typename InputType, typename ComputeType> inline
+TSManifoldMesh const& Delaunay3<InputType, ComputeType>::GetGraph() const
+{
+    return mGraph;
+}
+
+template <typename InputType, typename ComputeType> inline
+std::vector<int> const& Delaunay3<InputType, ComputeType>::GetIndices() const
+{
+    return mIndices;
+}
+
+template <typename InputType, typename ComputeType> inline
+std::vector<int> const& Delaunay3<InputType, ComputeType>::GetAdjacencies() const
+{
+    return mAdjacencies;
+}
+
+template <typename InputType, typename ComputeType>
+bool Delaunay3<InputType, ComputeType>::GetHull(std::vector<int>& hull) const
+{
+    if (mDimension == 3)
+    {
+        // Count the number of triangles that are not shared by two
+        // tetrahedra.
+        int numTriangles = 0;
+        for (auto adj : mAdjacencies)
+        {
+            if (adj == -1)
+            {
+                ++numTriangles;
+            }
+        }
+
+        if (numTriangles > 0)
+        {
+            // Enumerate the triangles.  The prototypical case is the single
+            // tetrahedron V[0] = (0,0,0), V[1] = (1,0,0), V[2] = (0,1,0), and
+            // V[3] = (0,0,1) with no adjacent tetrahedra.  The mIndices[]
+            // array is <0,1,2,3>.
+            //   i = 0, face = 0:
+            //     skip index 0, <x,1,2,3>, no swap, triangle = <1,2,3>
+            //   i = 1, face = 1:
+            //     skip index 1, <0,x,2,3>, swap,    triangle = <0,3,2>
+            //   i = 2, face = 2:
+            //     skip index 2, <0,1,x,3>, no swap, triangle = <0,1,3>
+            //   i = 3, face = 3:
+            //     skip index 3, <0,1,2,x>, swap,    triangle = <0,2,1>
+            // To guarantee counterclockwise order of triangles when viewed
+            // outside the tetrahedron, the swap of the last two indices
+            // occurs when face is an odd number; (face % 2) != 0.
+            hull.resize(3 * numTriangles);
+            int current = 0, i = 0;
+            for (auto adj : mAdjacencies)
+            {
+                if (adj == -1)
+                {
+                    int tetra = i / 4, face = i % 4;
+                    for (int j = 0; j < 4; ++j)
+                    {
+                        if (j != face)
+                        {
+                            hull[current++] = mIndices[4 * tetra + j];
+                        }
+                    }
+                    if ((face % 2) != 0)
+                    {
+                        std::swap(hull[current - 1], hull[current - 2]);
+                    }
+                }
+                ++i;
+            }
+            return true;
+        }
+        else
+        {
+            LogError("Unexpected.  There must be at least one tetrahedron.");
+        }
+    }
+    else
+    {
+        LogError("The dimension must be 3.");
+    }
+    return false;
+}
+
+template <typename InputType, typename ComputeType>
+bool Delaunay3<InputType, ComputeType>::GetIndices(int i, std::array<int, 4>& indices) const
+{
+    if (mDimension == 3)
+    {
+        int numTetrahedra = static_cast<int>(mIndices.size() / 4);
+        if (0 <= i && i < numTetrahedra)
+        {
+            indices[0] = mIndices[4 * i];
+            indices[1] = mIndices[4 * i + 1];
+            indices[2] = mIndices[4 * i + 2];
+            indices[3] = mIndices[4 * i + 3];
+            return true;
+        }
+    }
+    else
+    {
+        LogError("The dimension must be 3.");
+    }
+    return false;
+}
+
+template <typename InputType, typename ComputeType>
+bool Delaunay3<InputType, ComputeType>::GetAdjacencies(int i, std::array<int, 4>& adjacencies) const
+{
+    if (mDimension == 3)
+    {
+        int numTetrahedra = static_cast<int>(mIndices.size() / 4);
+        if (0 <= i && i < numTetrahedra)
+        {
+            adjacencies[0] = mAdjacencies[4 * i];
+            adjacencies[1] = mAdjacencies[4 * i + 1];
+            adjacencies[2] = mAdjacencies[4 * i + 2];
+            adjacencies[3] = mAdjacencies[4 * i + 3];
+            return true;
+        }
+    }
+    else
+    {
+        LogError("The dimension must be 3.");
+    }
+    return false;
+}
+
+template <typename InputType, typename ComputeType>
+int Delaunay3<InputType, ComputeType>::GetContainingTetrahedron(Vector3<InputType> const& p, SearchInfo& info) const
+{
+    if (mDimension == 3)
+    {
+        Vector3<ComputeType> test{ p[0], p[1], p[2] };
+
+        int numTetrahedra = static_cast<int>(mIndices.size() / 4);
+        info.path.resize(numTetrahedra);
+        info.numPath = 0;
+        int tetrahedron;
+        if (0 <= info.initialTetrahedron
+            && info.initialTetrahedron < numTetrahedra)
+        {
+            tetrahedron = info.initialTetrahedron;
+        }
+        else
+        {
+            info.initialTetrahedron = 0;
+            tetrahedron = 0;
+        }
+
+        // Use tetrahedron faces as binary separating planes.
+        for (int i = 0; i < numTetrahedra; ++i)
+        {
+            int ibase = 4 * tetrahedron;
+            int const* v = &mIndices[ibase];
+
+            info.path[info.numPath++] = tetrahedron;
+            info.finalTetrahedron = tetrahedron;
+            info.finalV[0] = v[0];
+            info.finalV[1] = v[1];
+            info.finalV[2] = v[2];
+            info.finalV[3] = v[3];
+
+            // <V1,V2,V3> counterclockwise when viewed outside tetrahedron.
+            if (mQuery.ToPlane(test, v[1], v[2], v[3]) > 0)
+            {
+                tetrahedron = mAdjacencies[ibase];
+                if (tetrahedron == -1)
+                {
+                    info.finalV[0] = v[1];
+                    info.finalV[1] = v[2];
+                    info.finalV[2] = v[3];
+                    info.finalV[3] = v[0];
+                    return -1;
+                }
+                continue;
+            }
+
+            // <V0,V3,V2> counterclockwise when viewed outside tetrahedron.
+            if (mQuery.ToPlane(test, v[0], v[2], v[3]) < 0)
+            {
+                tetrahedron = mAdjacencies[ibase + 1];
+                if (tetrahedron == -1)
+                {
+                    info.finalV[0] = v[0];
+                    info.finalV[1] = v[2];
+                    info.finalV[2] = v[3];
+                    info.finalV[3] = v[1];
+                    return -1;
+                }
+                continue;
+            }
+
+            // <V0,V1,V3> counterclockwise when viewed outside tetrahedron.
+            if (mQuery.ToPlane(test, v[0], v[1], v[3]) > 0)
+            {
+                tetrahedron = mAdjacencies[ibase + 2];
+                if (tetrahedron == -1)
+                {
+                    info.finalV[0] = v[0];
+                    info.finalV[1] = v[1];
+                    info.finalV[2] = v[3];
+                    info.finalV[3] = v[2];
+                    return -1;
+                }
+                continue;
+            }
+
+            // <V0,V2,V1> counterclockwise when viewed outside tetrahedron.
+            if (mQuery.ToPlane(test, v[0], v[1], v[2]) < 0)
+            {
+                tetrahedron = mAdjacencies[ibase + 3];
+                if (tetrahedron == -1)
+                {
+                    info.finalV[0] = v[0];
+                    info.finalV[1] = v[1];
+                    info.finalV[2] = v[2];
+                    info.finalV[3] = v[3];
+                    return -1;
+                }
+                continue;
+            }
+
+            return tetrahedron;
+        }
+    }
+    else
+    {
+        LogError("The dimension must be 3.");
+    }
+    return -1;
+}
+
+template <typename InputType, typename ComputeType>
+bool Delaunay3<InputType, ComputeType>::GetContainingTetrahedron(int i, std::shared_ptr<Tetrahedron>& tetra) const
+{
+    int numTetrahedra = static_cast<int>(mGraph.GetTetrahedra().size());
+    for (int t = 0; t < numTetrahedra; ++t)
+    {
+        int j;
+        for (j = 0; j < 4; ++j)
+        {
+            auto const& opposite = TetrahedronKey<true>::oppositeFace;
+            int v0 = tetra->V[opposite[j][0]];
+            int v1 = tetra->V[opposite[j][1]];
+            int v2 = tetra->V[opposite[j][2]];
+            if (mQuery.ToPlane(i, v0, v1, v2) > 0)
+            {
+                // Point i sees face <v0,v1,v2> from outside the tetrahedron.
+                auto adjTetra = tetra->S[j].lock();
+                if (adjTetra)
+                {
+                    // Traverse to the tetrahedron sharing the face.
+                    tetra = adjTetra;
+                    break;
+                }
+                else
+                {
+                    // We reached a hull face, so the point is outside the
+                    // hull.  TODO:  Once a hull data structure is in place,
+                    // return tetra->S[j] as the candidate for starting a
+                    // search for visible hull faces.
+                    return false;
+                }
+            }
+
+        }
+
+        if (j == 4)
+        {
+            // The point is inside all four faces, so the point is inside
+            // a tetrahedron.
+            return true;
+        }
+    }
+
+    LogError("Unexpected termination of GetContainingTetrahedron.");
+    return false;
+}
+
+template <typename InputType, typename ComputeType>
+bool Delaunay3<InputType, ComputeType>::GetAndRemoveInsertionPolyhedron(int i,
+    std::set<std::shared_ptr<Tetrahedron>>& candidates, std::set<TriangleKey<true>>& boundary)
+{
+    // Locate the tetrahedra that make up the insertion polyhedron.
+    TSManifoldMesh polyhedron;
+    while (candidates.size() > 0)
+    {
+        std::shared_ptr<Tetrahedron> tetra = *candidates.begin();
+        candidates.erase(candidates.begin());
+
+        for (int j = 0; j < 4; ++j)
+        {
+            auto adj = tetra->S[j].lock();
+            if (adj && candidates.find(adj) == candidates.end())
+            {
+                int a0 = adj->V[0];
+                int a1 = adj->V[1];
+                int a2 = adj->V[2];
+                int a3 = adj->V[3];
+                if (mQuery.ToCircumsphere(i, a0, a1, a2, a3) <= 0)
+                {
+                    // Point i is in the circumsphere.
+                    candidates.insert(adj);
+                }
+            }
+        }
+
+        int v0 = tetra->V[0];
+        int v1 = tetra->V[1];
+        int v2 = tetra->V[2];
+        int v3 = tetra->V[3];
+        if (!polyhedron.Insert(v0, v1, v2, v3))
+        {
+            return false;
+        }
+        if (!mGraph.Remove(v0, v1, v2, v3))
+        {
+            return false;
+        }
+    }
+
+    // Get the boundary triangles of the insertion polyhedron.
+    for (auto const& element : polyhedron.GetTetrahedra())
+    {
+        std::shared_ptr<Tetrahedron> tetra = element.second;
+        for (int j = 0; j < 4; ++j)
+        {
+            if (!tetra->S[j].lock())
+            {
+                auto const& opposite = TetrahedronKey<true>::oppositeFace;
+                int v0 = tetra->V[opposite[j][0]];
+                int v1 = tetra->V[opposite[j][1]];
+                int v2 = tetra->V[opposite[j][2]];
+                boundary.insert(TriangleKey<true>(v0, v1, v2));
+            }
+        }
+    }
+    return true;
+}
+
+template <typename InputType, typename ComputeType>
+bool Delaunay3<InputType, ComputeType>::Update(int i)
+{
+    auto const& smap = mGraph.GetTetrahedra();
+    std::shared_ptr<Tetrahedron> tetra = smap.begin()->second;
+    if (GetContainingTetrahedron(i, tetra))
+    {
+        // The point is inside the convex hull.  The insertion polyhedron
+        // contains only tetrahedra in the current tetrahedralization; the
+        // hull does not change.
+
+        // Use a depth-first search for those tetrahedra whose circumspheres
+        // contain point i.
+        std::set<std::shared_ptr<Tetrahedron>> candidates;
+        candidates.insert(tetra);
+
+        // Get the boundary of the insertion polyhedron C that contains the
+        // tetrahedra whose circumspheres contain point i.  C contains the
+        // point i.
+        std::set<TriangleKey<true>> boundary;
+        if (!GetAndRemoveInsertionPolyhedron(i, candidates, boundary))
+        {
+            return false;
+        }
+
+        // The insertion polyhedron consists of the tetrahedra formed by
+        // point i and the faces of C.
+        for (auto const& key : boundary)
+        {
+            int v0 = key.V[0];
+            int v1 = key.V[1];
+            int v2 = key.V[2];
+            if (mQuery.ToPlane(i, v0, v1, v2) < 0)
+            {
+                if (!mGraph.Insert(i, v0, v1, v2))
+                {
+                    return false;
+                }
+            }
+            // else:  Point i is on an edge or face of 'tetra', so the
+            // subdivision has degenerate tetrahedra.  Ignore these.
+        }
+    }
+    else
+    {
+        // The point is outside the convex hull.  The insertion polyhedron
+        // is formed by point i and any tetrahedra in the current
+        // tetrahedralization whose circumspheres contain point i.
+
+        // Locate the convex hull of the tetrahedra.  TODO:  Maintain a hull
+        // data structure that is updated incrementally.
+        std::set<TriangleKey<true>> hull;
+        for (auto const& element : smap)
+        {
+            std::shared_ptr<Tetrahedron> t = element.second;
+            for (int j = 0; j < 4; ++j)
+            {
+                if (!t->S[j].lock())
+                {
+                    auto const& opposite = TetrahedronKey<true>::oppositeFace;
+                    int v0 = t->V[opposite[j][0]];
+                    int v1 = t->V[opposite[j][1]];
+                    int v2 = t->V[opposite[j][2]];
+                    hull.insert(TriangleKey<true>(v0, v1, v2));
+                }
+            }
+        }
+
+        // TODO:  Until the hull change, for now just iterate over all the
+        // hull faces and use the ones visible to point i to locate the
+        // insertion polyhedron.
+        auto const& tmap = mGraph.GetTriangles();
+        std::set<std::shared_ptr<Tetrahedron>> candidates;
+        std::set<TriangleKey<true>> visible;
+        for (auto const& key : hull)
+        {
+            int v0 = key.V[0];
+            int v1 = key.V[1];
+            int v2 = key.V[2];
+            if (mQuery.ToPlane(i, v0, v1, v2) > 0)
+            {
+                auto iter = tmap.find(TriangleKey<false>(v0, v1, v2));
+                if (iter != tmap.end() && iter->second->T[1].lock() == nullptr)
+                {
+                    auto adj = iter->second->T[0].lock();
+                    if (adj && candidates.find(adj) == candidates.end())
+                    {
+                        int a0 = adj->V[0];
+                        int a1 = adj->V[1];
+                        int a2 = adj->V[2];
+                        int a3 = adj->V[3];
+                        if (mQuery.ToCircumsphere(i, a0, a1, a2, a3) <= 0)
+                        {
+                            // Point i is in the circumsphere.
+                            candidates.insert(adj);
+                        }
+                        else
+                        {
+                            // Point i is not in the circumsphere but the hull
+                            // face is visible.
+                            visible.insert(key);
+                        }
+                    }
+                }
+                else
+                {
+                    LogError("Unexpected condition (ComputeType not exact?)");
+                    return false;
+                }
+            }
+        }
+
+        // Get the boundary of the insertion subpolyhedron C that contains the
+        // tetrahedra whose circumspheres contain point i.
+        std::set<TriangleKey<true>> boundary;
+        if (!GetAndRemoveInsertionPolyhedron(i, candidates, boundary))
+        {
+            return false;
+        }
+
+        // The insertion polyhedron P consists of the tetrahedra formed by
+        // point i and the back faces of C *and* the visible faces of
+        // mGraph-C.
+        for (auto const& key : boundary)
+        {
+            int v0 = key.V[0];
+            int v1 = key.V[1];
+            int v2 = key.V[2];
+            if (mQuery.ToPlane(i, v0, v1, v2) < 0)
+            {
+                // This is a back face of the boundary.
+                if (!mGraph.Insert(i, v0, v1, v2))
+                {
+                    return false;
+                }
+            }
+        }
+        for (auto const& key : visible)
+        {
+            if (!mGraph.Insert(i, key.V[0], key.V[2], key.V[1]))
+            {
+                return false;
+            }
+        }
+    }
+
+    return true;
+}
+
+}
diff --git a/Engine/Plugins/Experimental/GeometryProcessing/Source/GeometryAlgorithms/Private/GTEngine/Mathematics/GteDelaunay3Mesh.h b/Engine/Plugins/Experimental/GeometryProcessing/Source/GeometryAlgorithms/Private/GTEngine/Mathematics/GteDelaunay3Mesh.h
new file mode 100644
--- /dev/null
+++ b/Engine/Plugins/Experimental/GeometryProcessing/Source/GeometryAlgorithms/Private/GTEngine/Mathematics/GteDelaunay3Mesh.h
@@ -0,0 +1,169 @@
+// David Eberly, Geometric Tools, Redmond WA 98052
+// Copyright (c) 1998-2019
+// Distributed under the Boost Software License, Version 1.0.
+// http://www.boost.org/LICENSE_1_0.txt
+// http://www.geometrictools.com/License/Boost/LICENSE_1_0.txt
+// File Version: 3.0.0 (2016/06/19)
+
+#pragma once
+
+#include <GTEngine/Mathematics/GteDelaunay3.h>
+
+namespace gte
+{
+
+template <typename InputType, typename ComputeType, typename RationalType>
+class Delaunay3Mesh
+{
+public:
+    // Construction.
+    Delaunay3Mesh(Delaunay3<InputType, ComputeType> const& delaunay);
+
+    // Mesh information.
+    inline int GetNumVertices() const;
+    inline int GetNumTetrahedra() const;
+    inline Vector3<InputType> const* GetVertices() const;
+    inline int const* GetIndices() const;
+    inline int const* GetAdjacencies() const;
+
+    // Containment queries.
+    int GetContainingTetrahedron(Vector3<InputType> const& P) const;
+    bool GetVertices(int t, std::array<Vector3<InputType>, 4>& vertices)
+        const;
+    bool GetIndices(int t, std::array<int, 4>& indices) const;
+    bool GetAdjacencies(int t, std::array<int, 4>& adjacencies) const;
+    bool GetBarycentrics(int t, Vector3<InputType> const& P,
+        std::array<InputType, 4>& bary) const;
+
+private:
+    Delaunay3<InputType, ComputeType> const* mDelaunay;
+};
+
+
+template <typename InputType, typename ComputeType, typename RationalType>
+Delaunay3Mesh<InputType, ComputeType, RationalType>::Delaunay3Mesh(
+    Delaunay3<InputType, ComputeType> const& delaunay)
+    :
+    mDelaunay(&delaunay)
+{
+}
+
+template <typename InputType, typename ComputeType, typename RationalType>
+inline int Delaunay3Mesh<InputType, ComputeType, RationalType>::
+GetNumVertices() const
+{
+    return mDelaunay->GetNumVertices();
+}
+
+template <typename InputType, typename ComputeType, typename RationalType>
+inline int Delaunay3Mesh<InputType, ComputeType, RationalType>::
+GetNumTetrahedra() const
+{
+    return mDelaunay->GetNumTetrahedra();
+}
+
+template <typename InputType, typename ComputeType, typename RationalType>
+inline Vector3<InputType> const*
+Delaunay3Mesh<InputType, ComputeType, RationalType>::
+GetVertices() const
+{
+    return mDelaunay->GetVertices();
+}
+
+template <typename InputType, typename ComputeType, typename RationalType>
+inline int const* Delaunay3Mesh<InputType, ComputeType, RationalType>::
+GetIndices() const
+{
+    return &mDelaunay->GetIndices()[0];
+}
+
+template <typename InputType, typename ComputeType, typename RationalType>
+inline int const* Delaunay3Mesh<InputType, ComputeType, RationalType>::
+GetAdjacencies() const
+{
+    return &mDelaunay->GetAdjacencies()[0];
+}
+
+template <typename InputType, typename ComputeType, typename RationalType>
+int Delaunay3Mesh<InputType, ComputeType, RationalType>::
+GetContainingTetrahedron(Vector3<InputType> const& P) const
+{
+    typename Delaunay3<InputType, ComputeType>::SearchInfo info;
+    return mDelaunay->GetContainingTetrahedron(P, info);
+}
+
+template <typename InputType, typename ComputeType, typename RationalType>
+bool Delaunay3Mesh<InputType, ComputeType, RationalType>::
+GetVertices(int t, std::array<Vector3<InputType>, 4>& vertices) const
+{
+    if (mDelaunay->GetDimension() == 3)
+    {
+        std::array<int, 4> indices;
+        if (mDelaunay->GetIndices(t, indices))
+        {
+            PrimalQuery3<ComputeType> const& query = mDelaunay->GetQuery();
+            Vector3<ComputeType> const* ctVertices = query.GetVertices();
+            for (int i = 0; i < 4; ++i)
+            {
+                Vector3<ComputeType> const& V = ctVertices[indices[i]];
+                for (int j = 0; j < 3; ++j)
+                {
+                    vertices[i][j] = (InputType)V[j];
+                }
+            }
+            return true;
+        }
+    }
+    return false;
+}
+
+template <typename InputType, typename ComputeType, typename RationalType>
+bool Delaunay3Mesh<InputType, ComputeType, RationalType>::
+GetIndices(int t, std::array<int, 4>& indices) const
+{
+    return mDelaunay->GetIndices(t, indices);
+}
+
+template <typename InputType, typename ComputeType, typename RationalType>
+bool Delaunay3Mesh<InputType, ComputeType, RationalType>::
+GetAdjacencies(int t, std::array<int, 4>& indices) const
+{
+    return mDelaunay->GetAdjacencies(t, indices);
+}
+
+template <typename InputType, typename ComputeType, typename RationalType>
+bool Delaunay3Mesh<InputType, ComputeType, RationalType>::
+GetBarycentrics(int t, Vector3<InputType> const& P,
+std::array<InputType, 4>& bary) const
+{
+    std::array<int, 4> indices;
+    if (mDelaunay->GetIndices(t, indices))
+    {
+        PrimalQuery3<ComputeType> const& query = mDelaunay->GetQuery();
+        Vector3<ComputeType> const* vertices = query.GetVertices();
+        Vector3<RationalType> rtP{ P[0], P[1], P[2] };
+        std::array<Vector3<RationalType>, 4> rtV;
+        for (int i = 0; i < 4; ++i)
+        {
+            Vector3<ComputeType> const& V = vertices[indices[i]];
+            for (int j = 0; j < 3; ++j)
+            {
+                rtV[i][j] = (RationalType)V[j];
+            }
+        };
+
+        RationalType rtBary[4];
+        if (ComputeBarycentrics(rtP, rtV[0], rtV[1], rtV[2], rtV[3], rtBary))
+        {
+            for (int i = 0; i < 4; ++i)
+            {
+                bary[i] = (InputType)rtBary[i];
+            }
+            return true;
+        }
+    }
+    return false;
+}
+
+
+}
diff --git a/Engine/Plugins/Experimental/GeometryProcessing/Source/GeometryAlgorithms/Private/GTEngine/Mathematics/GteDistAlignedBox3OrientedBox3.h b/Engine/Plugins/Experimental/GeometryProcessing/Source/GeometryAlgorithms/Private/GTEngine/Mathematics/GteDistAlignedBox3OrientedBox3.h
new file mode 100644
--- /dev/null
+++ b/Engine/Plugins/Experimental/GeometryProcessing/Source/GeometryAlgorithms/Private/GTEngine/Mathematics/GteDistAlignedBox3OrientedBox3.h
@@ -0,0 +1,154 @@
+// David Eberly, Geometric Tools, Redmond WA 98052
+// Copyright (c) 1998-2019
+// Distributed under the Boost Software License, Version 1.0.
+// http://www.boost.org/LICENSE_1_0.txt
+// http://www.geometrictools.com/License/Boost/LICENSE_1_0.txt
+// File Version: 3.5.2 (2018/10/05)
+
+#pragma once
+
+#include <GTEngine/Mathematics/GteDCPQuery.h>
+#include <GTEngine/Mathematics/GteLCPSolver.h>
+#include <GTEngine/Mathematics/GteAlignedBox.h>
+#include <GTEngine/Mathematics/GteOrientedBox.h>
+#include <GTEngine/Mathematics/GteVector3.h>
+
+namespace gte
+{
+
+template <typename Real>
+class DCPQuery<Real, AlignedBox3<Real>, OrientedBox3<Real>>
+{
+public:
+    struct Result
+    {
+        bool queryIsSuccessful;
+
+        // These members are valid only when queryIsSuccessful is true;
+        // otherwise, they are all set to zero.
+        Real distance, sqrDistance;
+        std::array<Real, 3> box0Parameter, box1Parameter;
+        Vector3<Real> closestPoint[2];
+
+        // The number of iterations used by LCPSolver regardless of whether
+        // the query is successful.
+        int numLCPIterations;
+    };
+
+    // The default maximum iterations is 144 (n = 12, maxIterations = n*n).  If
+    // the solver fails to converge, try increasing the maximum number of
+    // iterations.
+    void SetMaxLCPIterations(int maxLCPIterations);
+
+    Result operator()(AlignedBox3<Real> const& box0, OrientedBox3<Real> const& box1);
+
+private:
+    LCPSolver<Real, 12> mLCP;
+};
+
+
+template <typename Real>
+void DCPQuery<Real, AlignedBox3<Real>, OrientedBox3<Real>>::SetMaxLCPIterations(int maxLCPIterations)
+{
+    mLCP.SetMaxIterations(maxLCPIterations);
+}
+
+template <typename Real>
+typename DCPQuery<Real, AlignedBox3<Real>, OrientedBox3<Real>>::Result
+    DCPQuery<Real, AlignedBox3<Real>, OrientedBox3<Real>>::operator()(
+        AlignedBox3<Real> const& box0, OrientedBox3<Real> const& box1)
+{
+    Result result;
+
+    // Translate the boxes so that the aligned box becomes a canonical box.
+    // Modify the oriented box coefficients to be nonnegative.
+    Vector3<Real> K = box0.max - box0.min;
+    Vector3<Real> delta = box1.center - box0.min;
+    for (int i = 0; i < 3; ++i)
+    {
+        delta -= box1.extent[i] * box1.axis[i];
+    }
+
+    Vector3<Real> rotDelta;
+    for (int i = 0; i < 3; ++i)
+    {
+        rotDelta[i] = Dot(box1.axis[i], delta);
+    }
+
+    Vector3<Real> twoExtent = box1.extent * (Real)2;
+
+    // The LCP has 6 variables and 6 (nontrivial) inequality constraints.
+    std::array<Real, 12> q =
+    {
+        -delta[0], -delta[1], -delta[2], rotDelta[0], rotDelta[1], rotDelta[2],
+        K[0], K[1], K[2], twoExtent[0], twoExtent[1], twoExtent[2]
+    };
+
+    std::array<std::array<Real, 12>, 12> M;
+    {
+        Real const z = (Real)0;
+        Real const p = (Real)1;
+        Real const m = (Real)-1;
+        Vector3<Real> const& U0 = box1.axis[0];
+        Vector3<Real> const& U1 = box1.axis[1];
+        Vector3<Real> const& U2 = box1.axis[2];
+        M[ 0] = { p, z, z, -U0[0], -U1[0], -U1[2], p, z, z, z, z, z };
+        M[ 1] = { z, p, z, -U0[1], -U1[1], -U1[1], z, p, z, z, z, z };
+        M[ 2] = { z, z, p, -U0[2], -U1[2], -U1[2], z, z, p, z, z, z };
+        M[ 3] = { -U0[0], -U0[1], -U0[2], p, z, z, z, z, z, p, z, z };
+        M[ 4] = { -U1[0], -U1[1], -U1[2], z, p, z, z, z, z, z, p, z };
+        M[ 5] = { -U2[0], -U2[1], -U2[2], z, z, p, z, z, z, z, z, p };
+        M[ 6] = { m, z, z, z, z, z, z, z, z, z, z, z };
+        M[ 7] = { z, m, z, z, z, z, z, z, z, z, z, z };
+        M[ 8] = { z, z, m, z, z, z, z, z, z, z, z, z };
+        M[ 9] = { z, z, z, m, z, z, z, z, z, z, z, z };
+        M[10] = { z, z, z, z, m, z, z, z, z, z, z, z };
+        M[11] = { z, z, z, z, z, m, z, z, z, z, z, z };
+    }
+
+    std::array<Real, 12> w, z;
+    if (mLCP.Solve(q, M, w, z))
+    {
+        result.queryIsSuccessful = true;
+
+        for (int i = 0; i < 3; ++i)
+        {
+            result.box0Parameter[i] = z[i] + box0.min[i];
+            result.closestPoint[0][i] = result.box0Parameter[i];
+        }
+
+        result.closestPoint[1] = box1.center;
+        for (int i = 0, j = 3; i < 3; ++i, ++j)
+        {
+            result.box1Parameter[i] = z[j] - box1.extent[i];
+            result.closestPoint[1] += result.box1Parameter[i] * box1.axis[i];
+        }
+
+        Vector3<Real> diff = result.closestPoint[1] - result.closestPoint[0];
+        result.sqrDistance = Dot(diff, diff);
+        result.distance = std::sqrt(result.sqrDistance);
+    }
+    else
+    {
+        // If you reach this case, the maximum number of iterations was not
+        // specified to be large enough or there is a problem due to
+        // floating-point rounding errors.  If you believe the latter is
+        // true, file a bug report.
+        result.queryIsSuccessful = false;
+
+        for (int i = 0; i < 3; ++i)
+        {
+            result.box0Parameter[i] = (Real)0;
+            result.box1Parameter[i] = (Real)0;
+            result.closestPoint[0][i] = (Real)0;
+            result.closestPoint[1][i] = (Real)0;
+        }
+        result.distance = (Real)0;
+        result.sqrDistance = (Real)0;
+    }
+
+    result.numLCPIterations = mLCP.GetNumIterations();
+    return result;
+}
+
+}
diff --git a/Engine/Plugins/Experimental/GeometryProcessing/Source/GeometryAlgorithms/Private/GTEngine/Mathematics/GteDistAlignedBoxAlignedBox.h b/Engine/Plugins/Experimental/GeometryProcessing/Source/GeometryAlgorithms/Private/GTEngine/Mathematics/GteDistAlignedBoxAlignedBox.h
new file mode 100644
--- /dev/null
+++ b/Engine/Plugins/Experimental/GeometryProcessing/Source/GeometryAlgorithms/Private/GTEngine/Mathematics/GteDistAlignedBoxAlignedBox.h
@@ -0,0 +1,79 @@
+// David Eberly, Geometric Tools, Redmond WA 98052
+// Copyright (c) 1998-2019
+// Distributed under the Boost Software License, Version 1.0.
+// http://www.boost.org/LICENSE_1_0.txt
+// http://www.geometrictools.com/License/Boost/LICENSE_1_0.txt
+// File Version: 3.5.2 (2018/10/05)
+
+#pragma once
+
+#include <GTEngine/Mathematics/GteDCPQuery.h>
+#include <GTEngine/Mathematics/GteAlignedBox.h>
+#include <GTEngine/Mathematics/GteIntrIntervals.h>
+
+namespace gte
+{
+
+template <int N, typename Real>
+class DCPQuery<Real, AlignedBox<N, Real>, AlignedBox<N, Real>>
+{
+public:
+    struct Result
+    {
+        Real distance, sqrDistance;
+
+        // To compute a single closest point on each box, use
+        // Vector<N, Real> closest0 = (closestPoints[0].min + closestPoints[0].max)/2;
+        // Vector<N, Real> closest1 = (closestPoints[1].min + closestPoints[1].max)/2;
+        AlignedBox<N, Real> closestPoints[2];
+    };
+
+    Result operator()(AlignedBox<N, Real> const& box0, AlignedBox<N, Real> const& box1);
+};
+
+
+template <int N, typename Real>
+typename DCPQuery<Real, AlignedBox<N, Real>, AlignedBox<N, Real>>::Result
+DCPQuery<Real, AlignedBox<N, Real>, AlignedBox<N, Real>>::operator()(
+    AlignedBox<N, Real> const& box0, AlignedBox<N, Real> const& box1)
+{
+    Result result;
+    result.sqrDistance = (Real)0;
+    for (int i = 0; i < N; ++i)
+    {
+        if (box0.min[i] >= box1.max[i])
+        {
+            Real delta = box0.min[i] - box1.min[i];
+            result.sqrDistance += delta * delta;
+            result.closestPoints[0].min[i] = box0.min[i];
+            result.closestPoints[0].max[i] = box0.min[i];
+            result.closestPoints[1].min[i] = box1.max[i];
+            result.closestPoints[1].max[i] = box1.max[i];
+        }
+        else if (box1.min[i] >= box0.max[i])
+        {
+            Real delta = box1.min[i] - box0.max[i];
+            result.sqrDistance += delta * delta;
+            result.closestPoints[0].min[i] = box0.max[i];
+            result.closestPoints[0].max[i] = box0.max[i];
+            result.closestPoints[1].min[i] = box1.min[i];
+            result.closestPoints[1].max[i] = box1.min[i];
+        }
+        else
+        {
+            std::array<Real, 2> intr0 = { box0.min[i], box0.max[i] };
+            std::array<Real, 2> intr1 = { box1.min[i], box1.max[i] };
+            FIQuery<Real, std::array<Real, 2>, std::array<Real, 2>> query;
+            auto iiResult = query(intr0, intr1);
+            for (int j = 0; j < 2; ++j)
+            {
+                result.closestPoints[j].min[i] = iiResult.overlap[0];
+                result.closestPoints[j].max[i] = iiResult.overlap[1];
+            }
+        }
+    }
+    result.distance = std::sqrt(result.sqrDistance);
+    return result;
+}
+
+}
diff --git a/Engine/Plugins/Experimental/GeometryProcessing/Source/GeometryAlgorithms/Private/GTEngine/Mathematics/GteDistCircle3Circle3.h b/Engine/Plugins/Experimental/GeometryProcessing/Source/GeometryAlgorithms/Private/GTEngine/Mathematics/GteDistCircle3Circle3.h
new file mode 100644
--- /dev/null
+++ b/Engine/Plugins/Experimental/GeometryProcessing/Source/GeometryAlgorithms/Private/GTEngine/Mathematics/GteDistCircle3Circle3.h
@@ -0,0 +1,382 @@
+// David Eberly, Geometric Tools, Redmond WA 98052
+// Copyright (c) 1998-2019
+// Distributed under the Boost Software License, Version 1.0.
+// http://www.boost.org/LICENSE_1_0.txt
+// http://www.geometrictools.com/License/Boost/LICENSE_1_0.txt
+// File Version: 3.0.3 (2018/10/17)
+
+#pragma once
+
+#include <GTEngine/LowLevel/GteLogger.h>
+#include <GTEngine/Mathematics/GteMath.h>
+#include <GTEngine/Mathematics/GtePolynomial1.h>
+#include <GTEngine/Mathematics/GteDCPQuery.h>
+#include <GTEngine/Mathematics/GteCircle3.h>
+#include <GTEngine/Mathematics/GteRootsPolynomial.h>
+
+// The 3D circle-circle distance algorithm is described in
+// http://www.geometrictools.com/Documentation/DistanceToCircle3.pdf
+// The notation used in the code matches that of the document.
+
+namespace gte
+{
+    template <typename Real>
+    class DCPQuery<Real, Circle3<Real>, Circle3<Real>>
+    {
+    public:
+        struct Result
+        {
+            Real distance, sqrDistance;
+            int numClosestPairs;
+            Vector3<Real> circle0Closest[2], circle1Closest[2];
+            bool equidistant;
+        };
+
+        Result operator()(Circle3<Real> const& circle0, Circle3<Real> const& circle1)
+        {
+            Result result;
+            Vector3<Real> const vzero = Vector3<Real>::Zero();
+            Real const zero = (Real)0;
+
+            Vector3<Real> N0 = circle0.normal, N1 = circle1.normal;
+            Real r0 = circle0.radius, r1 = circle1.radius;
+            Vector3<Real> D = circle1.center - circle0.center;
+            Vector3<Real> N0xN1 = Cross(N0, N1);
+
+            if (N0xN1 != vzero)
+            {
+                // Get parameters for constructing the degree-8 polynomial phi.
+                Real const one = (Real)1, two = (Real)2;
+                Real r0sqr = r0 * r0, r1sqr = r1 * r1;
+
+                // Compute U1 and V1 for the plane of circle1.
+                Vector3<Real> basis[3];
+                basis[0] = circle1.normal;
+                ComputeOrthogonalComplement(1, basis);
+                Vector3<Real> U1 = basis[1], V1 = basis[2];
+
+                // Construct the polynomial phi(cos(theta)).
+                Vector3<Real> N0xD = Cross(N0, D);
+                Vector3<Real> N0xU1 = Cross(N0, U1), N0xV1 = Cross(N0, V1);
+                Real a0 = r1 * Dot(D, U1), a1 = r1 * Dot(D, V1);
+                Real a2 = Dot(N0xD, N0xD), a3 = r1 * Dot(N0xD, N0xU1);
+                Real a4 = r1 * Dot(N0xD, N0xV1), a5 = r1sqr * Dot(N0xU1, N0xU1);
+                Real a6 = r1sqr * Dot(N0xU1, N0xV1), a7 = r1sqr * Dot(N0xV1, N0xV1);
+                Polynomial1<Real> p0{ a2 + a7, two * a3, a5 - a7 };
+                Polynomial1<Real> p1{ two * a4, two * a6 };
+                Polynomial1<Real> p2{ zero, a1 };
+                Polynomial1<Real> p3{ -a0 };
+                Polynomial1<Real> p4{ -a6, a4, two * a6 };
+                Polynomial1<Real> p5{ -a3, a7 - a5 };
+                Polynomial1<Real> tmp0{ one, zero, -one };
+                Polynomial1<Real> tmp1 = p2 * p2 + tmp0 * p3 * p3;
+                Polynomial1<Real> tmp2 = two * p2 * p3;
+                Polynomial1<Real> tmp3 = p4 * p4 + tmp0 * p5 * p5;
+                Polynomial1<Real> tmp4 = two * p4 * p5;
+                Polynomial1<Real> p6 = p0 * tmp1 + tmp0 * p1 * tmp2 - r0sqr * tmp3;
+                Polynomial1<Real> p7 = p0 * tmp2 + p1 * tmp1 - r0sqr * tmp4;
+
+                // The use of 'double' is intentional in case Real is a BSNumber or
+                // BSRational type.  We want the bisections to terminate in a
+                // reasonable amount of time.
+                unsigned int const maxIterations = GTE_C_MAX_BISECTIONS_GENERIC;
+                Real roots[8], sn, temp;
+                int i, degree, numRoots;
+
+                // The RootsPolynomial<Real>::Find(...) function currently does not
+                // combine duplicate roots.  We need only the unique ones here.
+                std::set<Real> uniqueRoots;
+
+                std::array<std::pair<Real, Real>, 16> pairs;
+                int numPairs = 0;
+                if (p7.GetDegree() > 0 || p7[0] != zero)
+                {
+                    // H(cs,sn) = p6(cs) + sn * p7(cs)
+                    Polynomial1<Real> phi = p6 * p6 - tmp0 * p7 * p7;
+                    degree = static_cast<int>(phi.GetDegree());
+                    LogAssert(degree > 0, "Unexpected degree for phi.");
+                    numRoots = RootsPolynomial<Real>::Find(degree, &phi[0],
+                        maxIterations, roots);
+                    for (i = 0; i < numRoots; ++i)
+                    {
+                        uniqueRoots.insert(roots[i]);
+                    }
+
+                    for (auto cs : uniqueRoots)
+                    {
+                        if (std::abs(cs) <= one)
+                        {
+                            temp = p7(cs);
+                            if (temp != zero)
+                            {
+                                sn = -p6(cs) / temp;
+                                pairs[numPairs++] = std::make_pair(cs, sn);
+                            }
+                            else
+                            {
+                                temp = std::max(one - cs * cs, zero);
+                                sn = std::sqrt(temp);
+                                pairs[numPairs++] = std::make_pair(cs, sn);
+                                if (sn != zero)
+                                {
+                                    pairs[numPairs++] = std::make_pair(cs, -sn);
+                                }
+                            }
+                        }
+                    }
+                }
+                else
+                {
+                    // H(cs,sn) = p6(cs)
+                    degree = static_cast<int>(p6.GetDegree());
+                    LogAssert(degree > 0, "Unexpected degree for p6.");
+                    numRoots = RootsPolynomial<Real>::Find(degree, &p6[0],
+                        maxIterations, roots);
+                    for (i = 0; i < numRoots; ++i)
+                    {
+                        uniqueRoots.insert(roots[i]);
+                    }
+
+                    for (auto cs : uniqueRoots)
+                    {
+                        if (std::abs(cs) <= one)
+                        {
+                            temp = std::max(one - cs * cs, zero);
+                            sn = std::sqrt(temp);
+                            pairs[numPairs++] = std::make_pair(cs, sn);
+                            if (sn != zero)
+                            {
+                                pairs[numPairs++] = std::make_pair(cs, -sn);
+                            }
+                        }
+                    }
+                }
+
+                std::array<ClosestInfo, 16> candidates;
+                for (i = 0; i < numPairs; ++i)
+                {
+                    ClosestInfo& info = candidates[i];
+                    Vector3<Real> delta =
+                        D + r1 * (pairs[i].first * U1 + pairs[i].second * V1);
+                    info.circle1Closest = circle0.center + delta;
+                    Real N0dDelta = Dot(N0, delta);
+                    Real lenN0xDelta = Length(Cross(N0, delta));
+                    if (lenN0xDelta > 0)
+                    {
+                        Real diff = lenN0xDelta - r0;
+                        info.sqrDistance = N0dDelta * N0dDelta + diff * diff;
+                        delta -= N0dDelta * circle0.normal;
+                        Normalize(delta);
+                        info.circle0Closest = circle0.center + r0 * delta;
+                        info.equidistant = false;
+                    }
+                    else
+                    {
+                        Vector3<Real> r0U0 = r0 * GetOrthogonal(N0, true);
+                        Vector3<Real> diff = delta - r0U0;
+                        info.sqrDistance = Dot(diff, diff);
+                        info.circle0Closest = circle0.center + r0U0;
+                        info.equidistant = true;
+                    }
+                }
+
+                std::sort(candidates.begin(), candidates.begin() + numPairs);
+
+                result.numClosestPairs = 1;
+                result.sqrDistance = candidates[0].sqrDistance;
+                result.circle0Closest[0] = candidates[0].circle0Closest;
+                result.circle1Closest[0] = candidates[0].circle1Closest;
+                result.equidistant = candidates[0].equidistant;
+                if (numRoots > 1
+                    && candidates[1].sqrDistance == candidates[0].sqrDistance)
+                {
+                    result.numClosestPairs = 2;
+                    result.circle0Closest[1] = candidates[1].circle0Closest;
+                    result.circle1Closest[1] = candidates[1].circle1Closest;
+                }
+            }
+            else
+            {
+                // The planes of the circles are parallel.  Whether the planes are the
+                // same or different, the problem reduces to determining how two
+                // circles in the same plane are separated, tangent with one circle
+                // outside the other, overlapping, or one circle contained inside the
+                // other circle.
+                DoQueryParallelPlanes(circle0, circle1, D, result);
+            }
+
+            result.distance = std::sqrt(result.sqrDistance);
+            return result;
+        }
+
+    private:
+        class SCPolynomial
+        {
+        public:
+            SCPolynomial()
+            {
+            }
+
+            SCPolynomial(Real oneTerm, Real cosTerm, Real sinTerm)
+            {
+                mPoly[0] = Polynomial1<Real>{ oneTerm, cosTerm };
+                mPoly[1] = Polynomial1<Real>{ sinTerm };
+            }
+
+            inline Polynomial1<Real> const& operator[] (unsigned int i) const
+            {
+                return mPoly[i];
+            }
+
+            inline Polynomial1<Real>& operator[] (unsigned int i)
+            {
+                return mPoly[i];
+            }
+
+            SCPolynomial operator+(SCPolynomial const& object) const
+            {
+                SCPolynomial result;
+                result.mPoly[0] = mPoly[0] + object.mPoly[0];
+                result.mPoly[1] = mPoly[1] + object.mPoly[1];
+                return result;
+            }
+
+            SCPolynomial operator-(SCPolynomial const& object) const
+            {
+                SCPolynomial result;
+                result.mPoly[0] = mPoly[0] - object.mPoly[0];
+                result.mPoly[1] = mPoly[1] - object.mPoly[1];
+                return result;
+            }
+
+            SCPolynomial operator*(SCPolynomial const& object) const
+            {
+                Polynomial1<Real> omcsqr{ (Real)1, (Real)0, (Real)-1 };  // 1 - c^2
+                SCPolynomial result;
+                result.mPoly[0] = mPoly[0] * object.mPoly[0] + omcsqr * mPoly[1] * object.mPoly[1];
+                result.mPoly[1] = mPoly[0] * object.mPoly[1] + mPoly[1] * object.mPoly[0];
+                return result;
+            }
+
+            SCPolynomial operator*(Real scalar) const
+            {
+                SCPolynomial result;
+                result.mPoly[0] = scalar * mPoly[0];
+                result.mPoly[1] = scalar * mPoly[1];
+                return result;
+            }
+
+        private:
+            Polynomial1<Real> mPoly[2];  // poly0(c) + s * poly1(c)
+        };
+
+        struct ClosestInfo
+        {
+            Real sqrDistance;
+            Vector3<Real> circle0Closest, circle1Closest;
+            bool equidistant;
+
+            inline bool operator< (ClosestInfo const& info) const
+            {
+                return sqrDistance < info.sqrDistance;
+            }
+        };
+
+        // The two circles are in parallel planes where D = C1 - C0, the
+        // difference of circle centers.
+        void DoQueryParallelPlanes(Circle3<Real> const& circle0, Circle3<Real> const& circle1,
+            Vector3<Real> const& D, Result& result)
+        {
+            Real N0dD = Dot(circle0.normal, D);
+            Vector3<Real> normProj = N0dD * circle0.normal;
+            Vector3<Real> compProj = D - normProj;
+            Vector3<Real> U = compProj;
+            Real d = Normalize(U);
+
+            // The configuration is determined by the relative location of the
+            // intervals of projection of the circles on to the D-line.  Circle0
+            // projects to [-r0,r0] and circle1 projects to [d-r1,d+r1].
+            Real r0 = circle0.radius, r1 = circle1.radius;
+            Real dmr1 = d - r1;
+            Real distance;
+            if (dmr1 >= r0)  // d >= r0 + r1
+            {
+                // The circles are separated (d > r0 + r1) or tangent with one
+                // outside the other (d = r0 + r1).
+                distance = dmr1 - r0;
+                result.numClosestPairs = 1;
+                result.circle0Closest[0] = circle0.center + r0 * U;
+                result.circle1Closest[0] = circle1.center - r1 * U;
+                result.equidistant = false;
+            }
+            else // d < r0 + r1
+            {
+                // The cases implicitly use the knowledge that d >= 0.
+                Real dpr1 = d + r1;
+                if (dpr1 <= r0)
+                {
+                    // Circle1 is inside circle0.
+                    distance = r0 - dpr1;
+                    result.numClosestPairs = 1;
+                    if (d > (Real)0)
+                    {
+                        result.circle0Closest[0] = circle0.center + r0 * U;
+                        result.circle1Closest[0] = circle1.center + r1 * U;
+                        result.equidistant = false;
+                    }
+                    else
+                    {
+                        // The circles are concentric, so U = (0,0,0).  Construct
+                        // a vector perpendicular to N0 to use for closest points.
+                        U = GetOrthogonal(circle0.normal, true);
+                        result.circle0Closest[0] = circle0.center + r0 * U;
+                        result.circle1Closest[0] = circle1.center + r1 * U;
+                        result.equidistant = true;
+                    }
+                }
+                else if (dmr1 <= -r0)
+                {
+                    // Circle0 is inside circle1.
+                    distance = -r0 - dmr1;
+                    result.numClosestPairs = 1;
+                    if (d > (Real)0)
+                    {
+                        result.circle0Closest[0] = circle0.center - r0 * U;
+                        result.circle1Closest[0] = circle1.center - r1 * U;
+                        result.equidistant = false;
+                    }
+                    else
+                    {
+                        // The circles are concentric, so U = (0,0,0).  Construct
+                        // a vector perpendicular to N0 to use for closest points.
+                        U = GetOrthogonal(circle0.normal, true);
+                        result.circle0Closest[0] = circle0.center + r0 * U;
+                        result.circle1Closest[0] = circle1.center + r1 * U;
+                        result.equidistant = true;
+                    }
+                }
+                else
+                {
+                    // The circles are overlapping.  The two points of intersection
+                    // are C0 + s*(C1-C0) +/- h*Cross(N,U), where
+                    // s = (1 + (r0^2 - r1^2)/d^2)/2 and h = sqrt(r0^2 - s^2 * d^2).
+                    Real r0sqr = r0 * r0, r1sqr = r1 * r1, dsqr = d * d;
+                    Real s = ((Real)1 + (r0sqr - r1sqr) / dsqr) / (Real)2;
+                    Real arg = std::max(r0sqr - dsqr * s * s, (Real)0);
+                    Real h = std::sqrt(arg);
+                    Vector3<Real> midpoint = circle0.center + s * compProj;
+                    Vector3<Real> hNxU = h * Cross(circle0.normal, U);
+                    distance = (Real)0;
+                    result.numClosestPairs = 2;
+                    result.circle0Closest[0] = midpoint + hNxU;
+                    result.circle0Closest[1] = midpoint - hNxU;
+                    result.circle1Closest[0] = result.circle0Closest[0] + normProj;
+                    result.circle1Closest[1] = result.circle0Closest[1] + normProj;
+                    result.equidistant = false;
+                }
+            }
+
+            result.sqrDistance = distance * distance + N0dD * N0dD;
+        }
+    };
+}
diff --git a/Engine/Plugins/Experimental/GeometryProcessing/Source/GeometryAlgorithms/Private/GTEngine/Mathematics/GteDistLine3AlignedBox3.h b/Engine/Plugins/Experimental/GeometryProcessing/Source/GeometryAlgorithms/Private/GTEngine/Mathematics/GteDistLine3AlignedBox3.h
new file mode 100644
--- /dev/null
+++ b/Engine/Plugins/Experimental/GeometryProcessing/Source/GeometryAlgorithms/Private/GTEngine/Mathematics/GteDistLine3AlignedBox3.h
@@ -0,0 +1,553 @@
+// David Eberly, Geometric Tools, Redmond WA 98052
+// Copyright (c) 1998-2019
+// Distributed under the Boost Software License, Version 1.0.
+// http://www.boost.org/LICENSE_1_0.txt
+// http://www.geometrictools.com/License/Boost/LICENSE_1_0.txt
+// File Version: 3.0.1 (2018/10/05)
+
+#pragma once
+
+#include <GTEngine/Mathematics/GteVector3.h>
+#include <GTEngine/Mathematics/GteDistPointAlignedBox.h>
+#include <GTEngine/Mathematics/GteLine.h>
+
+namespace gte
+{
+
+template <typename Real>
+class DCPQuery<Real, Line3<Real>, AlignedBox3<Real>>
+{
+public:
+    struct Result
+    {
+        Real distance, sqrDistance;
+        Real lineParameter;
+        Vector3<Real> closestPoint[2];
+    };
+
+    Result operator()(Line3<Real> const& line, AlignedBox3<Real> const& box);
+
+protected:
+    // Compute the distance and closest point between a line and an
+    // axis-aligned box whose center is the origin.  On input, 'point' is the
+    // line origin and 'direction' is the line direction.  On output, 'point'
+    // is the point on the box closest to the line.  The 'direction' is
+    // non-const to allow transforming the problem into the first octant.
+    void DoQuery(Vector3<Real>& point, Vector3<Real>& direction,
+        Vector3<Real> const& boxExtent, Result& result);
+
+private:
+    void Face(int i0, int i1, int i2, Vector3<Real>& pnt,
+        Vector3<Real> const& dir, Vector3<Real> const& PmE,
+        Vector3<Real> const& boxExtent, Result& result);
+
+    void CaseNoZeros(Vector3<Real>& pnt, Vector3<Real> const& dir,
+        Vector3<Real> const& boxExtent, Result& result);
+
+    void Case0(int i0, int i1, int i2, Vector3<Real>& pnt,
+        Vector3<Real> const& dir, Vector3<Real> const& boxExtent,
+        Result& result);
+
+    void Case00(int i0, int i1, int i2, Vector3<Real>& pnt,
+        Vector3<Real> const& dir, Vector3<Real> const& boxExtent,
+        Result& result);
+
+    void Case000(Vector3<Real>& pnt, Vector3<Real> const& boxExtent,
+        Result& result);
+};
+
+
+template <typename Real>
+typename DCPQuery<Real, Line3<Real>, AlignedBox3<Real>>::Result
+DCPQuery<Real, Line3<Real>, AlignedBox3<Real>>::operator()(
+    Line3<Real> const& line, AlignedBox3<Real> const& box)
+{
+    // Translate the line and box so that the box has center at the origin.
+    Vector3<Real> boxCenter, boxExtent;
+    box.GetCenteredForm(boxCenter, boxExtent);
+    Vector3<Real> point = line.origin - boxCenter;
+    Vector3<Real> direction = line.direction;
+
+    Result result;
+    DoQuery(point, direction, boxExtent, result);
+
+    // Compute the closest point on the line.
+    result.closestPoint[0] =
+        line.origin + result.lineParameter*line.direction;
+
+    // Compute the closest point on the box.
+    result.closestPoint[1] = boxCenter + point;
+    return result;
+}
+
+template <typename Real>
+void DCPQuery<Real, Line3<Real>, AlignedBox3<Real>>::DoQuery(
+    Vector3<Real>& point, Vector3<Real>& direction,
+    Vector3<Real> const& boxExtent, Result& result)
+{
+    result.sqrDistance = (Real)0;
+    result.lineParameter = (Real)0;
+
+    // Apply reflections so that direction vector has nonnegative components.
+    bool reflect[3];
+    for (int i = 0; i < 3; ++i)
+    {
+        if (direction[i] < (Real)0)
+        {
+            point[i] = -point[i];
+            direction[i] = -direction[i];
+            reflect[i] = true;
+        }
+        else
+        {
+            reflect[i] = false;
+        }
+    }
+
+    if (direction[0] > (Real)0)
+    {
+        if (direction[1] > (Real)0)
+        {
+            if (direction[2] > (Real)0)  // (+,+,+)
+            {
+                CaseNoZeros(point, direction, boxExtent, result);
+            }
+            else  // (+,+,0)
+            {
+                Case0(0, 1, 2, point, direction, boxExtent, result);
+            }
+        }
+        else
+        {
+            if (direction[2] > (Real)0)  // (+,0,+)
+            {
+                Case0(0, 2, 1, point, direction, boxExtent, result);
+            }
+            else  // (+,0,0)
+            {
+                Case00(0, 1, 2, point, direction, boxExtent, result);
+            }
+        }
+    }
+    else
+    {
+        if (direction[1] > (Real)0)
+        {
+            if (direction[2] > (Real)0)  // (0,+,+)
+            {
+                Case0(1, 2, 0, point, direction, boxExtent, result);
+            }
+            else  // (0,+,0)
+            {
+                Case00(1, 0, 2, point, direction, boxExtent, result);
+            }
+        }
+        else
+        {
+            if (direction[2] > (Real)0)  // (0,0,+)
+            {
+                Case00(2, 0, 1, point, direction, boxExtent, result);
+            }
+            else  // (0,0,0)
+            {
+                Case000(point, boxExtent, result);
+            }
+        }
+    }
+
+    // Undo the reflections applied previously.
+    for (int i = 0; i < 3; ++i)
+    {
+        if (reflect[i])
+        {
+            point[i] = -point[i];
+        }
+    }
+
+    result.distance = std::sqrt(result.sqrDistance);
+}
+
+template <typename Real>
+void DCPQuery<Real, Line3<Real>, AlignedBox3<Real>>::Face(int i0, int i1,
+    int i2, Vector3<Real>& pnt, Vector3<Real> const& dir,
+    Vector3<Real> const& PmE, Vector3<Real> const& boxExtent, Result& result)
+{
+    Vector3<Real> PpE;
+    Real lenSqr, inv, tmp, param, t, delta;
+
+    PpE[i1] = pnt[i1] + boxExtent[i1];
+    PpE[i2] = pnt[i2] + boxExtent[i2];
+    if (dir[i0] * PpE[i1] >= dir[i1] * PmE[i0])
+    {
+        if (dir[i0] * PpE[i2] >= dir[i2] * PmE[i0])
+        {
+            // v[i1] >= -e[i1], v[i2] >= -e[i2] (distance = 0)
+            pnt[i0] = boxExtent[i0];
+            inv = ((Real)1) / dir[i0];
+            pnt[i1] -= dir[i1] * PmE[i0] * inv;
+            pnt[i2] -= dir[i2] * PmE[i0] * inv;
+            result.lineParameter = -PmE[i0] * inv;
+        }
+        else
+        {
+            // v[i1] >= -e[i1], v[i2] < -e[i2]
+            lenSqr = dir[i0] * dir[i0] + dir[i2] * dir[i2];
+            tmp = lenSqr*PpE[i1] - dir[i1] * (dir[i0] * PmE[i0] +
+                dir[i2] * PpE[i2]);
+            if (tmp <= ((Real)2)*lenSqr*boxExtent[i1])
+            {
+                t = tmp / lenSqr;
+                lenSqr += dir[i1] * dir[i1];
+                tmp = PpE[i1] - t;
+                delta = dir[i0] * PmE[i0] + dir[i1] * tmp + dir[i2] * PpE[i2];
+                param = -delta / lenSqr;
+                result.sqrDistance += PmE[i0] * PmE[i0] + tmp*tmp +
+                    PpE[i2] * PpE[i2] + delta*param;
+
+                result.lineParameter = param;
+                pnt[i0] = boxExtent[i0];
+                pnt[i1] = t - boxExtent[i1];
+                pnt[i2] = -boxExtent[i2];
+            }
+            else
+            {
+                lenSqr += dir[i1] * dir[i1];
+                delta = dir[i0] * PmE[i0] + dir[i1] * PmE[i1] + dir[i2] * PpE[i2];
+                param = -delta / lenSqr;
+                result.sqrDistance += PmE[i0] * PmE[i0] + PmE[i1] * PmE[i1]
+                    + PpE[i2] * PpE[i2] + delta*param;
+
+                result.lineParameter = param;
+                pnt[i0] = boxExtent[i0];
+                pnt[i1] = boxExtent[i1];
+                pnt[i2] = -boxExtent[i2];
+            }
+        }
+    }
+    else
+    {
+        if (dir[i0] * PpE[i2] >= dir[i2] * PmE[i0])
+        {
+            // v[i1] < -e[i1], v[i2] >= -e[i2]
+            lenSqr = dir[i0] * dir[i0] + dir[i1] * dir[i1];
+            tmp = lenSqr*PpE[i2] - dir[i2] * (dir[i0] * PmE[i0] +
+                dir[i1] * PpE[i1]);
+            if (tmp <= ((Real)2)*lenSqr*boxExtent[i2])
+            {
+                t = tmp / lenSqr;
+                lenSqr += dir[i2] * dir[i2];
+                tmp = PpE[i2] - t;
+                delta = dir[i0] * PmE[i0] + dir[i1] * PpE[i1] + dir[i2] * tmp;
+                param = -delta / lenSqr;
+                result.sqrDistance += PmE[i0] * PmE[i0] + PpE[i1] * PpE[i1] +
+                    tmp*tmp + delta*param;
+
+                result.lineParameter = param;
+                pnt[i0] = boxExtent[i0];
+                pnt[i1] = -boxExtent[i1];
+                pnt[i2] = t - boxExtent[i2];
+            }
+            else
+            {
+                lenSqr += dir[i2] * dir[i2];
+                delta = dir[i0] * PmE[i0] + dir[i1] * PpE[i1] + dir[i2] * PmE[i2];
+                param = -delta / lenSqr;
+                result.sqrDistance += PmE[i0] * PmE[i0] + PpE[i1] * PpE[i1] +
+                    PmE[i2] * PmE[i2] + delta*param;
+
+                result.lineParameter = param;
+                pnt[i0] = boxExtent[i0];
+                pnt[i1] = -boxExtent[i1];
+                pnt[i2] = boxExtent[i2];
+            }
+        }
+        else
+        {
+            // v[i1] < -e[i1], v[i2] < -e[i2]
+            lenSqr = dir[i0] * dir[i0] + dir[i2] * dir[i2];
+            tmp = lenSqr*PpE[i1] - dir[i1] * (dir[i0] * PmE[i0] +
+                dir[i2] * PpE[i2]);
+            if (tmp >= (Real)0)
+            {
+                // v[i1]-edge is closest
+                if (tmp <= ((Real)2)*lenSqr*boxExtent[i1])
+                {
+                    t = tmp / lenSqr;
+                    lenSqr += dir[i1] * dir[i1];
+                    tmp = PpE[i1] - t;
+                    delta = dir[i0] * PmE[i0] + dir[i1] * tmp + dir[i2] * PpE[i2];
+                    param = -delta / lenSqr;
+                    result.sqrDistance += PmE[i0] * PmE[i0] + tmp*tmp +
+                        PpE[i2] * PpE[i2] + delta*param;
+
+                    result.lineParameter = param;
+                    pnt[i0] = boxExtent[i0];
+                    pnt[i1] = t - boxExtent[i1];
+                    pnt[i2] = -boxExtent[i2];
+                }
+                else
+                {
+                    lenSqr += dir[i1] * dir[i1];
+                    delta = dir[i0] * PmE[i0] + dir[i1] * PmE[i1]
+                        + dir[i2] * PpE[i2];
+                    param = -delta / lenSqr;
+                    result.sqrDistance += PmE[i0] * PmE[i0] + PmE[i1] * PmE[i1]
+                        + PpE[i2] * PpE[i2] + delta*param;
+
+                    result.lineParameter = param;
+                    pnt[i0] = boxExtent[i0];
+                    pnt[i1] = boxExtent[i1];
+                    pnt[i2] = -boxExtent[i2];
+                }
+                return;
+            }
+
+            lenSqr = dir[i0] * dir[i0] + dir[i1] * dir[i1];
+            tmp = lenSqr*PpE[i2] - dir[i2] * (dir[i0] * PmE[i0] +
+                dir[i1] * PpE[i1]);
+            if (tmp >= (Real)0)
+            {
+                // v[i2]-edge is closest
+                if (tmp <= ((Real)2)*lenSqr*boxExtent[i2])
+                {
+                    t = tmp / lenSqr;
+                    lenSqr += dir[i2] * dir[i2];
+                    tmp = PpE[i2] - t;
+                    delta = dir[i0] * PmE[i0] + dir[i1] * PpE[i1] + dir[i2] * tmp;
+                    param = -delta / lenSqr;
+                    result.sqrDistance += PmE[i0] * PmE[i0] + PpE[i1] * PpE[i1] +
+                        tmp*tmp + delta*param;
+
+                    result.lineParameter = param;
+                    pnt[i0] = boxExtent[i0];
+                    pnt[i1] = -boxExtent[i1];
+                    pnt[i2] = t - boxExtent[i2];
+                }
+                else
+                {
+                    lenSqr += dir[i2] * dir[i2];
+                    delta = dir[i0] * PmE[i0] + dir[i1] * PpE[i1]
+                        + dir[i2] * PmE[i2];
+                    param = -delta / lenSqr;
+                    result.sqrDistance += PmE[i0] * PmE[i0] + PpE[i1] * PpE[i1]
+                        + PmE[i2] * PmE[i2] + delta*param;
+
+                    result.lineParameter = param;
+                    pnt[i0] = boxExtent[i0];
+                    pnt[i1] = -boxExtent[i1];
+                    pnt[i2] = boxExtent[i2];
+                }
+                return;
+            }
+
+            // (v[i1],v[i2])-corner is closest
+            lenSqr += dir[i2] * dir[i2];
+            delta = dir[i0] * PmE[i0] + dir[i1] * PpE[i1] + dir[i2] * PpE[i2];
+            param = -delta / lenSqr;
+            result.sqrDistance += PmE[i0] * PmE[i0] + PpE[i1] * PpE[i1]
+                + PpE[i2] * PpE[i2] + delta*param;
+
+            result.lineParameter = param;
+            pnt[i0] = boxExtent[i0];
+            pnt[i1] = -boxExtent[i1];
+            pnt[i2] = -boxExtent[i2];
+        }
+    }
+}
+
+template <typename Real>
+void DCPQuery<Real, Line3<Real>, AlignedBox3<Real>>::CaseNoZeros(
+    Vector3<Real>& pnt, Vector3<Real> const& dir,
+    Vector3<Real> const& boxExtent, Result& result)
+{
+    Vector3<Real> PmE = pnt - boxExtent;
+    Real prodDxPy = dir[0] * PmE[1];
+    Real prodDyPx = dir[1] * PmE[0];
+    Real prodDzPx, prodDxPz, prodDzPy, prodDyPz;
+
+    if (prodDyPx >= prodDxPy)
+    {
+        prodDzPx = dir[2] * PmE[0];
+        prodDxPz = dir[0] * PmE[2];
+        if (prodDzPx >= prodDxPz)
+        {
+            // line intersects x = e0
+            Face(0, 1, 2, pnt, dir, PmE, boxExtent, result);
+        }
+        else
+        {
+            // line intersects z = e2
+            Face(2, 0, 1, pnt, dir, PmE, boxExtent, result);
+        }
+    }
+    else
+    {
+        prodDzPy = dir[2] * PmE[1];
+        prodDyPz = dir[1] * PmE[2];
+        if (prodDzPy >= prodDyPz)
+        {
+            // line intersects y = e1
+            Face(1, 2, 0, pnt, dir, PmE, boxExtent, result);
+        }
+        else
+        {
+            // line intersects z = e2
+            Face(2, 0, 1, pnt, dir, PmE, boxExtent, result);
+        }
+    }
+}
+
+template <typename Real>
+void DCPQuery<Real, Line3<Real>, AlignedBox3<Real>>::Case0(int i0, int i1,
+    int i2, Vector3<Real>& pnt, Vector3<Real> const& dir,
+    Vector3<Real> const& boxExtent, Result& result)
+{
+    Real PmE0 = pnt[i0] - boxExtent[i0];
+    Real PmE1 = pnt[i1] - boxExtent[i1];
+    Real prod0 = dir[i1] * PmE0;
+    Real prod1 = dir[i0] * PmE1;
+    Real delta, invLSqr, inv;
+
+    if (prod0 >= prod1)
+    {
+        // line intersects P[i0] = e[i0]
+        pnt[i0] = boxExtent[i0];
+
+        Real PpE1 = pnt[i1] + boxExtent[i1];
+        delta = prod0 - dir[i0] * PpE1;
+        if (delta >= (Real)0)
+        {
+            invLSqr = ((Real)1) / (dir[i0] * dir[i0] + dir[i1] * dir[i1]);
+            result.sqrDistance += delta*delta*invLSqr;
+            pnt[i1] = -boxExtent[i1];
+            result.lineParameter = -(dir[i0] * PmE0 + dir[i1] * PpE1)*invLSqr;
+        }
+        else
+        {
+            inv = ((Real)1) / dir[i0];
+            pnt[i1] -= prod0*inv;
+            result.lineParameter = -PmE0*inv;
+        }
+    }
+    else
+    {
+        // line intersects P[i1] = e[i1]
+        pnt[i1] = boxExtent[i1];
+
+        Real PpE0 = pnt[i0] + boxExtent[i0];
+        delta = prod1 - dir[i1] * PpE0;
+        if (delta >= (Real)0)
+        {
+            invLSqr = ((Real)1) / (dir[i0] * dir[i0] + dir[i1] * dir[i1]);
+            result.sqrDistance += delta*delta*invLSqr;
+            pnt[i0] = -boxExtent[i0];
+            result.lineParameter = -(dir[i0] * PpE0 + dir[i1] * PmE1)*invLSqr;
+        }
+        else
+        {
+            inv = ((Real)1) / dir[i1];
+            pnt[i0] -= prod1*inv;
+            result.lineParameter = -PmE1*inv;
+        }
+    }
+
+    if (pnt[i2] < -boxExtent[i2])
+    {
+        delta = pnt[i2] + boxExtent[i2];
+        result.sqrDistance += delta*delta;
+        pnt[i2] = -boxExtent[i2];
+    }
+    else if (pnt[i2] > boxExtent[i2])
+    {
+        delta = pnt[i2] - boxExtent[i2];
+        result.sqrDistance += delta*delta;
+        pnt[i2] = boxExtent[i2];
+    }
+}
+
+template <typename Real>
+void DCPQuery<Real, Line3<Real>, AlignedBox3<Real>>::Case00(int i0, int i1,
+    int i2, Vector3<Real>& pnt, Vector3<Real> const& dir,
+    Vector3<Real> const& boxExtent, Result& result)
+{
+    Real delta;
+
+    result.lineParameter = (boxExtent[i0] - pnt[i0]) / dir[i0];
+
+    pnt[i0] = boxExtent[i0];
+
+    if (pnt[i1] < -boxExtent[i1])
+    {
+        delta = pnt[i1] + boxExtent[i1];
+        result.sqrDistance += delta*delta;
+        pnt[i1] = -boxExtent[i1];
+    }
+    else if (pnt[i1] > boxExtent[i1])
+    {
+        delta = pnt[i1] - boxExtent[i1];
+        result.sqrDistance += delta*delta;
+        pnt[i1] = boxExtent[i1];
+    }
+
+    if (pnt[i2] < -boxExtent[i2])
+    {
+        delta = pnt[i2] + boxExtent[i2];
+        result.sqrDistance += delta*delta;
+        pnt[i2] = -boxExtent[i2];
+    }
+    else if (pnt[i2] > boxExtent[i2])
+    {
+        delta = pnt[i2] - boxExtent[i2];
+        result.sqrDistance += delta*delta;
+        pnt[i2] = boxExtent[i2];
+    }
+}
+
+template <typename Real>
+void DCPQuery<Real, Line3<Real>, AlignedBox3<Real>>::Case000(
+    Vector3<Real>& pnt, Vector3<Real> const& boxExtent, Result& result)
+{
+    Real delta;
+
+    if (pnt[0] < -boxExtent[0])
+    {
+        delta = pnt[0] + boxExtent[0];
+        result.sqrDistance += delta*delta;
+        pnt[0] = -boxExtent[0];
+    }
+    else if (pnt[0] > boxExtent[0])
+    {
+        delta = pnt[0] - boxExtent[0];
+        result.sqrDistance += delta*delta;
+        pnt[0] = boxExtent[0];
+    }
+
+    if (pnt[1] < -boxExtent[1])
+    {
+        delta = pnt[1] + boxExtent[1];
+        result.sqrDistance += delta*delta;
+        pnt[1] = -boxExtent[1];
+    }
+    else if (pnt[1] > boxExtent[1])
+    {
+        delta = pnt[1] - boxExtent[1];
+        result.sqrDistance += delta*delta;
+        pnt[1] = boxExtent[1];
+    }
+
+    if (pnt[2] < -boxExtent[2])
+    {
+        delta = pnt[2] + boxExtent[2];
+        result.sqrDistance += delta*delta;
+        pnt[2] = -boxExtent[2];
+    }
+    else if (pnt[2] > boxExtent[2])
+    {
+        delta = pnt[2] - boxExtent[2];
+        result.sqrDistance += delta*delta;
+        pnt[2] = boxExtent[2];
+    }
+}
+
+
+}
