move GeometryProcessing plugin from NotForLicensees to Experimental

diff --git a/Engine/Plugins/Experimental/GeometryProcessing/Source/GeometryAlgorithms/Private/GTEngine/Mathematics/GteDistLine3Circle3.h b/Engine/Plugins/Experimental/GeometryProcessing/Source/GeometryAlgorithms/Private/GTEngine/Mathematics/GteDistLine3Circle3.h
new file mode 100644
--- /dev/null
+++ b/Engine/Plugins/Experimental/GeometryProcessing/Source/GeometryAlgorithms/Private/GTEngine/Mathematics/GteDistLine3Circle3.h
@@ -0,0 +1,465 @@
+// David Eberly, Geometric Tools, Redmond WA 98052
+// Copyright (c) 1998-2019
+// Distributed under the Boost Software License, Version 1.0.
+// http://www.boost.org/LICENSE_1_0.txt
+// http://www.geometrictools.com/License/Boost/LICENSE_1_0.txt
+// File Version: 3.0.2 (2018/10/05)
+
+#pragma once
+
+#include <GTEngine/Mathematics/GteMath.h>
+#include <GTEngine/Mathematics/GteDCPQuery.h>
+#include <GTEngine/Mathematics/GteCircle3.h>
+#include <GTEngine/Mathematics/GteLine.h>
+#include <GTEngine/Mathematics/GteRootsBisection.h>
+#include <GTEngine/Mathematics/GteRootsPolynomial.h>
+#include <algorithm>
+
+// The 3D line-circle distance algorithm is described in
+// http://www.geometrictools.com/Documentation/DistanceToCircle3.pdf
+// The notation used in the code matches that of the document.
+
+namespace gte
+{
+
+template <typename Real>
+class DCPQuery<Real, Line3<Real>, Circle3<Real>>
+{
+public:
+    // The possible number of closest line-circle pairs is 1, 2 or all circle
+    // points.  If 1 or 2, numClosestPairs is set to this number and
+    // 'equidistant' is false; the number of valid elements in lineClosest[]
+    // and circleClosest[] is numClosestPairs.  If all circle points are
+    // closest, the line must be C+t*N where C is the circle center, N is
+    // the normal to the plane of the circle, and lineClosest[0] is set to C.
+    // In this case, 'equidistant' is true and circleClosest[0] is set to
+    // C+r*U, where r is the circle radius and U is a vector perpendicular
+    // to N.
+    struct Result
+    {
+        Real distance, sqrDistance;
+        int numClosestPairs;
+        Vector3<Real> lineClosest[2], circleClosest[2];
+        bool equidistant;
+    };
+
+    // The polynomial-based algorithm.  Type Real can be floating-point or
+    // rational.
+    Result operator()(Line3<Real> const& line, Circle3<Real> const& circle);
+
+    // The nonpolynomial-based algorithm that uses bisection.  Because the
+    // bisection is iterative, you should choose Real to be a floating-point
+    // type.  However, the algorithm will still work for a rational type, but
+    // it is costly because of the increase in arbitrary-size integers used
+    // during the bisection.
+    Result Robust(Line3<Real> const& line, Circle3<Real> const& circle);
+
+private:
+    // Support for operator(...).
+    struct ClosestInfo
+    {
+        Real sqrDistance;
+        Vector3<Real> lineClosest, circleClosest;
+        bool equidistant;
+
+        bool operator< (ClosestInfo const& info) const
+        {
+            return sqrDistance < info.sqrDistance;
+        }
+    };
+
+    void GetPair(Line3<Real> const& line, Circle3<Real> const& circle,
+        Vector3<Real> const& D, Real t, Vector3<Real>& lineClosest,
+        Vector3<Real>& circleClosest);
+
+    // Support for Robust(...).  Bisect the function
+    //   F(s) = s + m2b2 - r*m0sqr*s/sqrt(m0sqr*s*s + b1sqr)
+    // on the specified interval [smin,smax].
+    Real Bisect(Real m2b2, Real rm0sqr, Real m0sqr, Real b1sqr,
+        Real smin, Real smax);
+};
+
+
+template <typename Real>
+typename DCPQuery<Real, Line3<Real>, Circle3<Real>>::Result
+DCPQuery<Real, Line3<Real>, Circle3<Real>>::operator()(
+    Line3<Real> const& line, Circle3<Real> const& circle)
+{
+    Result result;
+    Vector3<Real> const vzero = Vector3<Real>::Zero();
+    Real const zero = (Real)0;
+
+    Vector3<Real> D = line.origin - circle.center;
+    Vector3<Real> NxM = Cross(circle.normal, line.direction);
+    Vector3<Real> NxD = Cross(circle.normal, D);
+    Real t;
+
+    if (NxM != vzero)
+    {
+        if (NxD != vzero)
+        {
+            Real NdM = Dot(circle.normal, line.direction);
+            if (NdM != zero)
+            {
+                // H(t) = (a*t^2 + 2*b*t + c)*(t + d)^2 - r^2*(a*t + b)^2
+                //      = h0 + h1*t + h2*t^2 + h3*t^3 + h4*t^4
+                Real a = Dot(NxM, NxM), b = Dot(NxM, NxD);
+                Real c = Dot(NxD, NxD), d = Dot(line.direction, D);
+                Real rsqr = circle.radius * circle.radius;
+                Real asqr = a * a, bsqr = b * b, dsqr = d * d;
+                Real h0 = c * dsqr - bsqr * rsqr;
+                Real h1 = 2 * (c * d + b * dsqr - a * b * rsqr);
+                Real h2 = c + 4 * b * d + a * dsqr - asqr * rsqr;
+                Real h3 = 2 * (b + a * d);
+                Real h4 = a;
+
+                std::map<Real, int> rmMap;
+                RootsPolynomial<Real>::template SolveQuartic<Real>(
+                    h0, h1, h2, h3, h4, rmMap);
+                std::array<ClosestInfo, 4> candidates;
+                int numRoots = 0;
+                for (auto const& rm : rmMap)
+                {
+                    t = rm.first;
+                    ClosestInfo info;
+                    Vector3<Real> NxDelta = NxD + t * NxM;
+                    if (NxDelta != vzero)
+                    {
+                        GetPair(line, circle, D,t, info.lineClosest,
+                            info.circleClosest);
+                        info.equidistant = false;
+                    }
+                    else
+                    {
+                        Vector3<Real> U = GetOrthogonal(circle.normal, true);
+                        info.lineClosest = circle.center;
+                        info.circleClosest =
+                            circle.center + circle.radius * U;
+                        info.equidistant = true;
+                    }
+                    Vector3<Real> diff = info.lineClosest - info.circleClosest;
+                    info.sqrDistance = Dot(diff, diff);
+                    candidates[numRoots++] = info;
+                }
+
+                std::sort(candidates.begin(), candidates.begin() + numRoots);
+
+                result.numClosestPairs = 1;
+                result.lineClosest[0] = candidates[0].lineClosest;
+                result.circleClosest[0] = candidates[0].circleClosest;
+                if (numRoots > 1
+                    && candidates[1].sqrDistance == candidates[0].sqrDistance)
+                {
+                    result.numClosestPairs = 2;
+                    result.lineClosest[1] = candidates[1].lineClosest;
+                    result.circleClosest[1] = candidates[1].circleClosest;
+                }
+            }
+            else
+            {
+                // The line is parallel to the plane of the circle.  The
+                // polynomial has the form H(t) = (t+v)^2*[(t+v)^2-(r^2-u^2)].
+                Real u = Dot(NxM, D), v = Dot(line.direction, D);
+                Real discr = circle.radius * circle.radius - u * u;
+                if (discr > zero)
+                {
+                    result.numClosestPairs = 2;
+                    Real rootDiscr = std::sqrt(discr);
+                    t = -v + rootDiscr;
+                    GetPair(line, circle, D, t, result.lineClosest[0],
+                        result.circleClosest[0]);
+                    t = -v - rootDiscr;
+                    GetPair(line, circle, D, t, result.lineClosest[1],
+                        result.circleClosest[1]);
+                }
+                else
+                {
+                    result.numClosestPairs = 1;
+                    t = -v;
+                    GetPair(line, circle, D, t, result.lineClosest[0],
+                        result.circleClosest[0]);
+                }
+            }
+        }
+        else
+        {
+            // The line is C+t*M, where M is not parallel to N.  The polynomial
+            // is H(t) = |Cross(N,M)|^2*t^2*(t^2 - r^2*|Cross(N,M)|^2), where
+            // root t = 0 does not correspond to the global minimum.  The other
+            // roots produce the global minimum.
+            result.numClosestPairs = 2;
+            t = circle.radius * Length(NxM);
+            GetPair(line, circle, D, t, result.lineClosest[0],
+                result.circleClosest[0]);
+            t = -t;
+            GetPair(line, circle, D, t, result.lineClosest[1],
+                result.circleClosest[1]);
+        }
+        result.equidistant = false;
+    }
+    else
+    {
+        if (NxD != vzero)
+        {
+            // The line is A+t*N (perpendicular to plane) but with A != C.
+            // The polyhomial is H(t) = |Cross(N,D)|^2*(t + Dot(M,D))^2.
+            result.numClosestPairs = 1;
+            t = -Dot(line.direction, D);
+            GetPair(line, circle, D, t, result.lineClosest[0],
+                result.circleClosest[0]);
+            result.equidistant = false;
+        }
+        else
+        {
+            // The line is C+t*N, so C is the closest point for the line and
+            // all circle points are equidistant from it.
+            Vector3<Real> U = GetOrthogonal(circle.normal, true);
+            result.numClosestPairs = 1;
+            result.lineClosest[0] = circle.center;
+            result.circleClosest[0] = circle.center + circle.radius * U;
+            result.equidistant = true;
+        }
+    }
+
+    Vector3<Real> diff = result.lineClosest[0] - result.circleClosest[0];
+    result.sqrDistance = Dot(diff, diff);
+    result.distance = std::sqrt(result.sqrDistance);
+    return result;
+}
+
+template <typename Real>
+typename DCPQuery<Real, Line3<Real>, Circle3<Real>>::Result
+DCPQuery<Real, Line3<Real>, Circle3<Real>>::Robust(
+    Line3<Real> const& line, Circle3<Real> const& circle)
+{
+    // The line is P(t) = B+t*M.  The circle is |X-C| = r with Dot(N,X-C)=0.
+    Result result;
+    Vector3<Real> vzero = Vector3<Real>::Zero();
+    Real const zero = (Real)0;
+
+    Vector3<Real> D = line.origin - circle.center;
+    Vector3<Real> MxN = Cross(line.direction, circle.normal);
+    Vector3<Real> DxN = Cross(D, circle.normal);
+
+    Real m0sqr = Dot(MxN, MxN);
+    if (m0sqr > zero)
+    {
+        // Compute the critical points s for F'(s) = 0.
+        Vector3<Real> P0, P1;
+        Real s, t;
+        int numRoots = 0;
+        std::array<Real, 3> roots;
+
+        // The line direction M and the plane normal N are not parallel.  Move
+        // the line origin B = (b0,b1,b2) to B' = B + lambda*line.direction =
+        // (0,b1',b2').
+        Real m0 = std::sqrt(m0sqr);
+        Real rm0 = circle.radius * m0;
+        Real lambda = -Dot(MxN, DxN) / m0sqr;
+        Vector3<Real> oldD = D;
+        D += lambda * line.direction;
+        DxN += lambda * MxN;
+        Real m2b2 = Dot(line.direction, D);
+        Real b1sqr = Dot(DxN, DxN);
+        if (b1sqr > zero)
+        {
+            // B' = (0,b1',b2') where b1' != 0.  See Sections 1.1.2 and 1.2.2
+            // of the PDF documentation.
+            Real b1 = std::sqrt(b1sqr);
+            Real rm0sqr = circle.radius * m0sqr;
+            if (rm0sqr > b1)
+            {
+                Real const twoThirds = (Real)2 / (Real)3;
+                Real sHat = std::sqrt(std::pow(rm0sqr * b1sqr, twoThirds) - b1sqr) / m0;
+                Real gHat = rm0sqr * sHat / std::sqrt(m0sqr * sHat * sHat + b1sqr);
+                Real cutoff = gHat - sHat;
+                if (m2b2 <= -cutoff)
+                {
+                    s = Bisect(m2b2, rm0sqr, m0sqr, b1sqr, -m2b2,
+                        -m2b2 + rm0);
+                    roots[numRoots++] = s;
+                    if (m2b2 == -cutoff)
+                    {
+                        roots[numRoots++] = -sHat;
+                    }
+                }
+                else if (m2b2 >= cutoff)
+                {
+                    s = Bisect(m2b2, rm0sqr, m0sqr, b1sqr, -m2b2 - rm0,
+                        -m2b2);
+                    roots[numRoots++] = s;
+                    if (m2b2 == cutoff)
+                    {
+                        roots[numRoots++] = sHat;
+                    }
+                }
+                else
+                {
+                    if (m2b2 <= zero)
+                    {
+                        s = Bisect(m2b2, rm0sqr, m0sqr, b1sqr, -m2b2,
+                            -m2b2 + rm0);
+                        roots[numRoots++] = s;
+                        s = Bisect(m2b2, rm0sqr, m0sqr, b1sqr, -m2b2 - rm0,
+                            -sHat);
+                        roots[numRoots++] = s;
+                    }
+                    else
+                    {
+                        s = Bisect(m2b2, rm0sqr, m0sqr, b1sqr, -m2b2 - rm0,
+                            -m2b2);
+                        roots[numRoots++] = s;
+                        s = Bisect(m2b2, rm0sqr, m0sqr, b1sqr, sHat,
+                            -m2b2 + rm0);
+                        roots[numRoots++] = s;
+                    }
+                }
+            }
+            else
+            {
+                if (m2b2 < zero)
+                {
+                    s = Bisect(m2b2, rm0sqr, m0sqr, b1sqr, -m2b2,
+                        -m2b2 + rm0);
+                }
+                else if (m2b2 > zero)
+                {
+                    s = Bisect(m2b2, rm0sqr, m0sqr, b1sqr, -m2b2 - rm0,
+                        -m2b2);
+                }
+                else
+                {
+                    s = zero;
+                }
+                roots[numRoots++] = s;
+            }
+        }
+        else
+        {
+            // The new line origin is B' = (0,0,b2').
+            if (m2b2 < zero)
+            {
+                s = -m2b2 + rm0;
+                roots[numRoots++] = s;
+            }
+            else if (m2b2 > zero)
+            {
+                s = -m2b2 - rm0;
+                roots[numRoots++] = s;
+            }
+            else
+            {
+                s = -m2b2 + rm0;
+                roots[numRoots++] = s;
+                s = -m2b2 - rm0;
+                roots[numRoots++] = s;
+            }
+        }
+
+        std::array<ClosestInfo, 4> candidates;
+        for (int i = 0; i < numRoots; ++i)
+        {
+            t = roots[i] + lambda;
+            ClosestInfo info;
+            Vector3<Real> NxDelta = 
+                Cross(circle.normal, oldD + t * line.direction);
+            if (NxDelta != vzero)
+            {
+                GetPair(line, circle, oldD, t, info.lineClosest,
+                    info.circleClosest);
+                info.equidistant = false;
+            }
+            else
+            {
+                Vector3<Real> U = GetOrthogonal(circle.normal, true);
+                info.lineClosest = circle.center;
+                info.circleClosest = circle.center + circle.radius * U;
+                info.equidistant = true;
+            }
+            Vector3<Real> diff = info.lineClosest - info.circleClosest;
+            info.sqrDistance = Dot(diff, diff);
+            candidates[i] = info;
+        }
+
+        std::sort(candidates.begin(), candidates.begin() + numRoots);
+
+        result.numClosestPairs = 1;
+        result.lineClosest[0] = candidates[0].lineClosest;
+        result.circleClosest[0] = candidates[0].circleClosest;
+        if (numRoots > 1
+            && candidates[1].sqrDistance == candidates[0].sqrDistance)
+        {
+            result.numClosestPairs = 2;
+            result.lineClosest[1] = candidates[1].lineClosest;
+            result.circleClosest[1] = candidates[1].circleClosest;
+        }
+
+        result.equidistant = false;
+    }
+    else
+    {
+        // The line direction and the plane normal are parallel.
+        if (DxN != vzero)
+        {
+            // The line is A+t*N but with A != C.
+            result.numClosestPairs = 1;
+            GetPair(line, circle, D, -Dot(line.direction, D),
+                result.lineClosest[0], result.circleClosest[0]);
+            result.equidistant = false;
+        }
+        else
+        {
+            // The line is C+t*N, so C is the closest point for the line and
+            // all circle points are equidistant from it.
+            Vector3<Real> U = GetOrthogonal(circle.normal, true);
+            result.numClosestPairs = 1;
+            result.lineClosest[0] = circle.center;
+            result.circleClosest[0] = circle.center + circle.radius * U;
+            result.equidistant = true;
+        }
+    }
+
+    Vector3<Real> diff = result.lineClosest[0] - result.circleClosest[0];
+    result.sqrDistance = Dot(diff, diff);
+    result.distance = std::sqrt(result.sqrDistance);
+    return result;
+}
+
+template <typename Real>
+void DCPQuery<Real, Line3<Real>, Circle3<Real>>::GetPair(
+    Line3<Real> const& line, Circle3<Real> const& circle,
+    Vector3<Real> const& D, Real t, Vector3<Real>& lineClosest,
+    Vector3<Real>& circleClosest)
+{
+    Vector3<Real> delta = D + t * line.direction;
+    lineClosest = circle.center + delta;
+    delta -= Dot(circle.normal, delta) * circle.normal;
+    Normalize(delta);
+    circleClosest = circle.center + circle.radius * delta;
+}
+
+template <typename Real>
+Real DCPQuery<Real, Line3<Real>, Circle3<Real>>::Bisect(Real m2b2,
+    Real rm0sqr, Real m0sqr, Real b1sqr, Real smin, Real smax)
+{
+    std::function<Real(Real)> G = [&, m2b2, rm0sqr, m0sqr, b1sqr](Real s)
+    {
+        return s + m2b2 - rm0sqr*s / std::sqrt(m0sqr*s*s + b1sqr);
+    };
+
+    // The function is known to be increasing, so we can specify -1 and +1
+    // as the function values at the bounding interval endpoints.  The use
+    // of 'double' is intentional in case Real is a BSNumber or BSRational
+    // type.  We want the bisections to terminate in a reasonable amount of
+    // time.
+    unsigned int const maxIterations = GTE_C_MAX_BISECTIONS_GENERIC;
+    Real root;
+    RootsBisection<Real>::Find(G, smin, smax, (Real)-1, (Real)+1,
+        maxIterations, root);
+    Real gmin;
+    gmin = G(root);
+    return root;
+}
+
+
+}
diff --git a/Engine/Plugins/Experimental/GeometryProcessing/Source/GeometryAlgorithms/Private/GTEngine/Mathematics/GteDistLine3OrientedBox3.h b/Engine/Plugins/Experimental/GeometryProcessing/Source/GeometryAlgorithms/Private/GTEngine/Mathematics/GteDistLine3OrientedBox3.h
new file mode 100644
--- /dev/null
+++ b/Engine/Plugins/Experimental/GeometryProcessing/Source/GeometryAlgorithms/Private/GTEngine/Mathematics/GteDistLine3OrientedBox3.h
@@ -0,0 +1,66 @@
+// David Eberly, Geometric Tools, Redmond WA 98052
+// Copyright (c) 1998-2019
+// Distributed under the Boost Software License, Version 1.0.
+// http://www.boost.org/LICENSE_1_0.txt
+// http://www.geometrictools.com/License/Boost/LICENSE_1_0.txt
+// File Version: 3.0.0 (2016/06/19)
+
+#pragma once
+
+#include <GTEngine/Mathematics/GteDistLine3AlignedBox3.h>
+#include <GTEngine/Mathematics/GteLine.h>
+#include <GTEngine/Mathematics/GteOrientedBox.h>
+
+namespace gte
+{
+
+template <typename Real>
+class DCPQuery<Real, Line3<Real>, OrientedBox3<Real>>
+    :
+    public DCPQuery<Real, Line3<Real>, AlignedBox3<Real>>
+{
+public:
+    struct Result
+        :
+        public DCPQuery<Real, Line3<Real>, AlignedBox3<Real>>::Result
+    {
+        // No additional information to compute.
+    };
+
+    Result operator()(Line3<Real> const& line, OrientedBox3<Real> const& box);
+};
+
+
+template <typename Real>
+typename DCPQuery<Real, Line3<Real>, OrientedBox3<Real>>::Result
+DCPQuery<Real, Line3<Real>, OrientedBox3<Real>>::operator()(
+    Line3<Real> const& line, OrientedBox3<Real> const& box)
+{
+    // Transform the line to the coordinate system of the oriented box.
+    // In this system, the box is axis-aligned with center at the origin.
+    Vector3<Real> diff = line.origin - box.center;
+    Vector3<Real> point, direction;
+    for (int i = 0; i < 3; ++i)
+    {
+        point[i] = Dot(diff, box.axis[i]);
+        direction[i] = Dot(line.direction, box.axis[i]);
+    }
+
+    Result result;
+    this->DoQuery(point, direction, box.extent, result);
+
+    // Compute the closest point on the line.
+    result.closestPoint[0] =
+        line.origin + result.lineParameter*line.direction;
+
+    // Compute the closest point on the box.
+    result.closestPoint[1] = box.center;
+    for (int i = 0; i < 3; ++i)
+    {
+        result.closestPoint[1] += point[i] * box.axis[i];
+    }
+    return result;
+}
+
+
+}
diff --git a/Engine/Plugins/Experimental/GeometryProcessing/Source/GeometryAlgorithms/Private/GTEngine/Mathematics/GteDistLine3Rectangle3.h b/Engine/Plugins/Experimental/GeometryProcessing/Source/GeometryAlgorithms/Private/GTEngine/Mathematics/GteDistLine3Rectangle3.h
new file mode 100644
--- /dev/null
+++ b/Engine/Plugins/Experimental/GeometryProcessing/Source/GeometryAlgorithms/Private/GTEngine/Mathematics/GteDistLine3Rectangle3.h
@@ -0,0 +1,140 @@
+// David Eberly, Geometric Tools, Redmond WA 98052
+// Copyright (c) 1998-2019
+// Distributed under the Boost Software License, Version 1.0.
+// http://www.boost.org/LICENSE_1_0.txt
+// http://www.geometrictools.com/License/Boost/LICENSE_1_0.txt
+// File Version: 3.0.0 (2016/06/19)
+
+#pragma once
+
+#include <GTEngine/Mathematics/GteVector3.h>
+#include <GTEngine/Mathematics/GteDistLineSegment.h>
+#include <GTEngine/Mathematics/GteRectangle.h>
+#include <limits>
+
+namespace gte
+{
+
+template <typename Real>
+class DCPQuery<Real, Line3<Real>, Rectangle3<Real>>
+{
+public:
+    struct Result
+    {
+        Real distance, sqrDistance;
+        Real lineParameter, rectangleParameter[2];
+        Vector3<Real> closestPoint[2];
+    };
+
+    Result operator()(Line3<Real> const& line,
+        Rectangle3<Real> const& rectangle);
+};
+
+// Template alias for convenience.
+template <typename Real>
+using DCPLine3Rectangle3 =
+DCPQuery<Real, Line3<Real>, Rectangle3<Real>>;
+
+
+template <typename Real>
+typename DCPQuery<Real, Line3<Real>, Rectangle3<Real>>::Result
+DCPQuery<Real, Line3<Real>, Rectangle3<Real>>::operator()(
+    Line3<Real> const& line, Rectangle3<Real> const& rectangle)
+{
+    Result result;
+
+    // Test if line intersects rectangle.  If so, the squared distance is
+    // zero.
+    Vector3<Real> N = Cross(rectangle.axis[0], rectangle.axis[1]);
+    Real NdD = Dot(N, line.direction);
+    if (std::abs(NdD) > (Real)0)
+    {
+        // The line and rectangle are not parallel, so the line intersects
+        // the plane of the rectangle.
+        Vector3<Real> diff = line.origin - rectangle.center;
+        Vector3<Real> basis[3];  // {D, U, V}
+        basis[0] = line.direction;
+        ComputeOrthogonalComplement<Real>(1, basis);
+        Real UdD0 = Dot(basis[1], rectangle.axis[0]);
+        Real UdD1 = Dot(basis[1], rectangle.axis[1]);
+        Real UdPmC = Dot(basis[1], diff);
+        Real VdD0 = Dot(basis[2], rectangle.axis[0]);
+        Real VdD1 = Dot(basis[2], rectangle.axis[1]);
+        Real VdPmC = Dot(basis[2], diff);
+        Real invDet = ((Real)1) / (UdD0*VdD1 - UdD1*VdD0);
+
+        // Rectangle coordinates for the point of intersection.
+        Real s0 = (VdD1*UdPmC - UdD1*VdPmC)*invDet;
+        Real s1 = (UdD0*VdPmC - VdD0*UdPmC)*invDet;
+
+        if (std::abs(s0) <= rectangle.extent[0]
+            && std::abs(s1) <= rectangle.extent[1])
+        {
+            // Line parameter for the point of intersection.
+            Real DdD0 = Dot(line.direction, rectangle.axis[0]);
+            Real DdD1 = Dot(line.direction, rectangle.axis[1]);
+            Real DdDiff = Dot(line.direction, diff);
+            result.lineParameter = s0*DdD0 + s1*DdD1 - DdDiff;
+
+            // Rectangle coordinates for the point of intersection.
+            result.rectangleParameter[0] = s0;
+            result.rectangleParameter[1] = s1;
+
+            // The intersection point is inside or on the rectangle.
+            result.closestPoint[0] = line.origin +
+                result.lineParameter*line.direction;
+
+            result.closestPoint[1] = rectangle.center +
+                s0*rectangle.axis[0] + s1*rectangle.axis[1];
+
+            result.distance = (Real)0;
+            result.sqrDistance = (Real)0;
+            return result;
+        }
+    }
+
+    // Either (1) the line is not parallel to the rectangle and the point of
+    // intersection of the line and the plane of the rectangle is outside the
+    // rectangle or (2) the line and rectangle are parallel.  Regardless, the
+    // closest point on the rectangle is on an edge of the rectangle.  Compare
+    // the line to all four edges of the rectangle.
+    result.distance = std::numeric_limits<Real>::max();
+    result.sqrDistance = std::numeric_limits<Real>::max();
+    Vector3<Real> scaledDir[2] =
+    {
+        rectangle.extent[0] * rectangle.axis[0],
+        rectangle.extent[1] * rectangle.axis[1]
+    };
+    for (int i1 = 0, omi1 = 1; i1 <= 1; ++i1, --omi1)
+    {
+        for (int i0 = -1; i0 <= 1; i0 += 2)
+        {
+            Vector3<Real> segCenter =
+                rectangle.center + scaledDir[i1] * (Real)i0;
+            Vector3<Real> segDirection = rectangle.axis[omi1];
+            Real segExtent = rectangle.extent[omi1];
+            Segment3<Real> segment(segCenter, segDirection, segExtent);
+
+            DCPQuery<Real, Line3<Real>, Segment3<Real>> query;
+            auto lsResult = query(line, segment);
+            if (lsResult.sqrDistance < result.sqrDistance)
+            {
+                result.sqrDistance = lsResult.sqrDistance;
+                result.distance = lsResult.distance;
+                result.lineParameter = lsResult.parameter[0];
+                Real ratio = lsResult.parameter[1] / segExtent; // in [-1,1]
+                result.rectangleParameter[0] =
+                    rectangle.extent[0] * (omi1 * i0 + i1 * ratio);
+                result.rectangleParameter[1] =
+                    rectangle.extent[1] * (i1 * i0 + omi1 * ratio);
+                result.closestPoint[0] = lsResult.closestPoint[0];
+                result.closestPoint[1] = lsResult.closestPoint[1];
+            }
+        }
+    }
+
+    return result;
+}
+
+
+}
diff --git a/Engine/Plugins/Experimental/GeometryProcessing/Source/GeometryAlgorithms/Private/GTEngine/Mathematics/GteDistLine3Triangle3.h b/Engine/Plugins/Experimental/GeometryProcessing/Source/GeometryAlgorithms/Private/GTEngine/Mathematics/GteDistLine3Triangle3.h
new file mode 100644
--- /dev/null
+++ b/Engine/Plugins/Experimental/GeometryProcessing/Source/GeometryAlgorithms/Private/GTEngine/Mathematics/GteDistLine3Triangle3.h
@@ -0,0 +1,127 @@
+// David Eberly, Geometric Tools, Redmond WA 98052
+// Copyright (c) 1998-2019
+// Distributed under the Boost Software License, Version 1.0.
+// http://www.boost.org/LICENSE_1_0.txt
+// http://www.geometrictools.com/License/Boost/LICENSE_1_0.txt
+// File Version: 3.0.0 (2016/06/19)
+
+#pragma once
+
+#include <GTEngine/Mathematics/GteVector3.h>
+#include <GTEngine/Mathematics/GteDistLineSegment.h>
+#include <GTEngine/Mathematics/GteTriangle.h>
+#include <limits>
+
+namespace gte
+{
+
+template <typename Real>
+class DCPQuery<Real, Line3<Real>, Triangle3<Real>>
+{
+public:
+    struct Result
+    {
+        Real distance, sqrDistance;
+        Real lineParameter, triangleParameter[3];
+        Vector3<Real> closestPoint[2];
+    };
+
+    Result operator()(Line3<Real> const& line,
+        Triangle3<Real> const& triangle);
+};
+
+
+template <typename Real>
+typename DCPQuery<Real, Line3<Real>, Triangle3<Real>>::Result
+DCPQuery<Real, Line3<Real>, Triangle3<Real>>::operator()(
+    Line3<Real> const& line, Triangle3<Real> const& triangle)
+{
+    Result result;
+
+    // Test if line intersects triangle.  If so, the squared distance is zero.
+    Vector3<Real> edge0 = triangle.v[1] - triangle.v[0];
+    Vector3<Real> edge1 = triangle.v[2] - triangle.v[0];
+    Vector3<Real> normal = UnitCross(edge0, edge1);
+    Real NdD = Dot(normal, line.direction);
+    if (std::abs(NdD) > (Real)0)
+    {
+        // The line and triangle are not parallel, so the line intersects
+        // the plane of the triangle.
+        Vector3<Real> diff = line.origin - triangle.v[0];
+        Vector3<Real> basis[3];  // {D, U, V}
+        basis[0] = line.direction;
+        ComputeOrthogonalComplement<Real>(1, basis);
+        Real UdE0 = Dot(basis[1], edge0);
+        Real UdE1 = Dot(basis[1], edge1);
+        Real UdDiff = Dot(basis[1], diff);
+        Real VdE0 = Dot(basis[2], edge0);
+        Real VdE1 = Dot(basis[2], edge1);
+        Real VdDiff = Dot(basis[2], diff);
+        Real invDet = ((Real)1) / (UdE0*VdE1 - UdE1*VdE0);
+
+        // Barycentric coordinates for the point of intersection.
+        Real b1 = (VdE1*UdDiff - UdE1*VdDiff)*invDet;
+        Real b2 = (UdE0*VdDiff - VdE0*UdDiff)*invDet;
+        Real b0 = (Real)1 - b1 - b2;
+
+        if (b0 >= (Real)0 && b1 >= (Real)0 && b2 >= (Real)0)
+        {
+            // Line parameter for the point of intersection.
+            Real DdE0 = Dot(line.direction, edge0);
+            Real DdE1 = Dot(line.direction, edge1);
+            Real DdDiff = Dot(line.direction, diff);
+            result.lineParameter = b1*DdE0 + b2*DdE1 - DdDiff;
+
+            // Barycentric coordinates for the point of intersection.
+            result.triangleParameter[0] = b0;
+            result.triangleParameter[1] = b1;
+            result.triangleParameter[2] = b2;
+
+            // The intersection point is inside or on the triangle.
+            result.closestPoint[0] = line.origin +
+                result.lineParameter*line.direction;
+            result.closestPoint[1] = triangle.v[0] + b1*edge0 + b2*edge1;
+
+            result.distance = (Real)0;
+            result.sqrDistance = (Real)0;
+            return result;
+        }
+    }
+
+    // Either (1) the line is not parallel to the triangle and the point of
+    // intersection of the line and the plane of the triangle is outside the
+    // triangle or (2) the line and triangle are parallel.  Regardless, the
+    // closest point on the triangle is on an edge of the triangle.  Compare
+    // the line to all three edges of the triangle.
+    result.distance = std::numeric_limits<Real>::max();
+    result.sqrDistance = std::numeric_limits<Real>::max();
+    for (int i0 = 2, i1 = 0; i1 < 3; i0 = i1++)
+    {
+        Vector3<Real> segCenter = ((Real)0.5)*(triangle.v[i0] +
+            triangle.v[i1]);
+        Vector3<Real> segDirection = triangle.v[i1] - triangle.v[i0];
+        Real segExtent = ((Real)0.5)*Normalize(segDirection);
+        Segment3<Real> segment(segCenter, segDirection, segExtent);
+
+        DCPQuery<Real, Line3<Real>, Segment3<Real>> query;
+        auto lsResult = query(line, segment);
+        if (lsResult.sqrDistance < result.sqrDistance)
+        {
+            result.sqrDistance = lsResult.sqrDistance;
+            result.distance = lsResult.distance;
+            result.lineParameter = lsResult.parameter[0];
+            result.triangleParameter[i0] = ((Real)0.5)*((Real)1 -
+                lsResult.parameter[0] / segExtent);
+            result.triangleParameter[i1] = (Real)1 -
+                result.triangleParameter[i0];
+            result.triangleParameter[3 - i0 - i1] = (Real)0;
+            result.closestPoint[0] = lsResult.closestPoint[0];
+            result.closestPoint[1] = lsResult.closestPoint[1];
+        }
+    }
+
+    return result;
+}
+
+
+}
diff --git a/Engine/Plugins/Experimental/GeometryProcessing/Source/GeometryAlgorithms/Private/GTEngine/Mathematics/GteDistLineLine.h b/Engine/Plugins/Experimental/GeometryProcessing/Source/GeometryAlgorithms/Private/GTEngine/Mathematics/GteDistLineLine.h
new file mode 100644
--- /dev/null
+++ b/Engine/Plugins/Experimental/GeometryProcessing/Source/GeometryAlgorithms/Private/GTEngine/Mathematics/GteDistLineLine.h
@@ -0,0 +1,79 @@
+// David Eberly, Geometric Tools, Redmond WA 98052
+// Copyright (c) 1998-2019
+// Distributed under the Boost Software License, Version 1.0.
+// http://www.boost.org/LICENSE_1_0.txt
+// http://www.geometrictools.com/License/Boost/LICENSE_1_0.txt
+// File Version: 3.0.1 (2018/10/05)
+
+#pragma once
+
+#include <GTEngine/Mathematics/GteDCPQuery.h>
+#include <GTEngine/Mathematics/GteLine.h>
+
+namespace gte
+{
+
+template <int N, typename Real>
+class DCPQuery<Real, Line<N, Real>, Line<N, Real>>
+{
+public:
+    struct Result
+    {
+        Real distance, sqrDistance;
+        Real parameter[2];
+        Vector<N, Real> closestPoint[2];
+    };
+
+    Result operator()(Line<N, Real> const& line0, Line<N, Real> const& line1);
+};
+
+// Template aliases for convenience.
+template <int N, typename Real>
+using DCPLineLine = DCPQuery<Real, Line<N, Real>, Line<N, Real>>;
+
+template <typename Real>
+using DCPLine2Line2 = DCPLineLine<2, Real>;
+
+template <typename Real>
+using DCPLine3Line3 = DCPLineLine<3, Real>;
+
+
+template <int N, typename Real>
+typename DCPQuery<Real, Line<N, Real>, Line<N, Real>>::Result
+DCPQuery<Real, Line<N, Real>, Line<N, Real>>::operator()(
+    Line<N, Real> const& line0, Line<N, Real> const& line1)
+{
+    Result result;
+
+    Vector<N, Real> diff = line0.origin - line1.origin;
+    Real a01 = -Dot(line0.direction, line1.direction);
+    Real b0 = Dot(diff, line0.direction);
+    Real s0, s1;
+
+    if (std::abs(a01) < (Real)1)
+    {
+        // Lines are not parallel.
+        Real det = (Real)1 - a01 * a01;
+        Real b1 = -Dot(diff, line1.direction);
+        s0 = (a01 * b1 - b0) / det;
+        s1 = (a01 * b0 - b1) / det;
+    }
+    else
+    {
+        // Lines are parallel, select any pair of closest points.
+        s0 = -b0;
+        s1 = (Real)0;
+    }
+
+    result.parameter[0] = s0;
+    result.parameter[1] = s1;
+    result.closestPoint[0] = line0.origin + s0 * line0.direction;
+    result.closestPoint[1] = line1.origin + s1 * line1.direction;
+    diff = result.closestPoint[0] - result.closestPoint[1];
+    result.sqrDistance = Dot(diff, diff);
+    result.distance = std::sqrt(result.sqrDistance);
+    return result;
+}
+
+
+}
diff --git a/Engine/Plugins/Experimental/GeometryProcessing/Source/GeometryAlgorithms/Private/GTEngine/Mathematics/GteDistLineRay.h b/Engine/Plugins/Experimental/GeometryProcessing/Source/GeometryAlgorithms/Private/GTEngine/Mathematics/GteDistLineRay.h
new file mode 100644
--- /dev/null
+++ b/Engine/Plugins/Experimental/GeometryProcessing/Source/GeometryAlgorithms/Private/GTEngine/Mathematics/GteDistLineRay.h
@@ -0,0 +1,91 @@
+// David Eberly, Geometric Tools, Redmond WA 98052
+// Copyright (c) 1998-2019
+// Distributed under the Boost Software License, Version 1.0.
+// http://www.boost.org/LICENSE_1_0.txt
+// http://www.geometrictools.com/License/Boost/LICENSE_1_0.txt
+// File Version: 3.0.1 (2018/10/05)
+
+#pragma once
+
+#include <GTEngine/Mathematics/GteDCPQuery.h>
+#include <GTEngine/Mathematics/GteLine.h>
+#include <GTEngine/Mathematics/GteRay.h>
+
+namespace gte
+{
+
+template <int N, typename Real>
+class DCPQuery<Real, Line<N, Real>, Ray<N, Real>>
+{
+public:
+    struct Result
+    {
+        Real distance, sqrDistance;
+        Real parameter[2];
+        Vector<N, Real> closestPoint[2];
+    };
+
+    Result operator()(Line<N, Real> const& line, Ray<N, Real> const& ray);
+};
+
+// Template aliases for convenience.
+template <int N, typename Real>
+using DCPLineRay = DCPQuery<Real, Line<N, Real>, Ray<N, Real>>;
+
+template <typename Real>
+using DCPLine2Ray2 = DCPLineRay<2, Real>;
+
+template <typename Real>
+using DCPLine3Ray3 = DCPLineRay<3, Real>;
+
+
+template <int N, typename Real>
+typename DCPQuery<Real, Line<N, Real>, Ray<N, Real>>::Result
+DCPQuery<Real, Line<N, Real>, Ray<N, Real>>::operator()(
+    Line<N, Real> const& line, Ray<N, Real> const& ray)
+{
+    Result result;
+
+    Vector<N, Real> diff = line.origin - ray.origin;
+    Real a01 = -Dot(line.direction, ray.direction);
+    Real b0 = Dot(diff, line.direction);
+    Real s0, s1;
+
+    if (std::abs(a01) < (Real)1)
+    {
+        Real b1 = -Dot(diff, ray.direction);
+        s1 = a01 * b0 - b1;
+
+        if (s1 >= (Real)0)
+        {
+            // Two interior points are closest, one on line and one on ray.
+            Real det = (Real)1 - a01 * a01;
+            s0 = (a01 * b1 - b0) / det;
+            s1 /= det;
+        }
+        else
+        {
+            // Origin of ray and interior point of line are closest.
+            s0 = -b0;
+            s1 = (Real)0;
+        }
+    }
+    else
+    {
+        // Lines are parallel, closest pair with one point at ray origin.
+        s0 = -b0;
+        s1 = (Real)0;
+    }
+
+    result.parameter[0] = s0;
+    result.parameter[1] = s1;
+    result.closestPoint[0] = line.origin + s0 * line.direction;
+    result.closestPoint[1] = ray.origin + s1 * ray.direction;
+    diff = result.closestPoint[0] - result.closestPoint[1];
+    result.sqrDistance = Dot(diff, diff);
+    result.distance = std::sqrt(result.sqrDistance);
+    return result;
+}
+
+
+}
diff --git a/Engine/Plugins/Experimental/GeometryProcessing/Source/GeometryAlgorithms/Private/GTEngine/Mathematics/GteDistLineSegment.h b/Engine/Plugins/Experimental/GeometryProcessing/Source/GeometryAlgorithms/Private/GTEngine/Mathematics/GteDistLineSegment.h
new file mode 100644
--- /dev/null
+++ b/Engine/Plugins/Experimental/GeometryProcessing/Source/GeometryAlgorithms/Private/GTEngine/Mathematics/GteDistLineSegment.h
@@ -0,0 +1,113 @@
+// David Eberly, Geometric Tools, Redmond WA 98052
+// Copyright (c) 1998-2019
+// Distributed under the Boost Software License, Version 1.0.
+// http://www.boost.org/LICENSE_1_0.txt
+// http://www.geometrictools.com/License/Boost/LICENSE_1_0.txt
+// File Version: 3.0.1 (2018/10/05)
+
+#pragma once
+
+#include <GTEngine/Mathematics/GteDCPQuery.h>
+#include <GTEngine/Mathematics/GteLine.h>
+#include <GTEngine/Mathematics/GteSegment.h>
+
+namespace gte
+{
+
+template <int N, typename Real>
+class DCPQuery<Real, Line<N, Real>, Segment<N, Real>>
+{
+public:
+    struct Result
+    {
+        Real distance, sqrDistance;
+        Real parameter[2];
+        Vector<N, Real> closestPoint[2];
+    };
+
+    // The centered form of the 'segment' is used.  Thus, parameter[1] of
+    // the result is in [-e,e], where e = |segment.p[1] - segment.p[0]|/2.
+    Result operator()(Line<N, Real> const& line,
+        Segment<N, Real> const& segment);
+};
+
+// Template aliases for convenience.
+template <int N, typename Real>
+using DCPLineSegment = DCPQuery<Real, Line<N, Real>, Segment<N, Real>>;
+
+template <typename Real>
+using DCPLine2Segment2 = DCPLineSegment<2, Real>;
+
+template <typename Real>
+using DCPLine3Segment3 = DCPLineSegment<3, Real>;
+
+
+template <int N, typename Real>
+typename DCPQuery<Real, Line<N, Real>, Segment<N, Real>>::Result
+DCPQuery<Real, Line<N, Real>, Segment<N, Real>>::operator()(
+    Line<N, Real> const& line, Segment<N, Real> const& segment)
+{
+    Result result;
+
+    Vector<N, Real> segCenter, segDirection;
+    Real segExtent;
+    segment.GetCenteredForm(segCenter, segDirection, segExtent);
+
+    Vector<N, Real> diff = line.origin - segCenter;
+    Real a01 = -Dot(line.direction, segDirection);
+    Real b0 = Dot(diff, line.direction);
+    Real s0, s1;
+
+    if (std::abs(a01) < (Real)1)
+    {
+        // The line and segment are not parallel.
+        Real det = (Real)1 - a01 * a01;
+        Real extDet = segExtent * det;
+        Real b1 = -Dot(diff, segDirection);
+        s1 = a01 * b0 - b1;
+
+        if (s1 >= -extDet)
+        {
+            if (s1 <= extDet)
+            {
+                // Two interior points are closest, one on the line and one
+                // on the segment.
+                s0 = (a01 * b1 - b0) / det;
+                s1 /= det;
+            }
+            else
+            {
+                // The endpoint e1 of the segment and an interior point of
+                // the line are closest.
+                s1 = segExtent;
+                s0 = -(a01 * s1 + b0);
+            }
+        }
+        else
+        {
+            // The endpoint e0 of the segment and an interior point of the
+            // line are closest.
+            s1 = -segExtent;
+            s0 = -(a01 * s1 + b0);
+        }
+    }
+    else
+    {
+        // The line and segment are parallel.  Choose the closest pair so that
+        // one point is at segment origin.
+        s1 = (Real)0;
+        s0 = -b0;
+    }
+
+    result.parameter[0] = s0;
+    result.parameter[1] = s1;
+    result.closestPoint[0] = line.origin + s0 * line.direction;
+    result.closestPoint[1] = segCenter + s1 * segDirection;
+    diff = result.closestPoint[0] - result.closestPoint[1];
+    result.sqrDistance = Dot(diff, diff);
+    result.distance = std::sqrt(result.sqrDistance);
+    return result;
+}
+
+
+}
diff --git a/Engine/Plugins/Experimental/GeometryProcessing/Source/GeometryAlgorithms/Private/GTEngine/Mathematics/GteDistOrientedBox3OrientedBox3.h b/Engine/Plugins/Experimental/GeometryProcessing/Source/GeometryAlgorithms/Private/GTEngine/Mathematics/GteDistOrientedBox3OrientedBox3.h
new file mode 100644
--- /dev/null
+++ b/Engine/Plugins/Experimental/GeometryProcessing/Source/GeometryAlgorithms/Private/GTEngine/Mathematics/GteDistOrientedBox3OrientedBox3.h
@@ -0,0 +1,162 @@
+// David Eberly, Geometric Tools, Redmond WA 98052
+// Copyright (c) 1998-2019
+// Distributed under the Boost Software License, Version 1.0.
+// http://www.boost.org/LICENSE_1_0.txt
+// http://www.geometrictools.com/License/Boost/LICENSE_1_0.txt
+// File Version: 3.5.2 (2018/10/05)
+
+#pragma once
+
+#include <GTEngine/Mathematics/GteDCPQuery.h>
+#include <GTEngine/Mathematics/GteLCPSolver.h>
+#include <GTEngine/Mathematics/GteOrientedBox.h>
+#include <GTEngine/Mathematics/GteVector3.h>
+
+namespace gte
+{
+
+template <typename Real>
+class DCPQuery<Real, OrientedBox3<Real>, OrientedBox3<Real>>
+{
+public:
+    struct Result
+    {
+        bool queryIsSuccessful;
+
+        // These members are valid only when queryIsSuccessful is true;
+        // otherwise, they are all set to zero.
+        Real distance, sqrDistance;
+        std::array<Real, 3> box0Parameter, box1Parameter;
+        Vector3<Real> closestPoint[2];
+
+        // The number of iterations used by LCPSolver regardless of whether
+        // the query is successful.
+        int numLCPIterations;
+    };
+
+    // The default maximum iterations is 144 (n = 12, maxIterations = n*n).  If
+    // the solver fails to converge, try increasing the maximum number of
+    // iterations.
+    void SetMaxLCPIterations(int maxLCPIterations);
+
+    Result operator()(OrientedBox3<Real> const& box0, OrientedBox3<Real> const& box1);
+
+private:
+    LCPSolver<Real, 12> mLCP;
+};
+
+
+template <typename Real>
+void DCPQuery<Real, OrientedBox3<Real>, OrientedBox3<Real>>::SetMaxLCPIterations(int maxLCPIterations)
+{
+    mLCP.SetMaxIterations(maxLCPIterations);
+}
+
+template <typename Real>
+typename DCPQuery<Real, OrientedBox3<Real>, OrientedBox3<Real>>::Result
+    DCPQuery<Real, OrientedBox3<Real>, OrientedBox3<Real>>::operator()(
+        OrientedBox3<Real> const& box0, OrientedBox3<Real> const& box1)
+{
+    Result result;
+
+    // Translate the center of box0 to the origin.  Modify the oriented box
+    // coefficients to be nonnegative.
+    Vector3<Real> delta = box1.center - box0.center;
+    for (int i = 0; i < 3; ++i)
+    {
+        delta += box0.extent[i] * box0.axis[i];
+        delta -= box1.extent[i] * box1.axis[i];
+    }
+
+    Vector3<Real> R0Delta, R1Delta;
+    for (int i = 0; i < 3; ++i)
+    {
+        R0Delta[i] = Dot(box0.axis[i], delta);
+        R1Delta[i] = Dot(box1.axis[i], delta);
+    }
+
+    std::array<std::array<Real, 3>, 3> R0TR1;
+    for (int r = 0; r < 3; ++r)
+    {
+        for (int c = 0; c < 3; ++c)
+        {
+            R0TR1[r][c] = Dot(box0.axis[r], box1.axis[c]);
+        }
+    }
+
+    Vector3<Real> twoExtent0 = box0.extent * (Real)2;
+    Vector3<Real> twoExtent1 = box1.extent * (Real)2;
+
+    // The LCP has 6 variables and 6 (nontrivial) inequality constraints.
+    std::array<Real, 12> q =
+    {
+        -R0Delta[0], -R0Delta[1], -R0Delta[2], R1Delta[0], R1Delta[1], R1Delta[2],
+        twoExtent0[0], twoExtent0[1], twoExtent0[2], twoExtent1[0], twoExtent1[1], twoExtent1[2]
+    };
+
+    std::array<std::array<Real, 12>, 12> M;
+    {
+        Real const z = (Real)0;
+        Real const p = (Real)1;
+        Real const m = (Real)-1;
+        M[ 0] = { p, z, z, -R0TR1[0][0], -R0TR1[0][1], -R0TR1[0][2], p, z, z, z, z, z };
+        M[ 1] = { z, p, z, -R0TR1[1][0], -R0TR1[1][1], -R0TR1[1][2], z, p, z, z, z, z };
+        M[ 2] = { z, z, p, -R0TR1[2][0], -R0TR1[2][1], -R0TR1[2][2], z, z, p, z, z, z };
+        M[ 3] = { -R0TR1[0][0], -R0TR1[1][0], -R0TR1[2][0], p, z, z, z, z, z, p, z, z };
+        M[ 4] = { -R0TR1[0][1], -R0TR1[1][1], -R0TR1[2][1], z, p, z, z, z, z, z, p, z };
+        M[ 5] = { -R0TR1[0][2], -R0TR1[1][2], -R0TR1[2][2], z, z, p, z, z, z, z, z, p };
+        M[ 6] = { m, z, z, z, z, z, z, z, z, z, z, z };
+        M[ 7] = { z, m, z, z, z, z, z, z, z, z, z, z };
+        M[ 8] = { z, z, m, z, z, z, z, z, z, z, z, z };
+        M[ 9] = { z, z, z, m, z, z, z, z, z, z, z, z };
+        M[10] = { z, z, z, z, m, z, z, z, z, z, z, z };
+        M[11] = { z, z, z, z, z, m, z, z, z, z, z, z };
+    }
+
+    std::array<Real, 12> w, z;
+    if (mLCP.Solve(q, M, w, z))
+    {
+        result.queryIsSuccessful = true;
+
+        result.closestPoint[0] = box0.center;
+        for (int i = 0; i < 3; ++i)
+        {
+            result.box0Parameter[i] = z[i] - box0.extent[i];
+            result.closestPoint[0] += result.box0Parameter[i] * box0.axis[i];
+        }
+
+        result.closestPoint[1] = box1.center;
+        for (int i = 0, j = 3; i < 3; ++i, ++j)
+        {
+            result.box1Parameter[i] = z[j] - box1.extent[i];
+            result.closestPoint[1] += result.box1Parameter[i] * box1.axis[i];
+        }
+
+        Vector3<Real> diff = result.closestPoint[1] - result.closestPoint[0];
+        result.sqrDistance = Dot(diff, diff);
+        result.distance = std::sqrt(result.sqrDistance);
+    }
+    else
+    {
+        // If you reach this case, the maximum number of iterations was not
+        // specified to be large enough or there is a problem due to
+        // floating-point rounding errors.  If you believe the latter is
+        // true, file a bug report.
+        result.queryIsSuccessful = false;
+
+        for (int i = 0; i < 3; ++i)
+        {
+            result.box0Parameter[i] = (Real)0;
+            result.box1Parameter[i] = (Real)0;
+            result.closestPoint[0][i] = (Real)0;
+            result.closestPoint[1][i] = (Real)0;
+        }
+        result.distance = (Real)0;
+        result.sqrDistance = (Real)0;
+    }
+
+    result.numLCPIterations = mLCP.GetNumIterations();
+    return result;
+}
+
+}
diff --git a/Engine/Plugins/Experimental/GeometryProcessing/Source/GeometryAlgorithms/Private/GTEngine/Mathematics/GteDistPoint3Circle3.h b/Engine/Plugins/Experimental/GeometryProcessing/Source/GeometryAlgorithms/Private/GTEngine/Mathematics/GteDistPoint3Circle3.h
new file mode 100644
--- /dev/null
+++ b/Engine/Plugins/Experimental/GeometryProcessing/Source/GeometryAlgorithms/Private/GTEngine/Mathematics/GteDistPoint3Circle3.h
@@ -0,0 +1,76 @@
+// David Eberly, Geometric Tools, Redmond WA 98052
+// Copyright (c) 1998-2019
+// Distributed under the Boost Software License, Version 1.0.
+// http://www.boost.org/LICENSE_1_0.txt
+// http://www.geometrictools.com/License/Boost/LICENSE_1_0.txt
+// File Version: 3.0.2 (2018/10/05)
+
+#pragma once
+
+#include <GTEngine/Mathematics/GteDCPQuery.h>
+#include <GTEngine/Mathematics/GteCircle3.h>
+
+// The 3D point-circle distance algorithm is described in
+// http://www.geometrictools.com/Documentation/DistanceToCircle3.pdf
+// The notation used in the code matches that of the document.
+
+namespace gte
+{
+
+template <typename Real>
+class DCPQuery<Real, Vector3<Real>, Circle3<Real>>
+{
+public:
+    // Either a single point on the circle is closest to 'point', in which
+    // case 'equidistant' is false, or the entire circle is closest to
+    // 'point', in which case 'equidistant' is true.  In the latter case, the
+    // query returns the circle point C+r*U, where C is the circle center, r
+    // is the circle radius, and U is a vector perpendicular to the normal N
+    // for the plane of the circle.
+    struct Result
+    {
+        Real distance, sqrDistance;
+        Vector3<Real> circleClosest;
+        bool equidistant;
+    };
+
+    Result operator()(Vector3<Real> const& point,
+        Circle3<Real> const& circle);
+};
+
+
+template <typename Real>
+typename DCPQuery<Real, Vector3<Real>, Circle3<Real>>::Result
+DCPQuery<Real, Vector3<Real>, Circle3<Real>>::operator()(
+    Vector3<Real> const& point, Circle3<Real> const& circle)
+{
+    Result result;
+
+    // Projection of P-C onto plane is Q-C = P-C - Dot(N,P-C)*N.
+    Vector3<Real> PmC = point - circle.center;
+    Vector3<Real> QmC = PmC - Dot(circle.normal, PmC)*circle.normal;
+    Real lengthQmC = Length(QmC);
+    if (lengthQmC > (Real)0)
+    {
+        result.circleClosest =
+            circle.center + (circle.radius / lengthQmC) * QmC;
+        result.equidistant = false;
+    }
+    else
+    {
+        // All circle points are equidistant from P.  Return one of them.
+        Vector3<Real> basis[3];
+        basis[0] = circle.normal;
+        ComputeOrthogonalComplement(1, basis);
+        result.circleClosest = circle.center + circle.radius * basis[1];
+        result.equidistant = true;
+    }
+
+    Vector3<Real> diff = point - result.circleClosest;
+    result.sqrDistance = Dot(diff, diff);
+    result.distance = std::sqrt(result.sqrDistance);
+    return result;
+}
+
+
+}
diff --git a/Engine/Plugins/Experimental/GeometryProcessing/Source/GeometryAlgorithms/Private/GTEngine/Mathematics/GteDistPoint3ConvexPolyhedron3.h b/Engine/Plugins/Experimental/GeometryProcessing/Source/GeometryAlgorithms/Private/GTEngine/Mathematics/GteDistPoint3ConvexPolyhedron3.h
new file mode 100644
--- /dev/null
+++ b/Engine/Plugins/Experimental/GeometryProcessing/Source/GeometryAlgorithms/Private/GTEngine/Mathematics/GteDistPoint3ConvexPolyhedron3.h
@@ -0,0 +1,188 @@
+// David Eberly, Geometric Tools, Redmond WA 98052
+// Copyright (c) 1998-2019
+// Distributed under the Boost Software License, Version 1.0.
+// http://www.boost.org/LICENSE_1_0.txt
+// http://www.geometrictools.com/License/Boost/LICENSE_1_0.txt
+// File Version: 3.5.1 (2018/10/05)
+
+#pragma once
+
+#include <GTEngine/Mathematics/GteDCPQuery.h>
+#include <GTEngine/Mathematics/GteLCPSolver.h>
+#include <GTEngine/Mathematics/GteConvexPolyhedron3.h>
+#include <memory>
+
+namespace gte
+{
+
+template <typename Real>
+class DCPQuery<Real, Vector3<Real>, ConvexPolyhedron3<Real>>
+{
+public:
+    // Construction.  If you have no knowledge of the number of faces for the
+    // convex polyhedra you plan on applying the query to, pass 'numTriangles'
+    // of zero.  This is a request to the operator() function to create the
+    // LCP solver for each query, and this requires memory allocation and
+    // deallocation per query.  If you plan on applying the query multiple
+    // times to a single polyhedron, even if the vertices of the polyhedron
+    // are modified for each query, then pass 'numTriangles' to be the number
+    // of triangle faces for that polyhedron.  This lets the operator()
+    // function know to create the LCP solver once at construction time, thus
+    // avoiding the memory management costs during the query.
+    DCPQuery(int numTriangles = 0);
+
+    struct Result
+    {
+        bool queryIsSuccessful;
+
+        // These members are valid only when queryIsSuccessful is true;
+        // otherwise, they are all set to zero.
+        Real distance, sqrDistance;
+        Vector3<Real> closestPoint[2];
+
+        // The number of iterations used by LCPSolver regardless of whether
+        // the query is successful.
+        int numLCPIterations;
+    };
+
+    // The default maximum iterations is 81 (n = 9, maxIterations = n*n).  If
+    // the solver fails to converge, try increasing the maximum number of
+    // iterations.
+    void SetMaxLCPIterations(int maxLCPIterations);
+
+    Result operator()(Vector3<Real> const& point, ConvexPolyhedron3<Real> const& polyhedron);
+
+private:
+    int mMaxLCPIterations;
+    std::unique_ptr<LCPSolver<Real>> mLCP;
+};
+
+
+template <typename Real>
+DCPQuery<Real, Vector3<Real>, ConvexPolyhedron3<Real>>::DCPQuery(int numTriangles)
+{
+    if (numTriangles > 0)
+    {
+        int const n = numTriangles + 3;
+        mLCP = std::make_unique<LCPSolver<Real>>(n);
+        mMaxLCPIterations = mLCP->GetMaxIterations();
+    }
+    else
+    {
+        mMaxLCPIterations = 0;
+    }
+}
+
+template <typename Real>
+void DCPQuery<Real, Vector3<Real>, ConvexPolyhedron3<Real>>::SetMaxLCPIterations(int maxLCPIterations)
+{
+    mMaxLCPIterations = maxLCPIterations;
+    if (mLCP)
+    {
+        mLCP->SetMaxIterations(mMaxLCPIterations);
+    }
+}
+
+template <typename Real>
+typename DCPQuery<Real, Vector3<Real>, ConvexPolyhedron3<Real>>::Result
+DCPQuery<Real, Vector3<Real>, ConvexPolyhedron3<Real>>::operator()(
+    Vector3<Real> const& point, ConvexPolyhedron3<Real> const& polyhedron)
+{
+    Result result;
+
+    int const numTriangles = static_cast<int>(polyhedron.planes.size());
+    if (numTriangles == 0)
+    {
+        // The polyhedron planes and aligned box need to be created.
+        result.queryIsSuccessful = false;
+        for (int i = 0; i < 3; ++i)
+        {
+            result.closestPoint[0][i] = (Real)0;
+            result.closestPoint[1][i] = (Real)0;
+        }
+        result.distance = (Real)0;
+        result.sqrDistance = (Real)0;
+        result.numLCPIterations = 0;
+        return result;
+    }
+
+    int const n = numTriangles + 3;
+
+    // Translate the point and convex polyhedron so that the polyhedron is in
+    // the first octant.  The translation is not explicit; rather, the q and M
+    // for the LCP are initialized using the translation information.
+    Vector4<Real> hmin = HLift(polyhedron.alignedBox.min, (Real)1);
+
+    std::vector<Real> q(n);
+    for (int r = 0; r < 3; ++r)
+    {
+        q[r] = polyhedron.alignedBox.min[r] - point[r];
+    }
+    for (int r = 3, t = 0; r < n; ++r, ++t)
+    {
+        q[r] = -Dot(polyhedron.planes[t], hmin);
+    }
+
+    std::vector<Real> M(n * n);
+    M[0] = (Real)1;  M[1] = (Real)0;  M[2] = (Real)0;
+    M[n] = (Real)0;  M[n + 1] = (Real)1;  M[n + 2] = (Real)0;
+    M[2 * n] = (Real)0;  M[2 * n + 1] = (Real)0;  M[2 * n + 2] = (Real)1;
+    for (int t = 0, c = 3; t < numTriangles; ++t, ++c)
+    {
+        Vector3<Real> normal = HProject(polyhedron.planes[t]);
+        for (int r = 0; r < 3; ++r)
+        {
+            M[c + n * r] = normal[r];
+            M[r + n * c] = -normal[r];
+        }
+    }
+    for (int r = 3; r < n; ++r)
+    {
+        for (int c = 3; c < n; ++c)
+        {
+            M[c + n * r] = (Real)0;
+        }
+    }
+
+    bool needsLCP = (mLCP == nullptr);
+    if (needsLCP)
+    {
+        mLCP = std::make_unique<LCPSolver<Real>>(n);
+        if (mMaxLCPIterations > 0)
+        {
+            mLCP->SetMaxIterations(mMaxLCPIterations);
+        }
+    }
+
+    std::vector<Real> w(n), z(n);
+    if (mLCP->Solve(q, M, w, z))
+    {
+        result.queryIsSuccessful = true;
+        result.closestPoint[0] = point;
+        for (int i = 0; i < 3; ++i)
+        {
+            result.closestPoint[1][i] = z[i] + polyhedron.alignedBox.min[i];
+        }
+
+        Vector3<Real> diff = result.closestPoint[1] - result.closestPoint[0];
+        result.sqrDistance = Dot(diff, diff);
+        result.distance = std::sqrt(result.sqrDistance);
+    }
+    else
+    {
+        // If you reach this case, the maximum number of iterations was not
+        // specified to be large enough or there is a problem due to
+        // floating-point rounding errors.  If you believe the latter is
+        // true, file a bug report.
+        result.queryIsSuccessful = false;
+    }
+
+    result.numLCPIterations = mLCP->GetNumIterations();
+    if (needsLCP)
+    {
+        mLCP = nullptr;
+    }
+    return result;
+}
+
+}
diff --git a/Engine/Plugins/Experimental/GeometryProcessing/Source/GeometryAlgorithms/Private/GTEngine/Mathematics/GteDistPoint3Cylinder3.h b/Engine/Plugins/Experimental/GeometryProcessing/Source/GeometryAlgorithms/Private/GTEngine/Mathematics/GteDistPoint3Cylinder3.h
new file mode 100644
--- /dev/null
+++ b/Engine/Plugins/Experimental/GeometryProcessing/Source/GeometryAlgorithms/Private/GTEngine/Mathematics/GteDistPoint3Cylinder3.h
@@ -0,0 +1,126 @@
+// David Eberly, Geometric Tools, Redmond WA 98052
+// Copyright (c) 1998-2019
+// Distributed under the Boost Software License, Version 1.0.
+// http://www.boost.org/LICENSE_1_0.txt
+// http://www.geometrictools.com/License/Boost/LICENSE_1_0.txt
+// File Version: 3.0.1 (2018/10/05)
+
+#pragma once
+
+#include <GTEngine/Mathematics/GteVector3.h>
+#include <GTEngine/Mathematics/GteDCPQuery.h>
+#include <GTEngine/Mathematics/GteCylinder3.h>
+
+// The queries consider the cylinder to be a solid.
+
+namespace gte
+{
+
+template <typename Real>
+class DCPQuery<Real, Vector3<Real>, Cylinder3<Real>>
+{
+public:
+    struct Result
+    {
+        Real distance;
+        Vector3<Real> cylinderClosest;
+    };
+
+    Result operator()(Vector3<Real> const& point,
+        Cylinder3<Real> const& cylinder);
+
+private:
+    void DoQueryInfiniteCylinder(Vector3<Real> const& P, Real radius,
+        Result& result);
+
+    void DoQueryFiniteCylinder(Vector3<Real> const& P, Real radius,
+        Real height, Result& result);
+};
+
+
+template <typename Real>
+typename DCPQuery<Real, Vector3<Real>, Cylinder3<Real>>::Result
+DCPQuery<Real, Vector3<Real>, Cylinder3<Real>>::operator()(
+    Vector3<Real> const& point, Cylinder3<Real> const& cylinder)
+{
+    Result result;
+
+    // Convert the point to the cylinder coordinate system.  In this system,
+    // the point believes (0,0,0) is the cylinder axis origin and (0,0,1) is
+    // the cylinder axis direction.
+    Vector3<Real> basis[3];
+    basis[0] = cylinder.axis.direction;
+    ComputeOrthogonalComplement(1, basis);
+
+    Vector3<Real> delta = point - cylinder.axis.origin;
+    Vector3<Real> P
+    {
+        Dot(basis[1], delta),
+        Dot(basis[2], delta),
+        Dot(basis[0], delta)
+    };
+
+    if (cylinder.height == std::numeric_limits<Real>::max())
+    {
+        DoQueryInfiniteCylinder(P, cylinder.radius, result);
+    }
+    else
+    {
+        DoQueryFiniteCylinder(P, cylinder.radius, cylinder.height, result);
+    }
+
+    // Convert the closest point from the cylinder coordinate system to the
+    // original coordinate system.
+    result.cylinderClosest = cylinder.axis.origin +
+        result.cylinderClosest[0] * basis[1] +
+        result.cylinderClosest[1] * basis[2] +
+        result.cylinderClosest[2] * basis[0];
+
+    return result;
+}
+
+template <typename Real>
+void DCPQuery<Real, Vector3<Real>, Cylinder3<Real>>::DoQueryInfiniteCylinder(
+    Vector3<Real> const& P, Real radius, Result& result)
+{
+    Real sqrRadius = radius * radius;
+    Real sqrDistance = P[0] * P[0] + P[1] * P[1];
+    if (sqrDistance >= sqrRadius)
+    {
+        // The point is outside the cylinder or on the cylinder wall.
+        Real distance = std::sqrt(sqrDistance);
+        result.distance = distance - radius;
+        Real temp = radius / distance;
+        result.cylinderClosest[0] = P[0] * temp;
+        result.cylinderClosest[1] = P[1] * temp;
+        result.cylinderClosest[2] = P[2];
+    }
+    else
+    {
+        // The point is inside the cylinder.
+        result.distance = (Real)0;
+        result.cylinderClosest = P;
+    }
+}
+
+template <typename Real>
+void DCPQuery<Real, Vector3<Real>, Cylinder3<Real>>::DoQueryFiniteCylinder(
+    Vector3<Real> const& P, Real radius, Real height, Result& result)
+{
+    DoQueryInfiniteCylinder(P, radius, result);
+
+    // Clamp the infinite cylinder's closest point to the finite cylinder.
+    if (result.cylinderClosest[2] > height)
+    {
+        result.cylinderClosest[2] = height;
+        result.distance = Length(result.cylinderClosest - P);
+    }
+    else if (result.cylinderClosest[2] < -height)
+    {
+        result.cylinderClosest[2] = -height;
+        result.distance = Length(result.cylinderClosest - P);
+    }
+}
+
+
+}
diff --git a/Engine/Plugins/Experimental/GeometryProcessing/Source/GeometryAlgorithms/Private/GTEngine/Mathematics/GteDistPoint3Frustum3.h b/Engine/Plugins/Experimental/GeometryProcessing/Source/GeometryAlgorithms/Private/GTEngine/Mathematics/GteDistPoint3Frustum3.h
new file mode 100644
--- /dev/null
+++ b/Engine/Plugins/Experimental/GeometryProcessing/Source/GeometryAlgorithms/Private/GTEngine/Mathematics/GteDistPoint3Frustum3.h
@@ -0,0 +1,430 @@
+// David Eberly, Geometric Tools, Redmond WA 98052
+// Copyright (c) 1998-2019
+// Distributed under the Boost Software License, Version 1.0.
+// http://www.boost.org/LICENSE_1_0.txt
+// http://www.geometrictools.com/License/Boost/LICENSE_1_0.txt
+// File Version: 3.0.1 (2018/10/05)
+
+#pragma once
+
+#include <GTEngine/Mathematics/GteDCPQuery.h>
+#include <GTEngine/Mathematics/GteFrustum3.h>
+
+namespace gte
+{
+
+template <typename Real>
+class DCPQuery<Real, Vector3<Real>, Frustum3<Real>>
+{
+public:
+    struct Result
+    {
+        Real distance, sqrDistance;
+        Vector3<Real> frustumClosestPoint;
+    };
+
+    Result operator()(Vector3<Real> const& point,
+        Frustum3<Real> const& frustum);
+};
+
+
+template <typename Real>
+typename DCPQuery<Real, Vector3<Real>, Frustum3<Real>>::Result
+DCPQuery<Real, Vector3<Real>, Frustum3<Real>>::operator()(
+    Vector3<Real> const& point, Frustum3<Real> const& frustum)
+{
+    Result result;
+
+    // Compute coordinates of point with respect to frustum coordinate system.
+    Vector3<Real> diff = point - frustum.origin;
+    Vector3<Real> test = {
+        Dot(diff, frustum.rVector),
+        Dot(diff, frustum.uVector),
+        Dot(diff, frustum.dVector) };
+
+    // Perform calculations in octant with nonnegative R and U coordinates.
+    bool rSignChange;
+    if (test[0] < (Real)0)
+    {
+        rSignChange = true;
+        test[0] = -test[0];
+    }
+    else
+    {
+        rSignChange = false;
+    }
+
+    bool uSignChange;
+    if (test[1] < (Real)0)
+    {
+        uSignChange = true;
+        test[1] = -test[1];
+    }
+    else
+    {
+        uSignChange = false;
+    }
+
+    // Frustum derived parameters.
+    Real rmin = frustum.rBound;
+    Real rmax = frustum.GetDRatio()*rmin;
+    Real umin = frustum.uBound;
+    Real umax = frustum.GetDRatio()*umin;
+    Real dmin = frustum.dMin;
+    Real dmax = frustum.dMax;
+    Real rminSqr = rmin*rmin;
+    Real uminSqr = umin*umin;
+    Real dminSqr = dmin*dmin;
+    Real minRDDot = rminSqr + dminSqr;
+    Real minUDDot = uminSqr + dminSqr;
+    Real minRUDDot = rminSqr + minUDDot;
+    Real maxRDDot = frustum.GetDRatio()*minRDDot;
+    Real maxUDDot = frustum.GetDRatio()*minUDDot;
+    Real maxRUDDot = frustum.GetDRatio()*minRUDDot;
+
+    // Algorithm computes closest point in all cases by determining in which
+    // Voronoi region of the vertices, edges, and faces of the frustum that
+    // the test point lives.
+    Vector3<Real> closest;
+    Real rDot, uDot, rdDot, udDot, rudDot, rEdgeDot, uEdgeDot, t;
+    if (test[2] >= dmax)
+    {
+        if (test[0] <= rmax)
+        {
+            if (test[1] <= umax)
+            {
+                // F-face
+                closest[0] = test[0];
+                closest[1] = test[1];
+                closest[2] = dmax;
+            }
+            else
+            {
+                // UF-edge
+                closest[0] = test[0];
+                closest[1] = umax;
+                closest[2] = dmax;
+            }
+        }
+        else
+        {
+            if (test[1] <= umax)
+            {
+                // LF-edge
+                closest[0] = rmax;
+                closest[1] = test[1];
+                closest[2] = dmax;
+            }
+            else
+            {
+                // LUF-vertex
+                closest[0] = rmax;
+                closest[1] = umax;
+                closest[2] = dmax;
+            }
+        }
+    }
+    else if (test[2] <= dmin)
+    {
+        if (test[0] <= rmin)
+        {
+            if (test[1] <= umin)
+            {
+                // N-face
+                closest[0] = test[0];
+                closest[1] = test[1];
+                closest[2] = dmin;
+            }
+            else
+            {
+                udDot = umin*test[1] + dmin*test[2];
+                if (udDot >= maxUDDot)
+                {
+                    // UF-edge
+                    closest[0] = test[0];
+                    closest[1] = umax;
+                    closest[2] = dmax;
+                }
+                else if (udDot >= minUDDot)
+                {
+                    // U-face
+                    uDot = dmin*test[1] - umin*test[2];
+                    t = uDot / minUDDot;
+                    closest[0] = test[0];
+                    closest[1] = test[1] - t*dmin;
+                    closest[2] = test[2] + t*umin;
+                }
+                else
+                {
+                    // UN-edge
+                    closest[0] = test[0];
+                    closest[1] = umin;
+                    closest[2] = dmin;
+                }
+            }
+        }
+        else
+        {
+            if (test[1] <= umin)
+            {
+                rdDot = rmin*test[0] + dmin*test[2];
+                if (rdDot >= maxRDDot)
+                {
+                    // LF-edge
+                    closest[0] = rmax;
+                    closest[1] = test[1];
+                    closest[2] = dmax;
+                }
+                else if (rdDot >= minRDDot)
+                {
+                    // L-face
+                    rDot = dmin*test[0] - rmin*test[2];
+                    t = rDot / minRDDot;
+                    closest[0] = test[0] - t*dmin;
+                    closest[1] = test[1];
+                    closest[2] = test[2] + t*rmin;
+                }
+                else
+                {
+                    // LN-edge
+                    closest[0] = rmin;
+                    closest[1] = test[1];
+                    closest[2] = dmin;
+                }
+            }
+            else
+            {
+                rudDot = rmin*test[0] + umin*test[1] + dmin*test[2];
+                rEdgeDot = umin*rudDot - minRUDDot*test[1];
+                if (rEdgeDot >= (Real)0)
+                {
+                    rdDot = rmin*test[0] + dmin*test[2];
+                    if (rdDot >= maxRDDot)
+                    {
+                        // LF-edge
+                        closest[0] = rmax;
+                        closest[1] = test[1];
+                        closest[2] = dmax;
+                    }
+                    else if (rdDot >= minRDDot)
+                    {
+                        // L-face
+                        rDot = dmin*test[0] - rmin*test[2];
+                        t = rDot / minRDDot;
+                        closest[0] = test[0] - t*dmin;
+                        closest[1] = test[1];
+                        closest[2] = test[2] + t*rmin;
+                    }
+                    else
+                    {
+                        // LN-edge
+                        closest[0] = rmin;
+                        closest[1] = test[1];
+                        closest[2] = dmin;
+                    }
+                }
+                else
+                {
+                    uEdgeDot = rmin*rudDot - minRUDDot*test[0];
+                    if (uEdgeDot >= (Real)0)
+                    {
+                        udDot = umin*test[1] + dmin*test[2];
+                        if (udDot >= maxUDDot)
+                        {
+                            // UF-edge
+                            closest[0] = test[0];
+                            closest[1] = umax;
+                            closest[2] = dmax;
+                        }
+                        else if (udDot >= minUDDot)
+                        {
+                            // U-face
+                            uDot = dmin*test[1] - umin*test[2];
+                            t = uDot / minUDDot;
+                            closest[0] = test[0];
+                            closest[1] = test[1] - t*dmin;
+                            closest[2] = test[2] + t*umin;
+                        }
+                        else
+                        {
+                            // UN-edge
+                            closest[0] = test[0];
+                            closest[1] = umin;
+                            closest[2] = dmin;
+                        }
+                    }
+                    else
+                    {
+                        if (rudDot >= maxRUDDot)
+                        {
+                            // LUF-vertex
+                            closest[0] = rmax;
+                            closest[1] = umax;
+                            closest[2] = dmax;
+                        }
+                        else if (rudDot >= minRUDDot)
+                        {
+                            // LU-edge
+                            t = rudDot / minRUDDot;
+                            closest[0] = t*rmin;
+                            closest[1] = t*umin;
+                            closest[2] = t*dmin;
+                        }
+                        else
+                        {
+                            // LUN-vertex
+                            closest[0] = rmin;
+                            closest[1] = umin;
+                            closest[2] = dmin;
+                        }
+                    }
+                }
+            }
+        }
+    }
+    else
+    {
+        rDot = dmin*test[0] - rmin*test[2];
+        uDot = dmin*test[1] - umin*test[2];
+        if (rDot <= (Real)0.0)
+        {
+            if (uDot <= (Real)0)
+            {
+                // point inside frustum
+                closest = test;
+            }
+            else
+            {
+                udDot = umin*test[1] + dmin*test[2];
+                if (udDot >= maxUDDot)
+                {
+                    // UF-edge
+                    closest[0] = test[0];
+                    closest[1] = umax;
+                    closest[2] = dmax;
+                }
+                else
+                {
+                    // U-face
+                    t = uDot / minUDDot;
+                    closest[0] = test[0];
+                    closest[1] = test[1] - t*dmin;
+                    closest[2] = test[2] + t*umin;
+                }
+            }
+        }
+        else
+        {
+            if (uDot <= (Real)0)
+            {
+                rdDot = rmin*test[0] + dmin*test[2];
+                if (rdDot >= maxRDDot)
+                {
+                    // LF-edge
+                    closest[0] = rmax;
+                    closest[1] = test[1];
+                    closest[2] = dmax;
+                }
+                else
+                {
+                    // L-face
+                    t = rDot / minRDDot;
+                    closest[0] = test[0] - t*dmin;
+                    closest[1] = test[1];
+                    closest[2] = test[2] + t*rmin;
+                }
+            }
+            else
+            {
+                rudDot = rmin*test[0] + umin*test[1] + dmin*test[2];
+                rEdgeDot = umin*rudDot - minRUDDot*test[1];
+                if (rEdgeDot >= (Real)0)
+                {
+                    rdDot = rmin*test[0] + dmin*test[2];
+                    if (rdDot >= maxRDDot)
+                    {
+                        // LF-edge
+                        closest[0] = rmax;
+                        closest[1] = test[1];
+                        closest[2] = dmax;
+                    }
+                    else // assert( rdDot >= minRDDot ) from geometry
+                    {
+                        // L-face
+                        t = rDot / minRDDot;
+                        closest[0] = test[0] - t*dmin;
+                        closest[1] = test[1];
+                        closest[2] = test[2] + t*rmin;
+                    }
+                }
+                else
+                {
+                    uEdgeDot = rmin*rudDot - minRUDDot*test[0];
+                    if (uEdgeDot >= (Real)0)
+                    {
+                        udDot = umin*test[1] + dmin*test[2];
+                        if (udDot >= maxUDDot)
+                        {
+                            // UF-edge
+                            closest[0] = test[0];
+                            closest[1] = umax;
+                            closest[2] = dmax;
+                        }
+                        else // assert( udDot >= minUDDot ) from geometry
+                        {
+                            // U-face
+                            t = uDot / minUDDot;
+                            closest[0] = test[0];
+                            closest[1] = test[1] - t*dmin;
+                            closest[2] = test[2] + t*umin;
+                        }
+                    }
+                    else
+                    {
+                        if (rudDot >= maxRUDDot)
+                        {
+                            // LUF-vertex
+                            closest[0] = rmax;
+                            closest[1] = umax;
+                            closest[2] = dmax;
+                        }
+                        else // assert( rudDot >= minRUDDot ) from geometry
+                        {
+                            // LU-edge
+                            t = rudDot / minRUDDot;
+                            closest[0] = t*rmin;
+                            closest[1] = t*umin;
+                            closest[2] = t*dmin;
+                        }
+                    }
+                }
+            }
+        }
+    }
+
+    diff = test - closest;
+
+    // Convert back to original quadrant.
+    if (rSignChange)
+    {
+        closest[0] = -closest[0];
+    }
+
+    if (uSignChange)
+    {
+        closest[1] = -closest[1];
+    }
+
+    // Convert back to original coordinates.
+    result.frustumClosestPoint = frustum.origin +
+        closest[0] * frustum.rVector +
+        closest[1] * frustum.uVector +
+        closest[2] * frustum.dVector;
+
+    result.sqrDistance = Dot(diff, diff);
+    result.distance = std::sqrt(result.sqrDistance);
+    return result;
+}
+
+
+}
diff --git a/Engine/Plugins/Experimental/GeometryProcessing/Source/GeometryAlgorithms/Private/GTEngine/Mathematics/GteDistPoint3Plane3.h b/Engine/Plugins/Experimental/GeometryProcessing/Source/GeometryAlgorithms/Private/GTEngine/Mathematics/GteDistPoint3Plane3.h
new file mode 100644
--- /dev/null
+++ b/Engine/Plugins/Experimental/GeometryProcessing/Source/GeometryAlgorithms/Private/GTEngine/Mathematics/GteDistPoint3Plane3.h
@@ -0,0 +1,44 @@
+// David Eberly, Geometric Tools, Redmond WA 98052
+// Copyright (c) 1998-2019
+// Distributed under the Boost Software License, Version 1.0.
+// http://www.boost.org/LICENSE_1_0.txt
+// http://www.geometrictools.com/License/Boost/LICENSE_1_0.txt
+// File Version: 3.0.0 (2016/06/19)
+
+#pragma once
+
+#include <GTEngine/Mathematics/GteVector3.h>
+#include <GTEngine/Mathematics/GteDCPQuery.h>
+#include <GTEngine/Mathematics/GteHyperplane.h>
+
+namespace gte
+{
+
+template <typename Real>
+class DCPQuery<Real, Vector3<Real>, Plane3<Real>>
+{
+public:
+    struct Result
+    {
+        Real distance, signedDistance;
+        Vector3<Real> planeClosestPoint;
+    };
+
+    Result operator()(Vector3<Real> const& point, Plane3<Real> const& plane);
+};
+
+
+template <typename Real>
+typename DCPQuery<Real, Vector3<Real>, Plane3<Real>>::Result
+DCPQuery<Real, Vector3<Real>, Plane3<Real>>::operator()(
+    Vector3<Real> const& point, Plane3<Real> const& plane)
+{
+    Result result;
+    result.signedDistance = Dot(plane.normal, point) - plane.constant;
+    result.distance = std::abs(result.signedDistance);
+    result.planeClosestPoint = point - result.signedDistance*plane.normal;
+    return result;
+}
+
+
+}
diff --git a/Engine/Plugins/Experimental/GeometryProcessing/Source/GeometryAlgorithms/Private/GTEngine/Mathematics/GteDistPoint3Rectangle3.h b/Engine/Plugins/Experimental/GeometryProcessing/Source/GeometryAlgorithms/Private/GTEngine/Mathematics/GteDistPoint3Rectangle3.h
new file mode 100644
--- /dev/null
+++ b/Engine/Plugins/Experimental/GeometryProcessing/Source/GeometryAlgorithms/Private/GTEngine/Mathematics/GteDistPoint3Rectangle3.h
@@ -0,0 +1,85 @@
+// David Eberly, Geometric Tools, Redmond WA 98052
+// Copyright (c) 1998-2019
+// Distributed under the Boost Software License, Version 1.0.
+// http://www.boost.org/LICENSE_1_0.txt
+// http://www.geometrictools.com/License/Boost/LICENSE_1_0.txt
+// File Version: 3.0.1 (2018/10/05)
+
+#pragma once
+
+#include <GTEngine/Mathematics/GteVector3.h>
+#include <GTEngine/Mathematics/GteDCPQuery.h>
+#include <GTEngine/Mathematics/GteRectangle.h>
+
+namespace gte
+{
+
+template <typename Real>
+class DCPQuery<Real, Vector3<Real>, Rectangle3<Real>>
+{
+public:
+    struct Result
+    {
+        Real distance, sqrDistance;
+        Real rectangleParameter[2];
+        Vector3<Real> rectangleClosestPoint;
+    };
+
+    Result operator()(Vector3<Real> const& point,
+        Rectangle3<Real> const& rectangle);
+};
+
+
+template <typename Real>
+typename DCPQuery<Real, Vector3<Real>, Rectangle3<Real>>::Result
+DCPQuery<Real, Vector3<Real>, Rectangle3<Real>>::operator()(
+    Vector3<Real> const& point, Rectangle3<Real> const& rectangle)
+{
+    Result result;
+
+    Vector3<Real> diff = rectangle.center - point;
+    Real b0 = Dot(diff, rectangle.axis[0]);
+    Real b1 = Dot(diff, rectangle.axis[1]);
+    Real s0 = -b0, s1 = -b1;
+    result.sqrDistance = Dot(diff, diff);
+
+    if (s0 < -rectangle.extent[0])
+    {
+        s0 = -rectangle.extent[0];
+    }
+    else if (s0 > rectangle.extent[0])
+    {
+        s0 = rectangle.extent[0];
+    }
+    result.sqrDistance += s0*(s0 + ((Real)2)*b0);
+
+    if (s1 < -rectangle.extent[1])
+    {
+        s1 = -rectangle.extent[1];
+    }
+    else if (s1 > rectangle.extent[1])
+    {
+        s1 = rectangle.extent[1];
+    }
+    result.sqrDistance += s1*(s1 + ((Real)2)*b1);
+
+    // Account for numerical round-off error.
+    if (result.sqrDistance < (Real)0)
+    {
+        result.sqrDistance = (Real)0;
+    }
+
+    result.distance = std::sqrt(result.sqrDistance);
+    result.rectangleParameter[0] = s0;
+    result.rectangleParameter[1] = s1;
+    result.rectangleClosestPoint = rectangle.center;
+    for (int i = 0; i < 2; ++i)
+    {
+        result.rectangleClosestPoint +=
+            result.rectangleParameter[i] * rectangle.axis[i];
+    }
+    return result;
+}
+
+
+}
diff --git a/Engine/Plugins/Experimental/GeometryProcessing/Source/GeometryAlgorithms/Private/GTEngine/Mathematics/GteDistPoint3Tetrahedron3.h b/Engine/Plugins/Experimental/GeometryProcessing/Source/GeometryAlgorithms/Private/GTEngine/Mathematics/GteDistPoint3Tetrahedron3.h
new file mode 100644
--- /dev/null
+++ b/Engine/Plugins/Experimental/GeometryProcessing/Source/GeometryAlgorithms/Private/GTEngine/Mathematics/GteDistPoint3Tetrahedron3.h
@@ -0,0 +1,83 @@
+// David Eberly, Geometric Tools, Redmond WA 98052
+// Copyright (c) 1998-2019
+// Distributed under the Boost Software License, Version 1.0.
+// http://www.boost.org/LICENSE_1_0.txt
+// http://www.geometrictools.com/License/Boost/LICENSE_1_0.txt
+// File Version: 3.0.1 (2018/10/05)
+
+#pragma once
+
+#include <GTEngine/Mathematics/GteDistPointTriangle.h>
+#include <GTEngine/Mathematics/GteTetrahedron3.h>
+#include <limits>
+
+namespace gte
+{
+
+template <typename Real>
+class DCPQuery<Real, Vector3<Real>, Tetrahedron3<Real>>
+{
+public:
+    struct Result
+    {
+        Real distance, sqrDistance;
+        Vector3<Real> tetrahedronClosestPoint;
+    };
+
+    Result operator()(Vector3<Real> const& point,
+        Tetrahedron3<Real> const& tetrahedron);
+};
+
+
+template <typename Real>
+typename DCPQuery<Real, Vector3<Real>, Tetrahedron3<Real>>::Result
+DCPQuery<Real, Vector3<Real>, Tetrahedron3<Real>>::operator()(
+    Vector3<Real> const& point, Tetrahedron3<Real> const& tetrahedron)
+{
+    Result result;
+
+    // Construct the planes for the faces of the tetrahedron.  The normals
+    // are outer pointing, but specified not to be unit length.  We only need
+    // to know sidedness of the query point, so we will save cycles by not
+    // computing unit-length normals.
+    Plane3<Real> planes[4];
+    tetrahedron.GetPlanes(planes);
+
+    // Determine which faces are visible to the query point.  Only these
+    // need to be processed by point-to-triangle distance queries.
+    result.sqrDistance = std::numeric_limits<Real>::max();
+    result.tetrahedronClosestPoint = Vector3<Real>::Zero();
+    for (int i = 0; i < 4; ++i)
+    {
+        if (Dot(planes[i].normal, point) >= planes[i].constant)
+        {
+            int indices[3] = { 0, 0, 0 };
+            tetrahedron.GetFaceIndices(i, indices);
+            Triangle3<Real> triangle(
+                tetrahedron.v[indices[0]],
+                tetrahedron.v[indices[1]],
+                tetrahedron.v[indices[2]]);
+
+            DCPQuery<Real, Vector3<Real>, Triangle3<Real>> query;
+            auto ptResult = query(point, triangle);
+            if (ptResult.sqrDistance < result.sqrDistance)
+            {
+                result.sqrDistance = ptResult.sqrDistance;
+                result.tetrahedronClosestPoint = ptResult.closest;
+            }
+        }
+    }
+
+    if (result.sqrDistance == std::numeric_limits<Real>::max())
+    {
+        // The query point is inside the solid tetrahedron.  Report a zero
+        // distance.  The closest points are identical.
+        result.sqrDistance = (Real)0;
+        result.tetrahedronClosestPoint = point;
+    }
+    result.distance = std::sqrt(result.sqrDistance);
+    return result;
+}
+
+
+}
diff --git a/Engine/Plugins/Experimental/GeometryProcessing/Source/GeometryAlgorithms/Private/GTEngine/Mathematics/GteDistPointAlignedBox.h b/Engine/Plugins/Experimental/GeometryProcessing/Source/GeometryAlgorithms/Private/GTEngine/Mathematics/GteDistPointAlignedBox.h
new file mode 100644
--- /dev/null
+++ b/Engine/Plugins/Experimental/GeometryProcessing/Source/GeometryAlgorithms/Private/GTEngine/Mathematics/GteDistPointAlignedBox.h
@@ -0,0 +1,91 @@
+// David Eberly, Geometric Tools, Redmond WA 98052
+// Copyright (c) 1998-2019
+// Distributed under the Boost Software License, Version 1.0.
+// http://www.boost.org/LICENSE_1_0.txt
+// http://www.geometrictools.com/License/Boost/LICENSE_1_0.txt
+// File Version: 3.0.1 (2018/10/05)
+
+#pragma once
+
+#include <GTEngine/Mathematics/GteDCPQuery.h>
+#include <GTEngine/Mathematics/GteAlignedBox.h>
+
+namespace gte
+{
+
+template <int N, typename Real>
+class DCPQuery<Real, Vector<N, Real>, AlignedBox<N, Real>>
+{
+public:
+    struct Result
+    {
+        Real distance, sqrDistance;
+        Vector<N, Real> boxClosest;
+    };
+
+    Result operator()(Vector<N, Real> const& point,
+        AlignedBox<N, Real> const& box);
+
+protected:
+    // On input, 'point' is the difference of the query point and the box
+    // center.  On output, 'point' is the point on the box closest to the
+    // query point.
+    void DoQuery(Vector<N, Real>& point, Vector<N, Real> const& boxExtent,
+        Result& result);
+};
+
+// Template aliases for convenience.
+template <int N, typename Real>
+using DCPPointAlignedBox =
+DCPQuery<Real, Vector<N, Real>, AlignedBox<N, Real>>;
+
+template <typename Real>
+using DCPPoint2AlignedBox2 = DCPPointAlignedBox<2, Real>;
+
+template <typename Real>
+using DCPPoint3AlignedBox3 = DCPPointAlignedBox<3, Real>;
+
+
+template <int N, typename Real>
+typename DCPQuery<Real, Vector<N, Real>, AlignedBox<N, Real>>::Result
+DCPQuery<Real, Vector<N, Real>, AlignedBox<N, Real>>::operator()(
+    Vector<N, Real> const& point, AlignedBox<N, Real> const& box)
+{
+    // Translate the point and box so that the box has center at the origin.
+    Vector<N, Real> boxCenter, boxExtent;
+    box.GetCenteredForm(boxCenter, boxExtent);
+    Vector<N, Real> closest = point - boxCenter;
+
+    Result result;
+    DoQuery(closest, boxExtent, result);
+
+    // Compute the closest point on the box.
+    result.boxClosest = boxCenter + closest;
+    return result;
+}
+
+template <int N, typename Real>
+void DCPQuery<Real, Vector<N, Real>, AlignedBox<N, Real>>::DoQuery(
+    Vector<N, Real>& point, Vector<N, Real> const& boxExtent, Result& result)
+{
+    result.sqrDistance = (Real)0;
+    for (int i = 0; i < N; ++i)
+    {
+        if (point[i] < -boxExtent[i])
+        {
+            Real delta = point[i] + boxExtent[i];
+            result.sqrDistance += delta * delta;
+            point[i] = -boxExtent[i];
+        }
+        else if (point[i] > boxExtent[i])
+        {
+            Real delta = point[i] - boxExtent[i];
+            result.sqrDistance += delta * delta;
+            point[i] = boxExtent[i];
+        }
+    }
+    result.distance = std::sqrt(result.sqrDistance);
+}
+
+
+}
diff --git a/Engine/Plugins/Experimental/GeometryProcessing/Source/GeometryAlgorithms/Private/GTEngine/Mathematics/GteDistPointHyperellipsoid.h b/Engine/Plugins/Experimental/GeometryProcessing/Source/GeometryAlgorithms/Private/GTEngine/Mathematics/GteDistPointHyperellipsoid.h
new file mode 100644
--- /dev/null
+++ b/Engine/Plugins/Experimental/GeometryProcessing/Source/GeometryAlgorithms/Private/GTEngine/Mathematics/GteDistPointHyperellipsoid.h
@@ -0,0 +1,382 @@
+// David Eberly, Geometric Tools, Redmond WA 98052
+// Copyright (c) 1998-2019
+// Distributed under the Boost Software License, Version 1.0.
+// http://www.boost.org/LICENSE_1_0.txt
+// http://www.geometrictools.com/License/Boost/LICENSE_1_0.txt
+// File Version: 3.0.2 (2018/10/05)
+
+#pragma once
+
+#include <GTEngine/Mathematics/GteMath.h>
+#include <GTEngine/Mathematics/GteVector.h>
+#include <GTEngine/Mathematics/GteDCPQuery.h>
+#include <GTEngine/Mathematics/GteHyperellipsoid.h>
+
+// Compute the distance from a point to a hyperellipsoid.  In 2D, this is a
+// point-ellipse distance query.  In 3D, this is a point-ellipsoid distance
+// query.  The following document describes the algorithm.
+//   http://www.geometrictools.com/Documentation/DistancePointEllipseEllipsoid.pdf
+// The hyperellipsoid can have arbitrary center and orientation; that is, it
+// does not have to be axis-aligned with center at the origin.
+//
+// For the 2D query,
+//   Vector2<Real> point;  // initialized to something
+//   Ellipse2<Real> ellipse;  // initialized to something
+//   DCPPoint2Ellipse2<Real> query;
+//   auto result = query(point, ellipse);
+//   Real distance = result.distance;
+//   Vector2<Real> closestEllipsePoint = result.closest;
+//
+// For the 3D query,
+//   Vector3<Real> point;  // initialized to something
+//   Ellipsoid3<Real> ellipsoid;  // initialized to something
+//   DCPPoint3Ellipsoid3<Real> query;
+//   auto result = query(point, ellipsoid);
+//   Real distance = result.distance;
+//   Vector3<Real> closestEllipsoidPoint = result.closest;
+
+namespace gte
+{
+
+template <int N, typename Real>
+class DCPQuery<Real, Vector<N, Real>, Hyperellipsoid<N, Real>>
+{
+public:
+    struct Result
+    {
+        Real distance, sqrDistance;
+        Vector<N, Real> closest;
+    };
+
+    // The query for any hyperellipsoid.
+    Result operator()(Vector<N, Real> const& point,
+        Hyperellipsoid<N, Real> const& hyperellipsoid);
+
+    // The 'hyperellipsoid' is assumed to be axis-aligned and centered at the
+    // origin , so only the extent[] values are used.
+    Result operator()(Vector<N, Real> const& point,
+        Vector<N, Real> const& extent);
+
+private:
+    // The hyperellipsoid is sum_{d=0}^{N-1} (x[d]/e[d])^2 = 1 with no
+    // constraints on the orderind of the e[d].  The query point is
+    // (y[0],...,y[N-1]) with no constraints on the signs of the components.
+    // The function returns the squared distance from the query point to the
+    // hyperellipsoid.   It also computes the hyperellipsoid point
+    // (x[0],...,x[N-1]) that is closest to (y[0],...,y[N-1]).
+    Real SqrDistance(Vector<N, Real> const& e,
+        Vector<N, Real> const& y, Vector<N, Real>& x);
+
+    // The hyperellipsoid is sum_{d=0}^{N-1} (x[d]/e[d])^2 = 1 with the e[d]
+    // positive and nonincreasing:  e[d] >= e[d + 1] for all d.  The query
+    // point is (y[0],...,y[N-1]) with y[d] >= 0 for all d.  The function
+    // returns the squared distance from the query point to the
+    // hyperellipsoid.  It also computes the hyperellipsoid point
+    // (x[0],...,x[N-1]) that is closest to (y[0],...,y[N-1]), where
+    // x[d] >= 0 for all d.
+    Real SqrDistanceSpecial(Vector<N, Real> const& e,
+        Vector<N, Real> const& y, Vector<N, Real>& x);
+
+    // The bisection algorithm to find the unique root of F(t).
+    Real Bisector(int numComponents, Vector<N, Real> const& e,
+        Vector<N, Real> const& y, Vector<N, Real>& x);
+};
+
+// Template aliases for convenience.
+template <int N, typename Real>
+using DCPPointHyperellipsoid =
+DCPQuery<Real, Vector<N, Real>, Hyperellipsoid<N, Real>>;
+
+template <typename Real>
+using DCPPoint2Ellipse2 = DCPPointHyperellipsoid<2, Real>;
+
+template <typename Real>
+using DCPPoint3Ellipsoid3 = DCPPointHyperellipsoid<3, Real>;
+
+
+template <int N, typename Real>
+typename DCPQuery<Real, Vector<N, Real>, Hyperellipsoid<N, Real>>::Result
+DCPQuery<Real, Vector<N, Real>, Hyperellipsoid<N, Real>>::operator()(
+    Vector<N, Real> const& point,
+    Hyperellipsoid<N, Real> const& hyperellipsoid)
+{
+    Result result;
+
+    // Compute the coordinates of Y in the hyperellipsoid coordinate system.
+    Vector<N, Real> diff = point - hyperellipsoid.center;
+    Vector<N, Real> y;
+    for (int i = 0; i < N; ++i)
+    {
+        y[i] = Dot(diff, hyperellipsoid.axis[i]);
+    }
+
+    // Compute the closest hyperellipsoid point in the axis-aligned
+    // coordinate system.
+    Vector<N, Real> x;
+    result.sqrDistance = SqrDistance(hyperellipsoid.extent, y, x);
+    result.distance = std::sqrt(result.sqrDistance);
+
+    // Convert back to the original coordinate system.
+    result.closest = hyperellipsoid.center;
+    for (int i = 0; i < N; ++i)
+    {
+        result.closest += x[i] * hyperellipsoid.axis[i];
+    }
+
+    return result;
+}
+
+template <int N, typename Real>
+typename DCPQuery<Real, Vector<N, Real>, Hyperellipsoid<N, Real>>::Result
+DCPQuery<Real, Vector<N, Real>, Hyperellipsoid<N, Real>>::operator()(
+    Vector<N, Real> const& point, Vector<N, Real> const& extent)
+{
+    Result result;
+    result.sqrDistance = SqrDistance(extent, point, result.closest);
+    result.distance = std::sqrt(result.sqrDistance);
+    return result;
+}
+
+template <int N, typename Real>
+Real DCPQuery<Real, Vector<N, Real>, Hyperellipsoid<N, Real>>::SqrDistance(
+    Vector<N, Real> const& e, Vector<N, Real> const& y, Vector<N, Real>& x)
+{
+    // Determine negations for y to the first octant.
+    std::array<bool, N> negate;
+    int i, j;
+    for (i = 0; i < N; ++i)
+    {
+        negate[i] = (y[i] < (Real)0);
+    }
+
+    // Determine the axis order for decreasing extents.
+    std::array<std::pair<Real, int>, N> permute;
+    for (i = 0; i < N; ++i)
+    {
+        permute[i].first = -e[i];
+        permute[i].second = i;
+    }
+    std::sort(permute.begin(), permute.end());
+
+    std::array<int, N> invPermute;
+    for (i = 0; i < N; ++i)
+    {
+        invPermute[permute[i].second] = i;
+    }
+
+    Vector<N, Real> locE, locY;
+    for (i = 0; i < N; ++i)
+    {
+        j = permute[i].second;
+        locE[i] = e[j];
+        locY[i] = std::abs(y[j]);
+    }
+
+    Vector<N, Real> locX;
+    Real sqrDistance = SqrDistanceSpecial(locE, locY, locX);
+
+    // Restore the axis order and reflections.
+    for (i = 0; i < N; ++i)
+    {
+        j = invPermute[i];
+        if (negate[i])
+        {
+            locX[j] = -locX[j];
+        }
+        x[i] = locX[j];
+    }
+
+    return sqrDistance;
+}
+
+template <int N, typename Real>
+Real DCPQuery<Real, Vector<N, Real>, Hyperellipsoid<N, Real>>::
+SqrDistanceSpecial(Vector<N, Real> const& e, Vector<N, Real> const& y,
+    Vector<N, Real>& x)
+{
+    Real sqrDistance = (Real)0;
+
+    Vector<N, Real> ePos, yPos, xPos;
+    int numPos = 0;
+    int i;
+    for (i = 0; i < N; ++i)
+    {
+        if (y[i] >(Real)0)
+        {
+            ePos[numPos] = e[i];
+            yPos[numPos] = y[i];
+            ++numPos;
+        }
+        else
+        {
+            x[i] = (Real)0;
+        }
+    }
+
+    if (y[N - 1] > (Real)0)
+    {
+        sqrDistance = Bisector(numPos, ePos, yPos, xPos);
+    }
+    else  // y[N-1] = 0
+    {
+        Vector<N - 1, Real> numer, denom;
+        Real eNm1Sqr = e[N - 1] * e[N - 1];
+        for (i = 0; i < numPos; ++i)
+        {
+            numer[i] = ePos[i] * yPos[i];
+            denom[i] = ePos[i] * ePos[i] - eNm1Sqr;
+        }
+
+        bool inSubHyperbox = true;
+        for (i = 0; i < numPos; ++i)
+        {
+            if (numer[i] >= denom[i])
+            {
+                inSubHyperbox = false;
+                break;
+            }
+        }
+
+        bool inSubHyperellipsoid = false;
+        if (inSubHyperbox)
+        {
+            // yPos[] is inside the axis-aligned bounding box of the
+            // subhyperellipsoid.  This intermediate test is designed to guard
+            // against the division by zero when ePos[i] == e[N-1] for some i.
+            Vector<N - 1, Real> xde;
+            Real discr = (Real)1;
+            for (i = 0; i < numPos; ++i)
+            {
+                xde[i] = numer[i] / denom[i];
+                discr -= xde[i] * xde[i];
+            }
+            if (discr >(Real)0)
+            {
+                // yPos[] is inside the subhyperellipsoid.  The closest
+                // hyperellipsoid point has x[N-1] > 0.
+                sqrDistance = (Real)0;
+                for (i = 0; i < numPos; ++i)
+                {
+                    xPos[i] = ePos[i] * xde[i];
+                    Real diff = xPos[i] - yPos[i];
+                    sqrDistance += diff*diff;
+                }
+                x[N - 1] = e[N - 1] * std::sqrt(discr);
+                sqrDistance += x[N - 1] * x[N - 1];
+                inSubHyperellipsoid = true;
+            }
+        }
+
+        if (!inSubHyperellipsoid)
+        {
+            // yPos[] is outside the subhyperellipsoid.  The closest
+            // hyperellipsoid point has x[N-1] == 0 and is on the
+            // domain-boundary hyperellipsoid.
+            x[N - 1] = (Real)0;
+            sqrDistance = Bisector(numPos, ePos, yPos, xPos);
+        }
+    }
+
+    // Fill in those x[] values that were not zeroed out initially.
+    for (i = 0, numPos = 0; i < N; ++i)
+    {
+        if (y[i] >(Real)0)
+        {
+            x[i] = xPos[numPos];
+            ++numPos;
+        }
+    }
+
+    return sqrDistance;
+}
+
+template <int N, typename Real>
+Real DCPQuery<Real, Vector<N, Real>, Hyperellipsoid<N, Real>>::Bisector(
+    int numComponents, Vector<N, Real> const& e, Vector<N, Real> const& y,
+    Vector<N, Real>& x)
+{
+    Vector<N, Real> z;
+    Real sumZSqr = (Real)0;
+    int i;
+    for (i = 0; i < numComponents; ++i)
+    {
+        z[i] = y[i] / e[i];
+        sumZSqr += z[i] * z[i];
+    }
+
+    if (sumZSqr == (Real)1)
+    {
+        // The point is on the hyperellipsoid.
+        for (i = 0; i < numComponents; ++i)
+        {
+            x[i] = y[i];
+        }
+        return (Real)0;
+    }
+
+    Real emin = e[numComponents - 1];
+    Vector<N, Real> pSqr, numerator;
+    for (i = 0; i < numComponents; ++i)
+    {
+        Real p = e[i] / emin;
+        pSqr[i] = p * p;
+        numerator[i] = pSqr[i] * z[i];
+    }
+
+    Real s = (Real)0, smin = z[numComponents - 1] - (Real)1, smax;
+    if (sumZSqr < (Real)1)
+    {
+        // The point is strictly inside the hyperellipsoid.
+        smax = (Real)0;
+    }
+    else
+    {
+        // The point is strictly outside the hyperellipsoid.
+        smax = Length(numerator, true) - (Real)1;
+    }
+
+    // The use of 'double' is intentional in case Real is a BSNumber or
+    // BSRational type.  We want the bisections to terminate in a reasonable
+    // amount of time.
+    unsigned int const jmax = GTE_C_MAX_BISECTIONS_GENERIC;
+    for (unsigned int j = 0; j < jmax; ++j)
+    {
+        s = (smin + smax) * (Real)0.5;
+        if (s == smin || s == smax)
+        {
+            break;
+        }
+
+        Real g = (Real)-1;
+        for (i = 0; i < numComponents; ++i)
+        {
+            Real ratio = numerator[i] / (s + pSqr[i]);
+            g += ratio * ratio;
+        }
+
+        if (g >(Real)0)
+        {
+            smin = s;
+        }
+        else if (g < (Real)0)
+        {
+            smax = s;
+        }
+        else
+        {
+            break;
+        }
+    }
+
+    Real sqrDistance = (Real)0;
+    for (i = 0; i < numComponents; ++i)
+    {
+        x[i] = pSqr[i] * y[i] / (s + pSqr[i]);
+        Real diff = x[i] - y[i];
+        sqrDistance += diff*diff;
+    }
+    return sqrDistance;
+}
+
+
+}
