move GeometryProcessing plugin from NotForLicensees to Experimental

diff --git a/Engine/Plugins/Experimental/GeometryProcessing/Source/GeometryAlgorithms/Private/GTEngine/Mathematics/GteDistPointLine.h b/Engine/Plugins/Experimental/GeometryProcessing/Source/GeometryAlgorithms/Private/GTEngine/Mathematics/GteDistPointLine.h
new file mode 100644
--- /dev/null
+++ b/Engine/Plugins/Experimental/GeometryProcessing/Source/GeometryAlgorithms/Private/GTEngine/Mathematics/GteDistPointLine.h
@@ -0,0 +1,62 @@
+// David Eberly, Geometric Tools, Redmond WA 98052
+// Copyright (c) 1998-2019
+// Distributed under the Boost Software License, Version 1.0.
+// http://www.boost.org/LICENSE_1_0.txt
+// http://www.geometrictools.com/License/Boost/LICENSE_1_0.txt
+// File Version: 3.0.1 (2018/10/05)
+
+#pragma once
+
+#include <GTEngine/Mathematics/GteDCPQuery.h>
+#include <GTEngine/Mathematics/GteLine.h>
+
+namespace gte
+{
+
+template <int N, typename Real>
+class DCPQuery<Real, Vector<N, Real>, Line<N, Real>>
+{
+public:
+    struct Result
+    {
+        Real distance, sqrDistance;
+        Real lineParameter;  // t in (-infinity,+infinity)
+        Vector<N, Real> lineClosest;  // origin + t * direction
+    };
+
+    Result operator()(Vector<N, Real> const& point,
+        Line<N, Real> const& line);
+};
+
+// Template aliases for convenience.
+template <int N, typename Real>
+using DCPPointLine =
+DCPQuery<Real, Vector<N, Real>, Line<N, Real>>;
+
+template <typename Real>
+using DCPPoint2Line2 = DCPPointLine<2, Real>;
+
+template <typename Real>
+using DCPPoint3Line3 = DCPPointLine<3, Real>;
+
+
+template <int N, typename Real>
+typename DCPQuery<Real, Vector<N, Real>, Line<N, Real>>::Result
+DCPQuery<Real, Vector<N, Real>, Line<N, Real>>::operator()(
+    Vector<N, Real> const& point, Line<N, Real> const& line)
+{
+    Result result;
+
+    Vector<N, Real> diff = point - line.origin;
+    result.lineParameter = Dot(line.direction, diff);
+    result.lineClosest = line.origin + result.lineParameter*line.direction;
+
+    diff = point - result.lineClosest;
+    result.sqrDistance = Dot(diff, diff);
+    result.distance = std::sqrt(result.sqrDistance);
+
+    return result;
+}
+
+
+}
diff --git a/Engine/Plugins/Experimental/GeometryProcessing/Source/GeometryAlgorithms/Private/GTEngine/Mathematics/GteDistPointOrientedBox.h b/Engine/Plugins/Experimental/GeometryProcessing/Source/GeometryAlgorithms/Private/GTEngine/Mathematics/GteDistPointOrientedBox.h
new file mode 100644
--- /dev/null
+++ b/Engine/Plugins/Experimental/GeometryProcessing/Source/GeometryAlgorithms/Private/GTEngine/Mathematics/GteDistPointOrientedBox.h
@@ -0,0 +1,72 @@
+// David Eberly, Geometric Tools, Redmond WA 98052
+// Copyright (c) 1998-2019
+// Distributed under the Boost Software License, Version 1.0.
+// http://www.boost.org/LICENSE_1_0.txt
+// http://www.geometrictools.com/License/Boost/LICENSE_1_0.txt
+// File Version: 3.0.0 (2016/06/19)
+
+#pragma once
+
+#include <GTEngine/Mathematics/GteDistPointAlignedBox.h>
+#include <GTEngine/Mathematics/GteOrientedBox.h>
+
+namespace gte
+{
+
+template <int N, typename Real>
+class DCPQuery<Real, Vector<N, Real>, OrientedBox<N, Real>>
+    :
+    public DCPQuery<Real, Vector<N, Real>, AlignedBox<N, Real>>
+{
+public:
+    struct Result
+        :
+        public DCPQuery<Real, Vector<N, Real>, AlignedBox<N, Real>>::Result
+    {
+        // No additional information to compute.
+    };
+
+    Result operator()(Vector<N, Real> const& point,
+        OrientedBox<N, Real> const& box);
+};
+
+// Template aliases for convenience.
+template <int N, typename Real>
+using DCPPointOrientedBox =
+DCPQuery<Real, Vector<N, Real>, AlignedBox<N, Real>>;
+
+template <typename Real>
+using DCPPoint2OrientedBox2 = DCPPointOrientedBox<2, Real>;
+
+template <typename Real>
+using DCPPoint3OrientedBox3 = DCPPointOrientedBox<3, Real>;
+
+
+template <int N, typename Real>
+typename DCPQuery<Real, Vector<N, Real>, OrientedBox<N, Real>>::Result
+DCPQuery<Real, Vector<N, Real>, OrientedBox<N, Real>>::operator()(
+    Vector<N, Real> const& point, OrientedBox<N, Real> const& box)
+{
+    // Translate the point to the coordinate system of the box.  In this
+    // system, the box is axis-aligned with center at the origin.
+    Vector<N, Real> diff = point - box.center;
+    Vector<N, Real> closest;
+    for (int i = 0; i < N; ++i)
+    {
+        closest[i] = Dot(diff, box.axis[i]);
+    }
+
+    Result result;
+    this->DoQuery(closest, box.extent, result);
+
+    // Compute the closest point on the box.
+    result.boxClosest = box.center;
+    for (int i = 0; i < N; ++i)
+    {
+        result.boxClosest += closest[i] * box.axis[i];
+    }
+    return result;
+}
+
+
+}
diff --git a/Engine/Plugins/Experimental/GeometryProcessing/Source/GeometryAlgorithms/Private/GTEngine/Mathematics/GteDistPointRay.h b/Engine/Plugins/Experimental/GeometryProcessing/Source/GeometryAlgorithms/Private/GTEngine/Mathematics/GteDistPointRay.h
new file mode 100644
--- /dev/null
+++ b/Engine/Plugins/Experimental/GeometryProcessing/Source/GeometryAlgorithms/Private/GTEngine/Mathematics/GteDistPointRay.h
@@ -0,0 +1,68 @@
+// David Eberly, Geometric Tools, Redmond WA 98052
+// Copyright (c) 1998-2019
+// Distributed under the Boost Software License, Version 1.0.
+// http://www.boost.org/LICENSE_1_0.txt
+// http://www.geometrictools.com/License/Boost/LICENSE_1_0.txt
+// File Version: 3.0.1 (2018/10/05)
+
+#pragma once
+
+#include <GTEngine/Mathematics/GteDCPQuery.h>
+#include <GTEngine/Mathematics/GteRay.h>
+
+namespace gte
+{
+
+template <int N, typename Real>
+class DCPQuery<Real, Vector<N, Real>, Ray<N, Real>>
+{
+public:
+    struct Result
+    {
+        Real distance, sqrDistance;
+        Real rayParameter;  // t in [0,+infinity)
+        Vector<N, Real> rayClosest;  // origin + t * direction
+    };
+
+    Result operator()(Vector<N, Real> const& point, Ray<N, Real> const& ray);
+};
+
+// Template aliases for convenience.
+template <int N, typename Real>
+using DCPPointRay =
+DCPQuery<Real, Vector<N, Real>, Ray<N, Real>>;
+
+template <typename Real>
+using DCPPoint2Ray2 = DCPPointRay<2, Real>;
+
+template <typename Real>
+using DCPPoint3Ray3 = DCPPointRay<3, Real>;
+
+
+template <int N, typename Real>
+typename DCPQuery<Real, Vector<N, Real>, Ray<N, Real>>::Result
+DCPQuery<Real, Vector<N, Real>, Ray<N, Real>>::operator()(
+    Vector<N, Real> const& point, Ray<N, Real> const& ray)
+{
+    Result result;
+
+    Vector<N, Real> diff = point - ray.origin;
+    result.rayParameter = Dot(ray.direction, diff);
+    if (result.rayParameter > (Real)0)
+    {
+        result.rayClosest = ray.origin + result.rayParameter*ray.direction;
+    }
+    else
+    {
+        result.rayClosest = ray.origin;
+    }
+
+    diff = point - result.rayClosest;
+    result.sqrDistance = Dot(diff, diff);
+    result.distance = std::sqrt(result.sqrDistance);
+
+    return result;
+}
+
+
+}
diff --git a/Engine/Plugins/Experimental/GeometryProcessing/Source/GeometryAlgorithms/Private/GTEngine/Mathematics/GteDistPointSegment.h b/Engine/Plugins/Experimental/GeometryProcessing/Source/GeometryAlgorithms/Private/GTEngine/Mathematics/GteDistPointSegment.h
new file mode 100644
--- /dev/null
+++ b/Engine/Plugins/Experimental/GeometryProcessing/Source/GeometryAlgorithms/Private/GTEngine/Mathematics/GteDistPointSegment.h
@@ -0,0 +1,94 @@
+// David Eberly, Geometric Tools, Redmond WA 98052
+// Copyright (c) 1998-2019
+// Distributed under the Boost Software License, Version 1.0.
+// http://www.boost.org/LICENSE_1_0.txt
+// http://www.geometrictools.com/License/Boost/LICENSE_1_0.txt
+// File Version: 3.0.1 (2018/10/05)
+
+#pragma once
+
+#include <GTEngine/Mathematics/GteDCPQuery.h>
+#include <GTEngine/Mathematics/GteSegment.h>
+
+namespace gte
+{
+
+template <int N, typename Real>
+class DCPQuery<Real, Vector<N, Real>, Segment<N, Real>>
+{
+public:
+    struct Result
+    {
+        Real distance, sqrDistance;
+        Real segmentParameter;  // t in [0,1]
+        Vector<N, Real> segmentClosest;  // (1-t)*p[0] + t*p[1]
+    };
+
+    Result operator()(Vector<N, Real> const& point,
+        Segment<N, Real> const& segment);
+};
+
+// Template aliases for convenience.
+template <int N, typename Real>
+using DCPPointSegment =
+DCPQuery<Real, Vector<N, Real>, Segment<N, Real>>;
+
+template <typename Real>
+using DCPPoint2Segment2 = DCPPointSegment<2, Real>;
+
+template <typename Real>
+using DCPPoint3Segment3 = DCPPointSegment<3, Real>;
+
+
+template <int N, typename Real>
+typename DCPQuery<Real, Vector<N, Real>, Segment<N, Real>>::Result
+DCPQuery<Real, Vector<N, Real>, Segment<N, Real>>::operator()(
+    Vector<N, Real> const& point, Segment<N, Real> const& segment)
+{
+    Result result;
+
+    // The direction vector is not unit length.  The normalization is deferred
+    // until it is needed.
+    Vector<N, Real> direction = segment.p[1] - segment.p[0];
+    Vector<N, Real> diff = point - segment.p[1];
+    Real t = Dot(direction, diff);
+    if (t >= (Real)0)
+    {
+        result.segmentParameter = (Real)1;
+        result.segmentClosest = segment.p[1];
+    }
+    else
+    {
+        diff = point - segment.p[0];
+        t = Dot(direction, diff);
+        if (t <= (Real)0)
+        {
+            result.segmentParameter = (Real)0;
+            result.segmentClosest = segment.p[0];
+        }
+        else
+        {
+            Real sqrLength = Dot(direction, direction);
+            if (sqrLength > (Real)0)
+            {
+                t /= sqrLength;
+                result.segmentParameter = t;
+                result.segmentClosest = segment.p[0] + t * direction;
+            }
+            else
+            {
+                result.segmentParameter = (Real)0;
+                result.segmentClosest = segment.p[0];
+            }
+        }
+    }
+
+    diff = point - result.segmentClosest;
+    result.sqrDistance = Dot(diff, diff);
+    result.distance = std::sqrt(result.sqrDistance);
+
+    return result;
+}
+
+
+}
diff --git a/Engine/Plugins/Experimental/GeometryProcessing/Source/GeometryAlgorithms/Private/GTEngine/Mathematics/GteDistPointTriangle.h b/Engine/Plugins/Experimental/GeometryProcessing/Source/GeometryAlgorithms/Private/GTEngine/Mathematics/GteDistPointTriangle.h
new file mode 100644
--- /dev/null
+++ b/Engine/Plugins/Experimental/GeometryProcessing/Source/GeometryAlgorithms/Private/GTEngine/Mathematics/GteDistPointTriangle.h
@@ -0,0 +1,266 @@
+// David Eberly, Geometric Tools, Redmond WA 98052
+// Copyright (c) 1998-2019
+// Distributed under the Boost Software License, Version 1.0.
+// http://www.boost.org/LICENSE_1_0.txt
+// http://www.geometrictools.com/License/Boost/LICENSE_1_0.txt
+// File Version: 3.0.1 (2018/10/05)
+
+#pragma once
+
+#include <GTEngine/Mathematics/GteVector.h>
+#include <GTEngine/Mathematics/GteDCPQuery.h>
+#include <GTEngine/Mathematics/GteTriangle.h>
+
+namespace gte
+{
+
+template <int N, typename Real>
+class DCPQuery<Real, Vector<N, Real>, Triangle<N, Real>>
+{
+public:
+    struct Result
+    {
+        Real distance, sqrDistance;
+        Real parameter[3];  // barycentric coordinates for triangle.v[3]
+        Vector<N, Real> closest;
+    };
+
+    Result operator()(Vector<N, Real> const& point,
+        Triangle<N, Real> const& triangle);
+
+private:
+    inline void GetMinEdge02(Real const& a11, Real const& b1,
+        Vector<2, Real>& p) const;
+
+    inline void GetMinEdge12(Real const& a01, Real const& a11, Real const& b1,
+        Real const& f10, Real const& f01, Vector<2, Real>& p) const;
+
+    inline void GetMinInterior(Vector<2, Real> const& p0, Real const& h0,
+        Vector<2, Real> const& p1, Real const& h1, Vector<2, Real>& p) const;
+};
+
+// Template aliases for convenience.
+template <int N, typename Real>
+using DCPPointTriangle =
+DCPQuery<Real, Vector<N, Real>, Triangle<N, Real>>;
+
+template <typename Real>
+using DCPPoint2Triangle2 = DCPPointTriangle<2, Real>;
+
+template <typename Real>
+using DCPPoint3Triangle3 = DCPPointTriangle<3, Real>;
+
+
+template <int N, typename Real> inline
+void DCPQuery<Real, Vector<N, Real>, Triangle<N, Real>>::GetMinEdge02(
+    Real const& a11, Real const& b1, Vector<2, Real>& p) const
+{
+    p[0] = (Real)0;
+    if (b1 >= (Real)0)
+    {
+        p[1] = (Real)0;
+    }
+    else if (a11 + b1 <= (Real)0)
+    {
+        p[1] = (Real)1;
+    }
+    else
+    {
+        p[1] = -b1 / a11;
+    }
+}
+
+template <int N, typename Real> inline
+void DCPQuery<Real, Vector<N, Real>, Triangle<N, Real>>::GetMinEdge12(
+    Real const& a01, Real const& a11, Real const& b1, Real const& f10,
+    Real const& f01, Vector<2, Real>& p) const
+{
+    Real h0 = a01 + b1 - f10;
+    if (h0 >= (Real)0)
+    {
+        p[1] = (Real)0;
+    }
+    else
+    {
+        Real h1 = a11 + b1 - f01;
+        if (h1 <= (Real)0)
+        {
+            p[1] = (Real)1;
+        }
+        else
+        {
+            p[1] = h0 / (h0 - h1);
+        }
+    }
+    p[0] = (Real)1 - p[1];
+}
+
+template <int N, typename Real> inline
+void DCPQuery<Real, Vector<N, Real>, Triangle<N, Real>>::GetMinInterior(
+    Vector<2, Real> const& p0, Real const& h0, Vector<2, Real> const& p1,
+    Real const& h1, Vector<2, Real>& p) const
+{
+    Real z = h0 / (h0 - h1);
+    p = ((Real)1 - z) * p0 + z * p1;
+}
+
+template <int N,typename Real>
+typename DCPQuery<Real, Vector<N, Real>, Triangle<N, Real>>::Result
+DCPQuery<Real, Vector<N, Real>, Triangle<N, Real>>::operator()(
+    Vector<N, Real> const& point, Triangle<N, Real> const& triangle)
+{
+    Vector<N, Real> diff = point - triangle.v[0];
+    Vector<N, Real> edge0 = triangle.v[1] - triangle.v[0];
+    Vector<N, Real> edge1 = triangle.v[2] - triangle.v[0];
+    Real a00 = Dot(edge0, edge0);
+    Real a01 = Dot(edge0, edge1);
+    Real a11 = Dot(edge1, edge1);
+    Real b0 = -Dot(diff, edge0);
+    Real b1 = -Dot(diff, edge1);
+
+    Real f00 = b0;
+    Real f10 = b0 + a00;
+    Real f01 = b0 + a01;
+
+    Vector<2, Real> p0, p1, p;
+    Real dt1, h0, h1;
+
+    // Compute the endpoints p0 and p1 of the segment.  The segment is
+    // parameterized by L(z) = (1-z)*p0 + z*p1 for z in [0,1] and the
+    // directional derivative of half the quadratic on the segment is
+    // H(z) = Dot(p1-p0,gradient[Q](L(z))/2), where gradient[Q]/2 = (F,G).
+    // By design, F(L(z)) = 0 for cases (2), (4), (5), and (6).  Cases (1) and
+    // (3) can correspond to no-intersection or intersection of F = 0 with the
+    // triangle.
+    if (f00 >= (Real)0)
+    {
+        if (f01 >= (Real)0)
+        {
+            // (1) p0 = (0,0), p1 = (0,1), H(z) = G(L(z))
+            GetMinEdge02(a11, b1, p);
+        }
+        else
+        {
+            // (2) p0 = (0,t10), p1 = (t01,1-t01), H(z) = (t11 - t10)*G(L(z))
+            p0[0] = (Real)0;
+            p0[1] = f00 / (f00 - f01);
+            p1[0] = f01 / (f01 - f10);
+            p1[1] = (Real)1 - p1[0];
+            dt1 = p1[1] - p0[1];
+            h0 = dt1 * (a11 * p0[1] + b1);
+            if (h0 >= (Real)0)
+            {
+                GetMinEdge02(a11, b1, p);
+            }
+            else
+            {
+                h1 = dt1 * (a01 * p1[0] + a11 * p1[1] + b1);
+                if (h1 <= (Real)0)
+                {
+                    GetMinEdge12(a01, a11, b1, f10, f01, p);
+                }
+                else
+                {
+                    GetMinInterior(p0, h0, p1, h1, p);
+                }
+            }
+        }
+    }
+    else if (f01 <= (Real)0)
+    {
+        if (f10 <= (Real)0)
+        {
+            // (3) p0 = (1,0), p1 = (0,1), H(z) = G(L(z)) - F(L(z))
+            GetMinEdge12(a01, a11, b1, f10, f01, p);
+        }
+        else
+        {
+            // (4) p0 = (t00,0), p1 = (t01,1-t01), H(z) = t11*G(L(z))
+            p0[0] = f00 / (f00 - f10);
+            p0[1] = (Real)0;
+            p1[0] = f01 / (f01 - f10);
+            p1[1] = (Real)1 - p1[0];
+            h0 = p1[1] * (a01 * p0[0] + b1);
+            if (h0 >= (Real)0)
+            {
+                p = p0;  // GetMinEdge01
+            }
+            else
+            {
+                h1 = p1[1] * (a01 * p1[0] + a11 * p1[1] + b1);
+                if (h1 <= (Real)0)
+                {
+                    GetMinEdge12(a01, a11, b1, f10, f01, p);
+                }
+                else
+                {
+                    GetMinInterior(p0, h0, p1, h1, p);
+                }
+            }
+        }
+    }
+    else if (f10 <= (Real)0)
+    {
+        // (5) p0 = (0,t10), p1 = (t01,1-t01), H(z) = (t11 - t10)*G(L(z))
+        p0[0] = (Real)0;
+        p0[1] = f00 / (f00 - f01);
+        p1[0] = f01 / (f01 - f10);
+        p1[1] = (Real)1 - p1[0];
+        dt1 = p1[1] - p0[1];
+        h0 = dt1 * (a11 * p0[1] + b1);
+        if (h0 >= (Real)0)
+        {
+            GetMinEdge02(a11, b1, p);
+        }
+        else
+        {
+            h1 = dt1 * (a01 * p1[0] + a11 * p1[1] + b1);
+            if (h1 <= (Real)0)
+            {
+                GetMinEdge12(a01, a11, b1, f10, f01, p);
+            }
+            else
+            {
+                GetMinInterior(p0, h0, p1, h1, p);
+            }
+        }
+    }
+    else
+    {
+        // (6) p0 = (t00,0), p1 = (0,t11), H(z) = t11*G(L(z))
+        p0[0] = f00 / (f00 - f10);
+        p0[1] = (Real)0;
+        p1[0] = (Real)0;
+        p1[1] = f00 / (f00 - f01);
+        h0 = p1[1] * (a01 * p0[0] + b1);
+        if (h0 >= (Real)0)
+        {
+            p = p0;  // GetMinEdge01
+        }
+        else
+        {
+            h1 = p1[1] * (a11 * p1[1] + b1);
+            if (h1 <= (Real)0)
+            {
+                GetMinEdge02(a11, b1, p);
+            }
+            else
+            {
+                GetMinInterior(p0, h0, p1, h1, p);
+            }
+        }
+    }
+
+    Result result;
+    result.parameter[0] = (Real)1 - p[0] - p[1];
+    result.parameter[1] = p[0];
+    result.parameter[2] = p[1];
+    result.closest = triangle.v[0] + p[0] * edge0 + p[1] * edge1;
+    diff = point - result.closest;
+    result.sqrDistance = Dot(diff, diff);
+    result.distance = std::sqrt(result.sqrDistance);
+    return result;
+}
+
+
+}
diff --git a/Engine/Plugins/Experimental/GeometryProcessing/Source/GeometryAlgorithms/Private/GTEngine/Mathematics/GteDistPointTriangleExact.h b/Engine/Plugins/Experimental/GeometryProcessing/Source/GeometryAlgorithms/Private/GTEngine/Mathematics/GteDistPointTriangleExact.h
new file mode 100644
--- /dev/null
+++ b/Engine/Plugins/Experimental/GeometryProcessing/Source/GeometryAlgorithms/Private/GTEngine/Mathematics/GteDistPointTriangleExact.h
@@ -0,0 +1,238 @@
+// David Eberly, Geometric Tools, Redmond WA 98052
+// Copyright (c) 1998-2019
+// Distributed under the Boost Software License, Version 1.0.
+// http://www.boost.org/LICENSE_1_0.txt
+// http://www.geometrictools.com/License/Boost/LICENSE_1_0.txt
+// File Version: 3.0.0 (2016/06/19)
+
+#pragma once
+
+#include <GTEngine/Mathematics/GteDCPQuery.h>
+#include <GTEngine/Mathematics/GteTriangle.h>
+
+namespace gte
+{
+
+template <int N, typename Rational>
+class DistancePointTriangleExact
+{
+public:
+    struct Result
+    {
+        Rational sqrDistance;
+        Rational parameter[3];  // barycentric coordinates for triangle.v[3]
+        Vector<N, Rational> closest;
+    };
+
+    Result operator()(Vector<N, Rational> const& point,
+        Triangle<N, Rational> const& triangle);
+};
+
+
+template <int N, typename Rational>
+typename DistancePointTriangleExact<N, Rational>::Result
+DistancePointTriangleExact<N, Rational>::operator()(
+    Vector<N, Rational> const& point, Triangle<N, Rational> const& triangle)
+{
+    Vector<N, Rational> diff = point - triangle.v[0];
+    Vector<N, Rational> edge0 = triangle.v[1] - triangle.v[0];
+    Vector<N, Rational> edge1 = triangle.v[2] - triangle.v[0];
+    Rational a00 = Dot(edge0, edge0);
+    Rational a01 = Dot(edge0, edge1);
+    Rational a11 = Dot(edge1, edge1);
+    Rational b0 = -Dot(diff, edge0);
+    Rational b1 = -Dot(diff, edge1);
+    Rational const zero = (Rational)0;
+    Rational const one = (Rational)1;
+    Rational det = a00 * a11 - a01 * a01;
+    Rational t0 = a01 * b1 - a11 * b0;
+    Rational t1 = a01 * b0 - a00 * b1;
+
+    if (t0 + t1 <= det)
+    {
+        if (t0 < zero)
+        {
+            if (t1 < zero)  // region 4
+            {
+                if (b0 < zero)
+                {
+                    t1 = zero;
+                    if (-b0 >= a00)  // V1
+                    {
+                        t0 = one;
+                    }
+                    else  // E01
+                    {
+                        t0 = -b0 / a00;
+                    }
+                }
+                else
+                {
+                    t0 = zero;
+                    if (b1 >= zero)  // V0
+                    {
+                        t1 = zero;
+                    }
+                    else if (-b1 >= a11)  // V2
+                    {
+                        t1 = one;
+                    }
+                    else  // E20
+                    {
+                        t1 = -b1 / a11;
+                    }
+                }
+            }
+            else  // region 3
+            {
+                t0 = zero;
+                if (b1 >= zero)  // V0
+                {
+                    t1 = zero;
+                }
+                else if (-b1 >= a11)  // V2
+                {
+                    t1 = one;
+                }
+                else  // E20
+                {
+                    t1 = -b1 / a11;
+                }
+            }
+        }
+        else if (t1 < zero)  // region 5
+        {
+            t1 = zero;
+            if (b0 >= zero)  // V0
+            {
+                t0 = zero;
+            }
+            else if (-b0 >= a00)  // V1
+            {
+                t0 = one;
+            }
+            else  // E01
+            {
+                t0 = -b0 / a00;
+            }
+        }
+        else  // region 0, interior
+        {
+            Rational invDet = one / det;
+            t0 *= invDet;
+            t1 *= invDet;
+        }
+    }
+    else
+    {
+        Rational tmp0, tmp1, numer, denom;
+
+        if (t0 < zero)  // region 2
+        {
+            tmp0 = a01 + b0;
+            tmp1 = a11 + b1;
+            if (tmp1 > tmp0)
+            {
+                numer = tmp1 - tmp0;
+                denom = a00 - ((Rational)2)*a01 + a11;
+                if (numer >= denom)  // V1
+                {
+                    t0 = one;
+                    t1 = zero;
+                }
+                else  // E12
+                {
+                    t0 = numer / denom;
+                    t1 = one - t0;
+                }
+            }
+            else
+            {
+                t0 = zero;
+                if (tmp1 <= zero)  // V2
+                {
+                    t1 = one;
+                }
+                else if (b1 >= zero)  // V0
+                {
+                    t1 = zero;
+                }
+                else  // E20
+                {
+                    t1 = -b1 / a11;
+                }
+            }
+        }
+        else if (t1 < zero)  // region 6
+        {
+            tmp0 = a01 + b1;
+            tmp1 = a00 + b0;
+            if (tmp1 > tmp0)
+            {
+                numer = tmp1 - tmp0;
+                denom = a00 - ((Rational)2)*a01 + a11;
+                if (numer >= denom)  // V2
+                {
+                    t1 = one;
+                    t0 = zero;
+                }
+                else  // E12
+                {
+                    t1 = numer / denom;
+                    t0 = one - t1;
+                }
+            }
+            else
+            {
+                t1 = zero;
+                if (tmp1 <= zero)  // V1
+                {
+                    t0 = one;
+                }
+                else if (b0 >= zero)  // V0
+                {
+                    t0 = zero;
+                }
+                else  // E01
+                {
+                    t0 = -b0 / a00;
+                }
+            }
+        }
+        else  // region 1
+        {
+            numer = a11 + b1 - a01 - b0;
+            if (numer <= zero)  // V2
+            {
+                t0 = zero;
+                t1 = one;
+            }
+            else
+            {
+                denom = a00 - ((Rational)2)*a01 + a11;
+                if (numer >= denom)  // V1
+                {
+                    t0 = one;
+                    t1 = zero;
+                }
+                else  // 12
+                {
+                    t0 = numer / denom;
+                    t1 = one - t0;
+                }
+            }
+        }
+    }
+
+    Result result;
+    result.parameter[0] = one - t0 - t1;
+    result.parameter[1] = t0;
+    result.parameter[2] = t1;
+    result.closest = triangle.v[0] + t0 * edge0 + t1 * edge1;
+    diff = point - result.closest;
+    result.sqrDistance = Dot(diff, diff);
+    return result;
+}
+
+
+}
diff --git a/Engine/Plugins/Experimental/GeometryProcessing/Source/GeometryAlgorithms/Private/GTEngine/Mathematics/GteDistRay3AlignedBox3.h b/Engine/Plugins/Experimental/GeometryProcessing/Source/GeometryAlgorithms/Private/GTEngine/Mathematics/GteDistRay3AlignedBox3.h
new file mode 100644
--- /dev/null
+++ b/Engine/Plugins/Experimental/GeometryProcessing/Source/GeometryAlgorithms/Private/GTEngine/Mathematics/GteDistRay3AlignedBox3.h
@@ -0,0 +1,64 @@
+// David Eberly, Geometric Tools, Redmond WA 98052
+// Copyright (c) 1998-2019
+// Distributed under the Boost Software License, Version 1.0.
+// http://www.boost.org/LICENSE_1_0.txt
+// http://www.geometrictools.com/License/Boost/LICENSE_1_0.txt
+// File Version: 3.0.0 (2016/06/19)
+
+#pragma once
+
+#include <GTEngine/Mathematics/GteDistLine3AlignedBox3.h>
+#include <GTEngine/Mathematics/GteRay.h>
+
+namespace gte
+{
+
+template <typename Real>
+class DCPQuery<Real, Ray3<Real>, AlignedBox3<Real>>
+{
+public:
+    struct Result
+    {
+        Real distance, sqrDistance;
+        Real rayParameter;
+        Vector3<Real> closestPoint[2];
+    };
+
+    Result operator()(Ray3<Real> const& ray, AlignedBox3<Real> const& box);
+};
+
+
+template <typename Real>
+typename DCPQuery<Real, Ray3<Real>, AlignedBox3<Real>>::Result
+DCPQuery<Real, Ray3<Real>, AlignedBox3<Real >> ::operator()(
+    Ray3<Real> const& ray, AlignedBox3<Real> const& box)
+{
+    Result result;
+
+    Line3<Real> line(ray.origin, ray.direction);
+    DCPQuery<Real, Line3<Real>, AlignedBox3<Real>> lbQuery;
+    auto lbResult = lbQuery(line, box);
+
+    if (lbResult.lineParameter >= (Real)0)
+    {
+        result.sqrDistance = lbResult.sqrDistance;
+        result.distance = lbResult.distance;
+        result.rayParameter = lbResult.lineParameter;
+        result.closestPoint[0] = lbResult.closestPoint[0];
+        result.closestPoint[1] = lbResult.closestPoint[1];
+    }
+    else
+    {
+        DCPQuery<Real, Vector3<Real>, AlignedBox3<Real>> pbQuery;
+        auto pbResult = pbQuery(ray.origin, box);
+        result.sqrDistance = pbResult.sqrDistance;
+        result.distance = pbResult.distance;
+        result.rayParameter = (Real)0;
+        result.closestPoint[0] = ray.origin;
+        result.closestPoint[1] = pbResult.boxClosest;
+    }
+    return result;
+}
+
+
+}
diff --git a/Engine/Plugins/Experimental/GeometryProcessing/Source/GeometryAlgorithms/Private/GTEngine/Mathematics/GteDistRay3OrientedBox3.h b/Engine/Plugins/Experimental/GeometryProcessing/Source/GeometryAlgorithms/Private/GTEngine/Mathematics/GteDistRay3OrientedBox3.h
new file mode 100644
--- /dev/null
+++ b/Engine/Plugins/Experimental/GeometryProcessing/Source/GeometryAlgorithms/Private/GTEngine/Mathematics/GteDistRay3OrientedBox3.h
@@ -0,0 +1,65 @@
+// David Eberly, Geometric Tools, Redmond WA 98052
+// Copyright (c) 1998-2019
+// Distributed under the Boost Software License, Version 1.0.
+// http://www.boost.org/LICENSE_1_0.txt
+// http://www.geometrictools.com/License/Boost/LICENSE_1_0.txt
+// File Version: 3.0.0 (2016/06/19)
+
+#pragma once
+
+#include <GTEngine/Mathematics/GteDistLine3OrientedBox3.h>
+#include <GTEngine/Mathematics/GteDistPointOrientedBox.h>
+#include <GTEngine/Mathematics/GteRay.h>
+
+namespace gte
+{
+
+template <typename Real>
+class DCPQuery<Real, Ray3<Real>, OrientedBox3<Real>>
+{
+public:
+    struct Result
+    {
+        Real distance, sqrDistance;
+        Real rayParameter;
+        Vector3<Real> closestPoint[2];
+    };
+
+    Result operator()(Ray3<Real> const& ray, OrientedBox3<Real> const& box);
+};
+
+
+template <typename Real>
+typename DCPQuery<Real, Ray3<Real>, OrientedBox3<Real>>::Result
+DCPQuery<Real, Ray3<Real>, OrientedBox3<Real >> ::operator()(
+    Ray3<Real> const& ray, OrientedBox3<Real> const& box)
+{
+    Result result;
+
+    Line3<Real> line(ray.origin, ray.direction);
+    DCPQuery<Real, Line3<Real>, OrientedBox3<Real>> lbQuery;
+    auto lbResult = lbQuery(line, box);
+
+    if (lbResult.lineParameter >= (Real)0)
+    {
+        result.sqrDistance = lbResult.sqrDistance;
+        result.distance = lbResult.distance;
+        result.rayParameter = lbResult.lineParameter;
+        result.closestPoint[0] = lbResult.closestPoint[0];
+        result.closestPoint[1] = lbResult.closestPoint[1];
+    }
+    else
+    {
+        DCPQuery<Real, Vector3<Real>, OrientedBox3<Real>> pbQuery;
+        auto pbResult = pbQuery(ray.origin, box);
+        result.sqrDistance = pbResult.sqrDistance;
+        result.distance = pbResult.distance;
+        result.rayParameter = (Real)0;
+        result.closestPoint[0] = ray.origin;
+        result.closestPoint[1] = pbResult.boxClosest;
+    }
+    return result;
+}
+
+
+}
diff --git a/Engine/Plugins/Experimental/GeometryProcessing/Source/GeometryAlgorithms/Private/GTEngine/Mathematics/GteDistRay3Rectangle3.h b/Engine/Plugins/Experimental/GeometryProcessing/Source/GeometryAlgorithms/Private/GTEngine/Mathematics/GteDistRay3Rectangle3.h
new file mode 100644
--- /dev/null
+++ b/Engine/Plugins/Experimental/GeometryProcessing/Source/GeometryAlgorithms/Private/GTEngine/Mathematics/GteDistRay3Rectangle3.h
@@ -0,0 +1,70 @@
+// David Eberly, Geometric Tools, Redmond WA 98052
+// Copyright (c) 1998-2019
+// Distributed under the Boost Software License, Version 1.0.
+// http://www.boost.org/LICENSE_1_0.txt
+// http://www.geometrictools.com/License/Boost/LICENSE_1_0.txt
+// File Version: 3.0.0 (2016/06/19)
+
+#pragma once
+
+#include <GTEngine/Mathematics/GteDistLine3Rectangle3.h>
+#include <GTEngine/Mathematics/GteDistPoint3Rectangle3.h>
+#include <GTEngine/Mathematics/GteRay.h>
+
+namespace gte
+{
+
+template <typename Real>
+class DCPQuery<Real, Ray3<Real>, Rectangle3<Real>>
+{
+public:
+    struct Result
+    {
+        Real distance, sqrDistance;
+        Real rayParameter, rectangleParameter[2];
+        Vector3<Real> closestPoint[2];
+    };
+
+    Result operator()(Ray3<Real> const& ray,
+        Rectangle3<Real> const& rectangle);
+};
+
+
+template <typename Real>
+typename DCPQuery<Real, Ray3<Real>, Rectangle3<Real>>::Result
+DCPQuery<Real, Ray3<Real>, Rectangle3<Real>>::operator()(
+    Ray3<Real> const& ray, Rectangle3<Real> const& rectangle)
+{
+    Result result;
+
+    Line3<Real> line(ray.origin, ray.direction);
+    DCPQuery<Real, Line3<Real>, Rectangle3<Real>> lrQuery;
+    auto lrResult = lrQuery(line, rectangle);
+
+    if (lrResult.lineParameter >= (Real)0)
+    {
+        result.distance = lrResult.distance;
+        result.sqrDistance = lrResult.sqrDistance;
+        result.rayParameter = lrResult.lineParameter;
+        result.rectangleParameter[0] = lrResult.rectangleParameter[0];
+        result.rectangleParameter[1] = lrResult.rectangleParameter[1];
+        result.closestPoint[0] = lrResult.closestPoint[0];
+        result.closestPoint[1] = lrResult.closestPoint[1];
+    }
+    else
+    {
+        DCPQuery<Real, Vector3<Real>, Rectangle3<Real>> prQuery;
+        auto prResult = prQuery(ray.origin, rectangle);
+        result.distance = prResult.distance;
+        result.sqrDistance = prResult.sqrDistance;
+        result.rayParameter = (Real)0;
+        result.rectangleParameter[0] = prResult.rectangleParameter[0];
+        result.rectangleParameter[1] = prResult.rectangleParameter[1];
+        result.closestPoint[0] = ray.origin;
+        result.closestPoint[1] = prResult.rectangleClosestPoint;
+    }
+    return result;
+}
+
+
+}
diff --git a/Engine/Plugins/Experimental/GeometryProcessing/Source/GeometryAlgorithms/Private/GTEngine/Mathematics/GteDistRay3Triangle3.h b/Engine/Plugins/Experimental/GeometryProcessing/Source/GeometryAlgorithms/Private/GTEngine/Mathematics/GteDistRay3Triangle3.h
new file mode 100644
--- /dev/null
+++ b/Engine/Plugins/Experimental/GeometryProcessing/Source/GeometryAlgorithms/Private/GTEngine/Mathematics/GteDistRay3Triangle3.h
@@ -0,0 +1,71 @@
+// David Eberly, Geometric Tools, Redmond WA 98052
+// Copyright (c) 1998-2019
+// Distributed under the Boost Software License, Version 1.0.
+// http://www.boost.org/LICENSE_1_0.txt
+// http://www.geometrictools.com/License/Boost/LICENSE_1_0.txt
+// File Version: 3.0.0 (2016/06/19)
+
+#pragma once
+
+#include <GTEngine/Mathematics/GteDistLine3Triangle3.h>
+#include <GTEngine/Mathematics/GteDistPointTriangle.h>
+#include <GTEngine/Mathematics/GteRay.h>
+
+namespace gte
+{
+
+template <typename Real>
+class DCPQuery<Real, Ray3<Real>, Triangle3<Real>>
+{
+public:
+    struct Result
+    {
+        Real distance, sqrDistance;
+        Real rayParameter, triangleParameter[3];
+        Vector3<Real> closestPoint[2];
+    };
+
+    Result operator()(Ray3<Real> const& ray, Triangle3<Real> const& triangle);
+};
+
+
+template <typename Real>
+typename DCPQuery<Real, Ray3<Real>, Triangle3<Real>>::Result
+DCPQuery<Real, Ray3<Real>, Triangle3<Real>>::operator()(
+    Ray3<Real> const& ray, Triangle3<Real> const& triangle)
+{
+    Result result;
+
+    Line3<Real> line(ray.origin, ray.direction);
+    DCPQuery<Real, Line3<Real>, Triangle3<Real>> ltQuery;
+    auto ltResult = ltQuery(line, triangle);
+
+    if (ltResult.lineParameter >= (Real)0)
+    {
+        result.distance = ltResult.distance;
+        result.sqrDistance = ltResult.sqrDistance;
+        result.rayParameter = ltResult.lineParameter;
+        result.triangleParameter[0] = ltResult.triangleParameter[0];
+        result.triangleParameter[1] = ltResult.triangleParameter[1];
+        result.triangleParameter[2] = ltResult.triangleParameter[2];
+        result.closestPoint[0] = ltResult.closestPoint[0];
+        result.closestPoint[1] = ltResult.closestPoint[1];
+    }
+    else
+    {
+        DCPQuery<Real, Vector3<Real>, Triangle3<Real>> ptQuery;
+        auto ptResult = ptQuery(ray.origin, triangle);
+        result.distance = ptResult.distance;
+        result.sqrDistance = ptResult.sqrDistance;
+        result.rayParameter = (Real)0;
+        result.triangleParameter[0] = ptResult.parameter[0];
+        result.triangleParameter[1] = ptResult.parameter[1];
+        result.triangleParameter[2] = ptResult.parameter[2];
+        result.closestPoint[0] = ray.origin;
+        result.closestPoint[1] = ptResult.closest;
+    }
+    return result;
+}
+
+
+}
diff --git a/Engine/Plugins/Experimental/GeometryProcessing/Source/GeometryAlgorithms/Private/GTEngine/Mathematics/GteDistRayRay.h b/Engine/Plugins/Experimental/GeometryProcessing/Source/GeometryAlgorithms/Private/GTEngine/Mathematics/GteDistRayRay.h
new file mode 100644
--- /dev/null
+++ b/Engine/Plugins/Experimental/GeometryProcessing/Source/GeometryAlgorithms/Private/GTEngine/Mathematics/GteDistRayRay.h
@@ -0,0 +1,162 @@
+// David Eberly, Geometric Tools, Redmond WA 98052
+// Copyright (c) 1998-2019
+// Distributed under the Boost Software License, Version 1.0.
+// http://www.boost.org/LICENSE_1_0.txt
+// http://www.geometrictools.com/License/Boost/LICENSE_1_0.txt
+// File Version: 3.0.1 (2018/10/05)
+
+#pragma once
+
+#include <GTEngine/Mathematics/GteDCPQuery.h>
+#include <GTEngine/Mathematics/GteRay.h>
+
+namespace gte
+{
+
+template <int N, typename Real>
+class DCPQuery<Real, Ray<N, Real>, Ray<N, Real>>
+{
+public:
+    struct Result
+    {
+        Real distance, sqrDistance;
+        Real parameter[2];
+        Vector<N, Real> closestPoint[2];
+    };
+
+    Result operator()(Ray<N, Real> const& ray0, Ray<N, Real> const& ray1);
+};
+
+// Template aliases for convenience.
+template <int N, typename Real>
+using DCPRayRay = DCPQuery<Real, Ray<N, Real>, Ray<N, Real>>;
+
+template <typename Real>
+using DCPRay2Ray2 = DCPRayRay<2, Real>;
+
+template <typename Real>
+using DCPRay3Ray3 = DCPRayRay<3, Real>;
+
+
+template <int N, typename Real>
+typename DCPQuery<Real, Ray<N, Real>, Ray<N, Real>>::Result
+DCPQuery<Real, Ray<N, Real>, Ray<N, Real>>::operator()(
+    Ray<N, Real> const& ray0, Ray<N, Real> const& ray1)
+{
+    Result result;
+
+    Vector<N, Real> diff = ray0.origin - ray1.origin;
+    Real a01 = -Dot(ray0.direction, ray1.direction);
+    Real b0 = Dot(diff, ray0.direction), b1;
+    Real s0, s1;
+
+    if (std::abs(a01) < (Real)1)
+    {
+        // Rays are not parallel.
+        b1 = -Dot(diff, ray1.direction);
+        s0 = a01 * b1 - b0;
+        s1 = a01 * b0 - b1;
+
+        if (s0 >= (Real)0)
+        {
+            if (s1 >= (Real)0)  // region 0 (interior)
+            {
+                // Minimum at two interior points of rays.
+                Real det = (Real)1 - a01 * a01;
+                s0 /= det;
+                s1 /= det;
+            }
+            else  // region 3 (side)
+            {
+                s1 = (Real)0;
+                if (b0 >= (Real)0)
+                {
+                    s0 = (Real)0;
+                }
+                else
+                {
+                    s0 = -b0;
+                }
+            }
+        }
+        else
+        {
+            if (s1 >= (Real)0)  // region 1 (side)
+            {
+                s0 = (Real)0;
+                if (b1 >= (Real)0)
+                {
+                    s1 = (Real)0;
+                }
+                else
+                {
+                    s1 = -b1;
+                }
+            }
+            else  // region 2 (corner)
+            {
+                if (b0 < (Real)0)
+                {
+                    s0 = -b0;
+                    s1 = (Real)0;
+                }
+                else
+                {
+                    s0 = (Real)0;
+                    if (b1 >= (Real)0)
+                    {
+                        s1 = (Real)0;
+                    }
+                    else
+                    {
+                        s1 = -b1;
+                    }
+                }
+            }
+        }
+    }
+    else
+    {
+        // Rays are parallel.
+        if (a01 > (Real)0)
+        {
+            // Opposite direction vectors.
+            s1 = (Real)0;
+            if (b0 >= (Real)0)
+            {
+                s0 = (Real)0;
+            }
+            else
+            {
+                s0 = -b0;
+            }
+        }
+        else
+        {
+            // Same direction vectors.
+            if (b0 >= (Real)0)
+            {
+                b1 = -Dot(diff, ray1.direction);
+                s0 = (Real)0;
+                s1 = -b1;
+            }
+            else
+            {
+                s0 = -b0;
+                s1 = (Real)0;
+            }
+        }
+    }
+
+    result.parameter[0] = s0;
+    result.parameter[1] = s1;
+    result.closestPoint[0] = ray0.origin + s0 * ray0.direction;
+    result.closestPoint[1] = ray1.origin + s1 * ray1.direction;
+    diff = result.closestPoint[0] - result.closestPoint[1];
+    result.sqrDistance = Dot(diff, diff);
+    result.distance = std::sqrt(result.sqrDistance);
+    return result;
+}
+
+
+}
diff --git a/Engine/Plugins/Experimental/GeometryProcessing/Source/GeometryAlgorithms/Private/GTEngine/Mathematics/GteDistRaySegment.h b/Engine/Plugins/Experimental/GeometryProcessing/Source/GeometryAlgorithms/Private/GTEngine/Mathematics/GteDistRaySegment.h
new file mode 100644
--- /dev/null
+++ b/Engine/Plugins/Experimental/GeometryProcessing/Source/GeometryAlgorithms/Private/GTEngine/Mathematics/GteDistRaySegment.h
@@ -0,0 +1,179 @@
+// David Eberly, Geometric Tools, Redmond WA 98052
+// Copyright (c) 1998-2019
+// Distributed under the Boost Software License, Version 1.0.
+// http://www.boost.org/LICENSE_1_0.txt
+// http://www.geometrictools.com/License/Boost/LICENSE_1_0.txt
+// File Version: 3.0.1 (2018/10/05)
+
+#pragma once
+
+#include <GTEngine/Mathematics/GteDCPQuery.h>
+#include <GTEngine/Mathematics/GteRay.h>
+#include <GTEngine/Mathematics/GteSegment.h>
+
+namespace gte
+{
+
+template <int N, typename Real>
+class DCPQuery<Real, Ray<N, Real>, Segment<N, Real>>
+{
+public:
+    struct Result
+    {
+        Real distance, sqrDistance;
+        Real parameter[2];
+        Vector<N, Real> closestPoint[2];
+    };
+
+    // The centered form of the 'segment' is used.  Thus, parameter[1] of
+    // the result is in [-e,e], where e = |segment.p[1] - segment.p[0]|/2.
+    Result operator()(Ray<N, Real> const& ray,
+        Segment<N, Real> const& segment);
+};
+
+// Template aliases for convenience.
+template <int N, typename Real>
+using DCPRaySegment = DCPQuery<Real, Ray<N, Real>, Segment<N, Real>>;
+
+template <typename Real>
+using DCPRay2Segment2 = DCPRaySegment<2, Real>;
+
+template <typename Real>
+using DCPRay3Segment3 = DCPRaySegment<3, Real>;
+
+
+template <int N, typename Real>
+typename DCPQuery<Real, Ray<N, Real>, Segment<N, Real>>::Result
+DCPQuery<Real, Ray<N, Real>, Segment<N, Real>>::operator()(
+    Ray<N, Real> const& ray, Segment<N, Real> const& segment)
+{
+    Result result;
+
+    Vector<N, Real> segCenter, segDirection;
+    Real segExtent;
+    segment.GetCenteredForm(segCenter, segDirection, segExtent);
+
+    Vector<N, Real> diff = ray.origin - segCenter;
+    Real a01 = -Dot(ray.direction, segDirection);
+    Real b0 = Dot(diff, ray.direction);
+    Real s0, s1;
+
+    if (std::abs(a01) < (Real)1)
+    {
+        // The ray and segment are not parallel.
+        Real det = (Real)1 - a01 * a01;
+        Real extDet = segExtent*det;
+        Real b1 = -Dot(diff, segDirection);
+        s0 = a01 * b1 - b0;
+        s1 = a01 * b0 - b1;
+
+        if (s0 >= (Real)0)
+        {
+            if (s1 >= -extDet)
+            {
+                if (s1 <= extDet)  // region 0
+                {
+                    // Minimum at interior points of ray and segment.
+                    s0 /= det;
+                    s1 /= det;
+                }
+                else  // region 1
+                {
+                    s1 = segExtent;
+                    s0 = std::max(-(a01*s1 + b0), (Real)0);
+                }
+            }
+            else  // region 5
+            {
+                s1 = -segExtent;
+                s0 = std::max(-(a01*s1 + b0), (Real)0);
+            }
+        }
+        else
+        {
+            if (s1 <= -extDet)  // region 4
+            {
+                s0 = -(-a01*segExtent + b0);
+                if (s0 > (Real)0)
+                {
+                    s1 = -segExtent;
+                }
+                else
+                {
+                    s0 = (Real)0;
+                    s1 = -b1;
+                    if (s1 < -segExtent)
+                    {
+                        s1 = -segExtent;
+                    }
+                    else if (s1 > segExtent)
+                    {
+                        s1 = segExtent;
+                    }
+                }
+            }
+            else if (s1 <= extDet)  // region 3
+            {
+                s0 = (Real)0;
+                s1 = -b1;
+                if (s1 < -segExtent)
+                {
+                    s1 = -segExtent;
+                }
+                else if (s1 > segExtent)
+                {
+                    s1 = segExtent;
+                }
+            }
+            else  // region 2
+            {
+                s0 = -(a01*segExtent + b0);
+                if (s0 > (Real)0)
+                {
+                    s1 = segExtent;
+                }
+                else
+                {
+                    s0 = (Real)0;
+                    s1 = -b1;
+                    if (s1 < -segExtent)
+                    {
+                        s1 = -segExtent;
+                    }
+                    else if (s1 > segExtent)
+                    {
+                        s1 = segExtent;
+                    }
+                }
+            }
+        }
+    }
+    else
+    {
+        // Ray and segment are parallel.
+        if (a01 > (Real)0)
+        {
+            // Opposite direction vectors.
+            s1 = -segExtent;
+        }
+        else
+        {
+            // Same direction vectors.
+            s1 = segExtent;
+        }
+
+        s0 = std::max(-(a01*s1 + b0), (Real)0);
+    }
+
+    result.parameter[0] = s0;
+    result.parameter[1] = s1;
+    result.closestPoint[0] = ray.origin + s0 * ray.direction;
+    result.closestPoint[1] = segCenter + s1 * segDirection;
+    diff = result.closestPoint[0] - result.closestPoint[1];
+    result.sqrDistance = Dot(diff, diff);
+    result.distance = std::sqrt(result.sqrDistance);
+    return result;
+}
+
+
+}
diff --git a/Engine/Plugins/Experimental/GeometryProcessing/Source/GeometryAlgorithms/Private/GTEngine/Mathematics/GteDistRectangle3AlignedBox3.h b/Engine/Plugins/Experimental/GeometryProcessing/Source/GeometryAlgorithms/Private/GTEngine/Mathematics/GteDistRectangle3AlignedBox3.h
new file mode 100644
--- /dev/null
+++ b/Engine/Plugins/Experimental/GeometryProcessing/Source/GeometryAlgorithms/Private/GTEngine/Mathematics/GteDistRectangle3AlignedBox3.h
@@ -0,0 +1,145 @@
+// David Eberly, Geometric Tools, Redmond WA 98052
+// Copyright (c) 1998-2019
+// Distributed under the Boost Software License, Version 1.0.
+// http://www.boost.org/LICENSE_1_0.txt
+// http://www.geometrictools.com/License/Boost/LICENSE_1_0.txt
+// File Version: 3.5.2 (2018/10/05)
+
+#pragma once
+
+#include <GTEngine/Mathematics/GteDCPQuery.h>
+#include <GTEngine/Mathematics/GteLCPSolver.h>
+#include <GTEngine/Mathematics/GteAlignedBox.h>
+#include <GTEngine/Mathematics/GteRectangle.h>
+#include <GTEngine/Mathematics/GteVector3.h>
+
+namespace gte
+{
+
+template <typename Real>
+class DCPQuery<Real, Rectangle3<Real>, AlignedBox3<Real>>
+{
+public:
+    struct Result
+    {
+        bool queryIsSuccessful;
+
+        // These members are valid only when queryIsSuccessful is true;
+        // otherwise, they are all set to zero.
+        Real distance, sqrDistance;
+        std::array<Real, 2> rectangleParameter;
+        std::array<Real, 3> boxParameter;
+        Vector3<Real> closestPoint[2];
+
+        // The number of iterations used by LCPSolver regardless of whether
+        // the query is successful.
+        int numLCPIterations;
+    };
+
+    // The default maximum iterations is 81 (n = 9, maxIterations = n*n).  If
+    // the solver fails to converge, try increasing the maximum number of
+    // iterations.
+    void SetMaxLCPIterations(int maxLCPIterations);
+
+    Result operator()(Rectangle3<Real> const& rectangle, AlignedBox3<Real> const& box);
+
+private:
+    LCPSolver<Real, 10> mLCP;
+};
+
+
+template <typename Real>
+void DCPQuery<Real, Rectangle3<Real>, AlignedBox3<Real>>::SetMaxLCPIterations(int maxLCPIterations)
+{
+    mLCP.SetMaxIterations(maxLCPIterations);
+}
+
+template <typename Real>
+typename DCPQuery<Real, Rectangle3<Real>, AlignedBox3<Real>>::Result
+DCPQuery<Real, Rectangle3<Real>, AlignedBox3<Real>>::operator()(
+    Rectangle3<Real> const& rectangle, AlignedBox3<Real> const& box)
+{
+    Result result;
+
+    // Translate the rectangle and aligned box so that the aligned box becomes
+    // a canonical box.
+    Vector3<Real> K = box.max - box.min;
+    Vector3<Real> V = rectangle.center - box.min;
+
+    // Convert the oriented rectangle to a regular one (origin at a corner).
+    Vector3<Real> scaledE0 = rectangle.axis[0] * rectangle.extent[0];
+    Vector3<Real> scaledE1 = rectangle.axis[1] * rectangle.extent[1];
+    Vector3<Real> E0 = scaledE0 * (Real)2;
+    Vector3<Real> E1 = scaledE1 * (Real)2;
+    V -= scaledE0 + scaledE1;
+
+    // Compute quantities to initialize q and M in the LCP.
+    Real dotVE0 = Dot(V, E0);
+    Real dotVE1 = Dot(V, E1);
+    Real dotE0E0 = Dot(E0, E0);
+    Real dotE1E1 = Dot(E1, E1);
+
+    // The LCP has 5 variables and 5 (nontrivial) inequality constraints.
+    std::array<Real, 10> q =
+    {
+        -V[0], -V[1], -V[2], dotVE0, dotVE1, K[0], K[1], K[2], (Real)1, (Real)1
+    };
+
+    std::array<std::array<Real, 10>, 10> M;
+    M[0] = { (Real)1, (Real)0, (Real)0, -E0[0], -E1[0], (Real)1, (Real)0, (Real)0, (Real)0, (Real)0 };
+    M[1] = { (Real)0, (Real)1, (Real)0, -E0[1], -E1[1], (Real)0, (Real)1, (Real)0, (Real)0, (Real)0 };
+    M[2] = { (Real)0, (Real)0, (Real)1, -E0[2], -E1[2], (Real)0, (Real)0, (Real)1, (Real)0 , (Real)0 };
+    M[3] = { -E0[0], -E0[1], -E0[2], dotE0E0, (Real)0, (Real)0, (Real)0, (Real)0, (Real)1, (Real)0 };
+    M[4] = { -E1[0], -E1[1], -E1[2], (Real)0, dotE1E1, (Real)0, (Real)0, (Real)0, (Real)0, (Real)1 };
+    M[5] = { (Real)-1, (Real)0, (Real)0, (Real)0, (Real)0, (Real)0, (Real)0, (Real)0, (Real)0, (Real)0 };
+    M[6] = { (Real)0, (Real)-1, (Real)0, (Real)0, (Real)0, (Real)0, (Real)0, (Real)0, (Real)0, (Real)0 };
+    M[7] = { (Real)0, (Real)0, (Real)-1, (Real)0, (Real)0, (Real)0, (Real)0, (Real)0, (Real)0, (Real)0 };
+    M[8] = { (Real)0, (Real)0, (Real)0, (Real)-1, (Real)0, (Real)0, (Real)0, (Real)0, (Real)0, (Real)0 };
+    M[9] = { (Real)0, (Real)0, (Real)0, (Real)0, (Real)-1, (Real)0, (Real)0, (Real)0, (Real)0, (Real)0 };
+
+    std::array<Real, 10> w, z;
+    if (mLCP.Solve(q, M, w, z))
+    {
+        result.queryIsSuccessful = true;
+        Real t0 = (z[3] * (Real)2 - (Real)1) * rectangle.extent[0];
+        Real t1 = (z[4] * (Real)2 - (Real)1) * rectangle.extent[1];
+        result.rectangleParameter[0] = t0;
+        result.rectangleParameter[1] = t1;
+        result.closestPoint[0] = rectangle.center + t0 * rectangle.axis[0] + t1 * rectangle.axis[1];
+        for (int i = 0; i < 3; ++i)
+        {
+            result.boxParameter[i] = z[i] + box.min[i];
+            result.closestPoint[1][i] = result.boxParameter[i];
+        }
+
+        Vector3<Real> diff = result.closestPoint[1] - result.closestPoint[0];
+        result.sqrDistance = Dot(diff, diff);
+        result.distance = std::sqrt(result.sqrDistance);
+    }
+    else
+    {
+        // If you reach this case, the maximum number of iterations was not
+        // specified to be large enough or there is a problem due to
+        // floating-point rounding errors.  If you believe the latter is
+        // true, file a bug report.
+        result.queryIsSuccessful = false;
+
+        for (int i = 0; i < 2; ++i)
+        {
+            result.rectangleParameter[i] = (Real)0;
+        }
+        for (int i = 0; i < 3; ++i)
+        {
+            result.boxParameter[i] = (Real)0;
+            result.closestPoint[0][i] = (Real)0;
+            result.closestPoint[1][i] = (Real)0;
+        }
+        result.distance = (Real)0;
+        result.sqrDistance = (Real)0;
+    }
+
+    result.numLCPIterations = mLCP.GetNumIterations();
+    return result;
+}
+
+}
diff --git a/Engine/Plugins/Experimental/GeometryProcessing/Source/GeometryAlgorithms/Private/GTEngine/Mathematics/GteDistRectangle3OrientedBox3.h b/Engine/Plugins/Experimental/GeometryProcessing/Source/GeometryAlgorithms/Private/GTEngine/Mathematics/GteDistRectangle3OrientedBox3.h
new file mode 100644
--- /dev/null
+++ b/Engine/Plugins/Experimental/GeometryProcessing/Source/GeometryAlgorithms/Private/GTEngine/Mathematics/GteDistRectangle3OrientedBox3.h
@@ -0,0 +1,154 @@
+// David Eberly, Geometric Tools, Redmond WA 98052
+// Copyright (c) 1998-2019
+// Distributed under the Boost Software License, Version 1.0.
+// http://www.boost.org/LICENSE_1_0.txt
+// http://www.geometrictools.com/License/Boost/LICENSE_1_0.txt
+// File Version: 3.5.2 (2018/10/05)
+
+#pragma once
+
+#include <GTEngine/Mathematics/GteDCPQuery.h>
+#include <GTEngine/Mathematics/GteLCPSolver.h>
+#include <GTEngine/Mathematics/GteOrientedBox.h>
+#include <GTEngine/Mathematics/GteRectangle.h>
+#include <GTEngine/Mathematics/GteVector3.h>
+
+namespace gte
+{
+
+template <typename Real>
+class DCPQuery<Real, Rectangle3<Real>, OrientedBox3<Real>>
+{
+public:
+    struct Result
+    {
+        bool queryIsSuccessful;
+
+        // These members are valid only when queryIsSuccessful is true;
+        // otherwise, they are all set to zero.
+        Real distance, sqrDistance;
+        std::array<Real, 2> rectangleParameter;
+        std::array<Real, 3> boxParameter;
+        Vector3<Real> closestPoint[2];
+
+        // The number of iterations used by LCPSolver regardless of whether
+        // the query is successful.
+        int numLCPIterations;
+    };
+
+    // The default maximum iterations is 81 (n = 9, maxIterations = n*n).  If
+    // the solver fails to converge, try increasing the maximum number of
+    // iterations.
+    void SetMaxLCPIterations(int maxLCPIterations);
+
+    Result operator()(Rectangle3<Real> const& rectangle, OrientedBox3<Real> const& box);
+
+private:
+    LCPSolver<Real, 10> mLCP;
+};
+
+
+template <typename Real>
+void DCPQuery<Real, Rectangle3<Real>, OrientedBox3<Real>>::SetMaxLCPIterations(int maxLCPIterations)
+{
+    mLCP.SetMaxIterations(maxLCPIterations);
+}
+
+template <typename Real>
+typename DCPQuery<Real, Rectangle3<Real>, OrientedBox3<Real>>::Result
+DCPQuery<Real, Rectangle3<Real>, OrientedBox3<Real>>::operator()(
+    Rectangle3<Real> const& rectangle, OrientedBox3<Real> const& box)
+{
+    Result result;
+
+    // Rigidly transform the rectangle and oriented box so that the oriented
+    // box becomes a canonical box.
+    Vector3<Real> K = box.extent * (Real)2;
+    Vector3<Real> tempV = rectangle.center - box.center;
+    Vector3<Real> V, E0, E1;
+    for (int i = 0; i < 3; ++i)
+    {
+        V[i] = Dot(box.axis[i], tempV) + box.extent[i];
+        E0[i] = Dot(box.axis[i], rectangle.axis[0]);
+        E1[i] = Dot(box.axis[i], rectangle.axis[1]);
+    }
+
+    // Convert the oriented rectangle to a regular one (origin at a corner).
+    Vector3<Real> scaledE0 = E0 * rectangle.extent[0];
+    Vector3<Real> scaledE1 = E1 * rectangle.extent[1];
+    V -= scaledE0 + scaledE1;
+    E0 = scaledE0 * (Real)2;
+    E1 = scaledE1 * (Real)2;
+
+    // Compute quantities to initialize q and M in the LCP.
+    Real dotVE0 = Dot(V, E0);
+    Real dotVE1 = Dot(V, E1);
+    Real dotE0E0 = Dot(E0, E0);
+    Real dotE1E1 = Dot(E1, E1);
+
+    // The LCP has 5 variables and 5 (nontrivial) inequality constraints.
+    std::array<Real, 10> q =
+    {
+        -V[0], -V[1], -V[2], dotVE0, dotVE1, K[0], K[1], K[2], (Real)1, (Real)1
+    };
+
+    std::array<std::array<Real, 10>, 10> M;
+    M[0] = { (Real)1, (Real)0, (Real)0, -E0[0], -E1[0], (Real)1, (Real)0, (Real)0, (Real)0, (Real)0 };
+    M[1] = { (Real)0, (Real)1, (Real)0, -E0[1], -E1[1], (Real)0, (Real)1, (Real)0, (Real)0, (Real)0 };
+    M[2] = { (Real)0, (Real)0, (Real)1, -E0[2], -E1[2], (Real)0, (Real)0, (Real)1, (Real)0 , (Real)0 };
+    M[3] = { -E0[0], -E0[1], -E0[2], dotE0E0, (Real)0, (Real)0, (Real)0, (Real)0, (Real)1, (Real)0 };
+    M[4] = { -E1[0], -E1[1], -E1[2], (Real)0, dotE1E1, (Real)0, (Real)0, (Real)0, (Real)0, (Real)1 };
+    M[5] = { (Real)-1, (Real)0, (Real)0, (Real)0, (Real)0, (Real)0, (Real)0, (Real)0, (Real)0, (Real)0 };
+    M[6] = { (Real)0, (Real)-1, (Real)0, (Real)0, (Real)0, (Real)0, (Real)0, (Real)0, (Real)0, (Real)0 };
+    M[7] = { (Real)0, (Real)0, (Real)-1, (Real)0, (Real)0, (Real)0, (Real)0, (Real)0, (Real)0, (Real)0 };
+    M[8] = { (Real)0, (Real)0, (Real)0, (Real)-1, (Real)0, (Real)0, (Real)0, (Real)0, (Real)0, (Real)0 };
+    M[9] = { (Real)0, (Real)0, (Real)0, (Real)0, (Real)-1, (Real)0, (Real)0, (Real)0, (Real)0, (Real)0 };
+
+    std::array<Real, 10> w, z;
+    if (mLCP.Solve(q, M, w, z))
+    {
+        result.queryIsSuccessful = true;
+
+        Real t0 = (z[3] * (Real)2 - (Real)1) * rectangle.extent[0];
+        Real t1 = (z[4] * (Real)2 - (Real)1) * rectangle.extent[1];
+        result.rectangleParameter[0] = t0;
+        result.rectangleParameter[1] = t1;
+        result.closestPoint[0] = rectangle.center + t0 * rectangle.axis[0] + t1 * rectangle.axis[1];
+        result.closestPoint[1] = box.center;
+        for (int i = 0; i < 3; ++i)
+        {
+            result.boxParameter[i] = z[i] - box.extent[i];
+            result.closestPoint[1] += result.boxParameter[i] * box.axis[i];
+        }
+
+        Vector3<Real> diff = result.closestPoint[1] - result.closestPoint[0];
+        result.sqrDistance = Dot(diff, diff);
+        result.distance = std::sqrt(result.sqrDistance);
+    }
+    else
+    {
+        // If you reach this case, the maximum number of iterations was not
+        // specified to be large enough or there is a problem due to
+        // floating-point rounding errors.  If you believe the latter is
+        // true, file a bug report.
+        result.queryIsSuccessful = false;
+
+        for (int i = 0; i < 2; ++i)
+        {
+            result.rectangleParameter[i] = (Real)0;
+        }
+        for (int i = 0; i < 3; ++i)
+        {
+            result.boxParameter[i] = (Real)0;
+            result.closestPoint[0][i] = (Real)0;
+            result.closestPoint[1][i] = (Real)0;
+        }
+        result.distance = (Real)0;
+        result.sqrDistance = (Real)0;
+    }
+
+    result.numLCPIterations = mLCP.GetNumIterations();
+    return result;
+}
+
+}
diff --git a/Engine/Plugins/Experimental/GeometryProcessing/Source/GeometryAlgorithms/Private/GTEngine/Mathematics/GteDistRectangle3Rectangle3.h b/Engine/Plugins/Experimental/GeometryProcessing/Source/GeometryAlgorithms/Private/GTEngine/Mathematics/GteDistRectangle3Rectangle3.h
new file mode 100644
--- /dev/null
+++ b/Engine/Plugins/Experimental/GeometryProcessing/Source/GeometryAlgorithms/Private/GTEngine/Mathematics/GteDistRectangle3Rectangle3.h
@@ -0,0 +1,104 @@
+// David Eberly, Geometric Tools, Redmond WA 98052
+// Copyright (c) 1998-2019
+// Distributed under the Boost Software License, Version 1.0.
+// http://www.boost.org/LICENSE_1_0.txt
+// http://www.geometrictools.com/License/Boost/LICENSE_1_0.txt
+// File Version: 3.0.0 (2016/06/19)
+
+#pragma once
+
+#include <GTEngine/Mathematics/GteDistSegment3Rectangle3.h>
+
+namespace gte
+{
+
+template <typename Real>
+class DCPQuery<Real, Rectangle3<Real>, Rectangle3<Real>>
+{
+public:
+    struct Result
+    {
+        Real distance, sqrDistance;
+        Real rectangle0Parameter[2], rectangle1Parameter[2];
+        Vector3<Real> closestPoint[2];
+    };
+
+    Result operator()(Rectangle3<Real> const& rectangle0,
+        Rectangle3<Real> const& rectangle1);
+};
+
+
+template <typename Real>
+typename DCPQuery<Real, Rectangle3<Real>, Rectangle3<Real>>::Result
+DCPQuery<Real, Rectangle3<Real>, Rectangle3<Real>>::operator()(
+    Rectangle3<Real> const& rectangle0, Rectangle3<Real> const& rectangle1)
+{
+    Result result;
+
+    DCPQuery<Real, Segment3<Real>, Rectangle3<Real>> srQuery;
+    typename DCPQuery<Real, Segment3<Real>, Rectangle3<Real>>::Result
+        srResult;
+    result.sqrDistance = std::numeric_limits<Real>::max();
+
+    // Compare edges of rectangle0 to the interior of rectangle1.
+    for (int i1 = 0; i1 < 2; ++i1)
+    {
+        for (int i0 = -1; i0 <= 1; i0 += 2)
+        {
+            Real s = i0 * rectangle0.extent[1 - i1];
+            Vector3<Real> segCenter = rectangle0.center +
+                s * rectangle0.axis[1 - i1];
+            Segment3<Real> edge(segCenter, rectangle0.axis[i1],
+                rectangle0.extent[i1]);
+
+            srResult = srQuery(edge, rectangle1);
+            if (srResult.sqrDistance < result.sqrDistance)
+            {
+                result.distance = srResult.distance;
+                result.sqrDistance = srResult.sqrDistance;
+                result.rectangle0Parameter[i1] = s;
+                result.rectangle0Parameter[1 - i1] =
+                    srResult.segmentParameter;
+                result.rectangle1Parameter[0] =
+                    srResult.rectangleParameter[0];
+                result.rectangle1Parameter[1] =
+                    srResult.rectangleParameter[1];
+                result.closestPoint[0] = srResult.closestPoint[0];
+                result.closestPoint[1] = srResult.closestPoint[1];
+            }
+        }
+    }
+
+    // Compare edges of rectangle1 to the interior of rectangle0.
+    for (int i1 = 0; i1 < 2; ++i1)
+    {
+        for (int i0 = -1; i0 <= 1; i0 += 2)
+        {
+            Real s = i0 * rectangle1.extent[1 - i1];
+            Vector3<Real> segCenter = rectangle1.center +
+                s * rectangle1.axis[1 - i1];
+            Segment3<Real> edge(segCenter, rectangle0.axis[i1],
+                rectangle0.extent[i1]);
+
+            srResult = srQuery(edge, rectangle0);
+            if (srResult.sqrDistance < result.sqrDistance)
+            {
+                result.distance = srResult.distance;
+                result.sqrDistance = srResult.sqrDistance;
+                result.rectangle0Parameter[0] =
+                    srResult.rectangleParameter[0];
+                result.rectangle0Parameter[1] =
+                    srResult.rectangleParameter[1];
+                result.rectangle1Parameter[i1] = s;
+                result.rectangle1Parameter[1 - i1] =
+                    srResult.segmentParameter;
+                result.closestPoint[0] = srResult.closestPoint[1];
+                result.closestPoint[1] = srResult.closestPoint[0];
+            }
+        }
+    }
+    return result;
+}
+
+
+}
diff --git a/Engine/Plugins/Experimental/GeometryProcessing/Source/GeometryAlgorithms/Private/GTEngine/Mathematics/GteDistSegment3AlignedBox3.h b/Engine/Plugins/Experimental/GeometryProcessing/Source/GeometryAlgorithms/Private/GTEngine/Mathematics/GteDistSegment3AlignedBox3.h
new file mode 100644
--- /dev/null
+++ b/Engine/Plugins/Experimental/GeometryProcessing/Source/GeometryAlgorithms/Private/GTEngine/Mathematics/GteDistSegment3AlignedBox3.h
@@ -0,0 +1,84 @@
+// David Eberly, Geometric Tools, Redmond WA 98052
+// Copyright (c) 1998-2019
+// Distributed under the Boost Software License, Version 1.0.
+// http://www.boost.org/LICENSE_1_0.txt
+// http://www.geometrictools.com/License/Boost/LICENSE_1_0.txt
+// File Version: 3.0.1 (2016/11/17)
+
+#pragma once
+
+#include <GTEngine/Mathematics/GteDistLine3AlignedBox3.h>
+#include <GTEngine/Mathematics/GteSegment.h>
+
+namespace gte
+{
+
+template <typename Real>
+class DCPQuery<Real, Segment3<Real>, AlignedBox3<Real>>
+{
+public:
+    struct Result
+    {
+        Real distance, sqrDistance;
+        Real segmentParameter;
+        Vector3<Real> closestPoint[2];
+    };
+
+    Result operator()(Segment3<Real> const& segment,
+        AlignedBox3<Real> const& box);
+};
+
+
+template <typename Real>
+typename DCPQuery<Real, Segment3<Real>, AlignedBox3<Real>>::Result
+DCPQuery<Real, Segment3<Real>, AlignedBox3<Real>> ::operator()(
+    Segment3<Real> const& segment, AlignedBox3<Real> const& box)
+{
+    Result result;
+
+    Vector3<Real> segCenter, segDirection;
+    Real segExtent;
+    segment.GetCenteredForm(segCenter, segDirection, segExtent);
+
+    Line3<Real> line(segCenter, segDirection);
+    DCPQuery<Real, Line3<Real>, AlignedBox3<Real>> lbQuery;
+    auto lbResult = lbQuery(line, box);
+
+    if (lbResult.lineParameter >= -segExtent)
+    {
+        if (lbResult.lineParameter <= segExtent)
+        {
+            result.sqrDistance = lbResult.sqrDistance;
+            result.distance = lbResult.distance;
+            result.segmentParameter = lbResult.lineParameter;
+            result.closestPoint[0] = lbResult.closestPoint[0];
+            result.closestPoint[1] = lbResult.closestPoint[1];
+        }
+        else
+        {
+            DCPQuery<Real, Vector3<Real>, AlignedBox3<Real>> pbQuery;
+            Vector3<Real> point = segCenter + segExtent*segDirection;
+            auto pbResult = pbQuery(point, box);
+            result.sqrDistance = pbResult.sqrDistance;
+            result.distance = pbResult.distance;
+            result.segmentParameter = segExtent;
+            result.closestPoint[0] = point;
+            result.closestPoint[1] = pbResult.boxClosest;
+        }
+    }
+    else
+    {
+        DCPQuery<Real, Vector3<Real>, AlignedBox3<Real>> pbQuery;
+        Vector3<Real> point = segCenter - segExtent*segDirection;
+        auto pbResult = pbQuery(point, box);
+        result.sqrDistance = pbResult.sqrDistance;
+        result.distance = pbResult.distance;
+        result.segmentParameter = -segExtent;
+        result.closestPoint[0] = point;
+        result.closestPoint[1] = pbResult.boxClosest;
+    }
+    return result;
+}
+
+
+}
diff --git a/Engine/Plugins/Experimental/GeometryProcessing/Source/GeometryAlgorithms/Private/GTEngine/Mathematics/GteDistSegment3OrientedBox3.h b/Engine/Plugins/Experimental/GeometryProcessing/Source/GeometryAlgorithms/Private/GTEngine/Mathematics/GteDistSegment3OrientedBox3.h
new file mode 100644
--- /dev/null
+++ b/Engine/Plugins/Experimental/GeometryProcessing/Source/GeometryAlgorithms/Private/GTEngine/Mathematics/GteDistSegment3OrientedBox3.h
@@ -0,0 +1,86 @@
+// David Eberly, Geometric Tools, Redmond WA 98052
+// Copyright (c) 1998-2019
+// Distributed under the Boost Software License, Version 1.0.
+// http://www.boost.org/LICENSE_1_0.txt
+// http://www.geometrictools.com/License/Boost/LICENSE_1_0.txt
+// File Version: 3.0.0 (2016/06/19)
+
+#pragma once
+
+#include <GTEngine/Mathematics/GteDistLine3OrientedBox3.h>
+#include <GTEngine/Mathematics/GteDistPointOrientedBox.h>
+#include <GTEngine/Mathematics/GteRay.h>
+#include <GTEngine/Mathematics/GteSegment.h>
+
+namespace gte
+{
+
+template <typename Real>
+class DCPQuery<Real, Segment3<Real>, OrientedBox3<Real>>
+{
+public:
+    struct Result
+    {
+        Real distance, sqrDistance;
+        Real segmentParameter;
+        Vector3<Real> closestPoint[2];
+    };
+
+    Result operator()(Segment3<Real> const& segment,
+        OrientedBox3<Real> const& box);
+};
+
+
+template <typename Real>
+typename DCPQuery<Real, Segment3<Real>, OrientedBox3<Real>>::Result
+DCPQuery<Real, Segment3<Real>, OrientedBox3<Real>> ::operator()(
+    Segment3<Real> const& segment, OrientedBox3<Real> const& box)
+{
+    Result result;
+
+    Vector3<Real> segCenter, segDirection;
+    Real segExtent;
+    segment.GetCenteredForm(segCenter, segDirection, segExtent);
+
+    Line3<Real> line(segCenter, segDirection);
+    DCPQuery<Real, Line3<Real>, OrientedBox3<Real>> lbQuery;
+    auto lbResult = lbQuery(line, box);
+
+    if (lbResult.lineParameter >= -segExtent)
+    {
+        if (lbResult.lineParameter <= segExtent)
+        {
+            result.sqrDistance = lbResult.sqrDistance;
+            result.distance = lbResult.distance;
+            result.segmentParameter = lbResult.lineParameter;
+            result.closestPoint[0] = lbResult.closestPoint[0];
+            result.closestPoint[1] = lbResult.closestPoint[1];
+        }
+        else
+        {
+            DCPQuery<Real, Vector3<Real>, OrientedBox3<Real>> pbQuery;
+            Vector3<Real> point = segCenter + segExtent*segDirection;
+            auto pbResult = pbQuery(point, box);
+            result.sqrDistance = pbResult.sqrDistance;
+            result.distance = pbResult.distance;
+            result.segmentParameter = segExtent;
+            result.closestPoint[0] = point;
+            result.closestPoint[1] = pbResult.boxClosest;
+        }
+    }
+    else
+    {
+        DCPQuery<Real, Vector3<Real>, OrientedBox3<Real>> pbQuery;
+        Vector3<Real> point = segCenter - segExtent*segDirection;
+        auto pbResult = pbQuery(point, box);
+        result.sqrDistance = pbResult.sqrDistance;
+        result.distance = pbResult.distance;
+        result.segmentParameter = segExtent;
+        result.closestPoint[0] = point;
+        result.closestPoint[1] = pbResult.boxClosest;
+    }
+    return result;
+}
+
+
+}
diff --git a/Engine/Plugins/Experimental/GeometryProcessing/Source/GeometryAlgorithms/Private/GTEngine/Mathematics/GteDistSegment3Rectangle3.h b/Engine/Plugins/Experimental/GeometryProcessing/Source/GeometryAlgorithms/Private/GTEngine/Mathematics/GteDistSegment3Rectangle3.h
new file mode 100644
--- /dev/null
+++ b/Engine/Plugins/Experimental/GeometryProcessing/Source/GeometryAlgorithms/Private/GTEngine/Mathematics/GteDistSegment3Rectangle3.h
@@ -0,0 +1,87 @@
+// David Eberly, Geometric Tools, Redmond WA 98052
+// Copyright (c) 1998-2019
+// Distributed under the Boost Software License, Version 1.0.
+// http://www.boost.org/LICENSE_1_0.txt
+// http://www.geometrictools.com/License/Boost/LICENSE_1_0.txt
+// File Version: 3.0.0 (2016/06/19)
+
+#pragma once
+
+#include <GTEngine/Mathematics/GteDistLine3Rectangle3.h>
+#include <GTEngine/Mathematics/GteDistPoint3Rectangle3.h>
+#include <GTEngine/Mathematics/GteSegment.h>
+
+namespace gte
+{
+
+template <typename Real>
+class DCPQuery<Real, Segment3<Real>, Rectangle3<Real>>
+{
+public:
+    struct Result
+    {
+        Real distance, sqrDistance;
+        Real segmentParameter, rectangleParameter[2];
+        Vector3<Real> closestPoint[2];
+    };
+
+    Result operator()(Segment3<Real> const& segment,
+        Rectangle3<Real> const& rectangle);
+};
+
+
+template <typename Real>
+typename DCPQuery<Real, Segment3<Real>, Rectangle3<Real>>::Result
+DCPQuery<Real, Segment3<Real>, Rectangle3<Real>>::operator()(
+    Segment3<Real> const& segment, Rectangle3<Real> const& rectangle)
+{
+    Result result;
+
+    Vector3<Real> segCenter, segDirection;
+    Real segExtent;
+    segment.GetCenteredForm(segCenter, segDirection, segExtent);
+
+    Line3<Real> line(segCenter, segDirection);
+    DCPQuery<Real, Line3<Real>, Rectangle3<Real>> lrQuery;
+    auto lrResult = lrQuery(line, rectangle);
+
+    if (lrResult.lineParameter >= -segExtent)
+    {
+        if (lrResult.lineParameter <= segExtent)
+        {
+            result.distance = lrResult.distance;
+            result.sqrDistance = lrResult.sqrDistance;
+            result.segmentParameter = lrResult.lineParameter;
+            result.rectangleParameter[0] = lrResult.rectangleParameter[0];
+            result.rectangleParameter[1] = lrResult.rectangleParameter[1];
+            result.closestPoint[0] = lrResult.closestPoint[0];
+            result.closestPoint[1] = lrResult.closestPoint[1];
+        }
+        else
+        {
+            DCPQuery<Real, Vector3<Real>, Rectangle3<Real>> prQuery;
+            Vector3<Real> point = segCenter + segExtent*segDirection;
+            auto prResult = prQuery(point, rectangle);
+            result.sqrDistance = prResult.sqrDistance;
+            result.distance = prResult.distance;
+            result.segmentParameter = segExtent;
+            result.closestPoint[0] = point;
+            result.closestPoint[1] = prResult.rectangleClosestPoint;
+        }
+    }
+    else
+    {
+        DCPQuery<Real, Vector3<Real>, Rectangle3<Real>> prQuery;
+        Vector3<Real> point = segCenter - segExtent*segDirection;
+        auto prResult = prQuery(point, rectangle);
+        result.sqrDistance = prResult.sqrDistance;
+        result.distance = prResult.distance;
+        result.segmentParameter = segExtent;
+        result.closestPoint[0] = point;
+        result.closestPoint[1] = prResult.rectangleClosestPoint;
+    }
+    return result;
+}
+
+
+}
diff --git a/Engine/Plugins/Experimental/GeometryProcessing/Source/GeometryAlgorithms/Private/GTEngine/Mathematics/GteDistSegment3Triangle3.h b/Engine/Plugins/Experimental/GeometryProcessing/Source/GeometryAlgorithms/Private/GTEngine/Mathematics/GteDistSegment3Triangle3.h
new file mode 100644
--- /dev/null
+++ b/Engine/Plugins/Experimental/GeometryProcessing/Source/GeometryAlgorithms/Private/GTEngine/Mathematics/GteDistSegment3Triangle3.h
@@ -0,0 +1,88 @@
+// David Eberly, Geometric Tools, Redmond WA 98052
+// Copyright (c) 1998-2019
+// Distributed under the Boost Software License, Version 1.0.
+// http://www.boost.org/LICENSE_1_0.txt
+// http://www.geometrictools.com/License/Boost/LICENSE_1_0.txt
+// File Version: 3.0.0 (2016/06/19)
+
+#pragma once
+
+#include <GTEngine/Mathematics/GteDistLine3Triangle3.h>
+#include <GTEngine/Mathematics/GteDistPointTriangle.h>
+#include <GTEngine/Mathematics/GteSegment.h>
+
+namespace gte
+{
+
+template <typename Real>
+class DCPQuery<Real, Segment3<Real>, Triangle3<Real>>
+{
+public:
+    struct Result
+    {
+        Real distance, sqrDistance;
+        Real segmentParameter, triangleParameter[3];
+        Vector3<Real> closestPoint[2];
+    };
+
+    Result operator()(Segment3<Real> const& segment,
+        Triangle3<Real> const& triangle);
+};
+
+
+template <typename Real>
+typename DCPQuery<Real, Segment3<Real>, Triangle3<Real>>::Result
+DCPQuery<Real, Segment3<Real>, Triangle3<Real>>::operator()(
+    Segment3<Real> const& segment, Triangle3<Real> const& triangle)
+{
+    Result result;
+
+    Vector3<Real> segCenter, segDirection;
+    Real segExtent;
+    segment.GetCenteredForm(segCenter, segDirection, segExtent);
+
+    Line3<Real> line(segCenter, segDirection);
+    DCPQuery<Real, Line3<Real>, Triangle3<Real>> ltQuery;
+    auto ltResult = ltQuery(line, triangle);
+
+    if (ltResult.lineParameter >= -segExtent)
+    {
+        if (ltResult.lineParameter <= segExtent)
+        {
+            result.distance = ltResult.distance;
+            result.sqrDistance = ltResult.sqrDistance;
+            result.segmentParameter = ltResult.lineParameter;
+            result.triangleParameter[0] = ltResult.triangleParameter[0];
+            result.triangleParameter[1] = ltResult.triangleParameter[1];
+            result.triangleParameter[2] = ltResult.triangleParameter[2];
+            result.closestPoint[0] = ltResult.closestPoint[0];
+            result.closestPoint[1] = ltResult.closestPoint[1];
+        }
+        else
+        {
+            DCPQuery<Real, Vector3<Real>, Triangle3<Real>> ptQuery;
+            Vector3<Real> point = segCenter + segExtent*segDirection;
+            auto ptResult = ptQuery(point, triangle);
+            result.sqrDistance = ptResult.sqrDistance;
+            result.distance = ptResult.distance;
+            result.segmentParameter = segExtent;
+            result.closestPoint[0] = point;
+            result.closestPoint[1] = ptResult.closest;
+        }
+    }
+    else
+    {
+        DCPQuery<Real, Vector3<Real>, Triangle3<Real>> ptQuery;
+        Vector3<Real> point = segCenter - segExtent*segDirection;
+        auto ptResult = ptQuery(point, triangle);
+        result.sqrDistance = ptResult.sqrDistance;
+        result.distance = ptResult.distance;
+        result.segmentParameter = segExtent;
+        result.closestPoint[0] = point;
+        result.closestPoint[1] = ptResult.closest;
+    }
+    return result;
+}
+
+
+}
diff --git a/Engine/Plugins/Experimental/GeometryProcessing/Source/GeometryAlgorithms/Private/GTEngine/Mathematics/GteDistSegmentSegment.h b/Engine/Plugins/Experimental/GeometryProcessing/Source/GeometryAlgorithms/Private/GTEngine/Mathematics/GteDistSegmentSegment.h
new file mode 100644
--- /dev/null
+++ b/Engine/Plugins/Experimental/GeometryProcessing/Source/GeometryAlgorithms/Private/GTEngine/Mathematics/GteDistSegmentSegment.h
@@ -0,0 +1,427 @@
+// David Eberly, Geometric Tools, Redmond WA 98052
+// Copyright (c) 1998-2019
+// Distributed under the Boost Software License, Version 1.0.
+// http://www.boost.org/LICENSE_1_0.txt
+// http://www.geometrictools.com/License/Boost/LICENSE_1_0.txt
+// File Version: 3.0.1 (2018/10/05)
+
+#pragma once
+
+#include <GTEngine/Mathematics/GteDCPQuery.h>
+#include <GTEngine/Mathematics/GteSegment.h>
+
+// Compute the closest points on the line segments P(s) = (1-s)*P0 + s*P1 and
+// Q(t) = (1-t)*Q0 + t*Q1 for 0 <= s <= 1 and 0 <= t <= 1.  The algorithm is
+// robust even for nearly parallel segments.  Effectively, it uses a conjugate
+// gradient search for the minimum of the squared distance function, which
+// avoids the numerical problems introduced by divisions in the case the
+// minimum is located at an interior point of the domain.  See the document
+//   http://www.geometrictools.com/Documentation/DistanceLine3Line3.pdf
+// for details.
+
+namespace gte
+{
+
+template <int N, typename Real>
+class DCPQuery<Real, Segment<N, Real>, Segment<N, Real>>
+{
+public:
+    struct Result
+    {
+        Real distance, sqrDistance;
+        Real parameter[2];
+        Vector<N, Real> closest[2];
+    };
+
+    Result operator()(Segment<N, Real> const& segment0,
+        Segment<N, Real> const& segment1);
+
+    Result operator()(Vector<N, Real> const& P0, Vector<N, Real> const& P1,
+        Vector<N, Real> const& Q0, Vector<N, Real> const& Q1);
+
+private:
+    // Compute the root of h(z) = h0 + slope*z and clamp it to the interval
+    // [0,1].  It is required that for h1 = h(1), either (h0 < 0 and h1 > 0)
+    // or (h0 > 0 and h1 < 0).
+    Real GetClampedRoot(Real slope, Real h0, Real h1);
+
+    // Compute the intersection of the line dR/ds = 0 with the domain [0,1]^2.
+    // The direction of the line dR/ds is conjugate to (1,0), so the algorithm
+    // for minimization is effectively the conjugate gradient algorithm for a
+    // quadratic function.
+    void ComputeIntersection(Real const sValue[2], int const classify[2],
+        int edge[2], Real end[2][2]);
+
+    // Compute the location of the minimum of R on the segment of intersection
+    // for the line dR/ds = 0 and the domain [0,1]^2.
+    void ComputeMinimumParameters(int const edge[2], Real const end[2][2],
+        Real parameter[2]);
+
+    // The coefficients of R(s,t), not including the constant term.
+    Real mA, mB, mC, mD, mE;
+
+    // dR/ds(i,j) at the four corners of the domain
+    Real mF00, mF10, mF01, mF11;
+
+    // dR/dt(i,j) at the four corners of the domain
+    Real mG00, mG10, mG01, mG11;
+};
+
+// Template aliases for convenience.
+template <int N, typename Real>
+using DCPSegmentSegment = DCPQuery<Real, Segment<N, Real>, Segment<N, Real>>;
+
+template <typename Real>
+using DCPSegment2Segment2 = DCPSegmentSegment<2, Real>;
+
+template <typename Real>
+using DCPSegment3Segment3 = DCPSegmentSegment<3, Real>;
+
+
+template <int N, typename Real>
+typename DCPQuery<Real, Segment<N, Real>, Segment<N, Real>>::Result
+DCPQuery<Real, Segment<N, Real>, Segment<N, Real>>::operator()(
+    Segment<N, Real> const& segment0, Segment<N, Real> const& segment1)
+{
+    return operator()(segment0.p[0], segment0.p[1], segment1.p[0],
+        segment1.p[1]);
+}
+
+template <int N, typename Real>
+typename DCPQuery<Real, Segment<N, Real>, Segment<N, Real>>::Result
+DCPQuery<Real, Segment<N, Real>, Segment<N, Real>>::operator()(
+    Vector<N, Real> const& P0, Vector<N, Real> const& P1,
+    Vector<N, Real> const& Q0, Vector<N, Real> const& Q1)
+{
+    Result result;
+
+    // The code allows degenerate line segments; that is, P0 and P1 can be
+    // the same point or Q0 and Q1 can be the same point.  The quadratic
+    // function for squared distance between the segment is
+    //   R(s,t) = a*s^2 - 2*b*s*t + c*t^2 + 2*d*s - 2*e*t + f
+    // for (s,t) in [0,1]^2 where
+    //   a = Dot(P1-P0,P1-P0), b = Dot(P1-P0,Q1-Q0), c = Dot(Q1-Q0,Q1-Q0),
+    //   d = Dot(P1-P0,P0-Q0), e = Dot(Q1-Q0,P0-Q0), f = Dot(P0-Q0,P0-Q0)
+    Vector<N, Real> P1mP0 = P1 - P0;
+    Vector<N, Real> Q1mQ0 = Q1 - Q0;
+    Vector<N, Real> P0mQ0 = P0 - Q0;
+    mA = Dot(P1mP0, P1mP0);
+    mB = Dot(P1mP0, Q1mQ0);
+    mC = Dot(Q1mQ0, Q1mQ0);
+    mD = Dot(P1mP0, P0mQ0);
+    mE = Dot(Q1mQ0, P0mQ0);
+
+    mF00 = mD;
+    mF10 = mF00 + mA;
+    mF01 = mF00 - mB;
+    mF11 = mF10 - mB;
+
+    mG00 = -mE;
+    mG10 = mG00 - mB;
+    mG01 = mG00 + mC;
+    mG11 = mG10 + mC;
+
+    if (mA > (Real)0 && mC > (Real)0)
+    {
+        // Compute the solutions to dR/ds(s0,0) = 0 and dR/ds(s1,1) = 0.  The
+        // location of sI on the s-axis is stored in classifyI (I = 0 or 1).  If
+        // sI <= 0, classifyI is -1.  If sI >= 1, classifyI is 1.  If 0 < sI < 1,
+        // classifyI is 0.  This information helps determine where to search for
+        // the minimum point (s,t).  The fij values are dR/ds(i,j) for i and j in
+        // {0,1}.
+
+        Real sValue[2];
+        sValue[0] = GetClampedRoot(mA, mF00, mF10);
+        sValue[1] = GetClampedRoot(mA, mF01, mF11);
+
+        int classify[2];
+        for (int i = 0; i < 2; ++i)
+        {
+            if (sValue[i] <= (Real)0)
+            {
+                classify[i] = -1;
+            }
+            else if (sValue[i] >= (Real)1)
+            {
+                classify[i] = +1;
+            }
+            else
+            {
+                classify[i] = 0;
+            }
+        }
+
+        if (classify[0] == -1 && classify[1] == -1)
+        {
+            // The minimum must occur on s = 0 for 0 <= t <= 1.
+            result.parameter[0] = (Real)0;
+            result.parameter[1] = GetClampedRoot(mC, mG00, mG01);
+        }
+        else if (classify[0] == +1 && classify[1] == +1)
+        {
+            // The minimum must occur on s = 1 for 0 <= t <= 1.
+            result.parameter[0] = (Real)1;
+            result.parameter[1] = GetClampedRoot(mC, mG10, mG11);
+        }
+        else
+        {
+            // The line dR/ds = 0 intersects the domain [0,1]^2 in a
+            // nondegenerate segment.  Compute the endpoints of that segment,
+            // end[0] and end[1].  The edge[i] flag tells you on which domain
+            // edge end[i] lives: 0 (s=0), 1 (s=1), 2 (t=0), 3 (t=1).
+            int edge[2];
+            Real end[2][2];
+            ComputeIntersection(sValue, classify, edge, end);
+
+            // The directional derivative of R along the segment of
+            // intersection is
+            //   H(z) = (end[1][1]-end[1][0])*dR/dt((1-z)*end[0] + z*end[1])
+            // for z in [0,1].  The formula uses the fact that dR/ds = 0 on
+            // the segment.  Compute the minimum of H on [0,1].
+            ComputeMinimumParameters(edge, end, result.parameter);
+        }
+    }
+    else
+    {
+        if (mA > (Real)0)
+        {
+            // The Q-segment is degenerate (Q0 and Q1 are the same point) and
+            // the quadratic is R(s,0) = a*s^2 + 2*d*s + f and has (half)
+            // first derivative F(t) = a*s + d.  The closest P-point is
+            // interior to the P-segment when F(0) < 0 and F(1) > 0.
+            result.parameter[0] = GetClampedRoot(mA, mF00, mF10);
+            result.parameter[1] = (Real)0;
+        }
+        else if (mC > (Real)0)
+        {
+            // The P-segment is degenerate (P0 and P1 are the same point) and
+            // the quadratic is R(0,t) = c*t^2 - 2*e*t + f and has (half)
+            // first derivative G(t) = c*t - e.  The closest Q-point is
+            // interior to the Q-segment when G(0) < 0 and G(1) > 0.
+            result.parameter[0] = (Real)0;
+            result.parameter[1] = GetClampedRoot(mC, mG00, mG01);
+        }
+        else
+        {
+            // P-segment and Q-segment are degenerate.
+            result.parameter[0] = (Real)0;
+            result.parameter[1] = (Real)0;
+        }
+    }
+
+
+    result.closest[0] =
+        ((Real)1 - result.parameter[0]) * P0 + result.parameter[0] * P1;
+    result.closest[1] =
+        ((Real)1 - result.parameter[1]) * Q0 + result.parameter[1] * Q1;
+    Vector<N, Real> diff = result.closest[0] - result.closest[1];
+    result.sqrDistance = Dot(diff, diff);
+    result.distance = std::sqrt(result.sqrDistance);
+    return result;
+}
+
+template <int N, typename Real>
+Real DCPQuery<Real, Segment<N, Real>, Segment<N, Real>>::GetClampedRoot(
+    Real slope, Real h0, Real h1)
+{
+    // Theoretically, r is in (0,1).  However, when the slope is nearly zero,
+    // then so are h0 and h1.  Significant numerical rounding problems can
+    // occur when using floating-point arithmetic.  If the rounding causes r
+    // to be outside the interval, clamp it.  It is possible that r is in
+    // (0,1) and has rounding errors, but because h0 and h1 are both nearly
+    // zero, the quadratic is nearly constant on (0,1).  Any choice of p
+    // should not cause undesirable accuracy problems for the final distance
+    // computation.
+    //
+    // NOTE:  You can use bisection to recompute the root or even use
+    // bisection to compute the root and skip the division.  This is generally
+    // slower, which might be a problem for high-performance applications.
+
+    Real r;
+    if (h0 < (Real)0)
+    {
+        if (h1 > (Real)0)
+        {
+            r = -h0 / slope;
+            if (r > (Real)1)
+            {
+                r = (Real)0.5;
+            }
+            // The slope is positive and -h0 is positive, so there is no
+            // need to test for a negative value and clamp it.
+        }
+        else
+        {
+            r = (Real)1;
+        }
+    }
+    else
+    {
+        r = (Real)0;
+    }
+    return r;
+}
+
+template <int N, typename Real>
+void DCPQuery<Real, Segment<N, Real>, Segment<N, Real>>::ComputeIntersection(
+    Real const sValue[2], int const classify[2], int edge[2], Real end[2][2])
+{
+    // The divisions are theoretically numbers in [0,1].  Numerical rounding
+    // errors might cause the result to be outside the interval.  When this
+    // happens, it must be that both numerator and denominator are nearly
+    // zero.  The denominator is nearly zero when the segments are nearly
+    // perpendicular.  The numerator is nearly zero when the P-segment is
+    // nearly degenerate (mF00 = a is small).  The choice of 0.5 should not
+    // cause significant accuracy problems.
+    //
+    // NOTE:  You can use bisection to recompute the root or even use
+    // bisection to compute the root and skip the division.  This is generally
+    // slower, which might be a problem for high-performance applications.
+
+    if (classify[0] < 0)
+    {
+        edge[0] = 0;
+        end[0][0] = (Real)0;
+        end[0][1] = mF00 / mB;
+        if (end[0][1] < (Real)0 || end[0][1] > (Real)1)
+        {
+            end[0][1] = (Real)0.5;
+        }
+
+        if (classify[1] == 0)
+        {
+            edge[1] = 3;
+            end[1][0] = sValue[1];
+            end[1][1] = (Real)1;
+        }
+        else  // classify[1] > 0
+        {
+            edge[1] = 1;
+            end[1][0] = (Real)1;
+            end[1][1] = mF10 / mB;
+            if (end[1][1] < (Real)0 || end[1][1] > (Real)1)
+            {
+                end[1][1] = (Real)0.5;
+            }
+        }
+    }
+    else if (classify[0] == 0)
+    {
+        edge[0] = 2;
+        end[0][0] = sValue[0];
+        end[0][1] = (Real)0;
+
+        if (classify[1] < 0)
+        {
+            edge[1] = 0;
+            end[1][0] = (Real)0;
+            end[1][1] = mF00 / mB;
+            if (end[1][1] < (Real)0 || end[1][1] > (Real)1)
+            {
+                end[1][1] = (Real)0.5;
+            }
+        }
+        else if (classify[1] == 0)
+        {
+            edge[1] = 3;
+            end[1][0] = sValue[1];
+            end[1][1] = (Real)1;
+        }
+        else
+        {
+            edge[1] = 1;
+            end[1][0] = (Real)1;
+            end[1][1] = mF10 / mB;
+            if (end[1][1] < (Real)0 || end[1][1] > (Real)1)
+            {
+                end[1][1] = (Real)0.5;
+            }
+        }
+    }
+    else  // classify[0] > 0
+    {
+        edge[0] = 1;
+        end[0][0] = (Real)1;
+        end[0][1] = mF10 / mB;
+        if (end[0][1] < (Real)0 || end[0][1] > (Real)1)
+        {
+            end[0][1] = (Real)0.5;
+        }
+
+        if (classify[1] == 0)
+        {
+            edge[1] = 3;
+            end[1][0] = sValue[1];
+            end[1][1] = (Real)1;
+        }
+        else
+        {
+            edge[1] = 0;
+            end[1][0] = (Real)0;
+            end[1][1] = mF00 / mB;
+            if (end[1][1] < (Real)0 || end[1][1] > (Real)1)
+            {
+                end[1][1] = (Real)0.5;
+            }
+        }
+    }
+}
+
+template <int N, typename Real>
+void DCPQuery<Real, Segment<N, Real>, Segment<N, Real>>::
+ComputeMinimumParameters(int const edge[2], Real const end[2][2],
+    Real parameter[2])
+{
+    Real delta = end[1][1] - end[0][1];
+    Real h0 = delta * (-mB * end[0][0] + mC * end[0][1] - mE);
+    if (h0 >= (Real)0)
+    {
+        if (edge[0] == 0)
+        {
+            parameter[0] = (Real)0;
+            parameter[1] = GetClampedRoot(mC, mG00, mG01);
+        }
+        else if (edge[0] == 1)
+        {
+            parameter[0] = (Real)1;
+            parameter[1] = GetClampedRoot(mC, mG10, mG11);
+        }
+        else
+        {
+            parameter[0] = end[0][0];
+            parameter[1] = end[0][1];
+        }
+    }
+    else
+    {
+        Real h1 = delta * (-mB * end[1][0] + mC * end[1][1] - mE);
+        if (h1 <= (Real)0)
+        {
+            if (edge[1] == 0)
+            {
+                parameter[0] = (Real)0;
+                parameter[1] = GetClampedRoot(mC, mG00, mG01);
+            }
+            else if (edge[1] == 1)
+            {
+                parameter[0] = (Real)1;
+                parameter[1] = GetClampedRoot(mC, mG10, mG11);
+            }
+            else
+            {
+                parameter[0] = end[1][0];
+                parameter[1] = end[1][1];
+            }
+        }
+        else  // h0 < 0 and h1 > 0
+        {
+            Real z = std::min(std::max(h0 / (h0 - h1), (Real)0), (Real)1);
+            Real omz = (Real)1 - z;
+            parameter[0] = omz * end[0][0] + z * end[1][0];
+            parameter[1] = omz * end[0][1] + z * end[1][1];
+        }
+    }
+}
+
+
+}
diff --git a/Engine/Plugins/Experimental/GeometryProcessing/Source/GeometryAlgorithms/Private/GTEngine/Mathematics/GteDistSegmentSegmentExact.h b/Engine/Plugins/Experimental/GeometryProcessing/Source/GeometryAlgorithms/Private/GTEngine/Mathematics/GteDistSegmentSegmentExact.h
new file mode 100644
--- /dev/null
+++ b/Engine/Plugins/Experimental/GeometryProcessing/Source/GeometryAlgorithms/Private/GTEngine/Mathematics/GteDistSegmentSegmentExact.h
@@ -0,0 +1,289 @@
+// David Eberly, Geometric Tools, Redmond WA 98052
+// Copyright (c) 1998-2019
+// Distributed under the Boost Software License, Version 1.0.
+// http://www.boost.org/LICENSE_1_0.txt
+// http://www.geometrictools.com/License/Boost/LICENSE_1_0.txt
+// File Version: 3.0.0 (2016/06/19)
+
+#pragma once
+
+#include <GTEngine/Mathematics/GteDCPQuery.h>
+#include <GTEngine/Mathematics/GteSegment.h>
+
+// Compute the closest points on the line segments P(s) = (1-s)*P0 + s*P1 and
+// Q(t) = (1-t)*Q0 + t*Q1 for 0 <= s <= 1 and 0 <= t <= 1.  The algorithm
+// relies on exact rational arithmetic.
+
+namespace gte
+{
+
+template <int N, typename Rational>
+class DistanceSegmentSegmentExact
+{
+public:
+    struct Result
+    {
+        Rational sqrDistance;
+        Rational parameter[2];
+        Vector<N, Rational> closest[2];
+    };
+
+    Result operator()(Segment<N, Rational> const& segment0,
+        Segment<N, Rational> const& segment1);
+
+    Result operator()(
+        Vector<N, Rational> const& P0, Vector<N, Rational> const& P1,
+        Vector<N, Rational> const& Q0, Vector<N, Rational> const& Q1);
+};
+
+
+template <int N, typename Rational>
+typename DistanceSegmentSegmentExact<N, Rational>::Result
+DistanceSegmentSegmentExact<N, Rational>::operator()(
+    Segment<N, Rational> const& segment0,
+    Segment<N, Rational> const& segment1)
+{
+    return operator()(segment0.p[0], segment0.p[1], segment1.p[0],
+        segment1.p[1]);
+}
+
+template <int N, typename Rational>
+typename DistanceSegmentSegmentExact<N, Rational>::Result
+DistanceSegmentSegmentExact<N, Rational>::operator()(
+    Vector<N, Rational> const& P0, Vector<N, Rational> const& P1,
+    Vector<N, Rational> const& Q0, Vector<N, Rational> const& Q1)
+{
+    Vector<N, Rational> P1mP0 = P1 - P0;
+    Vector<N, Rational> Q1mQ0 = Q1 - Q0;
+    Vector<N, Rational> P0mQ0 = P0 - Q0;
+    Rational a = Dot(P1mP0, P1mP0);
+    Rational b = Dot(P1mP0, Q1mQ0);
+    Rational c = Dot(Q1mQ0, Q1mQ0);
+    Rational d = Dot(P1mP0, P0mQ0);
+    Rational e = Dot(Q1mQ0, P0mQ0);
+    Rational const zero = (Rational)0;
+    Rational const one = (Rational)1;
+    Rational det = a * c - b * b;
+    Rational s, t, nd, bmd, bte, ctd, bpe, ate, btd;
+
+    if (det > zero)
+    {
+        bte = b * e;
+        ctd = c * d;
+        if (bte <= ctd)  // s <= 0
+        {
+            s = zero;
+            if (e <= zero)  // t <= 0
+            {
+                // region 6
+                t = zero;
+                nd = -d;
+                if (nd >= a)
+                {
+                    s = one;
+                }
+                else if (nd > zero)
+                {
+                    s = nd / a;
+                }
+                // else: s is already zero
+            }
+            else if (e < c)  // 0 < t < 1
+            {
+                // region 5
+                t = e / c;
+            }
+            else  // t >= 1
+            {
+                // region 4
+                t = one;
+                bmd = b - d;
+                if (bmd >= a)
+                {
+                    s = one;
+                }
+                else if (bmd > zero)
+                {
+                    s = bmd / a;
+                }
+                // else:  s is already zero
+            }
+        }
+        else  // s > 0
+        {
+            s = bte - ctd;
+            if (s >= det)  // s >= 1
+            {
+                // s = 1
+                s = one;
+                bpe = b + e;
+                if (bpe <= zero)  // t <= 0
+                {
+                    // region 8
+                    t = zero;
+                    nd = -d;
+                    if (nd <= zero)
+                    {
+                        s = zero;
+                    }
+                    else if (nd < a)
+                    {
+                        s = nd / a;
+                    }
+                    // else: s is already one
+                }
+                else if (bpe < c)  // 0 < t < 1
+                {
+                    // region 1
+                    t = bpe / c;
+                }
+                else  // t >= 1
+                {
+                    // region 2
+                    t = one;
+                    bmd = b - d;
+                    if (bmd <= zero)
+                    {
+                        s = zero;
+                    }
+                    else if (bmd < a)
+                    {
+                        s = bmd / a;
+                    }
+                    // else:  s is already one
+                }
+            }
+            else  // 0 < s < 1
+            {
+                ate = a * e;
+                btd = b * d;
+                if (ate <= btd)  // t <= 0
+                {
+                    // region 7
+                    t = zero;
+                    nd = -d;
+                    if (nd <= zero)
+                    {
+                        s = zero;
+                    }
+                    else if (nd >= a)
+                    {
+                        s = one;
+                    }
+                    else
+                    {
+                        s = nd / a;
+                    }
+                }
+                else  // t > 0
+                {
+                    t = ate - btd;
+                    if (t >= det)  // t >= 1
+                    {
+                        // region 3
+                        t = one;
+                        bmd = b - d;
+                        if (bmd <= zero)
+                        {
+                            s = zero;
+                        }
+                        else if (bmd >= a)
+                        {
+                            s = one;
+                        }
+                        else
+                        {
+                            s = bmd / a;
+                        }
+                    }
+                    else  // 0 < t < 1
+                    {
+                        // region 0
+                        s /= det;
+                        t /= det;
+                    }
+                }
+            }
+        }
+    }
+    else
+    {
+        // The segments are parallel.  The quadratic factors to R(s,t) =
+        // a*(s-(b/a)*t)^2 + 2*d*(s - (b/a)*t) + f, where a*c = b^2,
+        // e = b*d/a, f = |P0-Q0|^2, and b is not zero.  R is constant along
+        // lines of the form s-(b/a)*t = k, and the minimum of R occurs on the
+        // line a*s - b*t + d = 0.  This line must intersect both the s-axis
+        // and the t-axis because 'a' and 'b' are not zero.  Because of
+        // parallelism, the line is also represented by -b*s + c*t - e = 0.
+        //
+        // The code determines an edge of the domain [0,1]^2 that intersects
+        // the minimum line, or if none of the edges intersect, it determines
+        // the closest corner to the minimum line.  The conditionals are
+        // designed to test first for intersection with the t-axis (s = 0)
+        // using -b*s + c*t - e = 0 and then with the s-axis (t = 0) using
+        // a*s - b*t + d = 0.
+
+        // When s = 0, solve c*t - e = 0 (t = e/c).
+        if (e <= zero)  // t <= 0
+        {
+            // Now solve a*s - b*t + d = 0 for t = 0 (s = -d/a).
+            t = zero;
+            nd = -d;
+            if (nd <= zero)  // s <= 0
+            {
+                // region 6
+                s = zero;
+            }
+            else if (nd >= a)  // s >= 1
+            {
+                // region 8
+                s = one;
+            }
+            else  // 0 < s < 1
+            {
+                // region 7
+                s = nd / a;
+            }
+        }
+        else if (e >= c)  // t >= 1
+        {
+            // Now solve a*s - b*t + d = 0 for t = 1 (s = (b-d)/a).
+            t = one;
+            bmd = b - d;
+            if (bmd <= zero)  // s <= 0
+            {
+                // region 4
+                s = zero;
+            }
+            else if (bmd >= a)  // s >= 1
+            {
+                // region 2
+                s = one;
+            }
+            else  // 0 < s < 1
+            {
+                // region 3
+                s = bmd / a;
+            }
+        }
+        else  // 0 < t < 1
+        {
+            // The point (0,e/c) is on the line and domain, so we have one
+            // point at which R is a minimum.
+            s = zero;
+            t = e / c;
+        }
+    }
+
+    Result result;
+    result.parameter[0] = s;
+    result.parameter[1] = t;
+    result.closest[0] = P0 + s * P1mP0;
+    result.closest[1] = Q0 + t * Q1mQ0;
+    Vector<N, Rational> diff = result.closest[1] - result.closest[0];
+    result.sqrDistance = Dot(diff, diff);
+    return result;
+}
+
+
+}
