move GeometryProcessing plugin from NotForLicensees to Experimental

diff --git a/Engine/Plugins/Experimental/GeometryProcessing/Source/GeometryAlgorithms/Private/GTEngine/Mathematics/GteDistTriangle3AlignedBox3.h b/Engine/Plugins/Experimental/GeometryProcessing/Source/GeometryAlgorithms/Private/GTEngine/Mathematics/GteDistTriangle3AlignedBox3.h
new file mode 100644
--- /dev/null
+++ b/Engine/Plugins/Experimental/GeometryProcessing/Source/GeometryAlgorithms/Private/GTEngine/Mathematics/GteDistTriangle3AlignedBox3.h
@@ -0,0 +1,136 @@
+// David Eberly, Geometric Tools, Redmond WA 98052
+// Copyright (c) 1998-2019
+// Distributed under the Boost Software License, Version 1.0.
+// http://www.boost.org/LICENSE_1_0.txt
+// http://www.geometrictools.com/License/Boost/LICENSE_1_0.txt
+// File Version: 3.5.2 (2018/10/05)
+
+#pragma once
+
+#include <GTEngine/Mathematics/GteDCPQuery.h>
+#include <GTEngine/Mathematics/GteLCPSolver.h>
+#include <GTEngine/Mathematics/GteAlignedBox.h>
+#include <GTEngine/Mathematics/GteTriangle.h>
+#include <GTEngine/Mathematics/GteVector3.h>
+
+namespace gte
+{
+
+template <typename Real>
+class DCPQuery<Real, Triangle3<Real>, AlignedBox3<Real>>
+{
+public:
+    struct Result
+    {
+        bool queryIsSuccessful;
+
+        // These members are valid only when queryIsSuccessful is true;
+        // otherwise, they are all set to zero.
+        Real distance, sqrDistance;
+        std::array<Real, 3> triangleParameter, boxParameter;
+        Vector3<Real> closestPoint[2];
+
+        // The number of iterations used by LCPSolver regardless of whether
+        // the query is successful.
+        int numLCPIterations;
+    };
+
+    // The default maximum iterations is 81 (n = 9, maxIterations = n*n).  If
+    // the solver fails to converge, try increasing the maximum number of
+    // iterations.
+    void SetMaxLCPIterations(int maxLCPIterations);
+
+    Result operator()(Triangle3<Real> const& triangle, AlignedBox3<Real> const& box);
+
+private:
+    LCPSolver<Real, 9> mLCP;
+};
+
+
+template <typename Real>
+void DCPQuery<Real, Triangle3<Real>, AlignedBox3<Real>>::SetMaxLCPIterations(int maxLCPIterations)
+{
+    mLCP.SetMaxIterations(maxLCPIterations);
+}
+
+template <typename Real>
+typename DCPQuery<Real, Triangle3<Real>, AlignedBox3<Real>>::Result
+DCPQuery<Real, Triangle3<Real>, AlignedBox3<Real>>::operator()(
+    Triangle3<Real> const& triangle, AlignedBox3<Real> const& box)
+{
+    Result result;
+
+    // Translate the triangle and aligned box so that the aligned box becomes
+    // a canonical box.
+    Vector3<Real> K = box.max - box.min;
+    Vector3<Real> V = triangle.v[0] - box.min;
+    Vector3<Real> E0 = triangle.v[1] - triangle.v[0];
+    Vector3<Real> E1 = triangle.v[2] - triangle.v[0];
+
+    // Compute quantities to initialize q and M in the LCP.
+    Real dotVE0 = Dot(V, E0);
+    Real dotVE1 = Dot(V, E1);
+    Real dotE0E0 = Dot(E0, E0);
+    Real dotE0E1 = Dot(E0, E1);
+    Real dotE1E1 = Dot(E1, E1);
+
+    // The LCP has 5 variables and 4 (nontrivial) inequality constraints.
+    std::array<Real, 9> q =
+    {
+        -V[0], -V[1], -V[2], dotVE0, dotVE1, K[0], K[1], K[2], (Real)1
+    };
+
+    std::array<std::array<Real, 9>, 9> M;
+    M[0] = { (Real)1, (Real)0, (Real)0, -E0[0], -E1[0], (Real)1, (Real)0, (Real)0, (Real)0 };
+    M[1] = { (Real)0, (Real)1, (Real)0, -E0[1], -E1[1], (Real)0, (Real)1, (Real)0, (Real)0 };
+    M[2] = { (Real)0, (Real)0, (Real)1, -E0[2], -E1[2], (Real)0, (Real)0, (Real)1, (Real)0 };
+    M[3] = { -E0[0], -E0[1], -E0[2], dotE0E0, dotE0E1, (Real)0, (Real)0, (Real)0, (Real)1 };
+    M[4] = { -E1[0], -E1[1], -E1[2], dotE0E1, dotE1E1, (Real)0, (Real)0, (Real)0, (Real)1 };
+    M[5] = { (Real)-1, (Real)0, (Real)0, (Real)0, (Real)0, (Real)0, (Real)0, (Real)0, (Real)0 };
+    M[6] = { (Real)0, (Real)-1, (Real)0, (Real)0, (Real)0, (Real)0, (Real)0, (Real)0, (Real)0 };
+    M[7] = { (Real)0, (Real)0, (Real)-1, (Real)0, (Real)0, (Real)0, (Real)0, (Real)0, (Real)0 };
+    M[8] = { (Real)0, (Real)0, (Real)0, (Real)-1, (Real)-1, (Real)0, (Real)0, (Real)0, (Real)0 };
+
+    std::array<Real, 9> w, z;
+    if (mLCP.Solve(q, M, w, z))
+    {
+        result.queryIsSuccessful = true;
+
+        result.triangleParameter[0] = (Real)1 - z[3] - z[4];
+        result.triangleParameter[1] = z[3];
+        result.triangleParameter[2] = z[4];
+        result.closestPoint[0] = triangle.v[0] + z[3] * E0 + z[4] * E1;
+        for (int i = 0; i < 3; ++i)
+        {
+            result.boxParameter[i] = z[i] + box.min[i];
+            result.closestPoint[1][i] = result.boxParameter[i];
+        }
+
+        Vector3<Real> diff = result.closestPoint[1] - result.closestPoint[0];
+        result.sqrDistance = Dot(diff, diff);
+        result.distance = std::sqrt(result.sqrDistance);
+    }
+    else
+    {
+        // If you reach this case, the maximum number of iterations was not
+        // specified to be large enough or there is a problem due to
+        // floating-point rounding errors.  If you believe the latter is
+        // true, file a bug report.
+        result.queryIsSuccessful = false;
+
+        for (int i = 0; i < 3; ++i)
+        {
+            result.triangleParameter[i] = (Real)0;
+            result.boxParameter[i] = (Real)0;
+            result.closestPoint[0][i] = (Real)0;
+            result.closestPoint[1][i] = (Real)0;
+        }
+        result.distance = (Real)0;
+        result.sqrDistance = (Real)0;
+    }
+
+    result.numLCPIterations = mLCP.GetNumIterations();
+    return result;
+}
+
+}
diff --git a/Engine/Plugins/Experimental/GeometryProcessing/Source/GeometryAlgorithms/Private/GTEngine/Mathematics/GteDistTriangle3OrientedBox3.h b/Engine/Plugins/Experimental/GeometryProcessing/Source/GeometryAlgorithms/Private/GTEngine/Mathematics/GteDistTriangle3OrientedBox3.h
new file mode 100644
--- /dev/null
+++ b/Engine/Plugins/Experimental/GeometryProcessing/Source/GeometryAlgorithms/Private/GTEngine/Mathematics/GteDistTriangle3OrientedBox3.h
@@ -0,0 +1,144 @@
+// David Eberly, Geometric Tools, Redmond WA 98052
+// Copyright (c) 1998-2019
+// Distributed under the Boost Software License, Version 1.0.
+// http://www.boost.org/LICENSE_1_0.txt
+// http://www.geometrictools.com/License/Boost/LICENSE_1_0.txt
+// File Version: 3.5.2 (2018/10/05)
+
+#pragma once
+
+#include <GTEngine/Mathematics/GteDCPQuery.h>
+#include <GTEngine/Mathematics/GteLCPSolver.h>
+#include <GTEngine/Mathematics/GteOrientedBox.h>
+#include <GTEngine/Mathematics/GteTriangle.h>
+#include <GTEngine/Mathematics/GteVector3.h>
+
+namespace gte
+{
+
+template <typename Real>
+class DCPQuery<Real, Triangle3<Real>, OrientedBox3<Real>>
+{
+public:
+    struct Result
+    {
+        bool queryIsSuccessful;
+
+        // These members are valid only when queryIsSuccessful is true;
+        // otherwise, they are all set to zero.
+        Real distance, sqrDistance;
+        std::array<Real, 3> triangleParameter, boxParameter;
+        Vector3<Real> closestPoint[2];
+
+        // The number of iterations used by LCPSolver regardless of whether
+        // the query is successful.
+        int numLCPIterations;
+    };
+
+    // The default maximum iterations is 81 (n = 9, maxIterations = n*n).  If
+    // the solver fails to converge, try increasing the maximum number of
+    // iterations.
+    void SetMaxLCPIterations(int maxLCPIterations);
+
+    Result operator()(Triangle3<Real> const& triangle, OrientedBox3<Real> const& box);
+
+private:
+    LCPSolver<Real, 9> mLCP;
+};
+
+
+template <typename Real>
+void DCPQuery<Real, Triangle3<Real>, OrientedBox3<Real>>::SetMaxLCPIterations(int maxLCPIterations)
+{
+    mLCP.SetMaxIterations(maxLCPIterations);
+}
+
+template <typename Real>
+typename DCPQuery<Real, Triangle3<Real>, OrientedBox3<Real>>::Result
+DCPQuery<Real, Triangle3<Real>, OrientedBox3<Real>>::operator()(
+    Triangle3<Real> const& triangle, OrientedBox3<Real> const& box)
+{
+    Result result;
+
+    // Rigidly transform the triangle and oriented box so that the oriented
+    // box becomes a canonical box.
+    Vector3<Real> K = box.extent * (Real)2;
+    Vector3<Real> tempV = triangle.v[0] - box.center;
+    Vector3<Real> tempE0 = triangle.v[1] - triangle.v[0];
+    Vector3<Real> tempE1 = triangle.v[2] - triangle.v[0];
+    Vector3<Real> V, E0, E1;
+    for (int i = 0; i < 3; ++i)
+    {
+        V[i] = Dot(box.axis[i], tempV) + box.extent[i];
+        E0[i] = Dot(box.axis[i], tempE0);
+        E1[i] = Dot(box.axis[i], tempE1);
+    }
+
+    // Compute quantities to initialize q and M in the LCP.
+    Real dotVE0 = Dot(V, E0);
+    Real dotVE1 = Dot(V, E1);
+    Real dotE0E0 = Dot(E0, E0);
+    Real dotE0E1 = Dot(E0, E1);
+    Real dotE1E1 = Dot(E1, E1);
+
+    // The LCP has 5 variables and 4 (nontrivial) inequality constraints.
+    std::array<Real, 9> q =
+    {
+        -V[0], -V[1], -V[2], dotVE0, dotVE1, K[0], K[1], K[2], (Real)1
+    };
+
+    std::array<std::array<Real, 9>, 9> M;
+    M[0] = { (Real)1, (Real)0, (Real)0, -E0[0], -E1[0], (Real)1, (Real)0, (Real)0, (Real)0 };
+    M[1] = { (Real)0, (Real)1, (Real)0, -E0[1], -E1[1], (Real)0, (Real)1, (Real)0, (Real)0 };
+    M[2] = { (Real)0, (Real)0, (Real)1, -E0[2], -E1[2], (Real)0, (Real)0, (Real)1, (Real)0 };
+    M[3] = { -E0[0], -E0[1], -E0[2], dotE0E0, dotE0E1, (Real)0, (Real)0, (Real)0, (Real)1 };
+    M[4] = { -E1[0], -E1[1], -E1[2], dotE0E1, dotE1E1, (Real)0, (Real)0, (Real)0, (Real)1 };
+    M[5] = { (Real)-1, (Real)0, (Real)0, (Real)0, (Real)0, (Real)0, (Real)0, (Real)0, (Real)0 };
+    M[6] = { (Real)0, (Real)-1, (Real)0, (Real)0, (Real)0, (Real)0, (Real)0, (Real)0, (Real)0 };
+    M[7] = { (Real)0, (Real)0, (Real)-1, (Real)0, (Real)0, (Real)0, (Real)0, (Real)0, (Real)0 };
+    M[8] = { (Real)0, (Real)0, (Real)0, (Real)-1, (Real)-1, (Real)0, (Real)0, (Real)0, (Real)0 };
+
+    std::array<Real, 9> w, z;
+    if (mLCP.Solve(q, M, w, z))
+    {
+        result.queryIsSuccessful = true;
+
+        result.triangleParameter[0] = (Real)1 - z[3] - z[4];
+        result.triangleParameter[1] = z[3];
+        result.triangleParameter[2] = z[4];
+        result.closestPoint[0] = triangle.v[0] + z[3] * tempE0 + z[4] * tempE1;
+        result.closestPoint[1] = box.center;
+        for (int i = 0; i < 3; ++i)
+        {
+            result.boxParameter[i] = z[i] - box.extent[i];
+            result.closestPoint[1] += result.boxParameter[i] * box.axis[i];
+        }
+
+        Vector3<Real> diff = result.closestPoint[1] - result.closestPoint[0];
+        result.sqrDistance = Dot(diff, diff);
+        result.distance = std::sqrt(result.sqrDistance);
+    }
+    else
+    {
+        // If you reach this case, the maximum number of iterations was not
+        // specified to be large enough or there is a problem due to
+        // floating-point rounding errors.  If you believe the latter is
+        // true, file a bug report.
+        result.queryIsSuccessful = false;
+
+        for (int i = 0; i < 3; ++i)
+        {
+            result.triangleParameter[i] = (Real)0;
+            result.boxParameter[i] = (Real)0;
+            result.closestPoint[0][i] = (Real)0;
+            result.closestPoint[1][i] = (Real)0;
+        }
+        result.distance = (Real)0;
+        result.sqrDistance = (Real)0;
+    }
+
+    result.numLCPIterations = mLCP.GetNumIterations();
+    return result;
+}
+
+}
diff --git a/Engine/Plugins/Experimental/GeometryProcessing/Source/GeometryAlgorithms/Private/GTEngine/Mathematics/GteDistTriangle3Rectangle3.h b/Engine/Plugins/Experimental/GeometryProcessing/Source/GeometryAlgorithms/Private/GTEngine/Mathematics/GteDistTriangle3Rectangle3.h
new file mode 100644
--- /dev/null
+++ b/Engine/Plugins/Experimental/GeometryProcessing/Source/GeometryAlgorithms/Private/GTEngine/Mathematics/GteDistTriangle3Rectangle3.h
@@ -0,0 +1,99 @@
+// David Eberly, Geometric Tools, Redmond WA 98052
+// Copyright (c) 1998-2019
+// Distributed under the Boost Software License, Version 1.0.
+// http://www.boost.org/LICENSE_1_0.txt
+// http://www.geometrictools.com/License/Boost/LICENSE_1_0.txt
+// File Version: 3.0.0 (2016/06/19)
+
+#pragma once
+
+#include <GTEngine/Mathematics/GteDistSegment3Rectangle3.h>
+#include <GTEngine/Mathematics/GteDistSegment3Triangle3.h>
+
+namespace gte
+{
+
+template <typename Real>
+class DCPQuery<Real, Triangle3<Real>, Rectangle3<Real>>
+{
+public:
+    struct Result
+    {
+        Real distance, sqrDistance;
+        Real triangleParameter[3], rectangleParameter[2];
+        Vector3<Real> closestPoint[2];
+    };
+
+    Result operator()(Triangle3<Real> const& triangle,
+        Rectangle3<Real> const& rectangle);
+};
+
+
+template <typename Real>
+typename DCPQuery<Real, Triangle3<Real>, Rectangle3<Real>>::Result
+DCPQuery<Real, Triangle3<Real>, Rectangle3<Real>>::operator()(
+    Triangle3<Real> const& triangle, Rectangle3<Real> const& rectangle)
+{
+    Result result;
+
+    result.sqrDistance = std::numeric_limits<Real>::max();
+
+    // Compare edges of triangle to the interior of rectangle.
+    for (int i0 = 2, i1 = 0; i1 < 3; i0 = i1++)
+    {
+        Vector3<Real> segCenter = ((Real)0.5)*(triangle.v[i0] +
+            triangle.v[i1]);
+        Vector3<Real> segDirection = triangle.v[i1] - triangle.v[i0];
+        Real segExtent = ((Real)0.5)*Normalize(segDirection);
+        Segment3<Real> edge(segCenter, segDirection, segExtent);
+
+        DCPQuery<Real, Segment3<Real>, Rectangle3<Real>> srQuery;
+        auto srResult = srQuery(edge, rectangle);
+        if (srResult.sqrDistance < result.sqrDistance)
+        {
+            result.distance = srResult.distance;
+            result.sqrDistance = srResult.sqrDistance;
+            Real ratio = srResult.segmentParameter / segExtent;  // in [-1,1]
+            result.triangleParameter[i0] = ((Real)0.5)*((Real)1 - ratio);
+            result.triangleParameter[i1] =
+                (Real)1 - result.triangleParameter[i0];
+            result.triangleParameter[3 - i0 - i1] = (Real)0;
+            result.rectangleParameter[0] = srResult.rectangleParameter[0];
+            result.rectangleParameter[1] = srResult.rectangleParameter[1];
+            result.closestPoint[0] = srResult.closestPoint[0];
+            result.closestPoint[1] = srResult.closestPoint[1];
+        }
+    }
+
+    // Compare edges of rectangle to the interior of triangle.
+    for (int i1 = 0; i1 < 2; ++i1)
+    {
+        for (int i0 = -1; i0 <= 1; i0 += 2)
+        {
+            Real s = i0 * rectangle.extent[1 - i1];
+            Vector3<Real> segCenter = rectangle.center +
+                s * rectangle.axis[1 - i1];
+            Segment3<Real> edge(segCenter, rectangle.axis[i1],
+                rectangle.extent[i1]);
+
+            DCPQuery<Real, Segment3<Real>, Triangle3<Real>> stQuery;
+            auto stResult = stQuery(edge, triangle);
+            if (stResult.sqrDistance < result.sqrDistance)
+            {
+                result.distance = stResult.distance;
+                result.sqrDistance = stResult.sqrDistance;
+                result.triangleParameter[0] = stResult.triangleParameter[0];
+                result.triangleParameter[1] = stResult.triangleParameter[1];
+                result.triangleParameter[2] = stResult.triangleParameter[2];
+                result.rectangleParameter[i1] = s;
+                result.rectangleParameter[1 - i1] = stResult.segmentParameter;
+                result.closestPoint[0] = stResult.closestPoint[1];
+                result.closestPoint[1] = stResult.closestPoint[0];
+            }
+        }
+    }
+    return result;
+}
+
+
+}
diff --git a/Engine/Plugins/Experimental/GeometryProcessing/Source/GeometryAlgorithms/Private/GTEngine/Mathematics/GteDistTriangle3Triangle3.h b/Engine/Plugins/Experimental/GeometryProcessing/Source/GeometryAlgorithms/Private/GTEngine/Mathematics/GteDistTriangle3Triangle3.h
new file mode 100644
--- /dev/null
+++ b/Engine/Plugins/Experimental/GeometryProcessing/Source/GeometryAlgorithms/Private/GTEngine/Mathematics/GteDistTriangle3Triangle3.h
@@ -0,0 +1,100 @@
+// David Eberly, Geometric Tools, Redmond WA 98052
+// Copyright (c) 1998-2019
+// Distributed under the Boost Software License, Version 1.0.
+// http://www.boost.org/LICENSE_1_0.txt
+// http://www.geometrictools.com/License/Boost/LICENSE_1_0.txt
+// File Version: 3.0.0 (2016/06/19)
+
+#pragma once
+
+#include <GTEngine/Mathematics/GteDistSegment3Triangle3.h>
+
+namespace gte
+{
+
+template <typename Real>
+class DCPQuery<Real, Triangle3<Real>, Triangle3<Real>>
+{
+public:
+    struct Result
+    {
+        Real distance, sqrDistance;
+        Real triangle0Parameter[3], triangle1Parameter[3];
+        Vector3<Real> closestPoint[2];
+    };
+
+    Result operator()(Triangle3<Real> const& triangle0,
+        Triangle3<Real> const& triangle1);
+};
+
+
+template <typename Real>
+typename DCPQuery<Real, Triangle3<Real>, Triangle3<Real>>::Result
+DCPQuery<Real, Triangle3<Real>, Triangle3<Real>>::operator()(
+    Triangle3<Real> const& triangle0, Triangle3<Real> const& triangle1)
+{
+    Result result;
+
+    DCPQuery<Real, Segment3<Real>, Triangle3<Real>> stQuery;
+    typename DCPQuery<Real, Segment3<Real>, Triangle3<Real>>::Result
+        stResult;
+    result.sqrDistance = std::numeric_limits<Real>::max();
+
+    // Compare edges of triangle0 to the interior of triangle1.
+    for (int i0 = 2, i1 = 0; i1 < 3; i0 = i1++)
+    {
+        Vector3<Real> segCenter = ((Real)0.5)*(triangle0.v[i0] +
+            triangle0.v[i1]);
+        Vector3<Real> segDirection = triangle0.v[i1] - triangle0.v[i0];
+        Real segExtent = ((Real)0.5)*Normalize(segDirection);
+        Segment3<Real> edge(segCenter, segDirection, segExtent);
+
+        stResult = stQuery(edge, triangle1);
+        if (stResult.sqrDistance < result.sqrDistance)
+        {
+            result.distance = stResult.distance;
+            result.sqrDistance = stResult.sqrDistance;
+            Real ratio = stResult.segmentParameter / segExtent;  // in [-1,1]
+            result.triangle0Parameter[i0] = ((Real)0.5)*((Real)1 - ratio);
+            result.triangle0Parameter[i1] =
+                (Real)1 - result.triangle0Parameter[i0];
+            result.triangle0Parameter[3 - i0 - i1] = (Real)0;
+            result.triangle1Parameter[0] = stResult.triangleParameter[0];
+            result.triangle1Parameter[1] = stResult.triangleParameter[1];
+            result.triangle1Parameter[2] = stResult.triangleParameter[2];
+            result.closestPoint[0] = stResult.closestPoint[0];
+            result.closestPoint[1] = stResult.closestPoint[1];
+        }
+    }
+
+    // Compare edges of triangle1 to the interior of triangle0.
+    for (int i0 = 2, i1 = 0; i1 < 3; i0 = i1++)
+    {
+        Vector3<Real> segCenter = ((Real)0.5)*(triangle1.v[i0] +
+            triangle1.v[i1]);
+        Vector3<Real> segDirection = triangle1.v[i1] - triangle1.v[i0];
+        Real segExtent = ((Real)0.5)*Normalize(segDirection);
+        Segment3<Real> edge(segCenter, segDirection, segExtent);
+
+        stResult = stQuery(edge, triangle0);
+        if (stResult.sqrDistance < result.sqrDistance)
+        {
+            result.distance = stResult.distance;
+            result.sqrDistance = stResult.sqrDistance;
+            Real ratio = stResult.segmentParameter / segExtent;  // in [-1,1]
+            result.triangle0Parameter[0] = stResult.triangleParameter[0];
+            result.triangle0Parameter[1] = stResult.triangleParameter[1];
+            result.triangle0Parameter[2] = stResult.triangleParameter[2];
+            result.triangle1Parameter[i0] = ((Real)0.5)*((Real)1 - ratio);
+            result.triangle1Parameter[i1] =
+                (Real)1 - result.triangle0Parameter[i0];
+            result.triangle1Parameter[3 - i0 - i1] = (Real)0;
+            result.closestPoint[0] = stResult.closestPoint[0];
+            result.closestPoint[1] = stResult.closestPoint[1];
+        }
+    }
+    return result;
+}
+
+
+}
diff --git a/Engine/Plugins/Experimental/GeometryProcessing/Source/GeometryAlgorithms/Private/GTEngine/Mathematics/GteETManifoldMesh.cpp b/Engine/Plugins/Experimental/GeometryProcessing/Source/GeometryAlgorithms/Private/GTEngine/Mathematics/GteETManifoldMesh.cpp
new file mode 100644
--- /dev/null
+++ b/Engine/Plugins/Experimental/GeometryProcessing/Source/GeometryAlgorithms/Private/GTEngine/Mathematics/GteETManifoldMesh.cpp
@@ -0,0 +1,384 @@
+// David Eberly, Geometric Tools, Redmond WA 98052
+// Copyright (c) 1998-2019
+// Distributed under the Boost Software License, Version 1.0.
+// http://www.boost.org/LICENSE_1_0.txt
+// http://www.geometrictools.com/License/Boost/LICENSE_1_0.txt
+// File Version: 3.0.1 (2017/01/02)
+
+#include <GTEngine/Mathematics/GteETManifoldMesh.h>
+#include <GTEngine/GTEnginePCH.h>
+#include <GTEngine/LowLevel/GteLogger.h>
+
+using namespace gte;
+
+ETManifoldMesh::~ETManifoldMesh()
+{
+}
+
+ETManifoldMesh::ETManifoldMesh(ECreator eCreator, TCreator tCreator)
+    :
+    mECreator(eCreator ? eCreator : CreateEdge),
+    mTCreator(tCreator ? tCreator : CreateTriangle),
+    mAssertOnNonmanifoldInsertion(true)
+{
+}
+
+ETManifoldMesh::ETManifoldMesh(ETManifoldMesh const& mesh)
+{
+    *this = mesh;
+}
+
+ETManifoldMesh& ETManifoldMesh::operator=(ETManifoldMesh const& mesh)
+{
+    Clear();
+
+    mECreator = mesh.mECreator;
+    mTCreator = mesh.mTCreator;
+    mAssertOnNonmanifoldInsertion = mesh.mAssertOnNonmanifoldInsertion;
+    for (auto const& element : mesh.mTMap)
+    {
+        Insert(element.first.V[0], element.first.V[1], element.first.V[2]);
+    }
+
+    return *this;
+}
+
+ETManifoldMesh::EMap const& ETManifoldMesh::GetEdges() const
+{
+    return mEMap;
+}
+
+ETManifoldMesh::TMap const& ETManifoldMesh::GetTriangles() const
+{
+    return mTMap;
+}
+
+bool ETManifoldMesh::AssertOnNonmanifoldInsertion(bool doAssert)
+{
+    std::swap(doAssert, mAssertOnNonmanifoldInsertion);
+    return doAssert;  // return the previous state
+}
+
+std::shared_ptr<ETManifoldMesh::Triangle> ETManifoldMesh::Insert(int v0, int v1, int v2)
+{
+    TriangleKey<true> tkey(v0, v1, v2);
+    if (mTMap.find(tkey) != mTMap.end())
+    {
+        // The triangle already exists.  Return a null pointer as a signal to
+        // the caller that the insertion failed.
+        return nullptr;
+    }
+
+    // Create the new triangle.  It will be added to mTMap at the end of the
+    // function so that if an assertion is triggered and the function returns
+    // early, the (bad) triangle will not be part of the mesh.
+    std::shared_ptr<Triangle> tri = mTCreator(v0, v1, v2);
+
+    // Add the edges to the mesh if they do not already exist.
+    for (int i0 = 2, i1 = 0; i1 < 3; i0 = i1++)
+    {
+        EdgeKey<false> ekey(tri->V[i0], tri->V[i1]);
+        std::shared_ptr<Edge> edge;
+        auto eiter = mEMap.find(ekey);
+        if (eiter == mEMap.end())
+        {
+            // This is the first time the edge is encountered.
+            edge = mECreator(tri->V[i0], tri->V[i1]);
+            mEMap[ekey] = edge;
+
+            // Update the edge and triangle.
+            edge->T[0] = tri;
+            tri->E[i0] = edge;
+        }
+        else
+        {
+            // This is the second time the edge is encountered.
+            edge = eiter->second;
+            if (!edge)
+            {
+                LogError("Unexpected condition.");
+                return nullptr;
+            }
+
+            // Update the edge.
+            if (edge->T[1].lock())
+            {
+                if (mAssertOnNonmanifoldInsertion)
+                {
+                    LogInformation("The mesh must be manifold.");
+                }
+                return nullptr;
+            }
+            edge->T[1] = tri;
+
+            // Update the adjacent triangles.
+            auto adjacent = edge->T[0].lock();
+            if (!adjacent)
+            {
+                LogError("Unexpected condition.");
+                return nullptr;
+            }
+            for (int j = 0; j < 3; ++j)
+            {
+                if (adjacent->E[j].lock() == edge)
+                {
+                    adjacent->T[j] = tri;
+                    break;
+                }
+            }
+
+            // Update the triangle.
+            tri->E[i0] = edge;
+            tri->T[i0] = adjacent;
+        }
+    }
+
+    mTMap[tkey] = tri;
+    return tri;
+}
+
+bool ETManifoldMesh::Remove(int v0, int v1, int v2)
+{
+    TriangleKey<true> tkey(v0, v1, v2);
+    auto titer = mTMap.find(tkey);
+    if (titer == mTMap.end())
+    {
+        // The triangle does not exist.
+        return false;
+    }
+
+    // Get the triangle.
+    std::shared_ptr<Triangle> tri = titer->second;
+
+    // Remove the edges and update adjacent triangles if necessary.
+    for (int i = 0; i < 3; ++i)
+    {
+        // Inform the edges the triangle is being deleted.
+        auto edge = tri->E[i].lock();
+        if (!edge)
+        {
+            // The triangle edge should be nonnull.
+            LogError("Unexpected condition.");
+            return false;
+        }
+
+        if (edge->T[0].lock() == tri)
+        {
+            // One-triangle edges always have pointer at index zero.
+            edge->T[0] = edge->T[1];
+            edge->T[1].reset();
+        }
+        else if (edge->T[1].lock() == tri)
+        {
+            edge->T[1].reset();
+        }
+        else
+        {
+            LogError("Unexpected condition.");
+            return false;
+        }
+
+        // Remove the edge if you have the last reference to it.
+        if (!edge->T[0].lock() && !edge->T[1].lock())
+        {
+            EdgeKey<false> ekey(edge->V[0], edge->V[1]);
+            mEMap.erase(ekey);
+        }
+
+        // Inform adjacent triangles the triangle is being deleted.
+        auto adjacent = tri->T[i].lock();
+        if (adjacent)
+        {
+            for (int j = 0; j < 3; ++j)
+            {
+                if (adjacent->T[j].lock() == tri)
+                {
+                    adjacent->T[j].reset();
+                    break;
+                }
+            }
+        }
+    }
+
+    mTMap.erase(tkey);
+    return true;
+}
+
+void ETManifoldMesh::Clear()
+{
+    mEMap.clear();
+    mTMap.clear();
+}
+
+bool ETManifoldMesh::IsClosed() const
+{
+    for (auto const& element : mEMap)
+    {
+        auto edge = element.second;
+        if (!edge->T[0].lock() || !edge->T[1].lock())
+        {
+            return false;
+        }
+    }
+    return true;
+}
+
+bool ETManifoldMesh::IsOriented() const
+{
+    for (auto const& element : mEMap)
+    {
+        auto edge = element.second;
+        if (edge->T[0].lock() && edge->T[1].lock())
+        {
+            // In each triangle, find the ordered edge that corresponds to the
+            // unordered edge element.first.  Also find the vertex opposite
+            // that edge.
+            bool edgePositive[2] = { false, false };
+            int vOpposite[2] = { -1, -1 };
+            for (int j = 0; j < 2; ++j)
+            {
+                auto tri = edge->T[j].lock();
+                for (int i = 0; i < 3; ++i)
+                {
+                    if (tri->V[i] == element.first.V[0])
+                    {
+                        int vNext = tri->V[(i + 1) % 3];
+                        if (vNext == element.first.V[1])
+                        {
+                            edgePositive[j] = true;
+                            vOpposite[j] = tri->V[(i + 2) % 3];
+                        }
+                        else
+                        {
+                            edgePositive[j] = false;
+                            vOpposite[j] = vNext;
+                        }
+                        break;
+                    }
+                }
+            }
+
+            // To be oriented consistently, the edges must have reversed
+            // ordering and the oppositive vertices cannot match.
+            if (edgePositive[0] == edgePositive[1] || vOpposite[0] == vOpposite[1])
+            {
+                return false;
+            }
+        }
+    }
+    return true;
+}
+
+void ETManifoldMesh::GetComponents(
+    std::vector<std::vector<std::shared_ptr<Triangle>>>& components) const
+{
+    // visited: 0 (unvisited), 1 (discovered), 2 (finished)
+    std::map<std::shared_ptr<Triangle>, int> visited;
+    for (auto const& element : mTMap)
+    {
+        visited.insert(std::make_pair(element.second, 0));
+    }
+
+    for (auto& element : mTMap)
+    {
+        auto tri = element.second;
+        if (visited[tri] == 0)
+        {
+            std::vector<std::shared_ptr<Triangle>> component;
+            DepthFirstSearch(tri, visited, component);
+            components.push_back(component);
+        }
+    }
+}
+
+void ETManifoldMesh::GetComponents(
+    std::vector<std::vector<TriangleKey<true>>>& components) const
+{
+    // visited: 0 (unvisited), 1 (discovered), 2 (finished)
+    std::map<std::shared_ptr<Triangle>, int> visited;
+    for (auto const& element : mTMap)
+    {
+        visited.insert(std::make_pair(element.second, 0));
+    }
+
+    for (auto& element : mTMap)
+    {
+        std::shared_ptr<Triangle> tri = element.second;
+        if (visited[tri] == 0)
+        {
+            std::vector<std::shared_ptr<Triangle>> component;
+            DepthFirstSearch(tri, visited, component);
+
+            std::vector<TriangleKey<true>> keyComponent;
+            keyComponent.reserve(component.size());
+            for (auto const& t : component)
+            {
+                keyComponent.push_back(TriangleKey<true>(t->V[0], t->V[1], t->V[2]));
+            }
+            components.push_back(keyComponent);
+        }
+    }
+}
+
+void ETManifoldMesh::DepthFirstSearch(std::shared_ptr<Triangle> const& tInitial,
+    std::map<std::shared_ptr<Triangle>, int>& visited,
+    std::vector<std::shared_ptr<Triangle>>& component) const
+{
+    // Allocate the maximum-size stack that can occur in the depth-first
+    // search.  The stack is empty when the index top is -1.
+    std::vector<std::shared_ptr<Triangle>> tStack(mTMap.size());
+    int top = -1;
+    tStack[++top] = tInitial;
+    while (top >= 0)
+    {
+        std::shared_ptr<Triangle> tri = tStack[top];
+        visited[tri] = 1;
+        int i;
+        for (i = 0; i < 3; ++i)
+        {
+            std::shared_ptr<Triangle> adj = tri->T[i].lock();
+            if (adj && visited[adj] == 0)
+            {
+                tStack[++top] = adj;
+                break;
+            }
+        }
+        if (i == 3)
+        {
+            visited[tri] = 2;
+            component.push_back(tri);
+            --top;
+        }
+    }
+}
+
+std::shared_ptr<ETManifoldMesh::Edge> ETManifoldMesh::CreateEdge(int v0, int v1)
+{
+    return std::make_shared<Edge>(v0, v1);
+}
+
+std::shared_ptr<ETManifoldMesh::Triangle> ETManifoldMesh::CreateTriangle(int v0, int v1, int v2)
+{
+    return std::make_shared<Triangle>(v0, v1, v2);
+}
+
+ETManifoldMesh::Edge::~Edge()
+{
+}
+
+ETManifoldMesh::Edge::Edge(int v0, int v1)
+{
+    V[0] = v0;
+    V[1] = v1;
+}
+
+ETManifoldMesh::Triangle::~Triangle()
+{
+}
+
+ETManifoldMesh::Triangle::Triangle(int v0, int v1, int v2)
+{
+    V[0] = v0;
+    V[1] = v1;
+    V[2] = v2;
+}
diff --git a/Engine/Plugins/Experimental/GeometryProcessing/Source/GeometryAlgorithms/Private/GTEngine/Mathematics/GteETManifoldMesh.h b/Engine/Plugins/Experimental/GeometryProcessing/Source/GeometryAlgorithms/Private/GTEngine/Mathematics/GteETManifoldMesh.h
new file mode 100644
--- /dev/null
+++ b/Engine/Plugins/Experimental/GeometryProcessing/Source/GeometryAlgorithms/Private/GTEngine/Mathematics/GteETManifoldMesh.h
@@ -0,0 +1,145 @@
+// David Eberly, Geometric Tools, Redmond WA 98052
+// Copyright (c) 1998-2019
+// Distributed under the Boost Software License, Version 1.0.
+// http://www.boost.org/LICENSE_1_0.txt
+// http://www.geometrictools.com/License/Boost/LICENSE_1_0.txt
+// File Version: 3.0.1 (2017/01/02)
+
+#pragma once
+
+#include <GTEngine/Mathematics/GteEdgeKey.h>
+#include <GTEngine/Mathematics/GteTriangleKey.h>
+#include <map>
+#include <memory>
+#include <vector>
+
+namespace gte
+{
+
+class GTE_IMPEXP ETManifoldMesh
+{
+public:
+    // Edge data types.
+    class Edge;
+    typedef std::shared_ptr<Edge> (*ECreator)(int, int);
+    typedef std::map<EdgeKey<false>, std::shared_ptr<Edge>> EMap;
+
+    // Triangle data types.
+    class Triangle;
+    typedef std::shared_ptr<Triangle> (*TCreator)(int, int, int);
+    typedef std::map<TriangleKey<true>, std::shared_ptr<Triangle>> TMap;
+
+    // Edge object.
+    class GTE_IMPEXP Edge
+    {
+    public:
+        virtual ~Edge();
+        Edge(int v0, int v1);
+
+        // Vertices of the edge.
+        int V[2];
+
+        // Triangles sharing the edge.
+        std::weak_ptr<Triangle> T[2];
+    };
+
+    // Triangle object.
+    class GTE_IMPEXP Triangle
+    {
+    public:
+        virtual ~Triangle();
+        Triangle(int v0, int v1, int v2);
+
+        // Vertices, listed in counterclockwise order (V[0],V[1],V[2]).
+        int V[3];
+
+        // Adjacent edges.  E[i] points to edge (V[i],V[(i+1)%3]).
+        std::weak_ptr<Edge> E[3];
+
+        // Adjacent triangles.  T[i] points to the adjacent triangle
+        // sharing edge E[i].
+        std::weak_ptr<Triangle> T[3];
+    };
+
+
+    // Construction and destruction.
+    virtual ~ETManifoldMesh();
+    ETManifoldMesh(ECreator eCreator = nullptr, TCreator tCreator = nullptr);
+
+    // Support for a deep copy of the mesh.  The mEMap and mTMap objects have
+    // dynamically allocated memory for edges and triangles.  A shallow copy
+    // of the pointers to this memory is problematic.  Allowing sharing, say,
+    // via std::shared_ptr, is an option but not really the intent of copying
+    // the mesh graph.
+    ETManifoldMesh(ETManifoldMesh const& mesh);
+    ETManifoldMesh& operator=(ETManifoldMesh const& mesh);
+
+    // Member access.
+    EMap const& GetEdges() const;
+    TMap const& GetTriangles() const;
+
+    // If the insertion of a triangle fails because the mesh would become
+    // nonmanifold, the default behavior is to trigger a LogError message.
+    // You can disable this behavior in situations where you want the Logger
+    // system on but you want to continue gracefully without an assertion.
+    // The return value is the previous value of the internal state
+    // mAssertOnNonmanifoldInsertion.
+    bool AssertOnNonmanifoldInsertion(bool doAssert);
+
+    // If <v0,v1,v2> is not in the mesh, a Triangle object is created and
+    // returned; otherwise, <v0,v1,v2> is in the mesh and nullptr is returned.
+    // If the insertion leads to a nonmanifold mesh, the call fails with a
+    // nullptr returned.
+    virtual std::shared_ptr<Triangle> Insert(int v0, int v1, int v2);
+
+    // If <v0,v1,v2> is in the mesh, it is removed and 'true' is returned;
+    // otherwise, <v0,v1,v2> is not in the mesh and 'false' is returned.
+    virtual bool Remove(int v0, int v1, int v2);
+
+    // Destroy the edges and triangles to obtain an empty mesh.
+    virtual void Clear();
+
+    // A manifold mesh is closed if each edge is shared twice.  A closed
+    // mesh is not necessarily oriented.  For example, you could have a
+    // mesh with spherical topology.  The upper hemisphere has outer-facing
+    // normals and the lower hemisphere has inner-facing normals.  The
+    // discontinuity in orientation occurs on the circle shared by the
+    // hemispheres.
+    bool IsClosed() const;
+
+    // Test whether all triangles in the mesh are oriented consistently and
+    // that no two triangles are coincident.  The latter means that you
+    // cannot have both triangles <v0,v1,v2> and <v0,v2,v1> in the mesh to
+    // be considered oriented.
+    bool IsOriented() const;
+
+    // Compute the connected components of the edge-triangle graph that the
+    // mesh represents.  The first function returns pointers into 'this'
+    // object's containers, so you must consume the components before
+    // clearing or destroying 'this'.  The second function returns triangle
+    // keys, which requires three times as much storage as the pointers but
+    // allows you to clear or destroy 'this' before consuming the components.
+    void GetComponents(std::vector<std::vector<std::shared_ptr<Triangle>>>& components) const;
+    void GetComponents(std::vector<std::vector<TriangleKey<true>>>& components) const;
+
+protected:
+    // The edge data and default edge creation.
+    static std::shared_ptr<Edge> CreateEdge(int v0, int v1);
+    ECreator mECreator;
+    EMap mEMap;
+
+    // The triangle data and default triangle creation.
+    static std::shared_ptr<Triangle> CreateTriangle(int v0, int v1, int v2);
+    TCreator mTCreator;
+    TMap mTMap;
+    bool mAssertOnNonmanifoldInsertion;  // default: true
+
+    // Support for computing connected components.  This is a straightforward
+    // depth-first search of the graph but uses a preallocated stack rather
+    // than a recursive function that could possibly overflow the call stack.
+    void DepthFirstSearch(std::shared_ptr<Triangle> const& tInitial,
+        std::map<std::shared_ptr<Triangle>, int>& visited,
+        std::vector<std::shared_ptr<Triangle>>& component) const;
+};
+
+}
diff --git a/Engine/Plugins/Experimental/GeometryProcessing/Source/GeometryAlgorithms/Private/GTEngine/Mathematics/GteETNonmanifoldMesh.cpp b/Engine/Plugins/Experimental/GeometryProcessing/Source/GeometryAlgorithms/Private/GTEngine/Mathematics/GteETNonmanifoldMesh.cpp
new file mode 100644
--- /dev/null
+++ b/Engine/Plugins/Experimental/GeometryProcessing/Source/GeometryAlgorithms/Private/GTEngine/Mathematics/GteETNonmanifoldMesh.cpp
@@ -0,0 +1,310 @@
+// David Eberly, Geometric Tools, Redmond WA 98052
+// Copyright (c) 1998-2018
+// Distributed under the Boost Software License, Version 1.0.
+// http://www.boost.org/LICENSE_1_0.txt
+// http://www.geometrictools.com/License/Boost/LICENSE_1_0.txt
+// File Version: 3.20.0 (2019/01/09)
+
+#include <GTEngine/Mathematics/GteETNonmanifoldMesh.h>
+#include <GTEngine/LowLevel/GteLogger.h>
+#include <GTEngine/GTEnginePCH.h>
+
+using namespace gte;
+
+ETNonmanifoldMesh::~ETNonmanifoldMesh()
+{
+}
+
+ETNonmanifoldMesh::ETNonmanifoldMesh(ECreator eCreator, TCreator tCreator)
+    :
+    mECreator(eCreator ? eCreator : CreateEdge),
+    mTCreator(tCreator ? tCreator : CreateTriangle)
+{
+}
+
+ETNonmanifoldMesh::ETNonmanifoldMesh(ETNonmanifoldMesh const& mesh)
+{
+    *this = mesh;
+}
+
+ETNonmanifoldMesh& ETNonmanifoldMesh::operator=(ETNonmanifoldMesh const& mesh)
+{
+    Clear();
+
+    mECreator = mesh.mECreator;
+    mTCreator = mesh.mTCreator;
+    for (auto const& element : mesh.mTMap)
+    {
+        Insert(element.first.V[0], element.first.V[1], element.first.V[2]);
+    }
+
+    return *this;
+}
+
+ETNonmanifoldMesh::EMap const& ETNonmanifoldMesh::GetEdges() const
+{
+    return mEMap;
+}
+
+ETNonmanifoldMesh::TMap const& ETNonmanifoldMesh::GetTriangles() const
+{
+    return mTMap;
+}
+
+std::shared_ptr<ETNonmanifoldMesh::Triangle> ETNonmanifoldMesh::Insert(int v0, int v1, int v2)
+{
+    TriangleKey<true> tkey(v0, v1, v2);
+    if (mTMap.find(tkey) != mTMap.end())
+    {
+        // The triangle already exists.  Return a null pointer as a signal to
+        // the caller that the insertion failed.
+        return nullptr;
+    }
+
+    // Create the new triangle.  It will be added to mTMap at the end of the
+    // function so that if an assertion is triggered and the function returns
+    // early, the (bad) triangle will not be part of the mesh.
+    std::shared_ptr<Triangle> tri = mTCreator(v0, v1, v2);
+
+    // Add the edges to the mesh if they do not already exist.
+    for (int i0 = 2, i1 = 0; i1 < 3; i0 = i1++)
+    {
+        EdgeKey<false> ekey(tri->V[i0], tri->V[i1]);
+        std::shared_ptr<Edge> edge;
+        auto eiter = mEMap.find(ekey);
+        if (eiter == mEMap.end())
+        {
+            // This is the first time the edge is encountered.
+            edge = mECreator(tri->V[i0], tri->V[i1]);
+            mEMap[ekey] = edge;
+        }
+        else
+        {
+            // The edge was previously encountered and created.
+            edge = eiter->second;
+            if (!edge)
+            {
+                LogError("Unexpected condition.");
+                return nullptr;
+            }
+        }
+
+        // Associate the edge with the triangle.
+        tri->E[i0] = edge;
+
+        // Update the adjacent set of triangles for the edge.
+        edge->T.insert(tri);
+    }
+
+    mTMap[tkey] = tri;
+    return tri;
+}
+
+bool ETNonmanifoldMesh::Remove(int v0, int v1, int v2)
+{
+    TriangleKey<true> tkey(v0, v1, v2);
+    auto titer = mTMap.find(tkey);
+    if (titer == mTMap.end())
+    {
+        // The triangle does not exist.
+        return false;
+    }
+
+    // Get the triangle.
+    std::shared_ptr<Triangle> tri = titer->second;
+
+    // Remove the edges and update adjacent triangles if necessary.
+    for (int i = 0; i < 3; ++i)
+    {
+        // Inform the edges the triangle is being deleted.
+        auto edge = tri->E[i].lock();
+        if (!edge)
+        {
+            LogError("Unexpected condition.");
+            return false;
+        }
+
+        // Remove the triangle from the edge's set of adjacent triangles.
+        size_t numRemoved = edge->T.erase(tri);
+        if (numRemoved == 0)
+        {
+            LogError("Unexpected condition.");
+            return false;
+        }
+
+        // Remove the edge if you have the last reference to it.
+        if (edge->T.size() == 0)
+        {
+            EdgeKey<false> ekey(edge->V[0], edge->V[1]);
+            mEMap.erase(ekey);
+        }
+    }
+
+    // Remove the triangle from the graph.
+    mTMap.erase(tkey);
+    return true;
+}
+
+void ETNonmanifoldMesh::Clear()
+{
+    mEMap.clear();
+    mTMap.clear();
+}
+
+bool ETNonmanifoldMesh::IsManifold() const
+{
+    for (auto const& element : mEMap)
+    {
+        if (element.second->T.size() > 2)
+        {
+            return false;
+        }
+    }
+    return true;
+}
+
+bool ETNonmanifoldMesh::IsClosed() const
+{
+    for (auto const& element : mEMap)
+    {
+        if (element.second->T.size() != 2)
+        {
+            return false;
+        }
+    }
+    return true;
+}
+
+void ETNonmanifoldMesh::GetComponents(std::vector<std::vector<std::shared_ptr<Triangle>>>& components) const
+{
+    // visited: 0 (unvisited), 1 (discovered), 2 (finished)
+    std::map<std::shared_ptr<Triangle>, int> visited;
+    for (auto const& element : mTMap)
+    {
+        visited.insert(std::make_pair(element.second, 0));
+    }
+
+    for (auto& element : mTMap)
+    {
+        auto tri = element.second;
+        if (visited[tri] == 0)
+        {
+            std::vector<std::shared_ptr<Triangle>> component;
+            DepthFirstSearch(tri, visited, component);
+            components.push_back(component);
+        }
+    }
+}
+
+void ETNonmanifoldMesh::GetComponents(std::vector<std::vector<TriangleKey<true>>>& components) const
+{
+    // visited: 0 (unvisited), 1 (discovered), 2 (finished)
+    std::map<std::shared_ptr<Triangle>, int> visited;
+    for (auto const& element : mTMap)
+    {
+        visited.insert(std::make_pair(element.second, 0));
+    }
+
+    for (auto& element : mTMap)
+    {
+        std::shared_ptr<Triangle> tri = element.second;
+        if (visited[tri] == 0)
+        {
+            std::vector<std::shared_ptr<Triangle>> component;
+            DepthFirstSearch(tri, visited, component);
+
+            std::vector<TriangleKey<true>> keyComponent;
+            keyComponent.reserve(component.size());
+            for (auto const& t : component)
+            {
+                keyComponent.push_back(TriangleKey<true>(t->V[0], t->V[1], t->V[2]));
+            }
+            components.push_back(keyComponent);
+        }
+    }
+}
+
+void ETNonmanifoldMesh::DepthFirstSearch(std::shared_ptr<Triangle> const& tInitial,
+    std::map<std::shared_ptr<Triangle>, int>& visited, std::vector<std::shared_ptr<Triangle>>& component) const
+{
+    // Allocate the maximum-size stack that can occur in the depth-first
+    // search.  The stack is empty when the index top is -1.
+    std::vector<std::shared_ptr<Triangle>> tStack(mTMap.size());
+    int top = -1;
+    tStack[++top] = tInitial;
+    while (top >= 0)
+    {
+        std::shared_ptr<Triangle> tri = tStack[top];
+        visited[tri] = 1;
+        int i;
+        for (i = 0; i < 3; ++i)
+        {
+            auto edge = tri->E[i].lock();
+            if (!edge)
+            {
+                LogError("Unexpected condition.");
+                return;
+            }
+
+            bool foundUnvisited = false;
+            for (auto const& adjw : edge->T)
+            {
+                auto adj = adjw.lock();
+                if (!adj)
+                {
+                    LogError("Unexpected condition.");
+                    return;
+                }
+
+                if (visited[adj] == 0)
+                {
+                    tStack[++top] = adj;
+                    foundUnvisited = true;
+                    break;
+                }
+            }
+
+            if (foundUnvisited)
+            {
+                break;
+            }
+        }
+        if (i == 3)
+        {
+            visited[tri] = 2;
+            component.push_back(tri);
+            --top;
+        }
+    }
+}
+
+std::shared_ptr<ETNonmanifoldMesh::Edge> ETNonmanifoldMesh::CreateEdge(int v0, int v1)
+{
+    return std::make_shared<Edge>(v0, v1);
+}
+
+std::shared_ptr<ETNonmanifoldMesh::Triangle> ETNonmanifoldMesh::CreateTriangle(int v0, int v1, int v2)
+{
+    return std::make_shared<Triangle>(v0, v1, v2);
+}
+
+ETNonmanifoldMesh::Edge::~Edge()
+{
+}
+
+ETNonmanifoldMesh::Edge::Edge(int v0, int v1)
+{
+    V[0] = v0;
+    V[1] = v1;
+}
+
+ETNonmanifoldMesh::Triangle::~Triangle()
+{
+}
+
+ETNonmanifoldMesh::Triangle::Triangle(int v0, int v1, int v2)
+{
+    V[0] = v0;
+    V[1] = v1;
+    V[2] = v2;
+}
diff --git a/Engine/Plugins/Experimental/GeometryProcessing/Source/GeometryAlgorithms/Private/GTEngine/Mathematics/GteETNonmanifoldMesh.h b/Engine/Plugins/Experimental/GeometryProcessing/Source/GeometryAlgorithms/Private/GTEngine/Mathematics/GteETNonmanifoldMesh.h
new file mode 100644
--- /dev/null
+++ b/Engine/Plugins/Experimental/GeometryProcessing/Source/GeometryAlgorithms/Private/GTEngine/Mathematics/GteETNonmanifoldMesh.h
@@ -0,0 +1,141 @@
+// David Eberly, Geometric Tools, Redmond WA 98052
+// Copyright (c) 1998-2018
+// Distributed under the Boost Software License, Version 1.0.
+// http://www.boost.org/LICENSE_1_0.txt
+// http://www.geometrictools.com/License/Boost/LICENSE_1_0.txt
+// File Version: 3.20.0 (2019/01/09)
+
+#pragma once
+
+#include <GTEngine/LowLevel/GteWeakPtrCompare.h>
+#include <GTEngine/Mathematics/GteEdgeKey.h>
+#include <GTEngine/Mathematics/GteTriangleKey.h>
+#include <map>
+#include <memory>
+#include <set>
+#include <vector>
+
+namespace gte
+{
+    class GTE_IMPEXP ETNonmanifoldMesh
+    {
+    public:
+        // Edge data types.
+        class Edge;
+        typedef std::shared_ptr<Edge>(*ECreator)(int, int);
+        typedef std::map<EdgeKey<false>, std::shared_ptr<Edge>> EMap;
+
+        // Triangle data types.
+        class Triangle;
+        typedef std::shared_ptr<Triangle>(*TCreator)(int, int, int);
+        typedef std::map<TriangleKey<true>, std::shared_ptr<Triangle>> TMap;
+
+        // Edge object.
+        class GTE_IMPEXP Edge
+        {
+        public:
+            virtual ~Edge();
+            Edge(int v0, int v1);
+
+            bool operator<(Edge const& other) const
+            {
+                return EdgeKey<false>(V[0], V[1]) < EdgeKey<false>(other.V[0], other.V[1]);
+            }
+
+            // Vertices of the edge.
+            int V[2];
+
+            // Triangles sharing the edge.
+            std::set<std::weak_ptr<Triangle>, WeakPtrLT<Triangle>> T;
+        };
+
+        // Triangle object.
+        class GTE_IMPEXP Triangle
+        {
+        public:
+            virtual ~Triangle();
+            Triangle(int v0, int v1, int v2);
+
+            bool operator<(Triangle const& other) const
+            {
+                return TriangleKey<true>(V[0], V[1], V[2]) < TriangleKey<true>(other.V[0], other.V[1], other.V[2]);
+            }
+
+            // Vertices listed in counterclockwise order (V[0],V[1],V[2]).
+            int V[3];
+
+            // Adjacent edges.  E[i] points to edge (V[i],V[(i+1)%3]).
+            std::weak_ptr<Edge> E[3];
+        };
+
+
+        // Construction and destruction.
+        virtual ~ETNonmanifoldMesh();
+        ETNonmanifoldMesh(ECreator eCreator = nullptr, TCreator tCreator = nullptr);
+
+        // Support for a deep copy of the mesh.  The mEMap and mTMap objects
+        // have dynamically allocated memory for edges and triangles.  A
+        // shallow copy of the pointers to this memory is problematic.
+        // Allowing sharing, say, via std::shared_ptr, is an option but not
+        // really the intent of copying the mesh graph.
+        ETNonmanifoldMesh(ETNonmanifoldMesh const& mesh);
+        ETNonmanifoldMesh& operator=(ETNonmanifoldMesh const& mesh);
+
+        // Member access.
+        EMap const& GetEdges() const;
+        TMap const& GetTriangles() const;
+
+        // If <v0,v1,v2> is not in the mesh, a Triangle object is created and
+        // returned; otherwise, <v0,v1,v2> is in the mesh and nullptr is
+        // returned.
+        virtual std::shared_ptr<Triangle> Insert(int v0, int v1, int v2);
+
+        // If <v0,v1,v2> is in the mesh, it is removed and 'true' is returned;
+        // otherwise, <v0,v1,v2> is not in the mesh and 'false' is returned.
+        virtual bool Remove(int v0, int v1, int v2);
+
+        // Destroy the edges and triangles to obtain an empty mesh.
+        virtual void Clear();
+
+        // A manifold mesh has the property that an edge is shared by at most
+        // two triangles sharing.
+        bool IsManifold() const;
+
+        // A manifold mesh is closed if each edge is shared twice.  A closed
+        // mesh is not necessarily oriented.  For example, you could have a
+        // mesh with spherical topology. The upper hemisphere has outer-facing
+        // normals and the lower hemisphere has inner-facing normals.  The
+        // discontinuity in orientation occurs on the circle shared by the
+        // hemispheres.
+        bool IsClosed() const;
+
+        // Compute the connected components of the edge-triangle graph that
+        // the mesh represents.  The first function returns pointers into
+        // 'this' object's containers, so you must consume the components
+        // before clearing or destroying 'this'.  The second function returns
+        // triangle keys, which requires three times as much storage as the
+        // pointers but allows you to clear or destroy 'this' before consuming
+        // the components.
+        void GetComponents(std::vector<std::vector<std::shared_ptr<Triangle>>>& components) const;
+        void GetComponents(std::vector<std::vector<TriangleKey<true>>>& components) const;
+
+    protected:
+        // The edge data and default edge creation.
+        static std::shared_ptr<Edge> CreateEdge(int v0, int v1);
+        ECreator mECreator;
+        EMap mEMap;
+
+        // The triangle data and default triangle creation.
+        static std::shared_ptr<Triangle> CreateTriangle(int v0, int v1, int v2);
+        TCreator mTCreator;
+        TMap mTMap;
+
+        // Support for computing connected components.  This is a
+        // straightforward depth-first search of the graph but uses a
+        // preallocated stack rather than a recursive function that could
+        // possibly overflow the call stack.
+        void DepthFirstSearch(std::shared_ptr<Triangle> const& tInitial,
+            std::map<std::shared_ptr<Triangle>, int>& visited,
+            std::vector<std::shared_ptr<Triangle>>& component) const;
+    };
+}
diff --git a/Engine/Plugins/Experimental/GeometryProcessing/Source/GeometryAlgorithms/Private/GTEngine/Mathematics/GteEdgeKey.h b/Engine/Plugins/Experimental/GeometryProcessing/Source/GeometryAlgorithms/Private/GTEngine/Mathematics/GteEdgeKey.h
new file mode 100644
--- /dev/null
+++ b/Engine/Plugins/Experimental/GeometryProcessing/Source/GeometryAlgorithms/Private/GTEngine/Mathematics/GteEdgeKey.h
@@ -0,0 +1,65 @@
+// David Eberly, Geometric Tools, Redmond WA 98052
+// Copyright (c) 1998-2019
+// Distributed under the Boost Software License, Version 1.0.
+// http://www.boost.org/LICENSE_1_0.txt
+// http://www.geometrictools.com/License/Boost/LICENSE_1_0.txt
+// File Version: 3.0.1 (2018/05/22)
+
+#pragma once
+
+#include <GTEngine/Mathematics/GteFeatureKey.h>
+
+namespace gte
+{
+    template <bool Ordered>
+    class EdgeKey : public FeatureKey<2, Ordered>
+    {
+    public:
+        // An ordered edge has (V[0],V[1]) = (v0,v1).  An unordered edge has
+        // (V[0],V[1]) = (min(V[0],V[1]),max(V[0],V[1])).
+        EdgeKey();  // creates key (-1,-1)
+        explicit EdgeKey(int v0, int v1);
+    };
+
+	template<>
+	inline
+	EdgeKey<true>::EdgeKey()
+	{
+		V[0] = -1;
+		V[1] = -1;
+	}
+
+	template<>
+	inline
+	EdgeKey<true>::EdgeKey(int v0, int v1)
+	{
+		V[0] = v0;
+		V[1] = v1;
+	}
+
+	template<>
+	inline
+	EdgeKey<false>::EdgeKey()
+	{
+		V[0] = -1;
+		V[1] = -1;
+	}
+
+	template<>
+	inline
+	EdgeKey<false>::EdgeKey(int v0, int v1)
+	{
+		if (v0 < v1)
+		{
+			// v0 is minimum
+			V[0] = v0;
+			V[1] = v1;
+		}
+		else
+		{
+			// v1 is minimum
+			V[0] = v1;
+			V[1] = v0;
+		}
+	}
+}
diff --git a/Engine/Plugins/Experimental/GeometryProcessing/Source/GeometryAlgorithms/Private/GTEngine/Mathematics/GteEllipse3.h b/Engine/Plugins/Experimental/GeometryProcessing/Source/GeometryAlgorithms/Private/GTEngine/Mathematics/GteEllipse3.h
new file mode 100644
--- /dev/null
+++ b/Engine/Plugins/Experimental/GeometryProcessing/Source/GeometryAlgorithms/Private/GTEngine/Mathematics/GteEllipse3.h
@@ -0,0 +1,157 @@
+// David Eberly, Geometric Tools, Redmond WA 98052
+// Copyright (c) 1998-2019
+// Distributed under the Boost Software License, Version 1.0.
+// http://www.boost.org/LICENSE_1_0.txt
+// http://www.geometrictools.com/License/Boost/LICENSE_1_0.txt
+// File Version: 3.0.0 (2016/06/19)
+
+#pragma once
+
+#include <GTEngine/Mathematics/GteVector2.h>
+#include <GTEngine/Mathematics/GteVector3.h>
+
+// The plane containing ellipse is Dot(N,X-C) = 0 where X is any point in the
+// plane, C is the ellipse center, and N is a unit-length normal to the plane.
+// Vectors A0, A1, and N form an orthonormal right-handed set.  The ellipse in
+// the plane is parameterized by X = C + e0*cos(t)*A0 + e1*sin(t)*A1, where A0
+// is the major axis, A1 is the minor axis, and e0 and e1 are the extents
+// along those axes.  The angle t is in [-pi,pi) and e0 >= e1 > 0.
+
+namespace gte
+{
+
+template <typename Real>
+class Ellipse3
+{
+public:
+    // Construction and destruction.  The default constructor sets center to
+    // (0,0,0), A0 to (1,0,0), A1 to (0,1,0), normal to (0,0,1), e0 to 1, and
+    // e1 to 1.
+    Ellipse3(); 
+    Ellipse3(Vector3<Real> const& inCenter, Vector3<Real> const& inNormal,
+        Vector3<Real> const inAxis[2], Vector2<Real> const& inExtent);
+
+    // Public member access.
+    Vector3<Real> center, normal;
+    Vector3<Real> axis[2];
+    Vector2<Real> extent;
+
+public:
+    // Comparisons to support sorted containers.
+    bool operator==(Ellipse3 const& ellipse) const;
+    bool operator!=(Ellipse3 const& ellipse) const;
+    bool operator< (Ellipse3 const& ellipse) const;
+    bool operator<=(Ellipse3 const& ellipse) const;
+    bool operator> (Ellipse3 const& ellipse) const;
+    bool operator>=(Ellipse3 const& ellipse) const;
+};
+
+
+template <typename Real>
+Ellipse3<Real>::Ellipse3()
+    :
+    center(Vector3<Real>::Zero()),
+    normal(Vector3<Real>::Unit(2)),
+    extent({ (Real)1, (Real)1 })
+{
+    axis[0] = Vector3<Real>::Unit(0);
+    axis[1] = Vector3<Real>::Unit(1);
+}
+
+template <typename Real>
+Ellipse3<Real>::Ellipse3(Vector3<Real> const& inCenter,
+    Vector3<Real> const& inNormal, Vector3<Real> const inAxis[2],
+    Vector2<Real> const& inExtent)
+    :
+    center(inCenter),
+    normal(inNormal),
+    extent(inExtent)
+{
+    for (int i = 0; i < 2; ++i)
+    {
+        axis[i] = inAxis[i];
+    }
+}
+
+template <typename Real>
+bool Ellipse3<Real>::operator==(Ellipse3 const& ellipse) const
+{
+    return center == ellipse.center
+        && normal == ellipse.normal
+        && axis[0] == ellipse.axis[0]
+        && axis[1] == ellipse.axis[1]
+        && extent == ellipse.extent;
+}
+
+template <typename Real>
+bool Ellipse3<Real>::operator!=(Ellipse3 const& ellipse) const
+{
+    return !operator==(ellipse);
+}
+
+template <typename Real>
+bool Ellipse3<Real>::operator<(Ellipse3 const& ellipse) const
+{
+    if (center < ellipse.center)
+    {
+        return true;
+    }
+
+    if (center > ellipse.center)
+    {
+        return false;
+    }
+
+    if (normal < ellipse.normal)
+    {
+        return true;
+    }
+
+    if (normal > ellipse.normal)
+    {
+        return false;
+    }
+
+    if (axis[0] < ellipse.axis[0])
+    {
+        return true;
+    }
+
+    if (axis[0] > ellipse.axis[0])
+    {
+        return false;
+    }
+
+    if (axis[1] < ellipse.axis[1])
+    {
+        return true;
+    }
+
+    if (axis[1] > ellipse.axis[1])
+    {
+        return false;
+    }
+
+    return extent < ellipse.extent;
+}
+
+template <typename Real>
+bool Ellipse3<Real>::operator<=(Ellipse3 const& ellipse) const
+{
+    return operator<(ellipse) || operator==(ellipse);
+}
+
+template <typename Real>
+bool Ellipse3<Real>::operator>(Ellipse3 const& ellipse) const
+{
+    return !operator<=(ellipse);
+}
+
+template <typename Real>
+bool Ellipse3<Real>::operator>=(Ellipse3 const& ellipse) const
+{
+    return !operator<(ellipse);
+}
+
+
+}
diff --git a/Engine/Plugins/Experimental/GeometryProcessing/Source/GeometryAlgorithms/Private/GTEngine/Mathematics/GteEulerAngles.h b/Engine/Plugins/Experimental/GeometryProcessing/Source/GeometryAlgorithms/Private/GTEngine/Mathematics/GteEulerAngles.h
new file mode 100644
--- /dev/null
+++ b/Engine/Plugins/Experimental/GeometryProcessing/Source/GeometryAlgorithms/Private/GTEngine/Mathematics/GteEulerAngles.h
@@ -0,0 +1,82 @@
+// David Eberly, Geometric Tools, Redmond WA 98052
+// Copyright (c) 1998-2019
+// Distributed under the Boost Software License, Version 1.0.
+// http://www.boost.org/LICENSE_1_0.txt
+// http://www.geometrictools.com/License/Boost/LICENSE_1_0.txt
+// File Version: 3.0.0 (2016/06/19)
+
+#pragma once
+
+#include <GTEngine/Mathematics/GteVector.h>
+
+namespace gte
+{
+
+// Factorization into Euler angles is not necessarily unique.  Let the
+// integer indices for the axes be (N0,N1,N2), which must be in the set
+//   {(0,1,2),(0,2,1),(1,0,2),(1,2,0),(2,0,1),(2,1,0),
+//    (0,1,0),(0,2,0),(1,0,1),(1,2,1),(2,0,2),(2,1,2)}
+// Let the corresponding angles be (angleN0,angleN1,angleN2).  If the
+// result is ER_NOT_UNIQUE_SUM, then the multiple solutions occur because
+// angleN2+angleN0 is constant.  If the result is ER_NOT_UNIQUE_DIF, then
+// the multiple solutions occur because angleN2-angleN0 is constant.  In
+// either type of nonuniqueness, the function returns angleN0=0.
+enum EulerResult
+{
+    // The solution is invalid (incorrect axis indices).
+    ER_INVALID,
+
+    // The solution is unique.
+    ER_UNIQUE,
+
+    // The solution is not unique.  A sum of angles is constant.
+    ER_NOT_UNIQUE_SUM,
+
+    // The solution is not unique.  A difference of angles is constant.
+    ER_NOT_UNIQUE_DIF
+};
+
+template <typename Real>
+class EulerAngles
+{
+public:
+    EulerAngles();
+    EulerAngles(int i0, int i1, int i2, Real a0, Real a1, Real a2);
+
+    int axis[3];
+    Real angle[3];
+
+    // This member is set during conversions from rotation matrices,
+    // quaternions, or axis-angles.
+    EulerResult result;
+};
+
+
+template <typename Real>
+EulerAngles<Real>::EulerAngles()
+    :
+    result(ER_INVALID)
+{
+    for (int i = 0; i < 3; ++i)
+    {
+        axis[i] = 0;
+        angle[i] = (Real)0;
+    }
+}
+
+template <typename Real>
+EulerAngles<Real>::EulerAngles(int i0, int i1, int i2, Real a0, Real a1,
+    Real a2)
+    :
+    result(ER_UNIQUE)
+{
+    axis[0] = i0;
+    axis[1] = i1;
+    axis[2] = i2;
+    angle[0] = a0;
+    angle[1] = a1;
+    angle[2] = a2;
+}
+
+
+}
diff --git a/Engine/Plugins/Experimental/GeometryProcessing/Source/GeometryAlgorithms/Private/GTEngine/Mathematics/GteExp2Estimate.h b/Engine/Plugins/Experimental/GeometryProcessing/Source/GeometryAlgorithms/Private/GTEngine/Mathematics/GteExp2Estimate.h
new file mode 100644
--- /dev/null
+++ b/Engine/Plugins/Experimental/GeometryProcessing/Source/GeometryAlgorithms/Private/GTEngine/Mathematics/GteExp2Estimate.h
@@ -0,0 +1,154 @@
+// David Eberly, Geometric Tools, Redmond WA 98052
+// Copyright (c) 1998-2019
+// Distributed under the Boost Software License, Version 1.0.
+// http://www.boost.org/LICENSE_1_0.txt
+// http://www.geometrictools.com/License/Boost/LICENSE_1_0.txt
+// File Version: 3.0.1 (2018/10/05)
+
+#pragma once
+
+#include <GTEngine/Mathematics/GteMath.h>
+
+// Minimax polynomial approximations to 2^x.  The polynomial p(x) of
+// degree D minimizes the quantity maximum{|2^x - p(x)| : x in [0,1]}
+// over all polynomials of degree D.
+
+namespace gte
+{
+
+template <typename Real>
+class Exp2Estimate
+{
+public:
+    // The input constraint is x in [0,1].  For example,
+    //   float x; // in [0,1]
+    //   float result = Exp2Estimate<float>::Degree<3>(x);
+    template <int D>
+    inline static Real Degree(Real x);
+
+    // The input x can be any real number.  Range reduction is used to
+    // generate a value y in [0,1], call Degree(y), and combine the output
+    // with the proper exponent to obtain the approximation.  For example,
+    //   float x;  // x >= 0
+    //   float result = Exp2Estimate<float>::DegreeRR<3>(x);
+    template <int D>
+    inline static Real DegreeRR(Real x);
+
+private:
+    // Metaprogramming and private implementation to allow specialization of
+    // a template member function.
+    template <int D> struct degree {};
+    inline static Real Evaluate(degree<1>, Real t);
+    inline static Real Evaluate(degree<2>, Real t);
+    inline static Real Evaluate(degree<3>, Real t);
+    inline static Real Evaluate(degree<4>, Real t);
+    inline static Real Evaluate(degree<5>, Real t);
+    inline static Real Evaluate(degree<6>, Real t);
+    inline static Real Evaluate(degree<7>, Real t);
+};
+
+
+template <typename Real>
+template <int D>
+inline Real Exp2Estimate<Real>::Degree(Real x)
+{
+    return Evaluate(degree<D>(), x);
+}
+
+template <typename Real>
+template <int D>
+inline Real Exp2Estimate<Real>::DegreeRR(Real x)
+{
+    Real p = std::floor(x);
+    Real y = x - p;
+    Real poly = Degree<D>(y);
+    Real result = std::ldexp(poly, (int)p);
+    return result;
+}
+
+template <typename Real>
+inline Real Exp2Estimate<Real>::Evaluate(degree<1>, Real t)
+{
+    Real poly;
+    poly = (Real)GTE_C_EXP2_DEG1_C1;
+    poly = (Real)GTE_C_EXP2_DEG1_C0 + poly * t;
+    return poly;
+}
+
+template <typename Real>
+inline Real Exp2Estimate<Real>::Evaluate(degree<2>, Real t)
+{
+    Real poly;
+    poly = (Real)GTE_C_EXP2_DEG2_C2;
+    poly = (Real)GTE_C_EXP2_DEG2_C1 + poly * t;
+    poly = (Real)GTE_C_EXP2_DEG2_C0 + poly * t;
+    return poly;
+}
+
+template <typename Real>
+inline Real Exp2Estimate<Real>::Evaluate(degree<3>, Real t)
+{
+    Real poly;
+    poly = (Real)GTE_C_EXP2_DEG3_C3;
+    poly = (Real)GTE_C_EXP2_DEG3_C2 + poly * t;
+    poly = (Real)GTE_C_EXP2_DEG3_C1 + poly * t;
+    poly = (Real)GTE_C_EXP2_DEG3_C0 + poly * t;
+    return poly;
+}
+
+template <typename Real>
+inline Real Exp2Estimate<Real>::Evaluate(degree<4>, Real t)
+{
+    Real poly;
+    poly = (Real)GTE_C_EXP2_DEG4_C4;
+    poly = (Real)GTE_C_EXP2_DEG4_C3 + poly * t;
+    poly = (Real)GTE_C_EXP2_DEG4_C2 + poly * t;
+    poly = (Real)GTE_C_EXP2_DEG4_C1 + poly * t;
+    poly = (Real)GTE_C_EXP2_DEG4_C0 + poly * t;
+    return poly;
+}
+
+template <typename Real>
+inline Real Exp2Estimate<Real>::Evaluate(degree<5>, Real t)
+{
+    Real poly;
+    poly = (Real)GTE_C_EXP2_DEG5_C5;
+    poly = (Real)GTE_C_EXP2_DEG5_C4 + poly * t;
+    poly = (Real)GTE_C_EXP2_DEG5_C3 + poly * t;
+    poly = (Real)GTE_C_EXP2_DEG5_C2 + poly * t;
+    poly = (Real)GTE_C_EXP2_DEG5_C1 + poly * t;
+    poly = (Real)GTE_C_EXP2_DEG5_C0 + poly * t;
+    return poly;
+}
+
+template <typename Real>
+inline Real Exp2Estimate<Real>::Evaluate(degree<6>, Real t)
+{
+    Real poly;
+    poly = (Real)GTE_C_EXP2_DEG6_C6;
+    poly = (Real)GTE_C_EXP2_DEG6_C5 + poly * t;
+    poly = (Real)GTE_C_EXP2_DEG6_C4 + poly * t;
+    poly = (Real)GTE_C_EXP2_DEG6_C3 + poly * t;
+    poly = (Real)GTE_C_EXP2_DEG6_C2 + poly * t;
+    poly = (Real)GTE_C_EXP2_DEG6_C1 + poly * t;
+    poly = (Real)GTE_C_EXP2_DEG6_C0 + poly * t;
+    return poly;
+}
+
+template <typename Real>
+inline Real Exp2Estimate<Real>::Evaluate(degree<7>, Real t)
+{
+    Real poly;
+    poly = (Real)GTE_C_EXP2_DEG7_C7;
+    poly = (Real)GTE_C_EXP2_DEG7_C6 + poly * t;
+    poly = (Real)GTE_C_EXP2_DEG7_C5 + poly * t;
+    poly = (Real)GTE_C_EXP2_DEG7_C4 + poly * t;
+    poly = (Real)GTE_C_EXP2_DEG7_C3 + poly * t;
+    poly = (Real)GTE_C_EXP2_DEG7_C2 + poly * t;
+    poly = (Real)GTE_C_EXP2_DEG7_C1 + poly * t;
+    poly = (Real)GTE_C_EXP2_DEG7_C0 + poly * t;
+    return poly;
+}
+
+
+}
diff --git a/Engine/Plugins/Experimental/GeometryProcessing/Source/GeometryAlgorithms/Private/GTEngine/Mathematics/GteExpEstimate.h b/Engine/Plugins/Experimental/GeometryProcessing/Source/GeometryAlgorithms/Private/GTEngine/Mathematics/GteExpEstimate.h
new file mode 100644
--- /dev/null
+++ b/Engine/Plugins/Experimental/GeometryProcessing/Source/GeometryAlgorithms/Private/GTEngine/Mathematics/GteExpEstimate.h
@@ -0,0 +1,56 @@
+// David Eberly, Geometric Tools, Redmond WA 98052
+// Copyright (c) 1998-2019
+// Distributed under the Boost Software License, Version 1.0.
+// http://www.boost.org/LICENSE_1_0.txt
+// http://www.geometrictools.com/License/Boost/LICENSE_1_0.txt
+// File Version: 3.0.0 (2016/06/19)
+
+#pragma once
+
+#include <GTEngine/Mathematics/GteExp2Estimate.h>
+
+// Minimax polynomial approximations to 2^x.  The polynomial p(x) of
+// degree D minimizes the quantity maximum{|2^x - p(x)| : x in [0,1]}
+// over all polynomials of degree D.  The natural exponential is
+// computed using exp(x) = 2^{x log(2)}, where log(2) is the natural
+// logarithm of 2.
+
+namespace gte
+{
+
+template <typename Real>
+class ExpEstimate
+{
+public:
+    // The input constraint is x in [0,1].  For example,
+    //   float x; // in [0,1]
+    //   float result = ExpEstimate<float>::Degree<3>(x);
+    template <int D>
+    inline static Real Degree(Real x);
+
+    // The input x can be any real number.  Range reduction is used to
+    // generate a value y in [0,1], call Degree(y), and combine the output
+    // with the proper exponent to obtain the approximation.  For example,
+    //   float x;  // x >= 0
+    //   float result = ExpEstimate<float>::DegreeRR<3>(x);
+    template <int D>
+    inline static Real DegreeRR(Real x);
+};
+
+
+template <typename Real>
+template <int D>
+inline Real ExpEstimate<Real>::Degree(Real x)
+{
+    return Exp2Estimate<Real>::Degree<D>((Real)GTE_C_LN_2 * x);
+}
+
+template <typename Real>
+template <int D>
+inline Real ExpEstimate<Real>::DegreeRR(Real x)
+{
+    return Exp2Estimate<Real>::DegreeRR<D>((Real)GTE_C_LN_2 * x);
+}
+
+
+}
diff --git a/Engine/Plugins/Experimental/GeometryProcessing/Source/GeometryAlgorithms/Private/GTEngine/Mathematics/GteFIQuery.h b/Engine/Plugins/Experimental/GeometryProcessing/Source/GeometryAlgorithms/Private/GTEngine/Mathematics/GteFIQuery.h
new file mode 100644
--- /dev/null
+++ b/Engine/Plugins/Experimental/GeometryProcessing/Source/GeometryAlgorithms/Private/GTEngine/Mathematics/GteFIQuery.h
@@ -0,0 +1,33 @@
+// David Eberly, Geometric Tools, Redmond WA 98052
+// Copyright (c) 1998-2019
+// Distributed under the Boost Software License, Version 1.0.
+// http://www.boost.org/LICENSE_1_0.txt
+// http://www.geometrictools.com/License/Boost/LICENSE_1_0.txt
+// File Version: 3.0.2 (2018/10/05)
+
+#pragma once
+
+#include <GTEngine/Mathematics/GteMath.h>
+
+namespace gte
+{
+
+// Find-intersection queries.
+
+template <typename Real, typename Type0, typename Type1>
+class FIQuery
+{
+public:
+    struct Result
+    {
+        // A FIQuery-base class B must define a B::Result struct with member
+        // 'bool intersect'.  A FIQuery-derived class D must also derive a
+        // D::Result from B:Result but may have no members.  The member
+        // 'intersect' is 'true' iff the primitives intersect.  The operator()
+        // is non-const to allow FIQuery to store and modify private state
+        // that supports the query.
+    };
+    Result operator()(Type0 const& primitive0, Type1 const& primitive1);
+};
+
+}
diff --git a/Engine/Plugins/Experimental/GeometryProcessing/Source/GeometryAlgorithms/Private/GTEngine/Mathematics/GteFeatureKey.h b/Engine/Plugins/Experimental/GeometryProcessing/Source/GeometryAlgorithms/Private/GTEngine/Mathematics/GteFeatureKey.h
new file mode 100644
--- /dev/null
+++ b/Engine/Plugins/Experimental/GeometryProcessing/Source/GeometryAlgorithms/Private/GTEngine/Mathematics/GteFeatureKey.h
@@ -0,0 +1,76 @@
+// David Eberly, Geometric Tools, Redmond WA 98052
+// Copyright (c) 1998-2019
+// Distributed under the Boost Software License, Version 1.0.
+// http://www.boost.org/LICENSE_1_0.txt
+// http://www.geometrictools.com/License/Boost/LICENSE_1_0.txt
+// File Version: 3.0.0 (2016/06/19)
+
+#pragma once
+
+#include <GTEngine/GTEngineDEF.h>
+
+namespace gte
+{
+
+template <int N, bool Ordered>
+class FeatureKey
+{
+protected:
+    // Abstract base class with V[] uninitialized.  The derived classes must
+    // set the V[] values accordingly.
+    //
+    // An ordered feature key has V[0] = min(V[]) with (V[0],V[1],...,V[N-1]) a
+    // permutation of N inputs with an even number of transpositions.
+    //
+    // An unordered feature key has V[0] < V[1] < ... < V[N-1].
+    //
+    // Note that the word 'order' is about the geometry of the feature, not
+    // the comparison order for any sorting.
+    FeatureKey();
+
+public:
+    bool operator<(FeatureKey const& key) const;
+    bool operator==(FeatureKey const& key) const;
+
+    int V[N];
+};
+
+
+template <int N, bool Ordered>
+FeatureKey<N, Ordered>::FeatureKey()
+{
+}
+
+template <int N, bool Ordered>
+bool FeatureKey<N, Ordered>::operator<(FeatureKey const& key) const
+{
+    for (int i = N - 1; i >= 0; --i)
+    {
+        if (V[i] < key.V[i])
+        {
+            return true;
+        }
+
+        if (V[i] > key.V[i])
+        {
+            return false;
+        }
+    }
+    return false;
+}
+
+template <int N, bool Ordered>
+bool FeatureKey<N, Ordered>::operator==(FeatureKey const& key) const
+{
+    for (int i = 0; i < N; ++i)
+    {
+        if (V[i] != key.V[i])
+        {
+            return false;
+        }
+    }
+    return true;
+}
+
+
+}
diff --git a/Engine/Plugins/Experimental/GeometryProcessing/Source/GeometryAlgorithms/Private/GTEngine/Mathematics/GteFrenetFrame.h b/Engine/Plugins/Experimental/GeometryProcessing/Source/GeometryAlgorithms/Private/GTEngine/Mathematics/GteFrenetFrame.h
new file mode 100644
--- /dev/null
+++ b/Engine/Plugins/Experimental/GeometryProcessing/Source/GeometryAlgorithms/Private/GTEngine/Mathematics/GteFrenetFrame.h
@@ -0,0 +1,159 @@
+// David Eberly, Geometric Tools, Redmond WA 98052
+// Copyright (c) 1998-2019
+// Distributed under the Boost Software License, Version 1.0.
+// http://www.boost.org/LICENSE_1_0.txt
+// http://www.geometrictools.com/License/Boost/LICENSE_1_0.txt
+// File Version: 3.0.2 (2018/10/05)
+
+#pragma once
+
+#include <GTEngine/Mathematics/GteVector2.h>
+#include <GTEngine/Mathematics/GteVector3.h>
+#include <GTEngine/Mathematics/GteParametricCurve.h>
+#include <memory>
+
+namespace gte
+{
+
+template <typename Real>
+class FrenetFrame2
+{
+public:
+    // Construction.  The curve must persist as long as the FrenetFrame2
+    // object does.
+    FrenetFrame2(std::shared_ptr<ParametricCurve<2, Real>> const& curve);
+
+    // The normal is perpendicular to the tangent, rotated clockwise by
+    // pi/2 radians.
+    void operator()(Real t, Vector2<Real>& position, Vector2<Real>& tangent,
+        Vector2<Real>& normal) const;
+
+    Real GetCurvature(Real t) const;
+
+private:
+    std::shared_ptr<ParametricCurve<2, Real>> mCurve;
+};
+
+
+template <typename Real>
+class FrenetFrame3
+{
+public:
+    // Construction.  The curve must persist as long as the FrenetFrame3
+    // object does.
+    FrenetFrame3(std::shared_ptr<ParametricCurve<3, Real>> const& curve);
+
+    // The binormal is Cross(tangent, normal).
+    void operator()(Real t, Vector3<Real>& position, Vector3<Real>& tangent,
+        Vector3<Real>& normal, Vector3<Real>& binormal) const;
+
+    Real GetCurvature(Real t) const;
+    Real GetTorsion(Real t) const;
+
+private:
+    std::shared_ptr<ParametricCurve<3, Real>> mCurve;
+};
+
+
+template <typename Real>
+FrenetFrame2<Real>::FrenetFrame2(std::shared_ptr<ParametricCurve<2, Real>> const& curve)
+    :
+    mCurve(curve)
+{
+}
+
+template <typename Real>
+void FrenetFrame2<Real>::operator()(Real t, Vector2<Real>& position,
+    Vector2<Real>& tangent, Vector2<Real>& normal) const
+{
+    Vector2<Real> values[4];
+    mCurve->Evaluate(t, 1, values);
+    position = values[0];
+    tangent = values[1];
+    Normalize(tangent);
+    normal = Perp(tangent);
+}
+
+template <typename Real>
+Real FrenetFrame2<Real>::GetCurvature(Real t) const
+{
+    Vector2<Real> values[4];
+    mCurve->Evaluate(t, 2, values);
+    Real speedSqr = Dot(values[1], values[1]);
+    if (speedSqr > (Real)0)
+    {
+        Real numer = DotPerp(values[1], values[2]);
+        Real denom = std::pow(speedSqr, (Real)1.5);
+        return numer / denom;
+    }
+    else
+    {
+        // Curvature is indeterminate, just return 0.
+        return (Real)0;
+    }
+}
+
+
+
+template <typename Real>
+FrenetFrame3<Real>::FrenetFrame3(std::shared_ptr<ParametricCurve<3, Real>> const& curve)
+    :
+    mCurve(curve)
+{
+}
+
+template <typename Real>
+void FrenetFrame3<Real>::operator()(Real t, Vector3<Real>& position,
+    Vector3<Real>& tangent, Vector3<Real>& normal, Vector3<Real>& binormal) const
+{
+    Vector3<Real> values[4];
+    mCurve->Evaluate(t, 2, values);
+    position = values[0];
+    Real VDotV = Dot(values[1], values[1]);
+    Real VDotA = Dot(values[1], values[2]);
+    normal = VDotV * values[2] - VDotA * values[1];
+    Normalize(normal);
+    tangent = values[1];
+    Normalize(tangent);
+    binormal = Cross(tangent, normal);
+}
+
+template <typename Real>
+Real FrenetFrame3<Real>::GetCurvature(Real t) const
+{
+    Vector3<Real> values[4];
+    mCurve->Evaluate(t, 2, values);
+    Real speedSqr = Dot(values[1], values[1]);
+    if (speedSqr > (Real)0)
+    {
+        Real numer = Length(Cross(values[1], values[2]));
+        Real denom = std::pow(speedSqr, (Real)1.5);
+        return numer / denom;
+    }
+    else
+    {
+        // Curvature is indeterminate, just return 0.
+        return (Real)0;
+    }
+}
+
+template <typename Real>
+Real FrenetFrame3<Real>::GetTorsion(Real t) const
+{
+    Vector3<Real> values[4];
+    mCurve->Evaluate(t, 3, values);
+    Vector3<Real> cross = Cross(values[1], values[2]);
+    Real denom = Dot(cross, cross);
+    if (denom > (Real)0)
+    {
+        Real numer = Dot(cross, values[3]);
+        return numer / denom;
+    }
+    else
+    {
+        // Torsion is indeterminate, just return 0.
+        return (Real)0;
+    }
+}
+
+}
diff --git a/Engine/Plugins/Experimental/GeometryProcessing/Source/GeometryAlgorithms/Private/GTEngine/Mathematics/GteFrustum3.h b/Engine/Plugins/Experimental/GeometryProcessing/Source/GeometryAlgorithms/Private/GTEngine/Mathematics/GteFrustum3.h
new file mode 100644
--- /dev/null
+++ b/Engine/Plugins/Experimental/GeometryProcessing/Source/GeometryAlgorithms/Private/GTEngine/Mathematics/GteFrustum3.h
@@ -0,0 +1,257 @@
+// David Eberly, Geometric Tools, Redmond WA 98052
+// Copyright (c) 1998-2019
+// Distributed under the Boost Software License, Version 1.0.
+// http://www.boost.org/LICENSE_1_0.txt
+// http://www.geometrictools.com/License/Boost/LICENSE_1_0.txt
+// File Version: 3.0.0 (2016/06/19)
+
+#pragma once
+
+#include <GTEngine/Mathematics/GteVector3.h>
+
+// Orthogonal frustum.  Let E be the origin, D be the direction vector, U be
+// the up vector, and R be the right vector.  Let u > 0 and r > 0 be the
+// extents in the U and R directions, respectively.  Let n and f be the
+// extents in the D direction with 0 < n < f.  The four corners of the frustum
+// in the near plane are E + n*D + s0*u*U + s1*r*R where |s0| = |s1| = 1 (four
+// choices).  The four corners of the frustum in the far plane are
+// E + f*D + (f/n)*(s0*u*U + s1*r*R) where |s0| = |s1| = 1 (four choices).
+
+namespace gte
+{
+
+template <typename Real>
+class Frustum3
+{
+public:
+    // Construction and destruction.  The default constructor sets the
+    // following values:  origin (E) to (0,0,0), dVector (D) to (0,0,1),
+    // uVector (U) to (0,1,0), rVector (R) to (1,0,0), dMin (n) to 1,
+    // dMax (f) to 2, uBound (u) to 1, and rBound (r) to 1.
+    Frustum3();
+    Frustum3(Vector3<Real> const& inOrigin, Vector3<Real> const& inDVector,
+        Vector3<Real> const& inUVector, Vector3<Real> const& inRVector,
+        Real inDMin, Real inDMax, Real inUBound, Real inRBound);
+
+    // The Update() function must be called whenever changes are made to DMIN,
+    // DMax, UBound, or RBound.  The values mDRatio, mMTwoUF, and mMTwoRF are
+    // dependent on the changes, so call the Get*() accessors only after the
+    // Update() call.
+    void Update();
+    inline Real GetDRatio() const;
+    inline Real GetMTwoUF() const;
+    inline Real GetMTwoRF() const;
+
+    void ComputeVertices(Vector3<Real> vertex[8]) const;
+
+    Vector3<Real> origin, dVector, uVector, rVector;
+    Real dMin, dMax, uBound, rBound;
+
+public:
+    // Comparisons to support sorted containers.
+    bool operator==(Frustum3 const& frustum) const;
+    bool operator!=(Frustum3 const& frustum) const;
+    bool operator< (Frustum3 const& frustum) const;
+    bool operator<=(Frustum3 const& frustum) const;
+    bool operator> (Frustum3 const& frustum) const;
+    bool operator>=(Frustum3 const& frustum) const;
+
+protected:
+    // Quantities derived from the constructor inputs.
+    Real mDRatio, mMTwoUF, mMTwoRF;
+};
+
+
+template <typename Real>
+Frustum3<Real>::Frustum3()
+    :
+    origin(Vector3<Real>::Zero()),
+    dVector(Vector3<Real>::Unit(2)),
+    uVector(Vector3<Real>::Unit(1)),
+    rVector(Vector3<Real>::Unit(0)),
+    dMin((Real)1),
+    dMax((Real)2),
+    uBound((Real)1),
+    rBound((Real)1)
+{
+    Update();
+}
+
+template <typename Real>
+Frustum3<Real>::Frustum3(Vector3<Real> const& inOrigin,
+    Vector3<Real> const& inDVector, Vector3<Real> const& inUVector,
+    Vector3<Real> const& inRVector, Real inDMin, Real inDMax, Real inUBound,
+    Real inRBound)
+    :
+    origin(inOrigin),
+    dVector(inDVector),
+    uVector(inUVector),
+    rVector(inRVector),
+    dMin(inDMin),
+    dMax(inDMax),
+    uBound(inUBound),
+    rBound(inRBound)
+{
+    Update();
+}
+
+template <typename Real>
+void Frustum3<Real>::Update()
+{
+    mDRatio = dMax / dMin;
+    mMTwoUF = ((Real)-2) * uBound * dMax;
+    mMTwoRF = ((Real)-2) * rBound * dMax;
+}
+
+template <typename Real> inline
+Real Frustum3<Real>::GetDRatio() const
+{
+    return mDRatio;
+}
+
+template <typename Real> inline
+Real Frustum3<Real>::GetMTwoUF() const
+{
+    return mMTwoUF;
+}
+
+template <typename Real> inline
+Real Frustum3<Real>::GetMTwoRF() const
+{
+    return mMTwoRF;
+}
+
+template <typename Real>
+void Frustum3<Real>::ComputeVertices(Vector3<Real> vertex[8]) const
+{
+    Vector3<Real> dScaled = dMin * dVector;
+    Vector3<Real> uScaled = uBound * uVector;
+    Vector3<Real> rScaled = rBound * rVector;
+
+    vertex[0] = dScaled - uScaled - rScaled;
+    vertex[1] = dScaled - uScaled + rScaled;
+    vertex[2] = dScaled + uScaled + rScaled;
+    vertex[3] = dScaled + uScaled - rScaled;
+
+    for (int i = 0, ip = 4; i < 4; ++i, ++ip)
+    {
+        vertex[ip] = origin + mDRatio * vertex[i];
+        vertex[i] += origin;
+    }
+}
+
+template <typename Real>
+bool Frustum3<Real>::operator==(Frustum3 const& frustum) const
+{
+    return origin == frustum.origin
+        && dVector == frustum.dVector
+        && uVector == frustum.uVector
+        && rVector == frustum.rVector
+        && dMin == frustum.dMin
+        && dMax == frustum.dMax
+        && uBound == frustum.uBound
+        && rBound == frustum.rBound;
+}
+
+template <typename Real>
+bool Frustum3<Real>::operator!=(Frustum3 const& frustum) const
+{
+    return !operator==(frustum);
+}
+
+template <typename Real>
+bool Frustum3<Real>::operator<(Frustum3 const& frustum) const
+{
+    if (origin < frustum.origin)
+    {
+        return true;
+    }
+
+    if (origin > frustum.origin)
+    {
+        return false;
+    }
+
+    if (dVector < frustum.dVector)
+    {
+        return true;
+    }
+
+    if (dVector > frustum.dVector)
+    {
+        return false;
+    }
+
+    if (uVector < frustum.uVector)
+    {
+        return true;
+    }
+
+    if (uVector > frustum.uVector)
+    {
+        return false;
+    }
+
+    if (rVector < frustum.rVector)
+    {
+        return true;
+    }
+
+    if (rVector > frustum.rVector)
+    {
+        return false;
+    }
+
+    if (dMin < frustum.dMin)
+    {
+        return true;
+    }
+
+    if (dMin > frustum.dMin)
+    {
+        return false;
+    }
+
+    if (dMax < frustum.dMax)
+    {
+        return true;
+    }
+
+    if (dMax > frustum.dMax)
+    {
+        return false;
+    }
+
+    if (uBound < frustum.uBound)
+    {
+        return true;
+    }
+
+    if (uBound > frustum.uBound)
+    {
+        return false;
+    }
+
+    return rBound < frustum.rBound;
+}
+
+template <typename Real>
+bool Frustum3<Real>::operator<=(Frustum3 const& frustum) const
+{
+    return operator<(frustum) || operator==(frustum);
+}
+
+template <typename Real>
+bool Frustum3<Real>::operator>(Frustum3 const& frustum) const
+{
+    return !operator<=(frustum);
+}
+
+template <typename Real>
+bool Frustum3<Real>::operator>=(Frustum3 const& frustum) const
+{
+    return !operator<(frustum);
+}
+
+
+}
