move GeometryProcessing plugin from NotForLicensees to Experimental

diff --git a/Engine/Plugins/Experimental/GeometryProcessing/Source/GeometryAlgorithms/Private/GTEngine/Mathematics/GteGMatrix.h b/Engine/Plugins/Experimental/GeometryProcessing/Source/GeometryAlgorithms/Private/GTEngine/Mathematics/GteGMatrix.h
new file mode 100644
--- /dev/null
+++ b/Engine/Plugins/Experimental/GeometryProcessing/Source/GeometryAlgorithms/Private/GTEngine/Mathematics/GteGMatrix.h
@@ -0,0 +1,847 @@
+// David Eberly, Geometric Tools, Redmond WA 98052
+// Copyright (c) 1998-2019
+// Distributed under the Boost Software License, Version 1.0.
+// http://www.boost.org/LICENSE_1_0.txt
+// http://www.geometrictools.com/License/Boost/LICENSE_1_0.txt
+// File Version: 3.0.2 (2018/10/05)
+
+#pragma once
+
+#include <GTEngine/Mathematics/GteGVector.h>
+#include <GTEngine/Mathematics/GteGaussianElimination.h>
+#include <algorithm>
+
+// Uncomment these to test for out-of-range indices and size mismatches.
+//#define GTE_ASSERT_ON_GMATRIX_INDEX_OUT_OF_RANGE
+//#define GTE_ASSERT_ON_GMATRIX_SIZE_MISMATCH
+
+namespace gte
+{
+
+template <typename Real>
+class GMatrix
+{
+public:
+    // The table is length zero and mNumRows and mNumCols are set to zero.
+    GMatrix();
+
+    // The table is length numRows*numCols and the elements are initialized
+    // to zero.
+    GMatrix(int numRow, int numCols);
+
+    // For 0 <= r < numRows and 0 <= c < numCols, element (r,c) is 1 and all
+    // others are 0.  If either of r or c is invalid, the zero matrix is
+    // created.  This is a convenience for creating the standard Euclidean
+    // basis matrices; see also MakeUnit(int,int) and Unit(int,int).
+    GMatrix(int numRows, int numCols, int r, int c);
+
+    // The copy constructor, destructor, and assignment operator are generated
+    // by the compiler.
+
+    // Member access for which the storage representation is transparent.  The
+    // matrix entry in row r and column c is A(r,c).  The first operator()
+    // returns a const reference rather than a Real value.  This supports
+    // writing via standard file operations that require a const pointer to
+    // data.
+    void SetSize(int numRows, int numCols);
+    inline void GetSize(int& numRows, int& numCols) const;
+    inline int GetNumRows() const;
+    inline int GetNumCols() const;
+    inline int GetNumElements() const;
+    inline Real const& operator()(int r, int c) const;
+    inline Real& operator()(int r, int c);
+
+    // Member access by rows or by columns.  The input vectors must have the
+    // correct number of elements for the matrix size.
+    void SetRow(int r, GVector<Real> const& vec);
+    void SetCol(int c, GVector<Real> const& vec);
+    GVector<Real> GetRow(int r) const;
+    GVector<Real> GetCol(int c) const;
+
+    // Member access by 1-dimensional index.  NOTE: These accessors are
+    // useful for the manipulation of matrix entries when it does not
+    // matter whether storage is row-major or column-major.  Do not use
+    // constructs such as M[c+NumCols*r] or M[r+NumRows*c] that expose the
+    // storage convention.
+    inline Real const& operator[](int i) const;
+    inline Real& operator[](int i);
+
+    // Comparisons for sorted containers and geometric ordering.
+    inline bool operator==(GMatrix const& mat) const;
+    inline bool operator!=(GMatrix const& mat) const;
+    inline bool operator< (GMatrix const& mat) const;
+    inline bool operator<=(GMatrix const& mat) const;
+    inline bool operator> (GMatrix const& mat) const;
+    inline bool operator>=(GMatrix const& mat) const;
+
+    // Special matrices.
+    void MakeZero();  // All components are 0.
+    void MakeUnit(int r, int c);  // Component (r,c) is 1, all others zero.
+    void MakeIdentity();  // Diagonal entries 1, others 0, even when nonsquare
+    static GMatrix Zero(int numRows, int numCols);
+    static GMatrix Unit(int numRows, int numCols, int r, int c);
+    static GMatrix Identity(int numRows, int numCols);
+
+protected:
+    // The matrix is stored as a 1-dimensional array.  The convention of
+    // row-major or column-major is your choice.
+    int mNumRows, mNumCols;
+    std::vector<Real> mElements;
+};
+
+// Unary operations.
+template <typename Real>
+GMatrix<Real> operator+(GMatrix<Real> const& M);
+
+template <typename Real>
+GMatrix<Real> operator-(GMatrix<Real> const& M);
+
+// Linear-algebraic operations.
+template <typename Real>
+GMatrix<Real> operator+(GMatrix<Real> const& M0, GMatrix<Real> const& M1);
+
+template <typename Real>
+GMatrix<Real> operator-(GMatrix<Real> const& M0, GMatrix<Real> const& M1);
+
+template <typename Real>
+GMatrix<Real> operator*(GMatrix<Real> const& M, Real scalar);
+
+template <typename Real>
+GMatrix<Real> operator*(Real scalar, GMatrix<Real> const& M);
+
+template <typename Real>
+GMatrix<Real> operator/(GMatrix<Real> const& M, Real scalar);
+
+template <typename Real>
+GMatrix<Real>& operator+=(GMatrix<Real>& M0, GMatrix<Real> const& M1);
+
+template <typename Real>
+GMatrix<Real>& operator-=(GMatrix<Real>& M0, GMatrix<Real> const& M1);
+
+template <typename Real>
+GMatrix<Real>& operator*=(GMatrix<Real>& M, Real scalar);
+
+template <typename Real>
+GMatrix<Real>& operator/=(GMatrix<Real>& M, Real scalar);
+
+// Geometric operations.
+template <typename Real>
+Real L1Norm(GMatrix<Real> const& M);
+
+template <typename Real>
+Real L2Norm(GMatrix<Real> const& M);
+
+template <typename Real>
+Real LInfinityNorm(GMatrix<Real> const& M);
+
+template <typename Real>
+GMatrix<Real> Inverse(GMatrix<Real> const& M,
+    bool* reportInvertibility = nullptr);
+
+template <typename Real>
+Real Determinant(GMatrix<Real> const& M);
+
+// M^T
+template <typename Real>
+GMatrix<Real> Transpose(GMatrix<Real> const& M);
+
+// M*V
+template <typename Real>
+GVector<Real> operator*(GMatrix<Real> const& M, GVector<Real> const& V);
+
+// V^T*M
+template <typename Real>
+GVector<Real> operator*(GVector<Real> const& V, GMatrix<Real> const& M);
+
+// A*B
+template <typename Real>
+GMatrix<Real> operator*(GMatrix<Real> const& A, GMatrix<Real> const& B);
+
+template <typename Real>
+GMatrix<Real> MultiplyAB(GMatrix<Real> const& A, GMatrix<Real> const& B);
+
+// A*B^T
+template <typename Real>
+GMatrix<Real> MultiplyABT(GMatrix<Real> const& A, GMatrix<Real> const& B);
+
+// A^T*B
+template <typename Real>
+GMatrix<Real> MultiplyATB(GMatrix<Real> const& A, GMatrix<Real> const& B);
+
+// A^T*B^T
+template <typename Real>
+GMatrix<Real> MultiplyATBT(GMatrix<Real> const& A, GMatrix<Real> const& B);
+
+// M*D, D is square diagonal (stored as vector)
+template <typename Real>
+GMatrix<Real> MultiplyMD(GMatrix<Real> const& M, GVector<Real> const& D);
+
+// D*M, D is square diagonal (stored as vector)
+template <typename Real>
+GMatrix<Real> MultiplyDM(GVector<Real> const& D, GMatrix<Real> const& M);
+
+// U*V^T, U is N-by-1, V is M-by-1, result is N-by-M.
+template <typename Real>
+GMatrix<Real> OuterProduct(GVector<Real> const& U, GVector<Real> const& V);
+
+// Initialization to a diagonal matrix whose diagonal entries are the
+// components of D.
+template <typename Real>
+void MakeDiagonal(GVector<Real> const& D, GMatrix<Real>& M);
+
+
+template <typename Real>
+GMatrix<Real>::GMatrix()
+    :
+    mNumRows(0),
+    mNumCols(0)
+{
+}
+
+template <typename Real>
+GMatrix<Real>::GMatrix(int numRows, int numCols)
+{
+    SetSize(numRows, numCols);
+    std::fill(mElements.begin(), mElements.end(), (Real)0);
+}
+
+template <typename Real>
+GMatrix<Real>::GMatrix(int numRows, int numCols, int r, int c)
+{
+    SetSize(numRows, numCols);
+    MakeUnit(r, c);
+}
+
+template <typename Real>
+void GMatrix<Real>::SetSize(int numRows, int numCols)
+{
+    if (numRows > 0 && numCols > 0)
+    {
+        mNumRows = numRows;
+        mNumCols = numCols;
+        mElements.resize(mNumRows * mNumCols);
+    }
+    else
+    {
+        mNumRows = 0;
+        mNumCols = 0;
+        mElements.clear();
+    }
+}
+
+template <typename Real> inline
+void GMatrix<Real>::GetSize(int& numRows, int& numCols) const
+{
+    numRows = mNumRows;
+    numCols = mNumCols;
+}
+
+template <typename Real> inline
+int GMatrix<Real>::GetNumRows() const
+{
+    return mNumRows;
+}
+
+template <typename Real> inline
+int GMatrix<Real>::GetNumCols() const
+{
+    return mNumCols;
+}
+
+template <typename Real> inline
+int GMatrix<Real>::GetNumElements() const
+{
+    return static_cast<int>(mElements.size());
+}
+
+template <typename Real> inline
+Real const& GMatrix<Real>::operator()(int r, int c) const
+{
+#if defined(GTE_ASSERT_ON_GMATRIX_INDEX_OUT_OF_RANGE)
+    LogAssert(0 <= r && r < GetNumRows() && 0 <= c && c < GetNumCols(),
+        "Invalid index.");
+#endif
+
+#if defined(GTE_USE_ROW_MAJOR)
+    return mElements[c + mNumCols*r];
+#else
+    return mElements[r + mNumRows*c];
+#endif
+}
+
+template <typename Real> inline
+Real& GMatrix<Real>::operator()(int r, int c)
+{
+#if defined(GTE_ASSERT_ON_GMATRIX_INDEX_OUT_OF_RANGE)
+    LogAssert(0 <= r && r < GetNumRows() && 0 <= c && c < GetNumCols(),
+        "Invalid index.");
+#endif
+
+#if defined(GTE_USE_ROW_MAJOR)
+    return mElements[c + mNumCols*r];
+#else
+    return mElements[r + mNumRows*c];
+#endif
+}
+
+template <typename Real>
+void GMatrix<Real>::SetRow(int r, GVector<Real> const& vec)
+{
+#if defined(GTE_ASSERT_ON_GMATRIX_SIZE_MISMATCH)
+    LogAssert(vec.GetSize() == GetNumCols(), "Mismatched size.");
+#endif
+    for (int c = 0; c < mNumCols; ++c)
+    {
+        operator()(r, c) = vec[c];
+    }
+}
+
+template <typename Real>
+void GMatrix<Real>::SetCol(int c, GVector<Real> const& vec)
+{
+#if defined(GTE_ASSERT_ON_GMATRIX_SIZE_MISMATCH)
+    LogAssert(vec.GetSize() == GetNumRows(), "Mismatched size.");
+#endif
+    for (int r = 0; r < mNumRows; ++r)
+    {
+        operator()(r, c) = vec[r];
+    }
+}
+
+template <typename Real>
+GVector<Real> GMatrix<Real>::GetRow(int r) const
+{
+    GVector<Real> vec(mNumCols);
+    for (int c = 0; c < mNumCols; ++c)
+    {
+        vec[c] = operator()(r, c);
+    }
+    return vec;
+}
+
+template <typename Real>
+GVector<Real> GMatrix<Real>::GetCol(int c) const
+{
+    GVector<Real> vec(mNumRows);
+    for (int r = 0; r < mNumRows; ++r)
+    {
+        vec[r] = operator()(r, c);
+    }
+    return vec;
+}
+
+template <typename Real> inline
+Real const& GMatrix<Real>::operator[](int i) const
+{
+    return mElements[i];
+}
+
+template <typename Real> inline
+Real& GMatrix<Real>::operator[](int i)
+{
+    return mElements[i];
+}
+
+template <typename Real> inline
+bool GMatrix<Real>::operator==(GMatrix const& mat) const
+{
+    return mNumRows == mat.mNumRows && mNumCols == mat.mNumCols
+        && mElements == mat.mElements;
+}
+
+template <typename Real> inline
+bool GMatrix<Real>::operator!=(GMatrix const& mat) const
+{
+    return !operator==(mat);
+}
+
+template <typename Real> inline
+bool GMatrix<Real>::operator<(GMatrix const& mat) const
+{
+    return mNumRows == mat.mNumRows && mNumCols == mat.mNumCols
+        && mElements < mat.mElements;
+}
+
+template <typename Real> inline
+bool GMatrix<Real>::operator<=(GMatrix const& mat) const
+{
+    return mNumRows == mat.mNumRows && mNumCols == mat.mNumCols
+        && mElements <= mat.mElements;
+}
+
+template <typename Real> inline
+bool GMatrix<Real>::operator>(GMatrix const& mat) const
+{
+    return mNumRows == mat.mNumRows && mNumCols == mat.mNumCols
+        && mElements > mat.mElements;
+}
+
+template <typename Real> inline
+bool GMatrix<Real>::operator>=(GMatrix const& mat) const
+{
+    return mNumRows == mat.mNumRows && mNumCols == mat.mNumCols
+        && mElements >= mat.mElements;
+}
+
+template <typename Real>
+void GMatrix<Real>::MakeZero()
+{
+    std::fill(mElements.begin(), mElements.end(), (Real)0);
+}
+
+template <typename Real>
+void GMatrix<Real>::MakeUnit(int r, int c)
+{
+    MakeZero();
+    if (0 <= r && r < mNumRows && 0 <= c && c < mNumCols)
+    {
+        operator()(r, c) = (Real)1;
+    }
+}
+
+template <typename Real>
+void GMatrix<Real>::MakeIdentity()
+{
+    MakeZero();
+    int const numDiagonal = (mNumRows <= mNumCols ? mNumRows : mNumCols);
+    for (int i = 0; i < numDiagonal; ++i)
+    {
+        operator()(i, i) = (Real)1;
+    }
+}
+
+template <typename Real>
+GMatrix<Real> GMatrix<Real>::Zero(int numRows, int numCols)
+{
+    GMatrix<Real> M(numRows, numCols);
+    M.MakeZero();
+    return M;
+}
+
+template <typename Real>
+GMatrix<Real> GMatrix<Real>::Unit(int numRows, int numCols, int r, int c)
+{
+    GMatrix<Real> M(numRows, numCols);
+    M.MakeUnit(r, c);
+    return M;
+}
+
+template <typename Real>
+GMatrix<Real> GMatrix<Real>::Identity(int numRows, int numCols)
+{
+    GMatrix<Real> M(numRows, numCols);
+    M.MakeIdentity();
+    return M;
+}
+
+
+
+template <typename Real>
+GMatrix<Real> operator+(GMatrix<Real> const& M)
+{
+    return M;
+}
+
+template <typename Real>
+GMatrix<Real> operator-(GMatrix<Real> const& M)
+{
+    GMatrix<Real> result(M.GetNumRows(), M.GetNumCols());
+    for (int i = 0; i < M.GetNumElements(); ++i)
+    {
+        result[i] = -M[i];
+    }
+    return result;
+}
+
+template <typename Real>
+GMatrix<Real> operator+(GMatrix<Real> const& M0, GMatrix<Real> const& M1)
+{
+    GMatrix<Real> result = M0;
+    return result += M1;
+}
+
+template <typename Real>
+GMatrix<Real> operator-(GMatrix<Real> const& M0, GMatrix<Real> const& M1)
+{
+    GMatrix<Real> result = M0;
+    return result -= M1;
+}
+
+template <typename Real>
+GMatrix<Real> operator*(GMatrix<Real> const& M, Real scalar)
+{
+    GMatrix<Real> result = M;
+    return result *= scalar;
+}
+
+template <typename Real>
+GMatrix<Real> operator*(Real scalar, GMatrix<Real> const& M)
+{
+    GMatrix<Real> result = M;
+    return result *= scalar;
+}
+
+template <typename Real>
+GMatrix<Real> operator/(GMatrix<Real> const& M, Real scalar)
+{
+    GMatrix<Real> result = M;
+    return result /= scalar;
+}
+
+template <typename Real>
+GMatrix<Real>& operator+=(GMatrix<Real>& M0, GMatrix<Real> const& M1)
+{
+#if defined(GTE_ASSERT_ON_GMATRIX_SIZE_MISMATCH)
+    LogAssert(M0.GetNumRows() == M1.GetNumRows() &&
+        M0.GetNumCols() == M1.GetNumCols(), "Mismatched size.");
+#endif
+
+    for (int i = 0; i < M0.GetNumElements(); ++i)
+    {
+        M0[i] += M1[i];
+    }
+    return M0;
+}
+
+template <typename Real>
+GMatrix<Real>& operator-=(GMatrix<Real>& M0, GMatrix<Real> const& M1)
+{
+#if defined(GTE_ASSERT_ON_GMATRIX_SIZE_MISMATCH)
+    LogAssert(M0.GetNumRows() == M1.GetNumRows() &&
+        M0.GetNumCols() == M1.GetNumCols(), "Mismatched size.");
+#endif
+
+    for (int i = 0; i < M0.GetNumElements(); ++i)
+    {
+        M0[i] -= M1[i];
+    }
+    return M0;
+}
+
+template <typename Real>
+GMatrix<Real>& operator*=(GMatrix<Real>& M, Real scalar)
+{
+    for (int i = 0; i < M.GetNumElements(); ++i)
+    {
+        M[i] *= scalar;
+    }
+    return M;
+}
+
+template <typename Real>
+GMatrix<Real>& operator/=(GMatrix<Real>& M, Real scalar)
+{
+    if (scalar != (Real)0)
+    {
+        Real invScalar = ((Real)1) / scalar;
+        for (int i = 0; i < M.GetNumElements(); ++i)
+        {
+            M[i] *= invScalar;
+        }
+    }
+    else
+    {
+        for (int i = 0; i < M.GetNumElements(); ++i)
+        {
+            M[i] = (Real)0;
+        }
+    }
+    return M;
+}
+
+template <typename Real>
+Real L1Norm(GMatrix<Real> const& M)
+{
+    Real sum = std::abs(M[0]);
+    for (int i = 1; i < M.GetNumElements(); ++i)
+    {
+        sum += std::abs(M[i]);
+    }
+    return sum;
+}
+
+template <typename Real>
+Real L2Norm(GMatrix<Real> const& M)
+{
+    Real sum = M[0] * M[0];
+    for (int i = 1; i < M.GetNumElements(); ++i)
+    {
+        sum += M[i] * M[i];
+    }
+    return std::sqrt(sum);
+}
+
+template <typename Real>
+Real LInfinityNorm(GMatrix<Real> const& M)
+{
+    Real maxAbsElement = M[0];
+    for (int i = 1; i < M.GetNumElements(); ++i)
+    {
+        Real absElement = std::abs(M[i]);
+        if (absElement > maxAbsElement)
+        {
+            maxAbsElement = absElement;
+        }
+    }
+    return maxAbsElement;
+}
+
+template <typename Real>
+GMatrix<Real> Inverse(GMatrix<Real> const& M, bool* reportInvertibility)
+{
+    GMatrix<Real> invM(M.GetNumRows(), M.GetNumCols());
+    if (M.GetNumRows() == M.GetNumCols())
+    {
+        Real determinant;
+        bool invertible = GaussianElimination<Real>()(M.GetNumRows(), &M[0],
+            &invM[0], determinant, nullptr, nullptr, nullptr, 0, nullptr);
+        if (reportInvertibility)
+        {
+            *reportInvertibility = invertible;
+        }
+    }
+    else
+    {
+#if defined(GTE_ASSERT_ON_GMATRIX_SIZE_MISMATCH)
+        LogError("Matrix must be square.");
+#endif
+        invM.MakeZero();
+        if (reportInvertibility)
+        {
+            *reportInvertibility = false;
+        }
+    }
+    return invM;
+}
+
+template <typename Real>
+Real Determinant(GMatrix<Real> const& M)
+{
+    Real determinant;
+    if (M.GetNumRows() == M.GetNumCols())
+    {
+        GaussianElimination<Real>()(M.GetNumRows(), &M[0], nullptr,
+            determinant, nullptr, nullptr, nullptr, 0, nullptr);
+    }
+    else
+    {
+#if defined(GTE_ASSERT_ON_GMATRIX_SIZE_MISMATCH)
+        LogError("Matrix must be square.");
+#endif
+        determinant = (Real)0;
+    }
+    return determinant;
+}
+
+template <typename Real>
+GMatrix<Real> Transpose(GMatrix<Real> const& M)
+{
+    GMatrix<Real> result(M.GetNumCols(), M.GetNumRows());
+    for (int r = 0; r < M.GetNumRows(); ++r)
+    {
+        for (int c = 0; c < M.GetNumCols(); ++c)
+        {
+            result(c, r) = M(r, c);
+        }
+    }
+    return result;
+}
+
+template <typename Real>
+GVector<Real> operator*(GMatrix<Real> const& M, GVector<Real> const& V)
+{
+#if defined(GTE_ASSERT_ON_GMATRIX_SIZE_MISMATCH)
+    LogAssert(V.GetSize() == M.GetNumRows(), "Mismatched size.");
+#endif
+    GVector<Real> result(M.GetNumRows());
+    for (int r = 0; r < M.GetNumRows(); ++r)
+    {
+        result[r] = (Real)0;
+        for (int c = 0; c < M.GetNumCols(); ++c)
+        {
+            result[r] += M(r, c) * V[c];
+        }
+    }
+    return result;
+}
+
+template <typename Real>
+GVector<Real> operator*(GVector<Real> const& V, GMatrix<Real> const& M)
+{
+#if defined(GTE_ASSERT_ON_GMATRIX_SIZE_MISMATCH)
+    LogAssert(V.GetSize() == M.GetNumCols(), "Mismatched size.");
+#endif
+    GVector<Real> result(M.GetNumCols());
+    for (int c = 0; c < M.GetNumCols(); ++c)
+    {
+        result[c] = (Real)0;
+        for (int r = 0; r < M.GetNumRows(); ++r)
+        {
+            result[c] += V[r] * M(r, c);
+        }
+    }
+    return result;
+}
+
+template <typename Real>
+GMatrix<Real> operator*(GMatrix<Real> const& A, GMatrix<Real> const& B)
+{
+    return MultiplyAB(A, B);
+}
+
+template <typename Real>
+GMatrix<Real> MultiplyAB(GMatrix<Real> const& A, GMatrix<Real> const& B)
+{
+#if defined(GTE_ASSERT_ON_GMATRIX_SIZE_MISMATCH)
+    LogAssert(A.GetNumCols() == B.GetNumRows(), "Mismatched size.");
+#endif
+    int const numCommon = A.GetNumCols();
+    GMatrix<Real> result(A.GetNumRows(), B.GetNumCols());
+    for (int r = 0; r < result.GetNumRows(); ++r)
+    {
+        for (int c = 0; c < result.GetNumCols(); ++c)
+        {
+            result(r, c) = (Real)0;
+            for (int i = 0; i < numCommon; ++i)
+            {
+                result(r, c) += A(r, i) * B(i, c);
+            }
+        }
+    }
+    return result;
+}
+
+template <typename Real>
+GMatrix<Real> MultiplyABT(GMatrix<Real> const& A, GMatrix<Real> const& B)
+{
+#if defined(GTE_ASSERT_ON_GMATRIX_SIZE_MISMATCH)
+    LogAssert(A.GetNumCols() == B.GetNumCols(), "Mismatched size.");
+#endif
+    int const numCommon = A.GetNumCols();
+    GMatrix<Real> result(A.GetNumRows(), B.GetNumRows());
+    for (int r = 0; r < result.GetNumRows(); ++r)
+    {
+        for (int c = 0; c < result.GetNumCols(); ++c)
+        {
+            result(r, c) = (Real)0;
+            for (int i = 0; i < numCommon; ++i)
+            {
+                result(r, c) += A(r, i) * B(c, i);
+            }
+        }
+    }
+    return result;
+}
+
+template <typename Real>
+GMatrix<Real> MultiplyATB(GMatrix<Real> const& A, GMatrix<Real> const& B)
+{
+#if defined(GTE_ASSERT_ON_GMATRIX_SIZE_MISMATCH)
+    LogAssert(A.GetNumRows() == B.GetNumRows(), "Mismatched size.");
+#endif
+    int const numCommon = A.GetNumRows();
+    GMatrix<Real> result(A.GetNumCols(), B.GetNumCols());
+    for (int r = 0; r < result.GetNumRows(); ++r)
+    {
+        for (int c = 0; c < result.GetNumCols(); ++c)
+        {
+            result(r, c) = (Real)0;
+            for (int i = 0; i < numCommon; ++i)
+            {
+                result(r, c) += A(i, r) * B(i, c);
+            }
+        }
+    }
+    return result;
+}
+
+template <typename Real>
+GMatrix<Real> MultiplyATBT(GMatrix<Real> const& A, GMatrix<Real> const& B)
+{
+#if defined(GTE_ASSERT_ON_GMATRIX_SIZE_MISMATCH)
+    LogAssert(A.GetNumRows() == B.GetNumCols(), "Mismatched size.");
+#endif
+    int const numCommon = A.GetNumRows();
+    GMatrix<Real> result(A.GetNumCols(), B.GetNumRows());
+    for (int r = 0; r < result.GetNumRows(); ++r)
+    {
+        for (int c = 0; c < result.GetNumCols(); ++c)
+        {
+            result(r, c) = (Real)0;
+            for (int i = 0; i < numCommon; ++i)
+            {
+                result(r, c) += A(i, r) * B(c, i);
+            }
+        }
+    }
+    return result;
+}
+
+template <typename Real>
+GMatrix<Real> MultiplyMD(GMatrix<Real> const& M, GVector<Real> const& D)
+{
+#if defined(GTE_ASSERT_ON_GMATRIX_SIZE_MISMATCH)
+    LogAssert(D.GetSize() == M.GetNumCols(), "Mismatched size.");
+#endif
+    GMatrix<Real> result(M.GetNumRows(), M.GetNumCols());
+    for (int r = 0; r < result.GetNumRows(); ++r)
+    {
+        for (int c = 0; c < result.GetNumCols(); ++c)
+        {
+            result(r, c) = M(r, c) * D[c];
+        }
+    }
+    return result;
+}
+
+template <typename Real>
+GMatrix<Real> MultiplyDM(GVector<Real> const& D, GMatrix<Real> const& M)
+{
+#if defined(GTE_ASSERT_ON_GMATRIX_SIZE_MISMATCH)
+    LogAssert(D.GetSize() == M.GetNumRows(), "Mismatched size.");
+#endif
+    GMatrix<Real> result(M.GetNumRows(), M.GetNumCols());
+    for (int r = 0; r < result.GetNumRows(); ++r)
+    {
+        for (int c = 0; c < result.GetNumCols(); ++c)
+        {
+            result(r, c) = D[r] * M(r, c);
+        }
+    }
+    return result;
+}
+
+template <typename Real>
+GMatrix<Real> OuterProduct(GVector<Real> const& U, GVector<Real> const& V)
+{
+    GMatrix<Real> result(U.GetSize(), V.GetSize());
+    for (int r = 0; r < result.GetNumRows(); ++r)
+    {
+        for (int c = 0; c < result.GetNumCols(); ++c)
+        {
+            result(r, c) = U[r] * V[c];
+        }
+    }
+    return result;
+}
+
+template <typename Real>
+void MakeDiagonal(GVector<Real> const& D, GMatrix<Real>& M)
+{
+#if defined(GTE_ASSERT_ON_GMATRIX_SIZE_MISMATCH)
+    LogAssert(M.GetNumRows() == M.GetNumCols(), "Mismatched size.");
+#endif
+    int const N = M.GetNumRows();
+    for (int i = 0; i < N*N; ++i)
+    {
+        M[i] = (Real)0;
+    }
+
+    for (int i = 0; i < N; ++i)
+    {
+        M(i, i) = D[i];
+    }
+}
+
+
+}
diff --git a/Engine/Plugins/Experimental/GeometryProcessing/Source/GeometryAlgorithms/Private/GTEngine/Mathematics/GteGVector.h b/Engine/Plugins/Experimental/GeometryProcessing/Source/GeometryAlgorithms/Private/GTEngine/Mathematics/GteGVector.h
new file mode 100644
--- /dev/null
+++ b/Engine/Plugins/Experimental/GeometryProcessing/Source/GeometryAlgorithms/Private/GTEngine/Mathematics/GteGVector.h
@@ -0,0 +1,646 @@
+// David Eberly, Geometric Tools, Redmond WA 98052
+// Copyright (c) 1998-2019
+// Distributed under the Boost Software License, Version 1.0.
+// http://www.boost.org/LICENSE_1_0.txt
+// http://www.geometrictools.com/License/Boost/LICENSE_1_0.txt
+// File Version: 3.0.3 (2018/10/05)
+
+#pragma once
+
+#include <GTEngine/LowLevel/GteLogger.h>
+#include <cmath>
+#include <vector>
+
+// Uncomment this to test for size mismatches that are wrapped by
+// GVector::ValidateSize.  If the test is not enabled, the compiler should
+// not generate any code for ValidateSize in a release build.
+//#define GTE_ASSERT_ON_GVECTOR_SIZE_MISMATCH
+
+namespace gte
+{
+
+template <typename Real>
+class GVector
+{
+public:
+    // The tuple is length zero (uninitialized).
+    GVector();
+
+    // The tuple is length 'size' and the elements are uninitialized.
+    GVector(int size);
+
+    // For 0 <= d <= size, element d is 1 and all others are zero.  If d is
+    // invalid, the zero vector is created.  This is a convenience for
+    // creating the standard Euclidean basis vectors; see also
+    // MakeUnit(int,int) and Unit(int,int).
+    GVector(int size, int d);
+
+    // The copy constructor, destructor, and assignment operator are generated
+    // by the compiler.
+
+    // Member access.  SetSize(int) does not initialize the tuple.  The first
+    // operator[] returns a const reference rather than a Real value.  This
+    // supports writing via standard file operations that require a const
+    // pointer to data.
+    void SetSize(int size);
+    inline int GetSize() const;
+    inline Real const& operator[](int i) const;
+    inline Real& operator[](int i);
+
+    // Comparison (for use by STL containers).
+    inline bool operator==(GVector const& vec) const;
+    inline bool operator!=(GVector const& vec) const;
+    inline bool operator< (GVector const& vec) const;
+    inline bool operator<=(GVector const& vec) const;
+    inline bool operator> (GVector const& vec) const;
+    inline bool operator>=(GVector const& vec) const;
+
+    // Special vectors.
+    void MakeZero();  // All components are 0.
+    void MakeUnit(int d);  // Component d is 1, all others are zero.
+    static GVector Zero(int size);
+    static GVector Unit(int size, int d);
+
+protected:
+    // This data structure takes advantage of the built-in operator[],
+    // range checking, and visualizers in MSVS.
+    std::vector<Real> mTuple;
+};
+
+// Unary operations.
+template <typename Real>
+GVector<Real> operator+(GVector<Real> const& v);
+
+template <typename Real>
+GVector<Real> operator-(GVector<Real> const& v);
+
+// Linear-algebraic operations.
+template <typename Real>
+GVector<Real> operator+(GVector<Real> const& v0, GVector<Real> const& v1);
+
+template <typename Real>
+GVector<Real> operator-(GVector<Real> const& v0, GVector<Real> const& v1);
+
+template <typename Real>
+GVector<Real> operator*(GVector<Real> const& v, Real scalar);
+
+template <typename Real>
+GVector<Real> operator*(Real scalar, GVector<Real> const& v);
+
+template <typename Real>
+GVector<Real> operator/(GVector<Real> const& v, Real scalar);
+
+template <typename Real>
+GVector<Real>& operator+=(GVector<Real>& v0, GVector<Real> const& v1);
+
+template <typename Real>
+GVector<Real>& operator-=(GVector<Real>& v0, GVector<Real> const& v1);
+
+template <typename Real>
+GVector<Real>& operator*=(GVector<Real>& v, Real scalar);
+
+template <typename Real>
+GVector<Real>& operator/=(GVector<Real>& v, Real scalar);
+
+// Geometric operations.  The functions with 'robust' set to 'false' use the
+// standard algorithm for normalizing a vector by computing the length as a
+// square root of the squared length and dividing by it.  The results can be
+// infinite (or NaN) if the length is zero.  When 'robust' is set to 'true',
+// the algorithm is designed to avoid floating-point overflow and sets the
+// normalized vector to zero when the length is zero.
+template <typename Real>
+Real Dot(GVector<Real> const& v0, GVector<Real> const& v1);
+
+template <typename Real>
+Real Length(GVector<Real> const& v, bool robust = false);
+
+template <typename Real>
+Real Normalize(GVector<Real>& v, bool robust = false);
+
+// Gram-Schmidt orthonormalization to generate orthonormal vectors from the
+// linearly independent inputs.  The function returns the smallest length of
+// the unnormalized vectors computed during the process.  If this value is
+// nearly zero, it is possible that the inputs are linearly dependent (within
+// numerical round-off errors).  On input, 1 <= numElements <= N and v[0]
+// through v[numElements-1] must be initialized.  On output, the vectors
+// v[0] through v[numElements-1] form an orthonormal set.
+template <typename Real>
+Real Orthonormalize(int numElements, GVector<Real>* v, bool robust = false);
+
+// Compute the axis-aligned bounding box of the vectors.  The return value is
+// 'true' iff the inputs are valid, in which case vmin and vmax have valid
+// values.
+template <typename Real>
+bool ComputeExtremes(int numVectors, GVector<Real> const* v,
+    GVector<Real>& vmin, GVector<Real>& vmax);
+
+// Lift n-tuple v to homogeneous (n+1)-tuple (v,last).
+template <typename Real>
+GVector<Real> HLift(GVector<Real> const& v, Real last);
+
+// Project homogeneous n-tuple v = (u,v[n-1]) to (n-1)-tuple u.
+template <typename Real>
+GVector<Real> HProject(GVector<Real> const& v);
+
+// Lift n-tuple v = (w0,w1) to (n+1)-tuple u = (w0,u[inject],w1).  By
+// inference, w0 is a (inject)-tuple [nonexistent when inject=0] and w1 is a
+// (n-inject)-tuple [nonexistent when inject=n].
+template <typename Real>
+GVector<Real> Lift(GVector<Real> const& v, int inject, Real value);
+
+// Project n-tuple v = (w0,v[reject],w1) to (n-1)-tuple u = (w0,w1).  By
+// inference, w0 is a (reject)-tuple [nonexistent when reject=0] and w1 is a
+// (n-1-reject)-tuple [nonexistent when reject=n-1].
+template <typename Real>
+GVector<Real> Project(GVector<Real> const& v, int reject);
+
+
+template <typename Real>
+GVector<Real>::GVector()
+{
+    // Uninitialized.
+}
+
+template <typename Real>
+GVector<Real>::GVector(int size)
+{
+    SetSize(size);
+}
+
+template <typename Real>
+GVector<Real>::GVector(int size, int d)
+{
+    SetSize(size);
+    MakeUnit(d);
+}
+
+template <typename Real>
+void GVector<Real>::SetSize(int size)
+{
+#if defined(GTE_ASSERT_ON_GVECTOR_SIZE_MISMATCH)
+    LogAssert(size >= 0, "Mismatched size.");
+#endif
+    if (size > 0)
+    {
+        mTuple.resize(size);
+    }
+}
+
+template <typename Real> inline
+int GVector<Real>::GetSize() const
+{
+    return static_cast<int>(mTuple.size());
+}
+
+template <typename Real> inline
+Real const& GVector<Real>::operator[](int i) const
+{
+    return mTuple[i];
+}
+
+template <typename Real> inline
+Real& GVector<Real>::operator[](int i)
+{
+    return mTuple[i];
+}
+
+template <typename Real> inline
+bool GVector<Real>::operator==(GVector const& vec) const
+{
+    return mTuple == vec.mTuple;
+}
+
+template <typename Real> inline
+bool GVector<Real>::operator!=(GVector const& vec) const
+{
+    return mTuple != vec.mTuple;
+}
+
+template <typename Real> inline
+bool GVector<Real>::operator<(const GVector& vec) const
+{
+    return mTuple < vec.mTuple;
+}
+
+template <typename Real> inline
+bool GVector<Real>::operator<=(const GVector& vec) const
+{
+    return mTuple <= vec.mTuple;
+}
+
+template <typename Real> inline
+bool GVector<Real>::operator>(const GVector& vec) const
+{
+    return mTuple > vec.mTuple;
+}
+
+template <typename Real> inline
+bool GVector<Real>::operator>=(const GVector& vec) const
+{
+    return mTuple >= vec.mTuple;
+}
+
+template <typename Real>
+void GVector<Real>::MakeZero()
+{
+    std::fill(mTuple.begin(), mTuple.end(), (Real)0);
+}
+
+template <typename Real>
+void GVector<Real>::MakeUnit(int d)
+{
+    std::fill(mTuple.begin(), mTuple.end(), (Real)0);
+    if (0 <= d && d < (int)mTuple.size())
+    {
+        mTuple[d] = (Real)1;
+    }
+}
+
+template <typename Real>
+GVector<Real> GVector<Real>::Zero(int size)
+{
+    GVector<Real> v(size);
+    v.MakeZero();
+    return v;
+}
+
+template <typename Real>
+GVector<Real> GVector<Real>::Unit(int size, int d)
+{
+    GVector<Real> v(size);
+    v.MakeUnit(d);
+    return v;
+}
+
+
+
+template <typename Real>
+GVector<Real> operator+(GVector<Real> const& v)
+{
+    return v;
+}
+
+template <typename Real>
+GVector<Real> operator-(GVector<Real> const& v)
+{
+    GVector<Real> result(v.GetSize());
+    for (int i = 0; i < v.GetSize(); ++i)
+    {
+        result[i] = -v[i];
+    }
+    return result;
+}
+
+template <typename Real>
+GVector<Real> operator+(GVector<Real> const& v0, GVector<Real> const& v1)
+{
+    GVector<Real> result = v0;
+    return result += v1;
+}
+
+template <typename Real>
+GVector<Real> operator-(GVector<Real> const& v0, GVector<Real> const& v1)
+{
+    GVector<Real> result = v0;
+    return result -= v1;
+}
+
+template <typename Real>
+GVector<Real> operator*(GVector<Real> const& v, Real scalar)
+{
+    GVector<Real> result = v;
+    return result *= scalar;
+}
+
+template <typename Real>
+GVector<Real> operator*(Real scalar, GVector<Real> const& v)
+{
+    GVector<Real> result = v;
+    return result *= scalar;
+}
+
+template <typename Real>
+GVector<Real> operator/(GVector<Real> const& v, Real scalar)
+{
+    GVector<Real> result = v;
+    return result /= scalar;
+}
+
+template <typename Real>
+GVector<Real>& operator+=(GVector<Real>& v0, GVector<Real> const& v1)
+{
+    if (v0.GetSize() == v1.GetSize())
+    {
+        for (int i = 0; i < v0.GetSize(); ++i)
+        {
+            v0[i] += v1[i];
+        }
+    }
+    else
+    {
+#if defined(GTE_ASSERT_ON_GVECTOR_SIZE_MISMATCH)
+        LogError("Mismatched size.");
+#endif
+    }
+    return v0;
+}
+
+template <typename Real>
+GVector<Real>& operator-=(GVector<Real>& v0, GVector<Real> const& v1)
+{
+    if (v0.GetSize() == v1.GetSize())
+    {
+        for (int i = 0; i < v0.GetSize(); ++i)
+        {
+            v0[i] -= v1[i];
+        }
+    }
+    else
+    {
+#if defined(GTE_ASSERT_ON_GVECTOR_SIZE_MISMATCH)
+        LogError("Mismatched size.");
+#endif
+    }
+    return v0;
+}
+
+template <typename Real>
+GVector<Real>& operator*=(GVector<Real>& v, Real scalar)
+{
+    for (int i = 0; i < v.GetSize(); ++i)
+    {
+        v[i] *= scalar;
+    }
+    return v;
+}
+
+template <typename Real>
+GVector<Real>& operator/=(GVector<Real>& v, Real scalar)
+{
+    if (scalar != (Real)0)
+    {
+        Real invScalar = ((Real)1) / scalar;
+        for (int i = 0; i < v.GetSize(); ++i)
+        {
+            v[i] *= invScalar;
+        }
+    }
+    else
+    {
+        for (int i = 0; i < v.GetSize(); ++i)
+        {
+            v[i] = (Real)0;
+        }
+    }
+    return v;
+}
+
+template <typename Real>
+Real Dot(GVector<Real> const& v0, GVector<Real> const& v1)
+{
+    if (v0.GetSize() == v1.GetSize())
+    {
+        Real dot = v0[0] * v1[0];
+        for (int i = 1; i < v0.GetSize(); ++i)
+        {
+            dot += v0[i] * v1[i];
+        }
+        return dot;
+    }
+    else
+    {
+#if defined(GTE_ASSERT_ON_GVECTOR_SIZE_MISMATCH)
+        LogError("Mismatched size.");
+#endif
+        return (Real)0;
+    }
+}
+
+template <typename Real>
+Real Length(GVector<Real> const& v, bool robust)
+{
+    if (robust)
+    {
+        Real maxAbsComp = std::abs(v[0]);
+        for (int i = 1; i < v.GetSize(); ++i)
+        {
+            Real absComp = std::abs(v[i]);
+            if (absComp > maxAbsComp)
+            {
+                maxAbsComp = absComp;
+            }
+        }
+
+        Real length;
+        if (maxAbsComp > (Real)0)
+        {
+            GVector<Real> scaled = v / maxAbsComp;
+            length = maxAbsComp * std::sqrt(Dot(scaled, scaled));
+        }
+        else
+        {
+            length = (Real)0;
+        }
+        return length;
+    }
+    else
+    {
+        return std::sqrt(Dot(v, v));
+    }
+}
+
+template <typename Real>
+Real Normalize(GVector<Real>& v, bool robust)
+{
+    if (robust)
+    {
+        Real maxAbsComp = std::abs(v[0]);
+        for (int i = 1; i < v.GetSize(); ++i)
+        {
+            Real absComp = std::abs(v[i]);
+            if (absComp > maxAbsComp)
+            {
+                maxAbsComp = absComp;
+            }
+        }
+
+        Real length;
+        if (maxAbsComp > (Real)0)
+        {
+            v /= maxAbsComp;
+            length = std::sqrt(Dot(v, v));
+            v /= length;
+            length *= maxAbsComp;
+        }
+        else
+        {
+            length = (Real)0;
+            for (int i = 0; i < v.GetSize(); ++i)
+            {
+                v[i] = (Real)0;
+            }
+        }
+        return length;
+    }
+    else
+    {
+        Real length = std::sqrt(Dot(v, v));
+        if (length > (Real)0)
+        {
+            v /= length;
+        }
+        else
+        {
+            for (int i = 0; i < v.GetSize(); ++i)
+            {
+                v[i] = (Real)0;
+            }
+        }
+        return length;
+    }
+}
+
+template <typename Real>
+Real Orthonormalize(int numInputs, GVector<Real>* v, bool robust)
+{
+    if (v && 1 <= numInputs && numInputs <= v[0].GetSize())
+    {
+#if defined(GTE_ASSERT_ON_GVECTOR_SIZE_MISMATCH)
+        for (int i = 1; i < numInputs; ++i)
+        {
+            LogAssert(v[0].GetSize() == v[i].GetSize(), "Mismatched size.");
+        }
+#endif
+        Real minLength = Normalize(v[0], robust);
+        for (int i = 1; i < numInputs; ++i)
+        {
+            for (int j = 0; j < i; ++j)
+            {
+                Real dot = Dot(v[i], v[j]);
+                v[i] -= v[j] * dot;
+            }
+            Real length = Normalize(v[i], robust);
+            if (length < minLength)
+            {
+                minLength = length;
+            }
+        }
+        return minLength;
+    }
+
+    LogError("Invalid input.");
+    return (Real)0;
+}
+
+template <typename Real>
+bool ComputeExtremes(int numVectors, GVector<Real> const* v,
+    GVector<Real>& vmin, GVector<Real>& vmax)
+{
+    if (v && numVectors > 0)
+    {
+#if defined(GTE_ASSERT_ON_GVECTOR_SIZE_MISMATCH)
+        for (int i = 1; i < numVectors; ++i)
+        {
+            LogAssert(v[0].GetSize() == v[i].GetSize(), "Mismatched size.");
+        }
+#endif
+        int const size = v[0].GetSize();
+        vmin = v[0];
+        vmax = vmin;
+        for (int j = 1; j < numVectors; ++j)
+        {
+            GVector<Real> const& vec = v[j];
+            for (int i = 0; i < size; ++i)
+            {
+                if (vec[i] < vmin[i])
+                {
+                    vmin[i] = vec[i];
+                }
+                else if (vec[i] > vmax[i])
+                {
+                    vmax[i] = vec[i];
+                }
+            }
+        }
+        return true;
+    }
+
+    LogError("Invalid input.");
+    return false;
+}
+
+template <typename Real>
+GVector<Real> HLift(GVector<Real> const& v, Real last)
+{
+    int const size = v.GetSize();
+    GVector<Real> result(size + 1);
+    for (int i = 0; i < size; ++i)
+    {
+        result[i] = v[i];
+    }
+    result[size] = last;
+    return result;
+}
+
+template <typename Real>
+GVector<Real> HProject(GVector<Real> const& v)
+{
+    int const size = v.GetSize();
+    if (size > 1)
+    {
+        GVector<Real> result(size - 1);
+        for (int i = 0; i < size - 1; ++i)
+        {
+            result[i] = v[i];
+        }
+        return result;
+    }
+    else
+    {
+        return GVector<Real>();
+    }
+}
+
+template <typename Real>
+GVector<Real> Lift(GVector<Real> const& v, int inject, Real value)
+{
+    int const size = v.GetSize();
+    GVector<Real> result(size + 1);
+    int i;
+    for (i = 0; i < inject; ++i)
+    {
+        result[i] = v[i];
+    }
+    result[i] = value;
+    int j = i;
+    for (++j; i < size; ++i, ++j)
+    {
+        result[j] = v[i];
+    }
+    return result;
+}
+
+template <typename Real>
+GVector<Real> Project(GVector<Real> const& v, int reject)
+{
+    int const size = v.GetSize();
+    if (size > 1)
+    {
+        GVector<Real> result(size - 1);
+        for (int i = 0, j = 0; i < size - 1; ++i, ++j)
+        {
+            if (j == reject)
+            {
+                ++j;
+            }
+            result[i] = v[j];
+        }
+        return result;
+    }
+    else
+    {
+        return GVector<Real>();
+    }
+}
+
+}
diff --git a/Engine/Plugins/Experimental/GeometryProcessing/Source/GeometryAlgorithms/Private/GTEngine/Mathematics/GteGaussNewtonMinimizer.h b/Engine/Plugins/Experimental/GeometryProcessing/Source/GeometryAlgorithms/Private/GTEngine/Mathematics/GteGaussNewtonMinimizer.h
new file mode 100644
--- /dev/null
+++ b/Engine/Plugins/Experimental/GeometryProcessing/Source/GeometryAlgorithms/Private/GTEngine/Mathematics/GteGaussNewtonMinimizer.h
@@ -0,0 +1,228 @@
+// David Eberly, Geometric Tools, Redmond WA 98052
+// Copyright (c) 1998-2019
+// Distributed under the Boost Software License, Version 1.0.
+// http://www.boost.org/LICENSE_1_0.txt
+// http://www.geometrictools.com/License/Boost/LICENSE_1_0.txt
+// File Version: 3.14.0 (2018/07/16)
+
+#pragma once
+
+#include <GTEngine/Mathematics/GteCholeskyDecomposition.h>
+#include <functional>
+
+// Let F(p) = (F_{0}(p), F_{1}(p), ..., F_{n-1}(p)) be a vector-valued
+// function of the parameters p = (p_{0}, p_{1}, ..., p_{m-1}).  The
+// nonlinear least-squares problem is to minimize the real-valued error
+// function E(p) = |F(p)|^2, which is the squared length of F(p).
+//
+// Let J = dF/dp = [dF_{r}/dp_{c}] denote the Jacobian matrix, which is the
+// matrix of first-order partial derivatives of F.  The matrix has n rows and
+// m columns, and the indexing (r,c) refers to row r and column c.  A
+// first-order approximation is F(p + d) = F(p) + J(p)d, where d is an m-by-1
+// vector with small length.  Consequently, an approximation to E is E(p + d)
+// = |F(p + d)|^2 = |F(p) + J(p)d|^2.  The goal is to choose d to minimize
+// |F(p) + J(p)d|^2 and, hopefully, with E(p + d) < E(p).  Choosing an initial
+// p_{0}, the hope is that the algorithm generates a sequence p_{i} for which
+// E(p_{i+1}) < E(p_{i}) and, in the limit, E(p_{j}) approaches the global
+// minimum of E.  The algorithm is referred to as Gauss-Newton iteration.  If
+// E does not decrease for a step of the algorithm, one can modify the
+// algorithm to the Levenberg-Marquardt iteration.  See
+// GteLevenbergMarquardtMinimizer.h for a description and an implementation.
+//
+// For a single Gauss-Newton iteration, we need to choose d to minimize
+// |F(p) + J(p)d|^2 where p is fixed.  This is a linear least squares problem
+// which can be formulated using the normal equations
+// (J^T(p)*J(p))*d = -J^T(p)*F(p).  The matrix J^T*J is positive semidefinite.
+// If it is invertible, then d = -(J^T(p)*J(p))^{-1}*F(p).  If it is not
+// invertible, some other algorithm must be used to choose d; one option is
+// to use gradient descent for the step.  A Cholesky decomposition can be
+// used to solve the linear system.
+//
+// Although an implementation can allow the caller to pass an array of
+// functions F_{i}(p) and an array of derivatives dF_{r}/dp_{c}, some
+// applications might involve a very large n that precludes storing all
+// the computed Jacobian matrix entries because of excessive memory
+// requirements.  In such an application, it is better to compute instead
+// the entries of the m-by-m matrix J^T*J and the m-by-1 vector J^T*F.
+// Typically, m is small, so the memory requirements are not excessive.  Also,
+// there might be additional structure to F for which the caller can take
+// advantage; for example, 3-tuples of components of F(p) might correspond to
+// vectors that can be manipulated using an already existing mathematics
+// library.  The implementation here supports both approaches.
+
+namespace gte
+{
+    template <typename Real>
+    class GaussNewtonMinimizer
+    {
+    public:
+        // Convenient types for the domain vectors, the range vectors, the
+        // function F and the Jacobian J.
+        typedef GVector<Real> DVector;  // NumPDimensions elements
+        typedef GVector<Real> RVector;  // NumFDimensions elements
+        typedef GMatrix<Real> JMatrix;  // NumFDimensions-by-NumPDimensions elements
+        typedef GMatrix<Real> JTJMatrix;  // NumPDimensions-by-NumPDimensions
+        typedef GVector<Real> JTFVector;  // NumPDimensions elements
+        typedef std::function<void(DVector const&, RVector&)> FFunction;
+        typedef std::function<void(DVector const&, JMatrix&)> JFunction;
+        typedef std::function<void(DVector const&, JTJMatrix&, JTFVector&)> JPlusFunction;
+
+        // NOTE: The C++ compiler for Microsoft Visual Studio 12.0.21005.1 REL
+        // (MSVS 2013) appears to have a bug regarding passing std::function
+        // arguments.  The unit-test code for the first GaussNewtonMinimizer
+        // constructor was passed std::function objects of type FFunction and
+        // JFunction.  The compiler claimed that the constructor call is
+        // ambiguous because the JFunction portion of the first constructor's
+        // signature matches the JPlusFunction portion of the second
+        // constructor's signature, which is clearly not correct.  When using
+        // MSVS 2013, you need to explicitly typecast using
+        //   GaussNewtonMinimizer<Real> minimizer(
+        //       numPDimensions, numFDimensions,
+        //       (GaussNewtonMinimizer<Real>::FFunction const&) myFFunction,
+        //       (GaussNewtonMinimizer<Real>::JFunction const&) myJFunction);
+        // The same typecasting is also required for the second constructor.
+
+        // Create the minimizer that computes F(p) and J(p) directly.
+        GaussNewtonMinimizer(int numPDimensions, int numFDimensions,
+            FFunction const& inFFunction, JFunction const& inJFunction)
+            :
+            mNumPDimensions(numPDimensions),
+            mNumFDimensions(numFDimensions),
+            mFFunction(inFFunction),
+            mJFunction(inJFunction),
+            mF(mNumFDimensions),
+            mJ(mNumFDimensions, mNumPDimensions),
+            mJTJ(mNumPDimensions, mNumPDimensions),
+            mNegJTF(mNumPDimensions),
+            mDecomposer(mNumPDimensions),
+            mUseJFunction(true)
+        {
+            LogAssert(mNumPDimensions > 0 && mNumFDimensions > 0, "Invalid dimensions.");
+        }
+
+        // Create the minimizer that computes J^T(p)*J(p) and -J(p)*F(p).
+        GaussNewtonMinimizer(int numPDimensions, int numFDimensions,
+            FFunction const& inFFunction, JPlusFunction const& inJPlusFunction)
+            :
+            mNumPDimensions(numPDimensions),
+            mNumFDimensions(numFDimensions),
+            mFFunction(inFFunction),
+            mJPlusFunction(inJPlusFunction),
+            mF(mNumFDimensions),
+            mJ(mNumFDimensions, mNumPDimensions),
+            mJTJ(mNumPDimensions, mNumPDimensions),
+            mNegJTF(mNumPDimensions),
+            mDecomposer(mNumPDimensions),
+            mUseJFunction(false)
+        {
+            LogAssert(mNumPDimensions > 0 && mNumFDimensions > 0, "Invalid dimensions.");
+        }
+
+        // Disallow copy, assignment and move semantics.
+        GaussNewtonMinimizer(GaussNewtonMinimizer const&) = delete;
+        GaussNewtonMinimizer& operator=(GaussNewtonMinimizer const&) = delete;
+        GaussNewtonMinimizer(GaussNewtonMinimizer&&) = delete;
+        GaussNewtonMinimizer& operator=(GaussNewtonMinimizer&&) = delete;
+
+        inline int GetNumPDimensions() const { return mNumPDimensions; }
+        inline int GetNumFDimensions() const { return mNumFDimensions; }
+
+        struct Result
+        {
+            DVector minLocation;
+            Real minError;
+            Real minErrorDifference;
+            Real minUpdateLength;
+            size_t numIterations;
+            bool converged;
+        };
+
+        Result operator()(DVector const& p0, size_t maxIterations,
+            Real updateLengthTolerance, Real errorDifferenceTolerance)
+        {
+            Result result;
+            result.minLocation = p0;
+            result.minError = std::numeric_limits<Real>::max();
+            result.minErrorDifference = std::numeric_limits<Real>::max();
+            result.minUpdateLength = (Real)0;
+            result.numIterations = 0;
+            result.converged = false;
+
+            // As a simple precaution, ensure the tolerances are nonnegative.
+            updateLengthTolerance = std::max(updateLengthTolerance, (Real)0);
+            errorDifferenceTolerance = std::max(errorDifferenceTolerance, (Real)0);
+
+            // Compute the initial error.
+            mFFunction(p0, mF);
+            result.minError = Dot(mF, mF);
+
+            // Do the Gauss-Newton iterations.
+            auto pCurrent = p0;
+            for (result.numIterations = 1; result.numIterations <= maxIterations; ++result.numIterations)
+            {
+                ComputeLinearSystemInputs(pCurrent);
+                if (!mDecomposer.Factor(mJTJ))
+                {
+                    // TODO: The matrix mJTJ is positive semi-definite, so the
+                    // failure can occur when mJTJ has a zero eigenvalue in
+                    // which case mJTJ is not invertible.  Generate an iterate
+                    // anyway, perhaps using gradient descent?
+                    return result;
+                }
+                mDecomposer.SolveLower(mJTJ, mNegJTF);
+                mDecomposer.SolveUpper(mJTJ, mNegJTF);
+
+                auto pNext = pCurrent + mNegJTF;
+                mFFunction(pNext, mF);
+                Real error = Dot(mF, mF);
+                if (error < result.minError)
+                {
+                    result.minErrorDifference = result.minError - error;
+                    result.minUpdateLength = Length(mNegJTF);
+                    result.minLocation = pNext;
+                    result.minError = error;
+                    if (result.minErrorDifference <= errorDifferenceTolerance
+                        || result.minUpdateLength <= updateLengthTolerance)
+                    {
+                        result.converged = true;
+                        return result;
+                    }
+                }
+
+                pCurrent = pNext;
+            }
+
+            return result;
+        }
+
+    private:
+        void ComputeLinearSystemInputs(DVector const& pCurrent)
+        {
+            if (mUseJFunction)
+            {
+                mJFunction(pCurrent, mJ);
+                mJTJ = MultiplyATB(mJ, mJ);
+                mNegJTF = -(mF * mJ);
+            }
+            else
+            {
+                mJPlusFunction(pCurrent, mJTJ, mNegJTF);
+            }
+        }
+
+        int mNumPDimensions, mNumFDimensions;
+        FFunction mFFunction;
+        JFunction mJFunction;
+        JPlusFunction mJPlusFunction;
+
+        // Storage for J^T(p)*J(p) and -J^T(p)*F(p) during the iterations.
+        RVector mF;
+        JMatrix mJ;
+        JTJMatrix mJTJ;
+        JTFVector mNegJTF;
+
+        CholeskyDecomposition<Real> mDecomposer;
+
+        bool mUseJFunction;
+    };
+}
diff --git a/Engine/Plugins/Experimental/GeometryProcessing/Source/GeometryAlgorithms/Private/GTEngine/Mathematics/GteGaussianElimination.h b/Engine/Plugins/Experimental/GeometryProcessing/Source/GeometryAlgorithms/Private/GTEngine/Mathematics/GteGaussianElimination.h
new file mode 100644
--- /dev/null
+++ b/Engine/Plugins/Experimental/GeometryProcessing/Source/GeometryAlgorithms/Private/GTEngine/Mathematics/GteGaussianElimination.h
@@ -0,0 +1,290 @@
+// David Eberly, Geometric Tools, Redmond WA 98052
+// Copyright (c) 1998-2019
+// Distributed under the Boost Software License, Version 1.0.
+// http://www.boost.org/LICENSE_1_0.txt
+// http://www.geometrictools.com/License/Boost/LICENSE_1_0.txt
+// File Version: 3.0.0 (2016/06/19)
+
+#pragma once
+
+#include <GTEngine/LowLevel/GteLexicoArray2.h>
+#include <GTEngine/LowLevel/GteLogger.h>
+#include <GTEngine/LowLevel/GteWrapper.h>
+#include <cmath>
+#include <cstring>
+#include <vector>
+
+// The input matrix M must be NxN.  The storage convention for element lookup
+// is determined by GTE_USE_ROW_MAJOR or GTE_USE_COL_MAJOR, whichever is
+// active.  If you want the inverse of M, pass a nonnull pointer inverseM;
+// this matrix must also be NxN and use the same storage convention as M.  If
+// you do not want the inverse of M, pass a nullptr for inverseM.  If you want
+// to solve M*X = B for X, where X and B are Nx1, pass nonnull pointers for B
+// and X.  If you want to solve M*Y = C for Y, where X and C are NxK, pass
+// nonnull pointers for C and Y and pass K to numCols.  In all cases, pass
+// N to numRows.
+
+namespace gte
+{
+
+template <typename Real>
+class GaussianElimination
+{
+public:
+    bool operator()(int numRows,
+        Real const* M, Real* inverseM, Real& determinant,
+        Real const* B, Real* X,
+        Real const* C, int numCols, Real* Y) const;
+
+private:
+    // Support for copying source to target or to set target to zero.  If
+    // source is nullptr, then target is set to zero; otherwise source is
+    // copied to target.  This function hides the type traits used to
+    // determine whether Real is native floating-point or otherwise (such
+    // as BSNumber or BSRational).
+    void Set(int numElements, Real const* source, Real* target) const;
+};
+
+
+template <typename Real>
+bool GaussianElimination<Real>::operator()(int numRows, Real const* M,
+    Real* inverseM, Real& determinant, Real const* B, Real* X, Real const* C,
+    int numCols, Real* Y) const
+{
+    if (numRows <= 0 || !M
+        || ((B != nullptr) != (X != nullptr))
+        || ((C != nullptr) != (Y != nullptr))
+        || (C != nullptr && numCols < 1))
+    {
+        LogError("Invalid input.");
+        return false;
+    }
+
+    int numElements = numRows * numRows;
+    bool wantInverse = (inverseM != nullptr);
+    std::vector<Real> localInverseM;
+    if (!wantInverse)
+    {
+        localInverseM.resize(numElements);
+        inverseM = localInverseM.data();
+    }
+    Set(numElements, M, inverseM);
+
+    if (B)
+    {
+        Set(numRows, B, X);
+    }
+
+    if (C)
+    {
+        Set(numRows * numCols, C, Y);
+    }
+
+#if defined(GTE_USE_ROW_MAJOR)
+    LexicoArray2<true, Real> matInvM(numRows, numRows, inverseM);
+    LexicoArray2<true, Real> matY(numRows, numCols, Y);
+#else
+    LexicoArray2<false, Real> matInvM(numRows, numRows, inverseM);
+    LexicoArray2<false, Real> matY(numRows, numCols, Y);
+#endif
+
+    std::vector<int> colIndex(numRows), rowIndex(numRows), pivoted(numRows);
+    std::fill(pivoted.begin(), pivoted.end(), 0);
+
+    Real const zero = (Real)0;
+    Real const one = (Real)1;
+    bool odd = false;
+    determinant = one;
+
+    // Elimination by full pivoting.
+    int i1, i2, row = 0, col = 0;
+    for (int i0 = 0; i0 < numRows; ++i0)
+    {
+        // Search matrix (excluding pivoted rows) for maximum absolute entry.
+        Real maxValue = zero;
+        for (i1 = 0; i1 < numRows; ++i1)
+        {
+            if (!pivoted[i1])
+            {
+                for (i2 = 0; i2 < numRows; ++i2)
+                {
+                    if (!pivoted[i2])
+                    {
+                        Real absValue = std::abs(matInvM(i1, i2));
+                        if (absValue > maxValue)
+                        {
+                            maxValue = absValue;
+                            row = i1;
+                            col = i2;
+                        }
+                    }
+                }
+            }
+        }
+
+        if (maxValue == zero)
+        {
+            // The matrix is not invertible.
+            if (wantInverse)
+            {
+                Set(numElements, nullptr, inverseM);
+            }
+            determinant = zero;
+
+            if (B)
+            {
+                Set(numRows, nullptr, X);
+            }
+
+            if (C)
+            {
+                Set(numRows * numCols, nullptr, Y);
+            }
+            return false;
+        }
+
+        pivoted[col] = true;
+
+        // Swap rows so that the pivot entry is in row 'col'.
+        if (row != col)
+        {
+            odd = !odd;
+            for (int i = 0; i < numRows; ++i)
+            {
+                std::swap(matInvM(row, i), matInvM(col, i));
+            }
+
+            if (B)
+            {
+                std::swap(X[row], X[col]);
+            }
+
+            if (C)
+            {
+                for (int i = 0; i < numCols; ++i)
+                {
+                    std::swap(matY(row, i), matY(col, i));
+                }
+            }
+        }
+
+        // Keep track of the permutations of the rows.
+        rowIndex[i0] = row;
+        colIndex[i0] = col;
+
+        // Scale the row so that the pivot entry is 1.
+        Real diagonal = matInvM(col, col);
+        determinant *= diagonal;
+        Real inv = one / diagonal;
+        matInvM(col, col) = one;
+        for (i2 = 0; i2 < numRows; ++i2)
+        {
+            matInvM(col, i2) *= inv;
+        }
+
+        if (B)
+        {
+            X[col] *= inv;
+        }
+
+        if (C)
+        {
+            for (i2 = 0; i2 < numCols; ++i2)
+            {
+                matY(col, i2) *= inv;
+            }
+        }
+
+        // Zero out the pivot column locations in the other rows.
+        for (i1 = 0; i1 < numRows; ++i1)
+        {
+            if (i1 != col)
+            {
+                Real save = matInvM(i1, col);
+                matInvM(i1, col) = zero;
+                for (i2 = 0; i2 < numRows; ++i2)
+                {
+                    matInvM(i1, i2) -= matInvM(col, i2) * save;
+                }
+
+                if (B)
+                {
+                    X[i1] -= X[col] * save;
+                }
+
+                if (C)
+                {
+                    for (i2 = 0; i2 < numCols; ++i2)
+                    {
+                        matY(i1, i2) -= matY(col, i2) * save;
+                    }
+                }
+            }
+        }
+    }
+
+    if (wantInverse)
+    {
+        // Reorder rows to undo any permutations in Gaussian elimination.
+        for (i1 = numRows - 1; i1 >= 0; --i1)
+        {
+            if (rowIndex[i1] != colIndex[i1])
+            {
+                for (i2 = 0; i2 < numRows; ++i2)
+                {
+                    std::swap(matInvM(i2, rowIndex[i1]),
+                        matInvM(i2, colIndex[i1]));
+                }
+            }
+        }
+    }
+
+    if (odd)
+    {
+        determinant = -determinant;
+    }
+
+    return true;
+}
+
+template <typename Real>
+void GaussianElimination<Real>::Set(int numElements, Real const* source,
+    Real* target) const
+{
+    if (std::is_floating_point<Real>() == std::true_type())
+    {
+        // Fast set/copy for native floating-point.
+        size_t numBytes = numElements * sizeof(Real);
+        if (source)
+        {
+            Memcpy(target, source, numBytes);
+        }
+        else
+        {
+            memset(target, 0, numBytes);
+        }
+    }
+    else
+    {
+        // The inputs are not std containers, so ensure assignment works
+        // correctly.
+        if (source)
+        {
+            for (int i = 0; i < numElements; ++i)
+            {
+                target[i] = source[i];
+            }
+        }
+        else
+        {
+            Real const zero = (Real)0;
+            for (int i = 0; i < numElements; ++i)
+            {
+                target[i] = zero;
+            }
+        }
+    }
+}
+
+
+}
diff --git a/Engine/Plugins/Experimental/GeometryProcessing/Source/GeometryAlgorithms/Private/GTEngine/Mathematics/GteGenerateMeshUV.cpp b/Engine/Plugins/Experimental/GeometryProcessing/Source/GeometryAlgorithms/Private/GTEngine/Mathematics/GteGenerateMeshUV.cpp
new file mode 100644
--- /dev/null
+++ b/Engine/Plugins/Experimental/GeometryProcessing/Source/GeometryAlgorithms/Private/GTEngine/Mathematics/GteGenerateMeshUV.cpp
@@ -0,0 +1,100 @@
+// David Eberly, Geometric Tools, Redmond WA 98052
+// Copyright (c) 1998-2019
+// Distributed under the Boost Software License, Version 1.0.
+// http://www.boost.org/LICENSE_1_0.txt
+// http://www.geometrictools.com/License/Boost/LICENSE_1_0.txt
+// File Version: 3.0.0 (2016/06/19)
+
+#include <GTEngine/Mathematics/GteGenerateMeshUV.h>
+#include <GTEngine/GTEnginePCH.h>
+
+using namespace gte;
+
+std::string const GenerateMeshUVBase::msGLSLSource =
+"uniform Bounds\n"
+"{\n"
+"    ivec2 bound;\n"
+"    int numBoundaryEdges;\n"
+"    int numInputs;\n"
+"};\n"
+"\n"
+"struct VertexGraphData\n"
+"{\n"
+"    int adjacent;\n"
+"    Real weight;\n"
+"};\n"
+"\n"
+"buffer vertexGraph { ivec3 data[]; } vertexGraphSB;\n"
+"buffer vertexGraphData { VertexGraphData data[]; } vertexGraphDataSB;\n"
+"buffer orderedVertices { int data[]; } orderedVerticesSB;\n"
+"buffer inTCoords { Real2 data[]; } inTCoordsSB;\n"
+"buffer outTCoords { Real2 data[]; } outTCoordsSB;\n"
+"\n"
+"layout (local_size_x = NUM_X_THREADS, local_size_y = NUM_Y_THREADS, local_size_z = 1) in;\n"
+"void main()\n"
+"{\n"
+"    ivec2 t = ivec2(gl_GlobalInvocationID.xy);\n"
+"    int index = t.x + bound.x * t.y;\n"
+"    if (step(index, numInputs-1) == 1)\n"
+"    {\n"
+"        int v = orderedVerticesSB.data[numBoundaryEdges + index];\n"
+"        ivec2 range = vertexGraphSB.data[v].yz;\n"
+"        Real2 tcoord = Real2(0, 0);\n"
+"        Real weightSum = 0;\n"
+"        for (int j = 0; j < range.y; ++j)\n"
+"        {\n"
+"            VertexGraphData vgd = vertexGraphDataSB.data[range.x + j];\n"
+"            weightSum += vgd.weight;\n"
+"            tcoord += vgd.weight * inTCoordsSB.data[vgd.adjacent];\n"
+"        }\n"
+"        tcoord /= weightSum;\n"
+"        outTCoordsSB.data[v] = tcoord;\n"
+"    }\n"
+"}\n";
+
+std::string const GenerateMeshUVBase::msHLSLSource =
+"cbuffer Bounds\n"
+"{\n"
+"    int2 bound;\n"
+"    int numBoundaryEdges;\n"
+"    int numInputs;\n"
+"};\n"
+"\n"
+"struct VertexGraphData\n"
+"{\n"
+"    int adjacent;\n"
+"    Real weight;\n"
+"};\n"
+"\n"
+"StructuredBuffer<int3> vertexGraph;\n"
+"StructuredBuffer<VertexGraphData> vertexGraphData;\n"
+"StructuredBuffer<int> orderedVertices;\n"
+"StructuredBuffer<Real2> inTCoords;\n"
+"RWStructuredBuffer<Real2> outTCoords;\n"
+"\n"
+"[numthreads(NUM_X_THREADS, NUM_Y_THREADS, 1)]\n"
+"void CSMain(int2 t : SV_DispatchThreadID)\n"
+"{\n"
+"    int index = t.x + bound.x * t.y;\n"
+"    if (step(index, numInputs-1))\n"
+"    {\n"
+"        int v = orderedVertices[numBoundaryEdges + index];\n"
+"        int2 range = vertexGraph[v].yz;\n"
+"        Real2 tcoord = Real2(0, 0);\n"
+"        Real weightSum = 0;\n"
+"        for (int j = 0; j < range.y; ++j)\n"
+"        {\n"
+"            VertexGraphData data = vertexGraphData[range.x + j];\n"
+"            weightSum += data.weight;\n"
+"            tcoord += data.weight * inTCoords[data.adjacent];\n"
+"        }\n"
+"        tcoord /= weightSum;\n"
+"        outTCoords[v] = tcoord;\n"
+"    }\n"
+"}\n";
+
+std::string const* GenerateMeshUVBase::msSource[] =
+{
+    &msGLSLSource,
+    &msHLSLSource
+};
diff --git a/Engine/Plugins/Experimental/GeometryProcessing/Source/GeometryAlgorithms/Private/GTEngine/Mathematics/GteGenerateMeshUV.h b/Engine/Plugins/Experimental/GeometryProcessing/Source/GeometryAlgorithms/Private/GTEngine/Mathematics/GteGenerateMeshUV.h
new file mode 100644
--- /dev/null
+++ b/Engine/Plugins/Experimental/GeometryProcessing/Source/GeometryAlgorithms/Private/GTEngine/Mathematics/GteGenerateMeshUV.h
@@ -0,0 +1,851 @@
+// David Eberly, Geometric Tools, Redmond WA 98052
+// Copyright (c) 1998-2019
+// Distributed under the Boost Software License, Version 1.0.
+// http://www.boost.org/LICENSE_1_0.txt
+// http://www.geometrictools.com/License/Boost/LICENSE_1_0.txt
+// File Version: 3.0.2 (2018/10/05)
+
+#pragma once
+
+#include <GTEngine/LowLevel/GteComputeModel.h>
+#include <GTEngine/LowLevel/GteWrapper.h>
+#include <GTEngine/Mathematics/GteVector2.h>
+#include <GTEngine/Mathematics/GteVector3.h>
+#include <GTEngine/Mathematics/GteETManifoldMesh.h>
+#include <GTEngine/Mathematics/GteMath.h>
+#if defined(GTE_COMPUTE_MODEL_ALLOW_GPGPU)
+#include <GTEngine/Graphics/GteConstantBuffer.h>
+#include <GTEngine/Graphics/GteStructuredBuffer.h>
+#include <GTEngine/Graphics/GteComputeProgram.h>
+#endif
+#include <algorithm>
+#include <cstring>
+#include <functional>
+#include <limits>
+#include <memory>
+#include <set>
+#include <thread>
+#include <vector>
+
+// This class is an implementation of the barycentric mapping algorithm
+// described in Section 5.3 of the book
+//     Polygon Mesh Processing
+//     Mario Botsch, Leif Kobbelt, Mark Pauly, Pierre Alliez, Bruno Levy
+//     AK Peters, Ltd., Natick MA, 2010
+// It uses the mean value weights described in Section 5.3.1 to allow the mesh
+// geometry to influence the texture coordinate generation, and it uses
+// Gauss-Seidel iteration to solve the sparse linear system.  The authors'
+// advice is that the Gauss-Seidel approach works well for at most about 5000
+// vertices, presumably the convergence rate degrading as the number of
+// vertices increases.
+//
+// The algorithm implemented here has an additional preprocessing step that
+// computes a topological distance transform of the vertices.  The boundary
+// texture coordinates are propagated inward by updating the vertices in
+// topological distance order, leading to fast convergence for large numbers
+// of vertices.
+
+namespace gte
+{
+
+class GenerateMeshUVBase
+{
+protected:
+    static std::string const msGLSLSource;
+    static std::string const msHLSLSource;
+    static std::string const* msSource[];
+};
+
+template <typename Real>
+class GenerateMeshUV : public GenerateMeshUVBase
+{
+public:
+    class UVComputeModel : public ComputeModel
+    {
+    public:
+        virtual ~UVComputeModel();
+
+        UVComputeModel();
+
+        UVComputeModel(unsigned int inNumThreads,
+            std::function<void(unsigned int)> const* inProgress);
+
+#if defined(GTE_COMPUTE_MODEL_ALLOW_GPGPU)
+        UVComputeModel(unsigned int inNumThreads,
+            std::shared_ptr<GraphicsEngine> const& inEngine,
+            std::shared_ptr<ProgramFactory> const& inFactory,
+            std::function<void(unsigned int)> const* inProgress);
+#endif
+
+        std::function<void(unsigned int)> const* progress;
+    };
+
+    // Construction.
+    GenerateMeshUV(std::shared_ptr<UVComputeModel> const& cmodel);
+
+    // The incoming mesh must be edge-triangle manifold and have rectangle
+    // topology (simply connected, closed polyline boundary).  The arrays
+    // 'vertices' and 'tcoords' must both have 'numVertices' elements.  Set
+    // 'useSquareTopology' to true for the generated coordinates to live
+    // in the uv-square [0,1]^2.  Set it to false for the generated
+    // coordinates to live in a convex polygon that inscribes the uv-disk
+    // of center (1/2,1/2) and radius 1/2.
+    void operator()(unsigned int numIterations, bool useSquareTopology,
+        int numVertices, Vector3<Real> const* vertices, int numIndices,
+        int const* indices, Vector2<Real>* tcoords);
+
+private:
+    void TopologicalVertexDistanceTransform();
+    void AssignBoundaryTextureCoordinatesSquare();
+    void AssignBoundaryTextureCoordinatesDisk();
+    void ComputeMeanValueWeights();
+    void SolveSystem(unsigned int numIterations);
+    void SolveSystemCPUSingle(unsigned int numIterations);
+    void SolveSystemCPUMultiple(unsigned int numIterations);
+
+    // Convenience members that store the input parameters to operator().
+    int mNumVertices;
+    Vector3<Real> const* mVertices;
+    Vector2<Real>* mTCoords;
+
+    // The edge-triangle manifold graph, where each edge is shared by at most
+    // two triangles.
+    ETManifoldMesh mGraph;
+
+    // The mVertexInfo array stores -1 for the interior vertices.  For a
+    // boundary edge <v0,v1> that is counterclockwise, mVertexInfo[v0] = v1,
+    // which gives us an orded boundary polyline.
+    enum { INTERIOR_VERTEX = -1 };
+    std::vector<int> mVertexInfo;
+    int mNumBoundaryEdges, mBoundaryStart;
+    typedef ETManifoldMesh::Edge Edge;
+    std::set<std::shared_ptr<Edge>> mInteriorEdges;
+
+    // The vertex graph required to set up a sparse linear system of equations
+    // to determine the texture coordinates.
+    struct Vertex
+    {
+        // The topological distance from the boundary of the mesh.
+        int distance;
+
+        // The value range[0] is the index into mVertexGraphData for the first
+        // adjacent vertex.  The value range[1] is the number of adjacent
+        // vertices.
+        std::array<int, 2> range;
+
+#if defined(GTE_COMPUTE_MODEL_ALLOW_GPGPU) && defined(GTE_DEV_OPENGL)
+        int _padding;  // GLSL will map the ivec3 Vertex to ivec4 in an array
+#endif
+    };
+    std::vector<Vertex> mVertexGraph;
+    std::vector<std::pair<int, Real>> mVertexGraphData;
+
+    // The vertices are listed in the order determined by a topological distance
+    // transform.  Boundary vertices have 'distance' 0.  Any vertices that are
+    // not boundary vertices but are edge-adjacent to boundary vertices have
+    // 'distance' 1.  Neighbors of those have distance '2', and so on.  The
+    // mOrderedVertices array stores distance-0 vertices first, distance-1
+    // vertices second, and so on.
+    std::vector<int> mOrderedVertices;
+
+    std::shared_ptr<UVComputeModel> mCModel;
+
+#if defined(GTE_COMPUTE_MODEL_ALLOW_GPGPU)
+    // Support for solving the sparse linear system on the GPU.
+    void SolveSystemGPU(unsigned int numIterations);
+    std::shared_ptr<ComputeProgram> mSolveSystem;
+    std::shared_ptr<ConstantBuffer> mBoundBuffer;
+    std::shared_ptr<StructuredBuffer> mVGBuffer;
+    std::shared_ptr<StructuredBuffer> mVGDBuffer;
+    std::shared_ptr<StructuredBuffer> mOVBuffer;
+    std::shared_ptr<StructuredBuffer> mTCoordsBuffer[2];
+#endif
+};
+
+
+template <typename Real>
+GenerateMeshUV<Real>::UVComputeModel::~UVComputeModel()
+{
+}
+
+template <typename Real>
+GenerateMeshUV<Real>::UVComputeModel::UVComputeModel()
+    :
+    progress(nullptr)
+{
+}
+
+template <typename Real>
+GenerateMeshUV<Real>::UVComputeModel::UVComputeModel(unsigned int inNumThreads,
+    std::function<void(unsigned int)> const* inProgress)
+    :
+    ComputeModel(inNumThreads),
+    progress(inProgress)
+{
+}
+
+#if defined(GTE_COMPUTE_MODEL_ALLOW_GPGPU)
+
+template <typename Real>
+GenerateMeshUV<Real>::UVComputeModel::UVComputeModel(unsigned int inNumThreads,
+    std::shared_ptr<GraphicsEngine> const& inEngine,
+    std::shared_ptr<ProgramFactory> const& inFactory,
+    std::function<void(unsigned int)> const* inProgress)
+    :
+    ComputeModel(inNumThreads, inEngine, inFactory),
+    progress(inProgress)
+{
+}
+
+#endif
+
+template <typename Real>
+GenerateMeshUV<Real>::GenerateMeshUV(std::shared_ptr<UVComputeModel> const& cmodel)
+    :
+    mNumVertices(0),
+    mVertices(nullptr),
+    mTCoords(nullptr),
+    mNumBoundaryEdges(0),
+    mBoundaryStart(0),
+    mCModel(cmodel)
+{
+}
+
+template <typename Real>
+void GenerateMeshUV<Real>::operator()(unsigned int numIterations,
+    bool useSquareTopology, int numVertices, Vector3<Real> const* vertices,
+    int numIndices, int const* indices, Vector2<Real>* tcoords)
+{
+    // Ensure that numIterations is even, which avoids having a memory
+    // copy from the temporary ping-pong buffer to 'tcoords'.
+    if (numIterations & 1)
+    {
+        ++numIterations;
+    }
+
+    mNumVertices = numVertices;
+    mVertices = vertices;
+    mTCoords = tcoords;
+
+    // The linear system solver has a first pass to initialize the texture
+    // coordinates to ensure the Gauss-Seidel iteration converges rapidly.
+    // This requires the texture coordinates all start as (-1,-1).
+    for (int i = 0; i < numVertices; ++i)
+    {
+        mTCoords[i][0] = (Real)-1;
+        mTCoords[i][1] = (Real)-1;
+    }
+
+    // Create the manifold mesh data structure.
+    mGraph.Clear();
+    int const numTriangles = numIndices / 3;
+    for (int t = 0; t < numTriangles; ++t)
+    {
+        int v0 = *indices++;
+        int v1 = *indices++;
+        int v2 = *indices++;
+        mGraph.Insert(v0, v1, v2);
+    }
+
+    TopologicalVertexDistanceTransform();
+
+    if (useSquareTopology)
+    {
+        AssignBoundaryTextureCoordinatesSquare();
+    }
+    else
+    {
+        AssignBoundaryTextureCoordinatesDisk();
+    }
+
+    ComputeMeanValueWeights();
+    SolveSystem(numIterations);
+}
+
+template <typename Real>
+void GenerateMeshUV<Real>::TopologicalVertexDistanceTransform()
+{
+    // Initialize the graph information.
+    mVertexInfo.resize(mNumVertices);
+    std::fill(mVertexInfo.begin(), mVertexInfo.end(), INTERIOR_VERTEX);
+    mVertexGraph.resize(mNumVertices);
+    mVertexGraphData.resize(2 * mGraph.GetEdges().size());
+    std::pair<int, Real> initialData = std::make_pair(-1, (Real)-1);
+    std::fill(mVertexGraphData.begin(), mVertexGraphData.end(), initialData);
+    mOrderedVertices.resize(mNumVertices);
+    mInteriorEdges.clear();
+    mNumBoundaryEdges = 0;
+    mBoundaryStart = std::numeric_limits<int>::max();
+
+    // Count the number of adjacent vertices for each vertex.  For data sets
+    // with a large number of vertices, this is a preprocessing step to avoid
+    // a dynamic data structure that has a large number of std:map objects
+    // that take a very long time to destroy when a debugger is attached to
+    // the executable.  Instead, we allocate a single array that stores all
+    // the adjacency information.  It is also necessary to bundle the data
+    // this way for a GPU version of the algorithm.
+    std::vector<int> numAdjacencies(mNumVertices);
+    std::fill(numAdjacencies.begin(), numAdjacencies.end(), 0);
+
+    for (auto const& element : mGraph.GetEdges())
+    {
+        ++numAdjacencies[element.first.V[0]];
+        ++numAdjacencies[element.first.V[1]];
+
+        if (element.second->T[1].lock())
+        {
+            // This is an interior edge.
+            mInteriorEdges.insert(element.second);
+        }
+        else
+        {
+            // This is a boundary edge.  Determine the ordering of the
+            // vertex indices to make the edge counterclockwise.
+            ++mNumBoundaryEdges;
+            int v0 = element.second->V[0], v1 = element.second->V[1];
+            auto tri = element.second->T[0].lock();
+            int i;
+            for (i = 0; i < 3; ++i)
+            {
+                int v2 = tri->V[i];
+                if (v2 != v0 && v2 != v1)
+                {
+                    // The vertex is opposite the boundary edge.
+                    v0 = tri->V[(i + 1) % 3];
+                    v1 = tri->V[(i + 2) % 3];
+                    mVertexInfo[v0] = v1;
+                    mBoundaryStart = std::min(mBoundaryStart, v0);
+                    break;
+                }
+            }
+        }
+    }
+
+    // Set the range data for each vertex.
+    for (int vIndex = 0, aIndex = 0; vIndex < mNumVertices; ++vIndex)
+    {
+        int numAdjacent = numAdjacencies[vIndex];
+        mVertexGraph[vIndex].range = { { aIndex, numAdjacent } };
+        aIndex += numAdjacent;
+
+#if defined(GTE_COMPUTE_MODEL_ALLOW_GPGPU) && defined(GTE_DEV_OPENGL)
+        // Initialize the padding, even though it is unused.
+        mVertexGraph[vIndex]._padding = 0;
+#endif
+    }
+
+    // Compute a topological distance transform of the vertices.
+    std::set<int> currFront;
+    for (auto const& element : mGraph.GetEdges())
+    {
+        int v0 = element.second->V[0], v1 = element.second->V[1];
+        for (int i = 0; i < 2; ++i)
+        {
+            if (mVertexInfo[v0] == INTERIOR_VERTEX)
+            {
+                mVertexGraph[v0].distance = -1;
+            }
+            else
+            {
+                mVertexGraph[v0].distance = 0;
+                currFront.insert(v0);
+            }
+
+            // Insert v1 into the first available slot of the adjacency array.
+            std::array<int, 2> range = mVertexGraph[v0].range;
+            for (int j = 0; j < range[1]; ++j)
+            {
+                std::pair<int, Real>& data = mVertexGraphData[range[0] + j];
+                if (data.second == (Real)-1)
+                {
+                    data.first = v1;
+                    data.second = (Real)0;
+                    break;
+                }
+            }
+
+            std::swap(v0, v1);
+        }
+    }
+
+    // Use a breadth-first search to propagate the distance information.
+    int nextDistance = 1;
+    size_t numFrontVertices = currFront.size();
+    std::copy(currFront.begin(), currFront.end(), mOrderedVertices.begin());
+    while (currFront.size() > 0)
+    {
+        std::set<int> nextFront;
+        for (auto v : currFront)
+        {
+            std::array<int, 2> range = mVertexGraph[v].range;
+            auto* current = &mVertexGraphData[range[0]];
+            for (int j = 0; j < range[1]; ++j, ++current)
+            {
+                int a = current->first;
+                if (mVertexGraph[a].distance == -1)
+                {
+                    mVertexGraph[a].distance = nextDistance;
+                    nextFront.insert(a);
+                }
+            }
+        }
+        std::copy(nextFront.begin(), nextFront.end(), mOrderedVertices.begin() + numFrontVertices);
+        numFrontVertices += nextFront.size();
+        currFront = std::move(nextFront);
+        ++nextDistance;
+    }
+}
+
+template <typename Real>
+void GenerateMeshUV<Real>::AssignBoundaryTextureCoordinatesSquare()
+{
+    // Map the boundary of the mesh to the unit square [0,1]^2.  The selection
+    // of square vertices is such that the relative distances between boundary
+    // vertices and the relative distances between polygon vertices is
+    // preserved, except that the four corners of the square are required to
+    // have boundary points mapped to them.  The first boundary point has an
+    // implied distance of zero.  The value distance[i] is the length of the
+    // boundary polyline from vertex 0 to vertex i+1.
+    std::vector<Real> distance(mNumBoundaryEdges);
+    Real total = (Real)0;
+    int v0 = mBoundaryStart, v1, i;
+    for (i = 0; i < mNumBoundaryEdges; ++i)
+    {
+        v1 = mVertexInfo[v0];
+        total += Length(mVertices[v1] - mVertices[v0]);
+        distance[i] = total;
+        v0 = v1;
+    }
+
+    Real invTotal = ((Real)1) / total;
+    for (auto& d : distance)
+    {
+        d *= invTotal;
+    }
+
+    auto begin = distance.begin(), end = distance.end();
+    int endYMin = (int)(std::lower_bound(begin, end, (Real)0.25) - begin);
+    int endXMax = (int)(std::lower_bound(begin, end, (Real)0.50) - begin);
+    int endYMax = (int)(std::lower_bound(begin, end, (Real)0.75) - begin);
+    int endXMin = (int)distance.size() - 1;
+
+    // The first polygon vertex is (0,0).  The remaining vertices are chosen
+    // counterclockwise around the square.
+    v0 = mBoundaryStart;
+    mTCoords[v0][0] = (Real)0;
+    mTCoords[v0][1] = (Real)0;
+    for (i = 0; i < endYMin; ++i)
+    {
+        v1 = mVertexInfo[v0];
+        mTCoords[v1][0] = distance[i] * (Real)4;
+        mTCoords[v1][1] = (Real)0;
+        v0 = v1;
+    }
+
+    v1 = mVertexInfo[v0];
+    mTCoords[v1][0] = (Real)1;
+    mTCoords[v1][1] = (Real)0;
+    v0 = v1;
+    for (++i; i < endXMax; ++i)
+    {
+        v1 = mVertexInfo[v0];
+        mTCoords[v1][0] = (Real)1;
+        mTCoords[v1][1] = distance[i] * (Real)4 - (Real)1;
+        v0 = v1;
+    }
+
+    v1 = mVertexInfo[v0];
+    mTCoords[v1][0] = (Real)1;
+    mTCoords[v1][1] = (Real)1;
+    v0 = v1;
+    for (++i; i < endYMax; ++i)
+    {
+        v1 = mVertexInfo[v0];
+        mTCoords[v1][0] = (Real)3 - distance[i] * (Real)4;
+        mTCoords[v1][1] = (Real)1;
+        v0 = v1;
+    }
+
+    v1 = mVertexInfo[v0];
+    mTCoords[v1][0] = (Real)0;
+    mTCoords[v1][1] = (Real)1;
+    v0 = v1;
+    for (++i; i < endXMin; ++i)
+    {
+        v1 = mVertexInfo[v0];
+        mTCoords[v1][0] = (Real)0;
+        mTCoords[v1][1] = (Real)4 - distance[i] * (Real)4;
+        v0 = v1;
+    }
+}
+
+template <typename Real>
+void GenerateMeshUV<Real>::AssignBoundaryTextureCoordinatesDisk()
+{
+    // Map the boundary of the mesh to a convex polygon.  The selection of
+    // convex polygon vertices is such that the relative distances between
+    // boundary vertices and the relative distances between polygon vertices
+    // is preserved.  The first boundary point has an implied distance of
+    // zero.  The value distance[i] is the length of the boundary polyline
+    // from vertex 0 to vertex i+1.
+    std::vector<Real> distance(mNumBoundaryEdges);
+    Real total = (Real)0;
+    int v0 = mBoundaryStart;
+    for (int i = 0; i < mNumBoundaryEdges; ++i)
+    {
+        int v1 = mVertexInfo[v0];
+        total += Length(mVertices[v1] - mVertices[v0]);
+        distance[i] = total;
+        v0 = v1;
+    }
+
+    // The convex polygon lives in [0,1]^2 and inscribes a circle with center
+    // (1/2,1/2) and radius 1/2.  The polygon center is not necessarily the
+    // circle center!  This is the case when a boundary edge has length larger
+    // than half the total length of the boundary polyline; we do not expect
+    // such data for our meshes.  The first polygon vertex is (1/2,0).  The
+    // remaining vertices are chosen counterclockwise around the polygon.
+    Real multiplier = ((Real)GTE_C_TWO_PI) / total;
+    v0 = mBoundaryStart;
+    mTCoords[v0][0] = (Real)1;
+    mTCoords[v0][1] = (Real)0.5;
+    for (int i = 1; i < mNumBoundaryEdges; ++i)
+    {
+        int v1 = mVertexInfo[v0];
+        Real angle = multiplier * distance[i - 1];
+        mTCoords[v1][0] = (std::cos(angle) + (Real)1) * (Real)0.5;
+        mTCoords[v1][1] = (std::sin(angle) + (Real)1) * (Real)0.5;
+        v0 = v1;
+    }
+}
+
+template <typename Real>
+void GenerateMeshUV<Real>::ComputeMeanValueWeights()
+{
+    for (auto const& edge : mInteriorEdges)
+    {
+        int v0 = edge->V[0], v1 = edge->V[1];
+        for (int i = 0; i < 2; ++i)
+        {
+            // Compute the direction from X0 to X1 and compute the length
+            // of the edge (X0,X1).
+            Vector3<Real> X0 = mVertices[v0];
+            Vector3<Real> X1 = mVertices[v1];
+            Vector3<Real> X1mX0 = X1 - X0;
+            Real x1mx0length = Normalize(X1mX0);
+            Real weight;
+            if (x1mx0length >(Real)0)
+            {
+                // Compute the weight for X0 associated with X1.
+                weight = (Real)0;
+                for (int j = 0; j < 2; ++j)
+                {
+                    // Find the vertex of triangle T[j] opposite edge <X0,X1>.
+                    auto tri = edge->T[j].lock();
+                    int k;
+                    for (k = 0; k < 3; ++k)
+                    {
+                        int v2 = tri->V[k];
+                        if (v2 != v0 && v2 != v1)
+                        {
+                            Vector3<Real> X2 = mVertices[v2];
+                            Vector3<Real> X2mX0 = X2 - X0;
+                            Real x2mx0Length = Normalize(X2mX0);
+                            if (x2mx0Length >(Real)0)
+                            {
+                                Real dot = Dot(X2mX0, X1mX0);
+                                Real cs = std::min(std::max(dot, (Real)-1), (Real)1);
+                                Real angle = std::acos(cs);
+                                weight += std::tan(angle * (Real)0.5);
+                            }
+                            else
+                            {
+                                weight += (Real)1;
+                            }
+                            break;
+                        }
+                    }
+                }
+                weight /= x1mx0length;
+            }
+            else
+            {
+                weight = (Real)1;
+            }
+
+            std::array<int, 2> range = mVertexGraph[v0].range;
+            for (int j = 0; j < range[1]; ++j)
+            {
+                std::pair<int, Real>& data = mVertexGraphData[range[0] + j];
+                if (data.first == v1)
+                {
+                    data.second = weight;
+                }
+            }
+
+            std::swap(v0, v1);
+        }
+    }
+}
+
+template <typename Real>
+void GenerateMeshUV<Real>::SolveSystem(unsigned int numIterations)
+{
+    // On the first pass, average only neighbors whose texture coordinates
+    // have been computed.  This is a good initial guess for the linear system
+    // and leads to relatively fast convergence of the Gauss-Seidel iterates.
+    Real zero = (Real)0;
+    for (int i = mNumBoundaryEdges; i < mNumVertices; ++i)
+    {
+        int v0 = mOrderedVertices[i];
+        std::array<int, 2> range = mVertexGraph[v0].range;
+        auto const* current = &mVertexGraphData[range[0]];
+        Vector2<Real> tcoord{ zero, zero };
+        Real weight, weightSum = zero;
+        for (int j = 0; j < range[1]; ++j, ++current)
+        {
+            int v1 = current->first;
+            if (mTCoords[v1][0] != -1.0f)
+            {
+                weight = current->second;
+                weightSum += weight;
+                tcoord += weight * mTCoords[v1];
+            }
+        }
+        tcoord /= weightSum;
+        mTCoords[v0] = tcoord;
+    }
+
+#if defined(GTE_COMPUTE_MODEL_ALLOW_GPGPU)
+    if (mCModel->engine)
+    {
+        SolveSystemGPU(numIterations);
+    }
+    else
+#endif
+    {
+        if (mCModel->numThreads > 1)
+        {
+            SolveSystemCPUMultiple(numIterations);
+        }
+        else
+        {
+            SolveSystemCPUSingle(numIterations);
+        }
+    }
+}
+
+template <typename Real>
+void GenerateMeshUV<Real>::SolveSystemCPUSingle(unsigned int numIterations)
+{
+    // Use ping-pong buffers for the texture coordinates.
+    std::vector<Vector2<Real>> tcoords(mNumVertices);
+    size_t numBytes = mNumVertices * sizeof(Vector2<Real>);
+    Memcpy(&tcoords[0], mTCoords, numBytes);
+    Vector2<Real>* inTCoords = mTCoords;
+    Vector2<Real>* outTCoords = &tcoords[0];
+
+    // The value numIterations is even, so we always swap an even number
+    // of times.  This ensures that on exit from the loop, outTCoords is
+    // tcoords.
+    for (unsigned int i = 1; i <= numIterations; ++i)
+    {
+        if (mCModel->progress)
+        {
+            (*mCModel->progress)(i);
+        }
+
+        for (int j = mNumBoundaryEdges; j < mNumVertices; ++j)
+        {
+            int v0 = mOrderedVertices[j];
+            std::array<int, 2> range = mVertexGraph[v0].range;
+            auto const* current = &mVertexGraphData[range[0]];
+            Vector2<Real> tcoord{ (Real)0, (Real)0 };
+            Real weight, weightSum = (Real)0;
+            for (int k = 0; k < range[1]; ++k, ++current)
+            {
+                int v1 = current->first;
+                weight = current->second;
+                weightSum += weight;
+                tcoord += weight * inTCoords[v1];
+            }
+            tcoord /= weightSum;
+            outTCoords[v0] = tcoord;
+        }
+
+        std::swap(inTCoords, outTCoords);
+    }
+}
+
+template <typename Real>
+void GenerateMeshUV<Real>::SolveSystemCPUMultiple(unsigned int numIterations)
+{
+    // Use ping-pong buffers for the texture coordinates.
+    std::vector<Vector2<Real>> tcoords(mNumVertices);
+    size_t numBytes = mNumVertices * sizeof(Vector2<Real>);
+    Memcpy(&tcoords[0], mTCoords, numBytes);
+    Vector2<Real>* inTCoords = mTCoords;
+    Vector2<Real>* outTCoords = &tcoords[0];
+
+    // Partition the data for multiple threads.
+    int numV = mNumVertices - mNumBoundaryEdges;
+    int numVPerThread = numV / mCModel->numThreads;
+    std::vector<int> vmin(mCModel->numThreads), vmax(mCModel->numThreads);
+    for (unsigned int t = 0; t < mCModel->numThreads; ++t)
+    {
+        vmin[t] = mNumBoundaryEdges + t * numVPerThread;
+        vmax[t] = vmin[t] + numVPerThread - 1;
+    }
+    vmax[mCModel->numThreads - 1] = mNumVertices - 1;
+
+    // The value numIterations is even, so we always swap an even number
+    // of times.  This ensures that on exit from the loop, outTCoords is
+    // tcoords.
+    for (unsigned int i = 1; i <= numIterations; ++i)
+    {
+        if (mCModel->progress)
+        {
+            (*mCModel->progress)(i);
+        }
+
+        // Execute Gauss-Seidel iterations in multiple threads.
+        std::vector<std::thread> process(mCModel->numThreads);
+        for (unsigned int t = 0; t < mCModel->numThreads; ++t)
+        {
+            process[t] = std::thread([this, t, &vmin, &vmax, inTCoords,
+                outTCoords]()
+            {
+                for (int j = vmin[t]; j <= vmax[t]; ++j)
+                {
+                    int v0 = mOrderedVertices[j];
+                    std::array<int, 2> range = mVertexGraph[v0].range;
+                    auto const* current = &mVertexGraphData[range[0]];
+                    Vector2<Real> tcoord{ (Real)0, (Real)0 };
+                    Real weight, weightSum = (Real)0;
+                    for (int k = 0; k < range[1]; ++k, ++current)
+                    {
+                        int v1 = current->first;
+                        weight = current->second;
+                        weightSum += weight;
+                        tcoord += weight * inTCoords[v1];
+                    }
+                    tcoord /= weightSum;
+                    outTCoords[v0] = tcoord;
+                }
+            });
+        }
+
+        // Wait for all threads to finish.
+        for (unsigned int t = 0; t < mCModel->numThreads; ++t)
+        {
+            process[t].join();
+        }
+
+        std::swap(inTCoords, outTCoords);
+    }
+}
+
+#if defined(GTE_COMPUTE_MODEL_ALLOW_GPGPU)
+
+template <typename Real>
+void GenerateMeshUV<Real>::SolveSystemGPU(unsigned int numIterations)
+{
+    mCModel->factory->defines.Set("NUM_X_THREADS", 8);
+    mCModel->factory->defines.Set("NUM_Y_THREADS", 8);
+    if (std::numeric_limits<Real>::max() == std::numeric_limits<float>::max())
+    {
+        mCModel->factory->defines.Set("Real", "float");
+#if defined(GTE_DEV_OPENGL)
+        mCModel->factory->defines.Set("Real2", "vec2");
+#else
+        mCModel->factory->defines.Set("Real2", "float2");
+#endif
+    }
+    else
+    {
+        mCModel->factory->defines.Set("Real", "double");
+#if defined(GTE_DEV_OPENGL)
+        mCModel->factory->defines.Set("Real2", "dvec2");
+#else
+        mCModel->factory->defines.Set("Real2", "double2");
+#endif
+    }
+
+    // TODO: Test mSolveSystem for null and respond accordingly.
+    int api = mCModel->factory->GetAPI();
+    mSolveSystem = mCModel->factory->CreateFromSource(*msSource[api]);
+    std::shared_ptr<ComputeShader> cshader = mSolveSystem->GetCShader();
+
+    // Compute the number of thread groups.
+    int numInputs = mNumVertices - mNumBoundaryEdges;
+    Real factor0 = std::ceil(std::sqrt((Real)numInputs));
+    Real factor1 = std::ceil((Real)numInputs / factor0);
+    int xElements = static_cast<int>(factor0);
+    int yElements = static_cast<int>(factor1);
+    int xRem = (xElements % 8);
+    if (xRem > 0)
+    {
+        xElements += 8 - xRem;
+    }
+    int yRem = (yElements % 8);
+    if (yRem > 0)
+    {
+        yElements += 8 - yRem;
+    }
+    unsigned int numXGroups = xElements / 8;
+    unsigned int numYGroups = yElements / 8;
+
+    mBoundBuffer = std::make_shared<ConstantBuffer>(4 * sizeof(int), false);
+    int* data = mBoundBuffer->Get<int>();
+    data[0] = xElements;
+    data[1] = yElements;
+    data[2] = mNumBoundaryEdges;
+    data[3] = numInputs;
+    cshader->Set("Bounds", mBoundBuffer);
+
+    unsigned int const vgSize = static_cast<unsigned int>(mVertexGraph.size());
+    mVGBuffer = std::make_shared<StructuredBuffer>(vgSize, sizeof(Vertex));
+    Memcpy(mVGBuffer->GetData(), &mVertexGraph[0], mVGBuffer->GetNumBytes());
+    cshader->Set("vertexGraph", mVGBuffer);
+
+    unsigned int const vgdSize = static_cast<unsigned int>(mVertexGraphData.size());
+    mVGDBuffer = std::make_shared<StructuredBuffer>(vgdSize, sizeof(std::pair<int, Real>));
+    Memcpy(mVGDBuffer->GetData(), &mVertexGraphData[0], mVGDBuffer->GetNumBytes());
+    cshader->Set("vertexGraphData", mVGDBuffer);
+
+    unsigned int const ovSize = static_cast<unsigned int>(mOrderedVertices.size());
+    mOVBuffer = std::make_shared<StructuredBuffer>(ovSize, sizeof(int));
+    Memcpy(mOVBuffer->GetData(), &mOrderedVertices[0], mOVBuffer->GetNumBytes());
+    cshader->Set("orderedVertices", mOVBuffer);
+
+    for (int j = 0; j < 2; ++j)
+    {
+        mTCoordsBuffer[j] = std::make_shared<StructuredBuffer>(mNumVertices, sizeof(Vector2<Real>));
+        mTCoordsBuffer[j]->SetUsage(Resource::SHADER_OUTPUT);
+        Memcpy(mTCoordsBuffer[j]->GetData(), mTCoords, mTCoordsBuffer[j]->GetNumBytes());
+    }
+    mTCoordsBuffer[0]->SetCopyType(Resource::COPY_STAGING_TO_CPU);
+
+    // The value numIterations is even, so we always swap an even number
+    // of times.  This ensures that on exit from the loop,
+    // mTCoordsBuffer[0] has the final output.
+    for (unsigned int i = 1; i <= numIterations; ++i)
+    {
+        if (mCModel->progress)
+        {
+            (*mCModel->progress)(i);
+        }
+
+        cshader->Set("inTCoords", mTCoordsBuffer[0]);
+        cshader->Set("outTCoords", mTCoordsBuffer[1]);
+        mCModel->engine->Execute(mSolveSystem, numXGroups, numYGroups, 1);
+        std::swap(mTCoordsBuffer[0], mTCoordsBuffer[1]);
+    }
+
+    mCModel->engine->CopyGpuToCpu(mTCoordsBuffer[0]);
+    Memcpy(mTCoords, mTCoordsBuffer[0]->GetData(), mTCoordsBuffer[0]->GetNumBytes());
+}
+
+#endif
+
+}
diff --git a/Engine/Plugins/Experimental/GeometryProcessing/Source/GeometryAlgorithms/Private/GTEngine/Mathematics/GteHalfspace.h b/Engine/Plugins/Experimental/GeometryProcessing/Source/GeometryAlgorithms/Private/GTEngine/Mathematics/GteHalfspace.h
new file mode 100644
--- /dev/null
+++ b/Engine/Plugins/Experimental/GeometryProcessing/Source/GeometryAlgorithms/Private/GTEngine/Mathematics/GteHalfspace.h
@@ -0,0 +1,113 @@
+// David Eberly, Geometric Tools, Redmond WA 98052
+// Copyright (c) 1998-2019
+// Distributed under the Boost Software License, Version 1.0.
+// http://www.boost.org/LICENSE_1_0.txt
+// http://www.geometrictools.com/License/Boost/LICENSE_1_0.txt
+// File Version: 3.0.0 (2016/06/19)
+
+#pragma once
+
+#include <GTEngine/Mathematics/GteVector.h>
+
+// The halfspace is represented as Dot(N,X) >= c where N is a unit-length
+// normal vector, c is the plane constant, and X is any point in space.
+// The user must ensure that the normal vector is unit length.
+
+namespace gte
+{
+
+template <int N, typename Real>
+class Halfspace
+{
+public:
+    // Construction and destruction.  The default constructor sets the normal
+    // to (0,...,0,1) and the constant to zero (halfspace x[N-1] >= 0).
+    Halfspace();
+
+    // Specify N and c directly.
+    Halfspace(Vector<N, Real> const& inNormal, Real inConstant);
+
+    // Public member access.
+    Vector<N, Real> normal;
+    Real constant;
+
+public:
+    // Comparisons to support sorted containers.
+    bool operator==(Halfspace const& halfspace) const;
+    bool operator!=(Halfspace const& halfspace) const;
+    bool operator< (Halfspace const& halfspace) const;
+    bool operator<=(Halfspace const& halfspace) const;
+    bool operator> (Halfspace const& halfspace) const;
+    bool operator>=(Halfspace const& halfspace) const;
+};
+
+// Template alias for convenience.
+template <typename Real>
+using Halfspace3 = Halfspace<3, Real>;
+
+
+template <int N, typename Real>
+Halfspace<N, Real>::Halfspace()
+    :
+    constant((Real)0)
+{
+    normal.MakeUnit(N - 1);
+}
+
+template <int N, typename Real>
+Halfspace<N, Real>::Halfspace(Vector<N, Real> const& inNormal,
+    Real inConstant)
+    :
+    normal(inNormal),
+    constant(inConstant)
+{
+}
+
+template <int N, typename Real>
+bool Halfspace<N, Real>::operator==(Halfspace const& halfspace) const
+{
+    return normal == halfspace.normal && constant == halfspace.constant;
+}
+
+template <int N, typename Real>
+bool Halfspace<N, Real>::operator!=(Halfspace const& halfspace) const
+{
+    return !operator==(halfspace);
+}
+
+template <int N, typename Real>
+bool Halfspace<N, Real>::operator<(Halfspace const& halfspace) const
+{
+    if (normal < halfspace.normal)
+    {
+        return true;
+    }
+
+    if (normal > halfspace.normal)
+    {
+        return false;
+    }
+
+    return constant < halfspace.constant;
+}
+
+template <int N, typename Real>
+bool Halfspace<N, Real>::operator<=(Halfspace const& halfspace) const
+{
+    return operator<(halfspace) || operator==(halfspace);
+}
+
+template <int N, typename Real>
+bool Halfspace<N, Real>::operator>(Halfspace const& halfspace) const
+{
+    return !operator<=(halfspace);
+}
+
+template <int N, typename Real>
+bool Halfspace<N, Real>::operator>=(Halfspace const& halfspace) const
+{
+    return !operator<(halfspace);
+}
+
+
+}
