move GeometryProcessing plugin from NotForLicensees to Experimental

diff --git a/Engine/Plugins/Experimental/GeometryProcessing/Source/GeometryAlgorithms/Private/GTEngine/Mathematics/GteHyperellipsoid.h b/Engine/Plugins/Experimental/GeometryProcessing/Source/GeometryAlgorithms/Private/GTEngine/Mathematics/GteHyperellipsoid.h
new file mode 100644
--- /dev/null
+++ b/Engine/Plugins/Experimental/GeometryProcessing/Source/GeometryAlgorithms/Private/GTEngine/Mathematics/GteHyperellipsoid.h
@@ -0,0 +1,383 @@
+// David Eberly, Geometric Tools, Redmond WA 98052
+// Copyright (c) 1998-2019
+// Distributed under the Boost Software License, Version 1.0.
+// http://www.boost.org/LICENSE_1_0.txt
+// http://www.geometrictools.com/License/Boost/LICENSE_1_0.txt
+// File Version: 3.0.2 (2018/10/05)
+
+#pragma once
+
+#include <GTEngine/Mathematics/GteMatrix.h>
+#include <GTEngine/Mathematics/GteSymmetricEigensolver.h>
+
+// A hyperellipsoid has center K; axis directions U[0] through U[N-1], all
+// unit-length vectors; and extents e[0] through e[N-1], all positive numbers.
+// A point X = K + sum_{d=0}^{N-1} y[d]*U[d] is on the hyperellipsoid whenever
+// sum_{d=0}^{N-1} (y[d]/e[d])^2 = 1.  An algebraic representation for the
+// hyperellipsoid is (X-K)^T * M * (X-K) = 1, where M is the NxN symmetric
+// matrix M = sum_{d=0}^{N-1} U[d]*U[d]^T/e[d]^2, where the superscript T
+// denotes transpose.  Observe that U[i]*U[i]^T is a matrix, not a scalar dot
+// product.  The hyperellipsoid is also represented by a quadratic equation
+// 0 = C + B^T*X + X^T*A*X, where C is a scalar, B is an Nx1 vector, and A is
+// an NxN symmetric matrix with positive eigenvalues.  The coefficients can be
+// stored from lowest degree to highest degree,
+//   C = k[0]
+//   B = k[1], ..., k[N]
+//   A = k[N+1], ..., k[(N+1)(N+2)/2 - 1]
+// where the A-coefficients are the upper-triangular elements of A listed in
+// row-major order.  For N = 2, X = (x[0],x[1]) and
+//   0 = k[0] +
+//       k[1]*x[0] + k[2]*x[1] + 
+//       k[3]*x[0]*x[0] + k[4]*x[0]*x[1]
+//                      + k[5]*x[1]*x[1]
+// For N = 3, X = (x[0],x[1],x[2]) and
+//   0 = k[0] +
+//       k[1]*x[0] + k[2]*x[1] + k[3]*x[2] +
+//       k[4]*x[0]*x[0] + k[5]*x[0]*x[1] + k[6]*x[0]*x[2] +
+//                      + k[7]*x[1]*x[1] + k[8]*x[1]*x[2] +
+//                                       + k[9]*x[2]*x[2]
+// This equation can be factored to the form (X-K)^T * M * (X-K) = 1, where
+// K = -A^{-1}*B/2, M = A/(B^T*A^{-1}*B/4-C).
+
+namespace gte
+{
+
+template <int N, typename Real>
+class Hyperellipsoid
+{
+public:
+
+    // Construction and destruction.  The default constructor sets the center
+    // to Vector<N,Real>::Zero(), the axes to Vector<N,Real>::Unit(d), and all
+    // extents to 1.
+    Hyperellipsoid();
+    Hyperellipsoid(Vector<N, Real> const& inCenter,
+        std::array<Vector<N, Real>, N> const inAxis,
+        Vector<N, Real> const& inExtent);
+
+    // Compute M = sum_{d=0}^{N-1} U[d]*U[d]^T/e[d]^2.
+    void GetM(Matrix<N, N, Real>& M) const;
+
+    // Compute M^{-1} = sum_{d=0}^{N-1} U[d]*U[d]^T*e[d]^2.
+    void GetMInverse(Matrix<N, N, Real>& MInverse) const;
+
+    // Construct the coefficients in the quadratic equation that represents
+    // the hyperellipsoid.
+    void ToCoefficients(std::array<Real, (N+1)*(N+2)/2>& coeff) const;
+    void ToCoefficients(Matrix<N, N, Real>& A, Vector<N, Real>& B,
+        Real& C) const;
+
+    // Construct C, U[i], and e[i] from the equation.  The return value is
+    // 'true' if and only if the input coefficients represent a
+    // hyperellipsoid.  If the function returns 'false', the hyperellipsoid
+    // data members are undefined.
+    bool FromCoefficients(std::array<Real, (N + 1)*(N + 2) / 2> const& coeff);
+    bool FromCoefficients(Matrix<N, N, Real> const& A,
+        Vector<N, Real> const& B, Real C);
+
+    // Public member access.
+    Vector<N, Real> center;
+    std::array<Vector<N, Real>, N> axis;
+    Vector<N, Real> extent;
+
+private:
+    static void Convert(std::array<Real, (N + 1)*(N + 2) / 2> const& coeff,
+        Matrix<N, N, Real>& A, Vector<N, Real>& B, Real& C);
+
+    static void Convert(Matrix<N, N, Real> const& A, Vector<N, Real> const& B,
+        Real C, std::array<Real, (N + 1)*(N + 2) / 2>& coeff);
+
+public:
+    // Comparisons to support sorted containers.
+    bool operator==(Hyperellipsoid const& hyperellipsoid) const;
+    bool operator!=(Hyperellipsoid const& hyperellipsoid) const;
+    bool operator< (Hyperellipsoid const& hyperellipsoid) const;
+    bool operator<=(Hyperellipsoid const& hyperellipsoid) const;
+    bool operator> (Hyperellipsoid const& hyperellipsoid) const;
+    bool operator>=(Hyperellipsoid const& hyperellipsoid) const;
+};
+
+// Template aliases for convenience.
+template <typename Real>
+using Ellipse2 = Hyperellipsoid<2, Real>;
+
+template <typename Real>
+using Ellipsoid3 = Hyperellipsoid<3, Real>;
+
+
+template <int N, typename Real>
+Hyperellipsoid<N, Real>::Hyperellipsoid()
+{
+    center.MakeZero();
+    for (int d = 0; d < N; ++d)
+    {
+        axis[d].MakeUnit(d);
+        extent[d] = (Real)1;
+    }
+}
+
+template <int N, typename Real>
+Hyperellipsoid<N, Real>::Hyperellipsoid(Vector<N, Real> const& inCenter,
+    std::array<Vector<N, Real>, N> const inAxis,
+    Vector<N, Real> const& inExtent)
+    :
+    center(inCenter),
+    axis(inAxis),
+    extent(inExtent)
+{
+}
+
+template <int N, typename Real>
+void Hyperellipsoid<N, Real>::GetM(Matrix<N, N, Real>& M) const
+{
+    M.MakeZero();
+    for (int d = 0; d < N; ++d)
+    {
+        Vector<N, Real> ratio = axis[d] / extent[d];
+        M += OuterProduct<N, N, Real>(ratio, ratio);
+    }
+}
+
+template <int N, typename Real>
+void Hyperellipsoid<N, Real>::GetMInverse(Matrix<N, N, Real>& MInverse) const
+{
+    MInverse.MakeZero();
+    for (int d = 0; d < N; ++d)
+    {
+        Vector<N, Real> product = axis[d] * extent[d];
+        MInverse += OuterProduct<N, N, Real>(product, product);
+    }
+}
+
+template <int N, typename Real>
+void Hyperellipsoid<N, Real>::ToCoefficients(
+    std::array<Real, (N + 1)*(N + 2) / 2>& coeff) const
+{
+    int const numCoefficients = (N + 1)*(N + 2) / 2;
+    Matrix<N, N, Real> A;
+    Vector<N, Real> B;
+    Real C;
+    ToCoefficients(A, B, C);
+    Convert(A, B, C, coeff);
+
+    // Arrange for one of the coefficients of the quadratic terms to be 1.
+    int quadIndex = numCoefficients - 1;
+    int maxIndex = quadIndex;
+    Real maxValue = std::abs(coeff[quadIndex]);
+    for (int d = 2; d < N; ++d)
+    {
+        quadIndex -= d;
+        Real absValue = std::abs(coeff[quadIndex]);
+        if (absValue > maxValue)
+        {
+            maxIndex = quadIndex;
+            maxValue = absValue;
+        }
+    }
+
+    Real invMaxValue = ((Real)1) / maxValue;
+    for (int i = 0; i < numCoefficients; ++i)
+    {
+        if (i != maxIndex)
+        {
+            coeff[i] *= invMaxValue;
+        }
+        else
+        {
+            coeff[i] = (Real)1;
+        }
+    }
+}
+
+template <int N, typename Real>
+void Hyperellipsoid<N, Real>::ToCoefficients(Matrix<N, N, Real>& A,
+    Vector<N, Real>& B, Real& C) const
+{
+    GetM(A);
+    Vector<N, Real> product = A * center;
+    B = ((Real)-2) * product;
+    C = Dot(center, product) - (Real)1;
+}
+
+template <int N, typename Real>
+bool Hyperellipsoid<N, Real>::FromCoefficients(
+    std::array<Real, (N + 1)*(N + 2) / 2> const& coeff)
+{
+    Matrix<N, N, Real> A;
+    Vector<N, Real> B;
+    Real C;
+    Convert(coeff, A, B, C);
+    return FromCoefficients(A, B, C);
+}
+
+template <int N, typename Real>
+bool Hyperellipsoid<N, Real>::FromCoefficients(Matrix<N, N, Real> const& A,
+    Vector<N, Real> const& B, Real C)
+{
+    // Compute the center K = -A^{-1}*B/2.
+    bool invertible;
+    Matrix<N, N, Real> invA = Inverse(A, &invertible);
+    if (!invertible)
+    {
+        return false;
+    }
+
+    center = ((Real)-0.5) * (invA * B);
+
+    // Compute B^T*A^{-1}*B/4 - C = K^T*A*K - C = -K^T*B/2 - C.
+    Real rightSide = -((Real)0.5) * Dot(center, B) - C;
+    if (rightSide == (Real)0)
+    {
+        return false;
+    }
+
+    // Compute M = A/(K^T*A*K - C).
+    Real invRightSide = ((Real)1) / rightSide;
+    Matrix<N, N, Real> M = invRightSide * A;
+
+    // Factor into M = R*D*R^T.  M is symmetric, so it does not matter whether
+    // the matrix is stored in row-major or column-major order; they are
+    // equivalent.  The output R, however, is in row-major order.
+    SymmetricEigensolver<Real> es(N, 32);
+    Matrix<N, N, Real> rotation;
+    std::array<Real, N> diagonal;
+    es.Solve(&M[0], +1);  // diagonal[i] are nondecreasing
+    es.GetEigenvalues(&diagonal[0]);
+    es.GetEigenvectors(&rotation[0]);
+    if (es.GetEigenvectorMatrixType() == 0)
+    {
+        auto negLast = -rotation.GetCol(N - 1);
+        rotation.SetCol(N - 1, negLast);
+    }
+
+    for (int d = 0; d < N; ++d)
+    {
+        if (diagonal[d] <= (Real)0)
+        {
+            return false;
+        }
+
+        extent[d] = ((Real)1) / std::sqrt(diagonal[d]);
+        axis[d] = rotation.GetCol(d);
+    }
+
+    return true;
+}
+
+template <int N, typename Real>
+void Hyperellipsoid<N, Real>::Convert(
+    std::array<Real, (N + 1)*(N + 2) / 2> const& coeff, Matrix<N, N, Real>& A,
+    Vector<N, Real>& B, Real& C)
+{
+    int i = 0;
+    C = coeff[i++];
+
+    for (int j = 0; j < N; ++j)
+    {
+        B[j] = coeff[i++];
+    }
+
+    for (int r = 0; r < N; ++r)
+    {
+        for (int c = 0; c < r; ++c)
+        {
+            A(r, c) = A(c, r);
+        }
+
+        A(r, r) = coeff[i++];
+
+        for (int c = r + 1; c < N; ++c)
+        {
+            A(r, c) = coeff[i++] * (Real)0.5;
+        }
+    }
+}
+
+template <int N, typename Real>
+void Hyperellipsoid<N, Real>::Convert(Matrix<N, N, Real> const& A,
+    Vector<N, Real> const& B, Real C,
+    std::array<Real, (N + 1)*(N + 2) / 2>& coeff)
+{
+    int i = 0;
+    coeff[i++] = C;
+
+    for (int j = 0; j < N; ++j)
+    {
+        coeff[i++] = B[j];
+    }
+
+    for (int r = 0; r < N; ++r)
+    {
+        coeff[i++] = A(r, r);
+        for (int c = r + 1; c < N; ++c)
+        {
+            coeff[i++] = A(r, c) * (Real)2;
+        }
+    }
+}
+
+template <int N, typename Real>
+bool Hyperellipsoid<N, Real>::operator==(
+    Hyperellipsoid const& hyperellipsoid) const
+{
+    return center == hyperellipsoid.center && axis == hyperellipsoid.axis
+        && extent == hyperellipsoid.extent;
+}
+
+template <int N, typename Real>
+bool Hyperellipsoid<N, Real>::operator!=(
+    Hyperellipsoid const& hyperellipsoid) const
+{
+    return !operator==(hyperellipsoid);
+}
+
+template <int N, typename Real>
+bool Hyperellipsoid<N, Real>::operator<(
+    Hyperellipsoid const& hyperellipsoid) const
+{
+    if (center < hyperellipsoid.center)
+    {
+        return true;
+    }
+
+    if (center > hyperellipsoid.center)
+    {
+        return false;
+    }
+
+    if (axis < hyperellipsoid.axis)
+    {
+        return true;
+    }
+
+    if (axis > hyperellipsoid.axis)
+    {
+        return false;
+    }
+
+    return extent < hyperellipsoid.extent;
+}
+
+template <int N, typename Real>
+bool Hyperellipsoid<N, Real>::operator<=(
+    Hyperellipsoid const& hyperellipsoid) const
+{
+    return operator<(hyperellipsoid) || operator==(hyperellipsoid);
+}
+
+template <int N, typename Real>
+bool Hyperellipsoid<N, Real>::operator>(
+    Hyperellipsoid const& hyperellipsoid) const
+{
+    return !operator<=(hyperellipsoid);
+}
+
+template <int N, typename Real>
+bool Hyperellipsoid<N, Real>::operator>=(
+    Hyperellipsoid const& hyperellipsoid) const
+{
+    return !operator<(hyperellipsoid);
+}
+
+
+}
diff --git a/Engine/Plugins/Experimental/GeometryProcessing/Source/GeometryAlgorithms/Private/GTEngine/Mathematics/GteHyperplane.h b/Engine/Plugins/Experimental/GeometryProcessing/Source/GeometryAlgorithms/Private/GTEngine/Mathematics/GteHyperplane.h
new file mode 100644
--- /dev/null
+++ b/Engine/Plugins/Experimental/GeometryProcessing/Source/GeometryAlgorithms/Private/GTEngine/Mathematics/GteHyperplane.h
@@ -0,0 +1,149 @@
+// David Eberly, Geometric Tools, Redmond WA 98052
+// Copyright (c) 1998-2019
+// Distributed under the Boost Software License, Version 1.0.
+// http://www.boost.org/LICENSE_1_0.txt
+// http://www.geometrictools.com/License/Boost/LICENSE_1_0.txt
+// File Version: 3.0.1 (2016/07/28)
+
+#pragma once
+
+#include <GTEngine/Mathematics/GteMatrix.h>
+#include <GTEngine/Mathematics/GteSingularValueDecomposition.h>
+
+// The plane is represented as Dot(U,X) = c where U is a unit-length normal
+// vector, c is the plane constant, and X is any point on the plane.  The user
+// must ensure that the normal vector is unit length.
+
+namespace gte
+{
+
+template <int N, typename Real>
+class Hyperplane
+{
+public:
+    // Construction and destruction.  The default constructor sets the normal
+    // to (0,...,0,1) and the constant to zero (plane z = 0).
+    Hyperplane();
+
+    // Specify U and c directly.
+    Hyperplane(Vector<N, Real> const& inNormal, Real inConstant);
+
+    // U is specified, c = Dot(U,p) where p is a point on the hyperplane.
+    Hyperplane(Vector<N, Real> const& inNormal, Vector<N, Real> const& p);
+
+    // U is a unit-length vector in the orthogonal complement of the set
+    // {p[1]-p[0],...,p[n-1]-p[0]} and c = Dot(U,p[0]), where the p[i] are
+    // pointson the hyperplane.
+    Hyperplane(std::array<Vector<N, Real>, N> const& p);
+
+    // Public member access.
+    Vector<N, Real> normal;
+    Real constant;
+
+public:
+    // Comparisons to support sorted containers.
+    bool operator==(Hyperplane const& hyperplane) const;
+    bool operator!=(Hyperplane const& hyperplane) const;
+    bool operator< (Hyperplane const& hyperplane) const;
+    bool operator<=(Hyperplane const& hyperplane) const;
+    bool operator> (Hyperplane const& hyperplane) const;
+    bool operator>=(Hyperplane const& hyperplane) const;
+};
+
+// Template alias for convenience.
+template <typename Real>
+using Plane3 = Hyperplane<3, Real>;
+
+
+template <int N, typename Real>
+Hyperplane<N, Real>::Hyperplane()
+    :
+    constant((Real)0)
+{
+    normal.MakeUnit(N - 1);
+}
+
+template <int N, typename Real>
+Hyperplane<N, Real>::Hyperplane(Vector<N, Real> const& inNormal,
+    Real inConstant)
+    :
+    normal(inNormal),
+    constant(inConstant)
+{
+}
+
+template <int N, typename Real>
+Hyperplane<N, Real>::Hyperplane(Vector<N, Real> const& inNormal,
+    Vector<N, Real> const& p)
+    :
+    normal(inNormal),
+    constant(Dot(inNormal, p))
+{
+}
+
+template <int N, typename Real>
+Hyperplane<N, Real>::Hyperplane(std::array<Vector<N, Real>, N> const& p)
+{
+    Matrix<N, N - 1, Real> edge;
+    for (int i = 0; i < N - 1; ++i)
+    {
+        edge.SetCol(i, p[i + 1] - p[0]);
+    }
+
+    // Compute the 1-dimensional orthogonal complement of the edges of the
+    // simplex formed by the points p[].
+    SingularValueDecomposition<Real> svd(N, N - 1, 32);
+    svd.Solve(&edge[0], -1);
+    svd.GetUColumn(N - 1, &normal[0]);
+
+    constant = Dot(normal, p[0]);
+}
+
+template <int N, typename Real>
+bool Hyperplane<N, Real>::operator==(Hyperplane const& hyperplane) const
+{
+    return normal == hyperplane.normal && constant == hyperplane.constant;
+}
+
+template <int N, typename Real>
+bool Hyperplane<N, Real>::operator!=(Hyperplane const& hyperplane) const
+{
+    return !operator==(hyperplane);
+}
+
+template <int N, typename Real>
+bool Hyperplane<N, Real>::operator<(Hyperplane const& hyperplane) const
+{
+    if (normal < hyperplane.normal)
+    {
+        return true;
+    }
+
+    if (normal > hyperplane.normal)
+    {
+        return false;
+    }
+
+    return constant < hyperplane.constant;
+}
+
+template <int N, typename Real>
+bool Hyperplane<N, Real>::operator<=(Hyperplane const& hyperplane) const
+{
+    return operator<(hyperplane) || operator==(hyperplane);
+}
+
+template <int N, typename Real>
+bool Hyperplane<N, Real>::operator>(Hyperplane const& hyperplane) const
+{
+    return !operator<=(hyperplane);
+}
+
+template <int N, typename Real>
+bool Hyperplane<N, Real>::operator>=(Hyperplane const& hyperplane) const
+{
+    return !operator<(hyperplane);
+}
+
+
+}
diff --git a/Engine/Plugins/Experimental/GeometryProcessing/Source/GeometryAlgorithms/Private/GTEngine/Mathematics/GteHypersphere.h b/Engine/Plugins/Experimental/GeometryProcessing/Source/GeometryAlgorithms/Private/GTEngine/Mathematics/GteHypersphere.h
new file mode 100644
--- /dev/null
+++ b/Engine/Plugins/Experimental/GeometryProcessing/Source/GeometryAlgorithms/Private/GTEngine/Mathematics/GteHypersphere.h
@@ -0,0 +1,114 @@
+// David Eberly, Geometric Tools, Redmond WA 98052
+// Copyright (c) 1998-2019
+// Distributed under the Boost Software License, Version 1.0.
+// http://www.boost.org/LICENSE_1_0.txt
+// http://www.geometrictools.com/License/Boost/LICENSE_1_0.txt
+// File Version: 3.0.0 (2016/06/19)
+
+#pragma once
+
+#include <GTEngine/Mathematics/GteVector.h>
+
+// The hypersphere is represented as |X-C| = R where C is the center and R is
+// the radius.  The hypersphere is a circle for dimension 2 or a sphere for
+// dimension 3.
+
+namespace gte
+{
+
+template <int N, typename Real>
+class Hypersphere
+{
+public:
+    // Construction and destruction.  The default constructor sets the center
+    // to (0,...,0) and the radius to 1.
+    Hypersphere();
+    Hypersphere(Vector<N, Real> const& inCenter, Real inRadius);
+
+    // Public member access.
+    Vector<N, Real> center;
+    Real radius;
+
+public:
+    // Comparisons to support sorted containers.
+    bool operator==(Hypersphere const& hypersphere) const;
+    bool operator!=(Hypersphere const& hypersphere) const;
+    bool operator< (Hypersphere const& hypersphere) const;
+    bool operator<=(Hypersphere const& hypersphere) const;
+    bool operator> (Hypersphere const& hypersphere) const;
+    bool operator>=(Hypersphere const& hypersphere) const;
+};
+
+// Template aliases for convenience.
+template <typename Real>
+using Circle2 = Hypersphere<2, Real>;
+
+template <typename Real>
+using Sphere3 = Hypersphere<3, Real>;
+
+
+template <int N, typename Real>
+Hypersphere<N, Real>::Hypersphere()
+    :
+    radius((Real)1)
+{
+    center.MakeZero();
+}
+
+template <int N, typename Real>
+Hypersphere<N, Real>::Hypersphere(Vector<N, Real> const& inCenter,
+    Real inRadius)
+    :
+    center(inCenter),
+    radius(inRadius)
+{
+}
+
+template <int N, typename Real>
+bool Hypersphere<N, Real>::operator==(Hypersphere const& hypersphere) const
+{
+    return center == hypersphere.center && radius == hypersphere.radius;
+}
+
+template <int N, typename Real>
+bool Hypersphere<N, Real>::operator!=(Hypersphere const& hypersphere) const
+{
+    return !operator==(hypersphere);
+}
+
+template <int N, typename Real>
+bool Hypersphere<N, Real>::operator<(Hypersphere const& hypersphere) const
+{
+    if (center < hypersphere.center)
+    {
+        return true;
+    }
+
+    if (center > hypersphere.center)
+    {
+        return false;
+    }
+
+    return radius < hypersphere.radius;
+}
+
+template <int N, typename Real>
+bool Hypersphere<N, Real>::operator<=(Hypersphere const& hypersphere) const
+{
+    return operator<(hypersphere) || operator==(hypersphere);
+}
+
+template <int N, typename Real>
+bool Hypersphere<N, Real>::operator>(Hypersphere const& hypersphere) const
+{
+    return !operator<=(hypersphere);
+}
+
+template <int N, typename Real>
+bool Hypersphere<N, Real>::operator>=(Hypersphere const& hypersphere) const
+{
+    return !operator<(hypersphere);
+}
+
+
+}
diff --git a/Engine/Plugins/Experimental/GeometryProcessing/Source/GeometryAlgorithms/Private/GTEngine/Mathematics/GteIEEEBinary.h b/Engine/Plugins/Experimental/GeometryProcessing/Source/GeometryAlgorithms/Private/GTEngine/Mathematics/GteIEEEBinary.h
new file mode 100644
--- /dev/null
+++ b/Engine/Plugins/Experimental/GeometryProcessing/Source/GeometryAlgorithms/Private/GTEngine/Mathematics/GteIEEEBinary.h
@@ -0,0 +1,448 @@
+// David Eberly, Geometric Tools, Redmond WA 98052
+// Copyright (c) 1998-2019
+// Distributed under the Boost Software License, Version 1.0.
+// http://www.boost.org/LICENSE_1_0.txt
+// http://www.geometrictools.com/License/Boost/LICENSE_1_0.txt
+// File Version: 3.0.0 (2016/06/19)
+
+#pragma once
+
+#include <GTEngine/GTEngineDEF.h>
+#include <cstdint>
+
+namespace gte
+{
+
+template <typename Float, typename UInt, int NumBits, int Precision>
+class IEEEBinary
+{
+public:
+    // For generic access of the template types.
+    typedef Float FloatType;
+    typedef UInt UIntType;
+
+    // Construction from an encoding.  Copy constructor, destructor, and
+    // assignment operator are implicitly generated.  For the 3-parameter
+    // constructor, see the comments for SetEncoding(...).
+    ~IEEEBinary ();
+    IEEEBinary (); // The encoding is uninitialized.
+    IEEEBinary (IEEEBinary const& object);
+    IEEEBinary (UInt inEncoding);
+    IEEEBinary (UInt inSign, UInt inBiased, UInt inTrailing);
+    IEEEBinary (Float inNumber);
+
+    // Implicit conversion to floating-point type.
+    operator UInt () const;
+    operator Float () const;
+
+    // Assignment.
+    IEEEBinary& operator= (IEEEBinary const& object);
+
+    // Special constants.
+    static int const NUM_ENCODING_BITS    = NumBits;
+    static int const NUM_EXPONENT_BITS    = NumBits - Precision;
+    static int const NUM_SIGNIFICAND_BITS = Precision;
+    static int const NUM_TRAILING_BITS    = Precision - 1;
+    static int const EXPONENT_BIAS        = (1 << (NUM_EXPONENT_BITS - 1)) - 1;
+    static int const MAX_BIASED_EXPONENT  = (1 << NUM_EXPONENT_BITS) - 1;
+    static int const MIN_SUB_EXPONENT     = 1 - EXPONENT_BIAS;
+    static int const MIN_EXPONENT         = MIN_SUB_EXPONENT - NUM_TRAILING_BITS;
+    static int const SIGN_SHIFT           = NumBits - 1;
+
+    static UInt const SIGN_MASK           = (UInt(1) << (NumBits - 1));   
+    static UInt const NOT_SIGN_MASK       = UInt(~SIGN_MASK);
+    static UInt const TRAILING_MASK       = (UInt(1) << NUM_TRAILING_BITS) - 1;
+    static UInt const EXPONENT_MASK       = NOT_SIGN_MASK & ~TRAILING_MASK;
+    static UInt const NAN_QUIET_MASK      = (UInt(1) << (NUM_TRAILING_BITS - 1));
+    static UInt const NAN_PAYLOAD_MASK    = (TRAILING_MASK >> 1);
+    static UInt const MAX_TRAILING        = TRAILING_MASK;
+    static UInt const SUP_TRAILING        = (UInt(1) << NUM_TRAILING_BITS);
+    static UInt const POS_ZERO            = UInt(0);
+    static UInt const NEG_ZERO            = SIGN_MASK;
+    static UInt const MIN_SUBNORMAL       = UInt(1);
+    static UInt const MAX_SUBNORMAL       = TRAILING_MASK;
+    static UInt const MIN_NORMAL          = SUP_TRAILING;
+    static UInt const MAX_NORMAL          = NOT_SIGN_MASK & ~SUP_TRAILING;
+    static UInt const POS_INFINITY        = EXPONENT_MASK;
+    static UInt const NEG_INFINITY        = SIGN_MASK | EXPONENT_MASK;
+
+    // The types of numbers.
+    enum Classification
+    {
+        CLASS_NEG_INFINITY,
+        CLASS_NEG_SUBNORMAL,
+        CLASS_NEG_NORMAL,
+        CLASS_NEG_ZERO,
+        CLASS_POS_ZERO,
+        CLASS_POS_SUBNORMAL,
+        CLASS_POS_NORMAL,
+        CLASS_POS_INFINITY,
+        CLASS_QUIET_NAN,
+        CLASS_SIGNALING_NAN
+    };
+
+    Classification GetClassification () const;
+    bool IsZero () const;
+    bool IsSignMinus () const;
+    bool IsSubnormal () const;
+    bool IsNormal () const;
+    bool IsFinite () const;
+    bool IsInfinite () const;
+    bool IsNaN () const;
+    bool IsSignalingNaN () const;
+
+    // Get neighboring numbers.
+    UInt GetNextUp () const;
+    UInt GetNextDown () const;
+
+    // Encode and decode the binary representation.  The sign is 0 (number
+    // is nonnegative) or 1 (number is negative).  The biased exponent is in
+    // the range [0,MAX_BIASED_EXPONENT].  The trailing significand is in the
+    // range [0,MAX_TRAILING].
+    UInt GetSign () const;
+    UInt GetBiased () const;
+    UInt GetTrailing () const;
+    void SetEncoding (UInt sign, UInt biased, UInt trailing);
+    void GetEncoding (UInt& sign, UInt& biased, UInt& trailing) const;
+
+    // Access for direct manipulation of the object.
+    union
+    {
+        UInt encoding;
+        Float number;
+    };
+};
+
+typedef IEEEBinary<float, uint32_t, 32, 24> IEEEBinary32;
+typedef IEEEBinary<double, uint64_t, 64, 53> IEEEBinary64;
+
+
+template <typename Float, typename UInt, int NumBits, int Precision>
+IEEEBinary<Float, UInt, NumBits, Precision>::~IEEEBinary()
+{
+}
+
+template <typename Float, typename UInt, int NumBits, int Precision>
+IEEEBinary<Float, UInt, NumBits, Precision>::IEEEBinary()
+{
+    // The member mEncoding is uninitialized.
+}
+
+template <typename Float, typename UInt, int NumBits, int Precision>
+IEEEBinary<Float, UInt, NumBits, Precision>::IEEEBinary(
+    IEEEBinary const& object)
+    :
+    encoding(object.encoding)
+{
+}
+
+template <typename Float, typename UInt, int NumBits, int Precision>
+IEEEBinary<Float, UInt, NumBits, Precision>::IEEEBinary(UInt inEncoding)
+    :
+    encoding(inEncoding)
+{
+}
+
+template <typename Float, typename UInt, int NumBits, int Precision>
+IEEEBinary<Float, UInt, NumBits, Precision>::IEEEBinary(UInt inSign,
+    UInt inBiased, UInt inTrailing)
+{
+    SetEncoding(inSign, inBiased, inTrailing);
+}
+
+template <typename Float, typename UInt, int NumBits, int Precision>
+IEEEBinary<Float, UInt, NumBits, Precision>::IEEEBinary(Float inNumber)
+    :
+    number(inNumber)
+{
+}
+
+template <typename Float, typename UInt, int NumBits, int Precision>
+IEEEBinary<Float, UInt, NumBits, Precision>::operator UInt () const
+{
+    return encoding;
+}
+
+template <typename Float, typename UInt, int NumBits, int Precision>
+IEEEBinary<Float, UInt, NumBits, Precision>::operator Float () const
+{
+    return number;
+}
+
+template <typename Float, typename UInt, int NumBits, int Precision>
+IEEEBinary<Float, UInt, NumBits, Precision>&
+IEEEBinary<Float, UInt, NumBits, Precision>::operator= (IEEEBinary const& object)
+{
+    encoding = object.encoding;
+    return *this;
+}
+
+template <typename Float, typename UInt, int NumBits, int Precision>
+typename IEEEBinary<Float, UInt, NumBits, Precision>::Classification
+IEEEBinary<Float, UInt, NumBits, Precision>::GetClassification() const
+{
+    UInt sign, biased, trailing;
+    GetEncoding(sign, biased, trailing);
+
+    if (biased == 0)
+    {
+        if (trailing == 0)
+        {
+            return (sign != 0 ? CLASS_NEG_ZERO : CLASS_POS_ZERO);
+        }
+        else
+        {
+            return (sign != 0 ? CLASS_NEG_SUBNORMAL : CLASS_POS_SUBNORMAL);
+        }
+    }
+    else if (biased < MAX_BIASED_EXPONENT)
+    {
+        return (sign != 0 ? CLASS_NEG_NORMAL : CLASS_POS_NORMAL);
+    }
+    else if (trailing == 0)
+    {
+        return (sign != 0 ? CLASS_NEG_INFINITY : CLASS_POS_INFINITY);
+    }
+    else if (trailing & NAN_QUIET_MASK)
+    {
+        return CLASS_QUIET_NAN;
+    }
+    else
+    {
+        return CLASS_SIGNALING_NAN;
+    }
+}
+
+template <typename Float, typename UInt, int NumBits, int Precision>
+bool IEEEBinary<Float, UInt, NumBits, Precision>::IsZero() const
+{
+    return encoding == POS_ZERO || encoding == NEG_ZERO;
+}
+
+template <typename Float, typename UInt, int NumBits, int Precision>
+bool IEEEBinary<Float, UInt, NumBits, Precision>::IsSignMinus() const
+{
+    return (encoding & SIGN_MASK) != 0;
+}
+
+template <typename Float, typename UInt, int NumBits, int Precision>
+bool IEEEBinary<Float, UInt, NumBits, Precision>::IsSubnormal() const
+{
+    return GetBiased() == 0 && GetTrailing() > 0;
+}
+
+template <typename Float, typename UInt, int NumBits, int Precision>
+bool IEEEBinary<Float, UInt, NumBits, Precision>::IsNormal() const
+{
+    UInt biased = GetBiased();
+    return 0 < biased && biased < MAX_BIASED_EXPONENT;
+}
+
+template <typename Float, typename UInt, int NumBits, int Precision>
+bool IEEEBinary<Float, UInt, NumBits, Precision>::IsFinite() const
+{
+    return GetBiased() < MAX_BIASED_EXPONENT;
+}
+
+template <typename Float, typename UInt, int NumBits, int Precision>
+bool IEEEBinary<Float, UInt, NumBits, Precision>::IsInfinite() const
+{
+    return GetBiased() == MAX_BIASED_EXPONENT && GetTrailing() == 0;
+}
+
+template <typename Float, typename UInt, int NumBits, int Precision>
+bool IEEEBinary<Float, UInt, NumBits, Precision>::IsNaN() const
+{
+    return GetBiased() == MAX_BIASED_EXPONENT && GetTrailing() != 0;
+}
+
+template <typename Float, typename UInt, int NumBits, int Precision>
+bool IEEEBinary<Float, UInt, NumBits, Precision>::IsSignalingNaN() const
+{
+    UInt trailing = GetTrailing();
+    return GetBiased() == MAX_BIASED_EXPONENT
+        && (trailing & NAN_QUIET_MASK) == 0
+        && (trailing & NAN_PAYLOAD_MASK) != 0;
+}
+
+template <typename Float, typename UInt, int NumBits, int Precision>
+UInt IEEEBinary<Float, UInt, NumBits, Precision>::GetNextUp() const
+{
+    UInt sign, biased, trailing;
+    GetEncoding(sign, biased, trailing);
+
+    if (biased == 0)
+    {
+        if (trailing == 0)
+        {
+            // The next-up for both -0 and +0 is MIN_SUBNORMAL.
+            return MIN_SUBNORMAL;
+        }
+        else
+        {
+            if (sign != 0)
+            {
+                // When trailing is 1, 'this' is -MIN_SUBNORMAL and next-up
+                // is -0.
+                --trailing;
+                return SIGN_MASK | trailing;
+            }
+            else
+            {
+                // When trailing is MAX_TRAILING, 'this' is MAX_SUBNORMAL
+                // and next-up is MIN_NORMAL.
+                ++trailing;
+                return trailing;
+            }
+        }
+    }
+    else if (biased < MAX_BIASED_EXPONENT)
+    {
+        UInt nonnegative = (encoding & NOT_SIGN_MASK);
+        if (sign != 0)
+        {
+            --nonnegative;
+            return SIGN_MASK | nonnegative;
+        }
+        else
+        {
+            ++nonnegative;
+            return nonnegative;
+        }
+    }
+    else if (trailing == 0)
+    {
+        if (sign != 0)
+        {
+            // The next-up of -INFINITY is -MAX_NORMAL.
+            return SIGN_MASK | MAX_NORMAL;
+        }
+        else
+        {
+            // The next-up of +INFINITY is +INFINITY.
+            return POS_INFINITY;
+        }
+    }
+    else if (trailing & NAN_QUIET_MASK)
+    {
+        // TODO.  The IEEE standard is not clear what to do here.  Figure
+        // out what it means.
+        return 0;
+    }
+    else
+    {
+        // TODO.  The IEEE standard is not clear what to do here.  Figure
+        // out what it means.
+        return 0;
+    }
+}
+
+template <typename Float, typename UInt, int NumBits, int Precision>
+UInt IEEEBinary<Float, UInt, NumBits, Precision>::GetNextDown() const
+{
+    UInt sign, biased, trailing;
+    GetEncoding(sign, biased, trailing);
+
+    if (biased == 0)
+    {
+        if (trailing == 0)
+        {
+            // The next-down for both -0 and +0 is -MIN_SUBNORMAL.
+            return SIGN_MASK | MIN_SUBNORMAL;
+        }
+        else
+        {
+            if (sign == 0)
+            {
+                // When trailing is 1, 'this' is MIN_SUBNORMAL and next-down
+                // is +0.
+                --trailing;
+                return trailing;
+            }
+            else
+            {
+                // When trailing is MAX_TRAILING, 'this' is -MAX_SUBNORMAL
+                // and next-down is -MIN_NORMAL.
+                ++trailing;
+                return SIGN_MASK | trailing;
+            }
+        }
+    }
+    else if (biased < MAX_BIASED_EXPONENT)
+    {
+        UInt nonnegative = (encoding & NOT_SIGN_MASK);
+        if (sign == 0)
+        {
+            --nonnegative;
+            return nonnegative;
+        }
+        else
+        {
+            ++nonnegative;
+            return SIGN_MASK | nonnegative;
+        }
+    }
+    else if (trailing == 0)
+    {
+        if (sign == 0)
+        {
+            // The next-down of +INFINITY is +MAX_NORMAL.
+            return MAX_NORMAL;
+        }
+        else
+        {
+            // The next-down of -INFINITY is -INFINITY.
+            return NEG_INFINITY;
+        }
+    }
+    else if (trailing & NAN_QUIET_MASK)
+    {
+        // TODO.  The IEEE standard is not clear what to do here.  Figure
+        // out what it means.
+        return 0;
+    }
+    else
+    {
+        // TODO.  The IEEE standard is not clear what to do here.  Figure
+        // out what it means.
+        return 0;
+    }
+}
+
+template <typename Float, typename UInt, int NumBits, int Precision>
+UInt IEEEBinary<Float, UInt, NumBits, Precision>::GetSign() const
+{
+    return (encoding & SIGN_MASK) >> SIGN_SHIFT;
+}
+
+template <typename Float, typename UInt, int NumBits, int Precision>
+UInt IEEEBinary<Float, UInt, NumBits, Precision>::GetBiased() const
+{
+    return (encoding & EXPONENT_MASK) >> NUM_TRAILING_BITS;
+}
+
+template <typename Float, typename UInt, int NumBits, int Precision>
+UInt IEEEBinary<Float, UInt, NumBits, Precision>::GetTrailing() const
+{
+    return encoding & TRAILING_MASK;
+}
+
+template <typename Float, typename UInt, int NumBits, int Precision>
+void IEEEBinary<Float, UInt, NumBits, Precision>::SetEncoding(UInt sign,
+    UInt biased, UInt trailing)
+{
+    encoding =
+        (sign << SIGN_SHIFT) | (biased << NUM_TRAILING_BITS) | trailing;
+}
+
+template <typename Float, typename UInt, int NumBits, int Precision>
+void IEEEBinary<Float, UInt, NumBits, Precision>::GetEncoding(UInt& sign,
+    UInt& biased, UInt& trailing) const
+{
+    sign = GetSign();
+    biased = GetBiased();
+    trailing = GetTrailing();
+}
+
+
+}
diff --git a/Engine/Plugins/Experimental/GeometryProcessing/Source/GeometryAlgorithms/Private/GTEngine/Mathematics/GteIEEEBinary16.cpp b/Engine/Plugins/Experimental/GeometryProcessing/Source/GeometryAlgorithms/Private/GTEngine/Mathematics/GteIEEEBinary16.cpp
new file mode 100644
--- /dev/null
+++ b/Engine/Plugins/Experimental/GeometryProcessing/Source/GeometryAlgorithms/Private/GTEngine/Mathematics/GteIEEEBinary16.cpp
@@ -0,0 +1,368 @@
+// David Eberly, Geometric Tools, Redmond WA 98052
+// Copyright (c) 1998-2019
+// Distributed under the Boost Software License, Version 1.0.
+// http://www.boost.org/LICENSE_1_0.txt
+// http://www.geometrictools.com/License/Boost/LICENSE_1_0.txt
+// File Version: 3.0.2 (2018/10/05)
+
+#include <GTEngine/Mathematics/GteIEEEBinary16.h>
+#include <GTEngine/GTEnginePCH.h>
+#include <GTEngine/Mathematics/GteBitHacks.h>
+
+namespace gte
+{
+
+IEEEBinary16::~IEEEBinary16()
+{
+}
+
+IEEEBinary16::IEEEBinary16()
+    :
+    IEEEBinary<int16_t, uint16_t, 16, 11>()
+{
+    // uninitialized
+}
+
+IEEEBinary16::IEEEBinary16(IEEEBinary16 const& object)
+    :
+    IEEEBinary<int16_t, uint16_t, 16, 11>(object)
+{
+}
+
+IEEEBinary16::IEEEBinary16(float number)
+    :
+    IEEEBinary<int16_t, uint16_t, 16, 11>()
+{
+    union { float n; uint32_t e; } temp = { number };
+    encoding = Convert32To16(temp.e);
+}
+
+IEEEBinary16::IEEEBinary16(double number)
+    :
+    IEEEBinary<int16_t, uint16_t, 16, 11>()
+{
+    union { float n; uint32_t e; } temp;
+    temp.n = (float)number;
+    encoding = Convert32To16(temp.e);
+}
+
+IEEEBinary16::IEEEBinary16(uint16_t encoding)
+    :
+    IEEEBinary<int16_t, uint16_t, 16, 11>(encoding)
+{
+}
+
+IEEEBinary16::operator float() const
+{
+    union { uint32_t e; float n; } temp = { Convert16To32(encoding) };
+    return temp.n;
+}
+
+IEEEBinary16::operator double() const
+{
+    union { uint32_t e; float n; } temp = { Convert16To32(encoding) };
+    return (double)temp.n;
+}
+
+IEEEBinary16& IEEEBinary16::operator= (IEEEBinary16 const& object)
+{
+    IEEEBinary<int16_t, uint16_t, 16, 11>::operator=(object);
+    return *this;
+}
+
+bool IEEEBinary16::operator==(IEEEBinary16 const& object) const
+{
+    return (float)*this == (float)object;
+}
+
+bool IEEEBinary16::operator!=(IEEEBinary16 const& object) const
+{
+    return (float)*this != (float)object;
+}
+
+bool IEEEBinary16::operator< (IEEEBinary16 const& object) const
+{
+    return (float)*this < (float)object;
+}
+
+bool IEEEBinary16::operator<=(IEEEBinary16 const& object) const
+{
+    return (float)*this <= (float)object;
+}
+
+bool IEEEBinary16::operator> (IEEEBinary16 const& object) const
+{
+    return (float)*this > (float)object;
+}
+
+bool IEEEBinary16::operator>=(IEEEBinary16 const& object) const
+{
+    return (float)*this >= (float)object;
+}
+
+uint16_t IEEEBinary16::Convert32To16(uint32_t encoding)
+{
+    // Extract the channels for the binary32 number.
+    uint32_t sign32 = (encoding & F32_SIGN_MASK);
+    uint32_t biased32 =
+        ((encoding & F32_BIASED_EXPONENT_MASK) >> F32_NUM_TRAILING_BITS);
+    uint32_t trailing32 = (encoding & F32_TRAILING_MASK);
+    uint32_t nonneg32 = (encoding & F32_NOT_SIGN_MASK);
+
+    // Generate the channels for the IEEEBinary16 number.
+    uint16_t sign16 = static_cast<uint16_t>(sign32 >> DIFF_NUM_ENCODING_BITS);
+    uint16_t biased16, trailing16;
+    uint32_t frcpart;
+
+    if (biased32 == 0)
+    {
+        // nonneg32 is 32-zero or 32-subnormal, nearest is 16-zero.
+        return sign16;
+    }
+
+    if (biased32 < F32_MAX_BIASED_EXPONENT)
+    {
+        // nonneg32 is 32-normal.
+        if (nonneg32 <= F16_AVR_MIN_SUBNORMAL_ZERO)
+        {
+            // nonneg32 <= 2^{-25}, nearest is 16-zero.
+            return sign16;
+        }
+
+        if (nonneg32 <= F16_MIN_SUBNORMAL)
+        {
+            // 2^{-25} < nonneg32 <= 2^{-24}, nearest is 16-min-subnormal.
+            return sign16 | IEEEBinary16::MIN_SUBNORMAL;
+        }
+
+        if (nonneg32 < F16_MIN_NORMAL)
+        {
+            // 2^{-24} < nonneg32 < 2^{-14}, round to nearest
+            // 16-subnormal with ties to even.  Note that biased16 is zero.
+            trailing16 = static_cast<uint16_t>(((trailing32 & INT_PART_MASK) >> DIFF_NUM_TRAILING_BITS));
+            frcpart = (trailing32 & FRC_PART_MASK);
+            if (frcpart > FRC_HALF || (frcpart == FRC_HALF && (trailing16 & 1)))
+            {
+                // If there is a carry into the exponent, the nearest is
+                // actually 16-min-normal 1.0*2^{-14}, so the high-order bit
+                // of trailing16 makes biased16 equal to 1 and the result is
+                // correct.
+                ++trailing16;
+            }
+            return sign16 | trailing16;
+        }
+
+        if (nonneg32 <= F16_MAX_NORMAL)
+        {
+            // 2^{-14} <= nonneg32 <= 1.1111111111*2^{15}, round to nearest
+            // 16-normal with ties to even.
+            biased16 = static_cast<uint16_t>((biased32 - F32_EXPONENT_BIAS +
+                IEEEBinary16::EXPONENT_BIAS)
+                << IEEEBinary16::NUM_TRAILING_BITS);
+            trailing16 = static_cast<uint16_t>(((trailing32 & INT_PART_MASK) >> DIFF_NUM_TRAILING_BITS));
+            frcpart = (trailing32 & FRC_PART_MASK);
+            if (frcpart > FRC_HALF || (frcpart == FRC_HALF && (trailing16 & 1)))
+            {
+                // If there is a carry into the exponent, the addition of
+                // trailing16 to biased16 (rather than or-ing) produces the
+                // correct result.
+                ++trailing16;
+            }
+            return sign16 | (biased16 + trailing16);
+        }
+
+        if (nonneg32 < F16_AVR_MAX_NORMAL_INFINITY)
+        {
+            // 1.1111111111*2^{15} < nonneg32 < (MAX_NORMAL+INFINITY)/2, so
+            // the number is closest to 16-max-normal.
+            return sign16 | IEEEBinary16::MAX_NORMAL;
+        }
+
+        // nonneg32 >= (MAX_NORMAL+INFINITY)/2, so convert to 16-infinite.
+        return sign16 | IEEEBinary16::POS_INFINITY;
+    }
+
+    if (trailing32 == 0)
+    {
+        // The number is 32-infinite.  Convert to 16-infinite.
+        return sign16 | IEEEBinary16::POS_INFINITY;
+    }
+
+    // The number is 32-NaN.  Convert to 16-NaN with 16-payload the high-order
+    // 9 bits of the 32-payload.  The code also grabs the 32-quietNaN mask
+    // bit.
+    uint16_t maskPayload = static_cast<uint16_t>(
+        (trailing32 & 0x007FE000u) >> 13);
+    return sign16 | IEEEBinary16::EXPONENT_MASK | maskPayload;
+}
+
+uint32_t IEEEBinary16::Convert16To32(uint16_t encoding)
+{
+    // Extract the channels for the IEEEBinary16 number.
+    uint16_t sign16 = (encoding & IEEEBinary16::SIGN_MASK);
+    uint16_t biased16 = ((encoding & IEEEBinary16::EXPONENT_MASK)
+        >> IEEEBinary16::NUM_TRAILING_BITS);
+    uint16_t trailing16 = (encoding & IEEEBinary16::TRAILING_MASK);
+
+    // Generate the channels for the binary32 number.
+    uint32_t sign32 = static_cast<uint32_t>(sign16 << DIFF_NUM_ENCODING_BITS);
+    uint32_t biased32, trailing32;
+
+    if (biased16 == 0)
+    {
+        if (trailing16 == 0)
+        {
+            // The number is 16-zero.  Convert to 32-zero.
+            return sign32;
+        }
+        else
+        {
+            // The number is 16-subnormal.  Convert to 32-normal.
+            trailing32 = static_cast<uint32_t>(trailing16);
+            int32_t leading = GetLeadingBit(trailing32);
+            int32_t shift = 23 - leading;
+            biased32 = static_cast<uint32_t>(F32_EXPONENT_BIAS - 1 - shift);
+            trailing32 = (trailing32 << shift) & F32_TRAILING_MASK;
+            return sign32 | (biased32 << F32_NUM_TRAILING_BITS) | trailing32;
+        }
+    }
+
+    if (biased16 < IEEEBinary16::MAX_BIASED_EXPONENT)
+    {
+        // The number is 16-normal.  Convert to 32-normal.
+        biased32 = static_cast<uint32_t>(biased16 - IEEEBinary16::EXPONENT_BIAS +
+            F32_EXPONENT_BIAS);
+        trailing32 = (static_cast<uint32_t>(
+            trailing16) << DIFF_NUM_TRAILING_BITS);
+        return sign32 | (biased32 << F32_NUM_TRAILING_BITS) | trailing32;
+    }
+
+    if (trailing16 == 0)
+    {
+        // The number is 16-infinite.  Convert to 32-infinite.
+        return sign32 | F32_BIASED_EXPONENT_MASK;
+    }
+
+    // The number is 16-NaN.  Convert to 32-NaN with 16-payload embedded in
+    // the high-order 9 bits of the 32-payload.  The code also copies the
+    // 16-quietNaN mask bit.
+    uint32_t maskPayload =
+        ((trailing16 & IEEEBinary16::TRAILING_MASK) << DIFF_PAYLOAD_SHIFT);
+    return sign32 | F32_BIASED_EXPONENT_MASK | maskPayload;
+}
+
+IEEEBinary16 operator- (IEEEBinary16 x)
+{
+    uint16_t result = static_cast<uint16_t>(x) ^ IEEEBinary16::SIGN_MASK;
+    return result;
+}
+
+float operator+ (IEEEBinary16 x, IEEEBinary16 y)
+{
+    return static_cast<float>(x)+static_cast<float>(y);
+}
+
+float operator- (IEEEBinary16 x, IEEEBinary16 y)
+{
+    return static_cast<float>(x)-static_cast<float>(y);
+}
+
+float operator* (IEEEBinary16 x, IEEEBinary16 y)
+{
+    return static_cast<float>(x)* static_cast<float>(y);
+}
+
+float operator/ (IEEEBinary16 x, IEEEBinary16 y)
+{
+    return static_cast<float>(x) / static_cast<float>(y);
+}
+
+float operator+ (IEEEBinary16 x, float y)
+{
+    return static_cast<float>(x)+y;
+}
+
+float operator- (IEEEBinary16 x, float y)
+{
+    return static_cast<float>(x)-y;
+}
+
+float operator* (IEEEBinary16 x, float y)
+{
+    return static_cast<float>(x)* y;
+}
+
+float operator/ (IEEEBinary16 x, float y)
+{
+    return static_cast<float>(x) / y;
+}
+
+float operator+ (float x, IEEEBinary16 y)
+{
+    return x + static_cast<float>(y);
+}
+
+float operator- (float x, IEEEBinary16 y)
+{
+    return x - static_cast<float>(y);
+}
+
+float operator* (float x, IEEEBinary16 y)
+{
+    return x * static_cast<float>(y);
+}
+
+float operator/ (float x, IEEEBinary16 y)
+{
+    return x / static_cast<float>(y);
+}
+
+IEEEBinary16& operator+= (IEEEBinary16& x, IEEEBinary16 y)
+{
+    x = static_cast<float>(x)+static_cast<float>(y);
+    return x;
+}
+
+IEEEBinary16& operator-= (IEEEBinary16& x, IEEEBinary16 y)
+{
+    x = static_cast<float>(x)-static_cast<float>(y);
+    return x;
+}
+
+IEEEBinary16& operator*= (IEEEBinary16& x, IEEEBinary16 y)
+{
+    x = static_cast<float>(x)* static_cast<float>(y);
+    return x;
+}
+
+IEEEBinary16& operator/= (IEEEBinary16& x, IEEEBinary16 y)
+{
+    x = static_cast<float>(x) / static_cast<float>(y);
+    return x;
+}
+
+IEEEBinary16& operator+= (IEEEBinary16& x, float y)
+{
+    x = static_cast<float>(x)+y;
+    return x;
+}
+
+IEEEBinary16& operator-= (IEEEBinary16& x, float y)
+{
+    x = static_cast<float>(x)-y;
+    return x;
+}
+
+IEEEBinary16& operator*= (IEEEBinary16& x, float y)
+{
+    x = static_cast<float>(x)* y;
+    return x;
+}
+
+IEEEBinary16& operator/= (IEEEBinary16& x, float y)
+{
+    x = static_cast<float>(x) / y;
+    return x;
+}
+
+}
diff --git a/Engine/Plugins/Experimental/GeometryProcessing/Source/GeometryAlgorithms/Private/GTEngine/Mathematics/GteIEEEBinary16.h b/Engine/Plugins/Experimental/GeometryProcessing/Source/GeometryAlgorithms/Private/GTEngine/Mathematics/GteIEEEBinary16.h
new file mode 100644
--- /dev/null
+++ b/Engine/Plugins/Experimental/GeometryProcessing/Source/GeometryAlgorithms/Private/GTEngine/Mathematics/GteIEEEBinary16.h
@@ -0,0 +1,288 @@
+// David Eberly, Geometric Tools, Redmond WA 98052
+// Copyright (c) 1998-2019
+// Distributed under the Boost Software License, Version 1.0.
+// http://www.boost.org/LICENSE_1_0.txt
+// http://www.geometrictools.com/License/Boost/LICENSE_1_0.txt
+// File Version: 3.0.1 (2018/10/05)
+
+#pragma once
+
+#include <GTEngine/Mathematics/GteMath.h>
+#include <GTEngine/Mathematics/GteIEEEBinary.h>
+
+namespace gte
+{
+    class IEEEBinary16 : public IEEEBinary<int16_t, uint16_t, 16, 11>
+    {
+    public:
+        // Construction and destruction.  The base class destructor is hidden,
+        // but this is safe because there are no side effects of the
+        // destruction.
+        ~IEEEBinary16();
+        IEEEBinary16();  // uninitialized
+        IEEEBinary16(IEEEBinary16 const& object);
+        IEEEBinary16(float number);
+        IEEEBinary16(double number);
+        IEEEBinary16(uint16_t encoding);
+
+        // Implicit conversions.
+        operator float() const;
+        operator double() const;
+
+        // Assignment.
+        IEEEBinary16& operator=(IEEEBinary16 const& object);
+
+        // Comparison.
+        bool operator==(IEEEBinary16 const& object) const;
+        bool operator!=(IEEEBinary16 const& object) const;
+        bool operator< (IEEEBinary16 const& object) const;
+        bool operator<=(IEEEBinary16 const& object) const;
+        bool operator> (IEEEBinary16 const& object) const;
+        bool operator>=(IEEEBinary16 const& object) const;
+
+    private:
+        // Support for conversions between encodings.
+        enum
+        {
+            F32_NUM_ENCODING_BITS = 32,
+            F32_NUM_TRAILING_BITS = 23,
+            F32_EXPONENT_BIAS = 127,
+            F32_MAX_BIASED_EXPONENT = 255,
+            F32_SIGN_MASK = 0x80000000,
+            F32_NOT_SIGN_MASK = 0x7FFFFFFF,
+            F32_BIASED_EXPONENT_MASK = 0x7F800000,
+            F32_TRAILING_MASK = 0x007FFFFF,
+            F16_AVR_MIN_SUBNORMAL_ZERO = 0x33000000,
+            F16_MIN_SUBNORMAL = 0x33800000,
+            F16_MIN_NORMAL = 0x38800000,
+            F16_MAX_NORMAL = 0x477FE000,
+            F16_AVR_MAX_NORMAL_INFINITY = 0x477FF000,
+            DIFF_NUM_ENCODING_BITS = 16,
+            DIFF_NUM_TRAILING_BITS = 13,
+            DIFF_PAYLOAD_SHIFT = 13,
+            INT_PART_MASK = 0x007FE000,
+            FRC_PART_MASK = 0x00001FFF,
+            FRC_HALF = 0x00001000
+        };
+
+        static uint16_t Convert32To16(uint32_t encoding);
+        static uint32_t Convert16To32(uint16_t encoding);
+    };
+
+    // Arithmetic operations (high-precision).
+    IEEEBinary16 operator-(IEEEBinary16 x);
+    float operator+(IEEEBinary16 x, IEEEBinary16 y);
+    float operator-(IEEEBinary16 x, IEEEBinary16 y);
+    float operator*(IEEEBinary16 x, IEEEBinary16 y);
+    float operator/(IEEEBinary16 x, IEEEBinary16 y);
+    float operator+(IEEEBinary16 x, float y);
+    float operator-(IEEEBinary16 x, float y);
+    float operator*(IEEEBinary16 x, float y);
+    float operator/(IEEEBinary16 x, float y);
+    float operator+(float x, IEEEBinary16 y);
+    float operator-(float x, IEEEBinary16 y);
+    float operator*(float x, IEEEBinary16 y);
+    float operator/(float x, IEEEBinary16 y);
+
+    // Arithmetic updates.
+    IEEEBinary16& operator+=(IEEEBinary16& x, IEEEBinary16 y);
+    IEEEBinary16& operator-=(IEEEBinary16& x, IEEEBinary16 y);
+    IEEEBinary16& operator*=(IEEEBinary16& x, IEEEBinary16 y);
+    IEEEBinary16& operator/=(IEEEBinary16& x, IEEEBinary16 y);
+    IEEEBinary16& operator+=(IEEEBinary16& x, float y);
+    IEEEBinary16& operator-=(IEEEBinary16& x, float y);
+    IEEEBinary16& operator*=(IEEEBinary16& x, float y);
+    IEEEBinary16& operator/=(IEEEBinary16& x, float y);
+
+}
+
+namespace std
+{
+    inline gte::IEEEBinary16 abs(gte::IEEEBinary16 x)
+    {
+        return (gte::IEEEBinary16)std::abs((float)x);
+    }
+
+    inline gte::IEEEBinary16 acos(gte::IEEEBinary16 x)
+    {
+        return (gte::IEEEBinary16)std::acos((float)x);
+    }
+
+    inline gte::IEEEBinary16 acosh(gte::IEEEBinary16 x)
+    {
+        return (gte::IEEEBinary16)std::acosh((float)x);
+    }
+
+    inline gte::IEEEBinary16 asin(gte::IEEEBinary16 x)
+    {
+        return (gte::IEEEBinary16)std::asin((float)x);
+    }
+
+    inline gte::IEEEBinary16 asinh(gte::IEEEBinary16 x)
+    {
+        return (gte::IEEEBinary16)std::asin((float)x);
+    }
+
+    inline gte::IEEEBinary16 atan(gte::IEEEBinary16 x)
+    {
+        return (gte::IEEEBinary16)std::atan((float)x);
+    }
+
+    inline gte::IEEEBinary16 atanh(gte::IEEEBinary16 x)
+    {
+        return (gte::IEEEBinary16)std::atanh((float)x);
+    }
+
+    inline gte::IEEEBinary16 atan2(gte::IEEEBinary16 y, gte::IEEEBinary16 x)
+    {
+        return (gte::IEEEBinary16)std::atan2((float)y, (float)x);
+    }
+
+    inline gte::IEEEBinary16 ceil(gte::IEEEBinary16 x)
+    {
+        return (gte::IEEEBinary16)std::ceil((float)x);
+    }
+
+    inline gte::IEEEBinary16 cos(gte::IEEEBinary16 x)
+    {
+        return (gte::IEEEBinary16)std::cos((float)x);
+    }
+
+    inline gte::IEEEBinary16 cosh(gte::IEEEBinary16 x)
+    {
+        return (gte::IEEEBinary16)std::cosh((float)x);
+    }
+
+    inline gte::IEEEBinary16 exp(gte::IEEEBinary16 x)
+    {
+        return (gte::IEEEBinary16)std::exp((float)x);
+    }
+
+    inline gte::IEEEBinary16 exp2(gte::IEEEBinary16 x)
+    {
+        return (gte::IEEEBinary16)std::exp2((float)x);
+    }
+
+    inline gte::IEEEBinary16 floor(gte::IEEEBinary16 x)
+    {
+        return (gte::IEEEBinary16)std::floor((float)x);
+    }
+
+    inline gte::IEEEBinary16 fmod(gte::IEEEBinary16 x, gte::IEEEBinary16 y)
+    {
+        return (gte::IEEEBinary16)std::fmod((float)x, (float)y);
+    }
+
+    inline gte::IEEEBinary16 frexp(gte::IEEEBinary16 x, int* exponent)
+    {
+        return (gte::IEEEBinary16)std::frexp((float)x, exponent);
+    }
+
+    inline gte::IEEEBinary16 ldexp(gte::IEEEBinary16 x, int exponent)
+    {
+        return (gte::IEEEBinary16)std::ldexp((float)x, exponent);
+    }
+
+    inline gte::IEEEBinary16 log(gte::IEEEBinary16 x)
+    {
+        return (gte::IEEEBinary16)std::log((float)x);
+    }
+
+    inline gte::IEEEBinary16 log2(gte::IEEEBinary16 x)
+    {
+        return (gte::IEEEBinary16)std::log2((float)x);
+    }
+
+    inline gte::IEEEBinary16 log10(gte::IEEEBinary16 x)
+    {
+        return (gte::IEEEBinary16)std::log10((float)x);
+    }
+
+    inline gte::IEEEBinary16 pow(gte::IEEEBinary16 x, gte::IEEEBinary16 y)
+    {
+        return (gte::IEEEBinary16)std::pow((float)x, (float)y);
+    }
+
+    inline gte::IEEEBinary16 sin(gte::IEEEBinary16 x)
+    {
+        return (gte::IEEEBinary16)std::sin((float)x);
+    }
+
+    inline gte::IEEEBinary16 sinh(gte::IEEEBinary16 x)
+    {
+        return (gte::IEEEBinary16)std::sinh((float)x);
+    }
+
+    inline gte::IEEEBinary16 sqrt(gte::IEEEBinary16 x)
+    {
+        return (gte::IEEEBinary16)std::sqrt((float)x);
+    }
+
+    inline gte::IEEEBinary16 tan(gte::IEEEBinary16 x)
+    {
+        return (gte::IEEEBinary16)std::tan((float)x);
+    }
+
+    inline gte::IEEEBinary16 tanh(gte::IEEEBinary16 x)
+    {
+        return (gte::IEEEBinary16)std::tanh((float)x);
+    }
+}
+
+namespace gte
+{
+    inline IEEEBinary16 atandivpi(IEEEBinary16 x)
+    {
+        return (IEEEBinary16)atandivpi((float)x);
+    }
+
+    inline IEEEBinary16 atan2divpi(IEEEBinary16 y, IEEEBinary16 x)
+    {
+        return (IEEEBinary16)atan2divpi((float)y, (float)x);
+    }
+
+    inline IEEEBinary16 clamp(IEEEBinary16 x, IEEEBinary16 xmin, IEEEBinary16 xmax)
+    {
+        return (IEEEBinary16)clamp((float)x, (float)xmin, (float)xmax);
+    }
+
+    inline IEEEBinary16 cospi(IEEEBinary16 x)
+    {
+        return (IEEEBinary16)cospi((float)x);
+    }
+
+    inline IEEEBinary16 exp10(IEEEBinary16 x)
+    {
+        return (IEEEBinary16)exp10((float)x);
+    }
+
+    inline IEEEBinary16 invsqrt(IEEEBinary16 x)
+    {
+        return (IEEEBinary16)invsqrt((float)x);
+    }
+
+    inline int isign(IEEEBinary16 x)
+    {
+        return isign((float)x);
+    }
+
+    inline IEEEBinary16 saturate(IEEEBinary16 x)
+    {
+        return (IEEEBinary16)saturate((float)x);
+    }
+
+    inline IEEEBinary16 sign(IEEEBinary16 x)
+    {
+        return (IEEEBinary16)sign((float)x);
+    }
+
+    inline IEEEBinary16 sinpi(IEEEBinary16 x)
+    {
+        return (IEEEBinary16)sinpi((float)x);
+    }
+
+    inline IEEEBinary16 sqr(IEEEBinary16 x)
+    {
+        return (IEEEBinary16)sqr((float)x);
+    }
+}
diff --git a/Engine/Plugins/Experimental/GeometryProcessing/Source/GeometryAlgorithms/Private/GTEngine/Mathematics/GteIndexAttribute.h b/Engine/Plugins/Experimental/GeometryProcessing/Source/GeometryAlgorithms/Private/GTEngine/Mathematics/GteIndexAttribute.h
new file mode 100644
--- /dev/null
+++ b/Engine/Plugins/Experimental/GeometryProcessing/Source/GeometryAlgorithms/Private/GTEngine/Mathematics/GteIndexAttribute.h
@@ -0,0 +1,102 @@
+// David Eberly, Geometric Tools, Redmond WA 98052
+// Copyright (c) 1998-2019
+// Distributed under the Boost Software License, Version 1.0.
+// http://www.boost.org/LICENSE_1_0.txt
+// http://www.geometrictools.com/License/Boost/LICENSE_1_0.txt
+// File Version: 3.3.1 (2016/11/14)
+
+#pragma once
+
+#include <GTEngine/GTEngineDEF.h>
+#include <cstddef>
+#include <cstdint>
+
+// The IndexAttribute class represents an array of triples of indices into a
+// vertex array for an indexed triangle mesh.  For now, the source must be
+// either uint16_t or uint32_t.
+
+namespace gte
+{
+
+struct IndexAttribute
+{
+    // Construction.
+    inline IndexAttribute();
+    inline IndexAttribute(void* inSource, size_t inSize);
+
+    // Triangle access.
+    inline void SetTriangle(uint32_t t, uint32_t v0, uint32_t v1, uint32_t v2);
+    inline void GetTriangle(uint32_t t, uint32_t& v0, uint32_t& v1, uint32_t& v2) const;
+
+    // The source pointer must be 4-byte aligned, which is guaranteed on
+    // 32-bit and 64-bit architectures.  The size is the number of bytes
+    // per index.
+    void* source;
+    size_t size;
+};
+
+
+inline IndexAttribute::IndexAttribute()
+    :
+    source(nullptr),
+    size(0)
+{
+}
+
+inline IndexAttribute::IndexAttribute(void* inSource, size_t inSize)
+    :
+    source(inSource),
+    size(inSize)
+{
+}
+
+inline void IndexAttribute::SetTriangle(uint32_t t, uint32_t v0, uint32_t v1, uint32_t v2)
+{
+    if (size == sizeof(uint32_t))
+    {
+        uint32_t* index = reinterpret_cast<uint32_t*>(source) + 3 * t;
+        index[0] = v0;
+        index[1] = v1;
+        index[2] = v2;
+        return;
+    }
+
+    if (size == sizeof(uint16_t))
+    {
+        uint16_t* index = reinterpret_cast<uint16_t*>(source) + 3 * t;
+        index[0] = static_cast<uint16_t>(v0);
+        index[1] = static_cast<uint16_t>(v1);
+        index[2] = static_cast<uint16_t>(v2);
+        return;
+    }
+
+    // Unsupported type.
+}
+
+inline void IndexAttribute::GetTriangle(uint32_t t, uint32_t& v0, uint32_t& v1, uint32_t& v2) const
+{
+    if (size == sizeof(uint32_t))
+    {
+        uint32_t* index = reinterpret_cast<uint32_t*>(source) + 3 * t;
+        v0 = index[0];
+        v1 = index[1];
+        v2 = index[2];
+        return;
+    }
+
+    if (size == sizeof(uint16_t))
+    {
+        uint16_t* index = reinterpret_cast<uint16_t*>(source) + 3 * t;
+        v0 = static_cast<uint32_t>(index[0]);
+        v1 = static_cast<uint32_t>(index[1]);
+        v2 = static_cast<uint32_t>(index[2]);
+        return;
+    }
+
+    // Unsupported type.
+    v0 = 0;
+    v1 = 0;
+    v2 = 0;
+}
+
+}
diff --git a/Engine/Plugins/Experimental/GeometryProcessing/Source/GeometryAlgorithms/Private/GTEngine/Mathematics/GteIntegration.h b/Engine/Plugins/Experimental/GeometryProcessing/Source/GeometryAlgorithms/Private/GTEngine/Mathematics/GteIntegration.h
new file mode 100644
--- /dev/null
+++ b/Engine/Plugins/Experimental/GeometryProcessing/Source/GeometryAlgorithms/Private/GTEngine/Mathematics/GteIntegration.h
@@ -0,0 +1,238 @@
+// David Eberly, Geometric Tools, Redmond WA 98052
+// Copyright (c) 1998-2019
+// Distributed under the Boost Software License, Version 1.0.
+// http://www.boost.org/LICENSE_1_0.txt
+// http://www.geometrictools.com/License/Boost/LICENSE_1_0.txt
+// File Version: 3.0.0 (2016/06/19)
+
+#pragma once
+
+#include <GTEngine/Mathematics/GteRootsPolynomial.h>
+#include <array>
+#include <functional>
+#include <vector>
+
+namespace gte
+{
+
+template <typename Real>
+class Integration
+{
+public:
+    // A simple algorithm, but slow to converge as the number of samples is
+    // increased.  The 'numSamples' needs to be two or larger.
+    static Real TrapezoidRule(int numSamples, Real a, Real b,
+        std::function<Real(Real)> const& integrand);
+
+    // The trapezoid rule is used to generate initial estimates, but then
+    // Richardson extrapolation is used to improve the estimates.  This is
+    // preferred over TrapezoidRule.  The 'order' must be positive.
+    static Real Romberg(int order, Real a, Real b,
+        std::function<Real(Real)> const& integrand);
+
+    // Gaussian quadrature estimates the integral of a function f(x) defined
+    // on [-1,1] using
+    //   integral_{-1}^{1} f(t) dt = sum_{i=0}^{n-1} c[i]*f(r[i])
+    // where r[i] are the roots to the Legendre polynomial p(t) of degree n
+    // and c[i] = integral_{-1}^{1} prod_{j=0,j!=i} (t-r[j]/(r[i]-r[j]) dt.
+    // To integrate over [a,b], a transformation to [-1,1] is applied
+    // internally:  x - ((b-a)*t + (b+a))/2.  The Legendre polynomials are
+    // generated by
+    //   P[0](x) = 1, P[1](x) = x,
+    //   P[k](x) = ((2*k-1)*x*P[k-1](x) - (k-1)*P[k-2](x))/k, k >= 2
+    // Implementing the polynomial generation is simple, and computing the
+    // roots requires a numerical method for finding polynomial roots.  The
+    // challenging task is to develop an efficient algorithm for computing
+    // the coefficients c[i] for a specified degree.  The 'degree' must be
+    // two or larger.
+
+    static void ComputeQuadratureInfo(int degree, std::vector<Real>& roots,
+        std::vector<Real>& coefficients);
+
+    static Real GaussianQuadrature(std::vector<Real> const& roots,
+        std::vector<Real>const & coefficients, Real a, Real b,
+        std::function<Real(Real)> const& integrand);
+};
+
+
+template <typename Real>
+Real Integration<Real>::TrapezoidRule(int numSamples, Real a, Real b,
+    std::function<Real(Real)> const& integrand)
+{
+    Real h = (b - a) / (Real)(numSamples - 1);
+    Real result = ((Real)0.5) * (integrand(a) + integrand(b));
+    for (int i = 1; i <= numSamples - 2; ++i)
+    {
+        result += integrand(a + i*h);
+    }
+    result *= h;
+    return result;
+}
+
+template <typename Real>
+Real Integration<Real>::Romberg(int order, Real a, Real b,
+    std::function<Real(Real)> const& integrand)
+{
+    Real const half = (Real)0.5;
+    std::vector<std::array<Real, 2>> rom(order);
+    Real h = b - a;
+    rom[0][0] = half * h * (integrand(a) + integrand(b));
+    for (int i0 = 2, p0 = 1; i0 <= order; ++i0, p0 *= 2, h *= half)
+    {
+        // Approximations via the trapezoid rule.
+        Real sum = (Real)0;
+        int i1;
+        for (i1 = 1; i1 <= p0; ++i1)
+        {
+            sum += integrand(a + h * (i1 - half));
+        }
+
+        // Richardson extrapolation.
+        rom[0][1] = half * (rom[0][0] + h * sum);
+        for (int i2 = 1, p2 = 4; i2 < i0; ++i2, p2 *= 4)
+        {
+            rom[i2][1] = (p2*rom[i2 - 1][1] - rom[i2 - 1][0]) / (p2 - 1);
+        }
+
+        for (i1 = 0; i1 < i0; ++i1)
+        {
+            rom[i1][0] = rom[i1][1];
+        }
+    }
+
+    Real result = rom[order - 1][0];
+    return result;
+}
+
+template <typename Real>
+void Integration<Real>::ComputeQuadratureInfo(int degree,
+    std::vector<Real>& roots, std::vector<Real>& coefficients)
+{
+    Real const zero = (Real)0;
+    Real const one = (Real)1;
+    Real const half = (Real)0.5;
+
+    std::vector<std::vector<Real>> poly(degree + 1);
+
+    poly[0].resize(1);
+    poly[0][0] = one;
+
+    poly[1].resize(2);
+    poly[1][0] = zero;
+    poly[1][1] = one;
+
+    for (int n = 2; n <= degree; ++n)
+    {
+        Real mult0 = (Real)(n - 1) / (Real)n;
+        Real mult1 = (Real)(2 * n - 1) / (Real)n;
+
+        poly[n].resize(n + 1);
+        poly[n][0] = -mult0 * poly[n - 2][0];
+        for (int i = 1; i <= n - 2; ++i)
+        {
+            poly[n][i] = mult1 * poly[n - 1][i - 1] - mult0 * poly[n - 2][i];
+        }
+        poly[n][n - 1] = mult1 * poly[n - 1][n - 2];
+        poly[n][n] = mult1 * poly[n - 1][n - 1];
+    }
+
+    roots.resize(degree);
+    RootsPolynomial<Real>::Find(degree, &poly[degree][0], 2048, &roots[0]);
+
+    coefficients.resize(roots.size());
+    size_t n = roots.size() - 1;
+    std::vector<Real> subroots(n);
+    for (size_t i = 0; i < roots.size(); ++i)
+    {
+        Real denominator = (Real)1;
+        for (size_t j = 0, k = 0; j < roots.size(); ++j)
+        {
+            if (j != i)
+            {
+                subroots[k++] = roots[j];
+                denominator *= roots[i] - roots[j];
+            }
+        }
+
+        std::array<Real, 2> delta =
+        {
+            -one - subroots.back(),
+            +one - subroots.back()
+        };
+
+        std::vector<std::array<Real, 2>> weights(n);
+        weights[0][0] = half * delta[0] * delta[0];
+        weights[0][1] = half * delta[1] * delta[1];
+        for (size_t k = 1; k < n; ++k)
+        {
+            Real dk = (Real)k;
+            Real mult = -dk / (dk + (Real)2);
+            weights[k][0] = mult * delta[0] * weights[k - 1][0];
+            weights[k][1] = mult * delta[1] * weights[k - 1][1];
+        }
+
+        struct Info
+        {
+            int numBits;
+            std::array<Real, 2> product;
+        };
+
+        int numElements = (1 << static_cast<unsigned int>(n - 1));
+        std::vector<Info> info(numElements);
+        info[0].numBits = 0;
+        info[0].product[0] = one;
+        info[0].product[1] = one;
+        for (int ipow = 1, r = 0; ipow < numElements; ipow <<= 1, ++r)
+        {
+            info[ipow].numBits = 1;
+            info[ipow].product[0] = -one - subroots[r];
+            info[ipow].product[1] = +one - subroots[r];
+            for (int m = 1, j = ipow + 1; m < ipow; ++m, ++j)
+            {
+                info[j].numBits = info[m].numBits + 1;
+                info[j].product[0] =
+                    info[ipow].product[0] * info[m].product[0];
+                info[j].product[1] =
+                    info[ipow].product[1] * info[m].product[1];
+            }
+        }
+
+        std::vector<std::array<Real, 2>> sum(n);
+        std::array<Real, 2> zero2 = { zero, zero };
+        std::fill(sum.begin(), sum.end(), zero2);
+        for (size_t k = 0; k < info.size(); ++k)
+        {
+            sum[info[k].numBits][0] += info[k].product[0];
+            sum[info[k].numBits][1] += info[k].product[1];
+        }
+
+        std::array<Real, 2> total = zero2;
+        for (size_t k = 0; k < n; ++k)
+        {
+            total[0] += weights[n - 1 - k][0] * sum[k][0];
+            total[1] += weights[n - 1 - k][1] * sum[k][1];
+        }
+
+        coefficients[i] = (total[1] - total[0]) / denominator;
+    }
+}
+
+template <typename Real>
+Real Integration<Real>::GaussianQuadrature(std::vector<Real> const& roots,
+    std::vector<Real>const & coefficients, Real a, Real b,
+    std::function<Real(Real)> const& integrand)
+{
+    Real const half = (Real)0.5;
+    Real radius = half * (b - a);
+    Real center = half * (b + a);
+    Real result = (Real)0;
+    for (size_t i = 0; i < roots.size(); ++i)
+    {
+        result += coefficients[i] * integrand(radius*roots[i] + center);
+    }
+    result *= radius;
+    return result;
+}
+
+
+}
diff --git a/Engine/Plugins/Experimental/GeometryProcessing/Source/GeometryAlgorithms/Private/GTEngine/Mathematics/GteIntpAkima1.h b/Engine/Plugins/Experimental/GeometryProcessing/Source/GeometryAlgorithms/Private/GTEngine/Mathematics/GteIntpAkima1.h
new file mode 100644
--- /dev/null
+++ b/Engine/Plugins/Experimental/GeometryProcessing/Source/GeometryAlgorithms/Private/GTEngine/Mathematics/GteIntpAkima1.h
@@ -0,0 +1,181 @@
+// David Eberly, Geometric Tools, Redmond WA 98052
+// Copyright (c) 1998-2019
+// Distributed under the Boost Software License, Version 1.0.
+// http://www.boost.org/LICENSE_1_0.txt
+// http://www.geometrictools.com/License/Boost/LICENSE_1_0.txt
+// File Version: 3.0.0 (2016/06/19)
+
+#pragma once
+
+#include <GTEngine/LowLevel/GteLogger.h>
+#include <algorithm>
+#include <cmath>
+#include <vector>
+
+namespace gte
+{
+
+template <typename Real>
+class IntpAkima1
+{
+protected:
+    // Construction (abstract base class).
+    IntpAkima1(int quantity, Real const* F);
+public:
+    // Abstract base class.
+    virtual ~IntpAkima1();
+
+    // Member access.
+    inline int GetQuantity() const;
+    inline Real const* GetF() const;
+    virtual Real GetXMin() const = 0;
+    virtual Real GetXMax() const = 0;
+
+    // Evaluate the function and its derivatives.  The functions clamp the
+    // inputs to xmin <= x <= xmax.  The first operator is for function
+    // evaluation.  The second operator is for function or derivative
+    // evaluations.  The 'order' argument is the order of the derivative or
+    // zero for the function itself.
+    Real operator()(Real x) const;
+    Real operator()(int order, Real x) const;
+
+protected:
+    class Polynomial
+    {
+    public:
+        // P(x) = c[0] + c[1]*x + c[2]*x^2 + c[3]*x^3
+        inline Real& operator[](int i);
+        Real operator()(Real x) const;
+        Real operator()(int order, Real x) const;
+
+    private:
+        Real mCoeff[4];
+    };
+
+    Real ComputeDerivative(Real* slope) const;
+    virtual void Lookup(Real x, int& index, Real& dx) const = 0;
+
+    int mQuantity;
+    Real const* mF;
+    std::vector<Polynomial> mPoly;
+};
+
+
+template <typename Real>
+IntpAkima1<Real>::IntpAkima1(int quantity, Real const* F)
+    :
+    mQuantity(quantity),
+    mF(F)
+{
+    // At least three data points are needed to construct the estimates of
+    // the boundary derivatives.
+    LogAssert(mQuantity >= 3 && F, "Invalid input.");
+
+    mPoly.resize(mQuantity - 1);
+}
+
+template <typename Real>
+IntpAkima1<Real>::~IntpAkima1()
+{
+}
+
+template <typename Real> inline
+int IntpAkima1<Real>::GetQuantity() const
+{
+    return mQuantity;
+}
+
+template <typename Real> inline
+Real const* IntpAkima1<Real>::GetF() const
+{
+    return mF;
+}
+
+template <typename Real>
+Real IntpAkima1<Real>::operator()(Real x) const
+{
+    x = std::min(std::max(x, GetXMin()), GetXMax());
+    int index;
+    Real dx;
+    Lookup(x, index, dx);
+    return mPoly[index](dx);
+}
+
+template <typename Real>
+Real IntpAkima1<Real>::operator()(int order, Real x) const
+{
+    x = std::min(std::max(x, GetXMin()), GetXMax());
+    int index;
+    Real dx;
+    Lookup(x, index, dx);
+    return mPoly[index](order, dx);
+}
+
+template <typename Real>
+Real IntpAkima1<Real>::ComputeDerivative(Real* slope) const
+{
+    if (slope[1] != slope[2])
+    {
+        if (slope[0] != slope[1])
+        {
+            if (slope[2] != slope[3])
+            {
+                Real ad0 = std::abs(slope[3] - slope[2]);
+                Real ad1 = std::abs(slope[0] - slope[1]);
+                return (ad0 * slope[1] + ad1 * slope[2]) / (ad0 + ad1);
+            }
+            else
+            {
+                return slope[2];
+            }
+        }
+        else
+        {
+            if (slope[2] != slope[3])
+            {
+                return slope[1];
+            }
+            else
+            {
+                return ((Real)0.5) * (slope[1] + slope[2]);
+            }
+        }
+    }
+    else
+    {
+        return slope[1];
+    }
+}
+
+template <typename Real> inline
+Real& IntpAkima1<Real>::Polynomial::operator[](int i)
+{
+    return mCoeff[i];
+}
+
+template <typename Real>
+Real IntpAkima1<Real>::Polynomial::operator()(Real x) const
+{
+    return mCoeff[0] + x * (mCoeff[1] + x * (mCoeff[2] + x * mCoeff[3]));
+}
+
+template <typename Real>
+Real IntpAkima1<Real>::Polynomial::operator()(int order, Real x) const
+{
+    switch (order)
+    {
+    case 0:
+        return mCoeff[0] + x * (mCoeff[1] + x * (mCoeff[2] + x * mCoeff[3]));
+    case 1:
+        return mCoeff[1] + x * (((Real)2)*mCoeff[2] + x * ((Real)3) * mCoeff[3]);
+    case 2:
+        return ((Real)2) * mCoeff[2] + x * ((Real)6) * mCoeff[3];
+    case 3:
+        return ((Real)6) * mCoeff[3];
+    }
+
+    return (Real)0;
+}
+
+
+}
diff --git a/Engine/Plugins/Experimental/GeometryProcessing/Source/GeometryAlgorithms/Private/GTEngine/Mathematics/GteIntpAkimaNonuniform1.h b/Engine/Plugins/Experimental/GeometryProcessing/Source/GeometryAlgorithms/Private/GTEngine/Mathematics/GteIntpAkimaNonuniform1.h
new file mode 100644
--- /dev/null
+++ b/Engine/Plugins/Experimental/GeometryProcessing/Source/GeometryAlgorithms/Private/GTEngine/Mathematics/GteIntpAkimaNonuniform1.h
@@ -0,0 +1,128 @@
+// David Eberly, Geometric Tools, Redmond WA 98052
+// Copyright (c) 1998-2019
+// Distributed under the Boost Software License, Version 1.0.
+// http://www.boost.org/LICENSE_1_0.txt
+// http://www.geometrictools.com/License/Boost/LICENSE_1_0.txt
+// File Version: 3.0.0 (2016/06/19)
+
+#pragma once
+
+#include <GTEngine/Mathematics/GteIntpAkima1.h>
+
+namespace gte
+{
+
+template <typename Real>
+class IntpAkimaNonuniform1 : public IntpAkima1<Real>
+{
+public:
+    // Construction.  The interpolator is for arbitrarily spaced x-values.
+    // The input arrays must have 'quantity' elements and the X[] array
+    // must store increasing values:  X[i + 1] > X[i] for all i.
+    IntpAkimaNonuniform1(int quantity, Real const* X, Real const* F);
+
+    // Member access.
+    Real const* GetX() const;
+    virtual Real GetXMin() const override;
+    virtual Real GetXMax() const override;
+
+protected:
+    virtual void Lookup(Real x, int& index, Real& dx) const override;
+
+    Real const* mX;
+};
+
+template <typename Real>
+IntpAkimaNonuniform1<Real>::IntpAkimaNonuniform1(int quantity, Real const* X,
+    Real const* F)
+    :
+    IntpAkima1<Real>(quantity, F),
+    mX(X)
+{
+#if !defined(GTE_NO_LOGGER)
+    LogAssert(X != nullptr, "Invalid input.");
+    for (int j0 = 0, j1 = 1; j1 < quantity; ++j0, ++j1)
+    {
+        LogAssert(X[j1] > X[j0], "Invalid input.");
+    }
+#endif
+
+    // Compute slopes.
+    std::vector<Real> slope(quantity + 3);
+    int i, ip1, ip2;
+    for (i = 0, ip1 = 1, ip2 = 2; i < quantity - 1; ++i, ++ip1, ++ip2)
+    {
+        Real dx = X[ip1] - X[i];
+        Real df = F[ip1] - F[i];
+        slope[ip2] = df / dx;
+    }
+
+    slope[1] = ((Real)2) * slope[2] - slope[3];
+    slope[0] = ((Real)2) * slope[1] - slope[2];
+    slope[quantity + 1] = ((Real)2) * slope[quantity] - slope[quantity - 1];
+    slope[quantity + 2] = ((Real)2) * slope[quantity + 1] - slope[quantity];
+
+    // Construct derivatives.
+    std::vector<Real> FDer(quantity);
+    for (i = 0; i < quantity; ++i)
+    {
+        FDer[i] = this->ComputeDerivative(&slope[i]);
+    }
+
+    // Construct polynomials.
+    for (i = 0, ip1 = 1; i < quantity - 1; ++i, ++ip1)
+    {
+        auto& poly = this->mPoly[i];
+
+        Real F0 = F[i];
+        Real F1 = F[ip1];
+        Real FDer0 = FDer[i];
+        Real FDer1 = FDer[ip1];
+        Real df = F1 - F0;
+        Real dx = X[ip1] - X[i];
+        Real dx2 = dx * dx;
+        Real dx3 = dx2 * dx;
+
+        poly[0] = F0;
+        poly[1] = FDer0;
+        poly[2] = (((Real)3) * df - dx * (FDer1 + ((Real)2) * FDer0)) / dx2;
+        poly[3] = (dx * (FDer0 + FDer1) - ((Real)2) * df) / dx3;
+    }
+}
+
+template <typename Real>
+Real const* IntpAkimaNonuniform1<Real>::GetX() const
+{
+    return mX;
+}
+
+template <typename Real>
+Real IntpAkimaNonuniform1<Real>::GetXMin() const
+{
+    return mX[0];
+}
+
+template <typename Real>
+Real IntpAkimaNonuniform1<Real>::GetXMax() const
+{
+    return mX[this->mQuantity - 1];
+}
+
+template <typename Real>
+void IntpAkimaNonuniform1<Real>::Lookup(Real x, int& index, Real& dx) const
+{
+    // The caller has ensured that mXMin <= x <= mXMax.
+    for (index = 0; index + 1 < this->mQuantity; ++index)
+    {
+        if (x < mX[index + 1])
+        {
+            dx = x - mX[index];
+            return;
+        }
+    }
+
+    --index;
+    dx = x - mX[index];
+}
+
+}
diff --git a/Engine/Plugins/Experimental/GeometryProcessing/Source/GeometryAlgorithms/Private/GTEngine/Mathematics/GteIntpAkimaUniform1.h b/Engine/Plugins/Experimental/GeometryProcessing/Source/GeometryAlgorithms/Private/GTEngine/Mathematics/GteIntpAkimaUniform1.h
new file mode 100644
--- /dev/null
+++ b/Engine/Plugins/Experimental/GeometryProcessing/Source/GeometryAlgorithms/Private/GTEngine/Mathematics/GteIntpAkimaUniform1.h
@@ -0,0 +1,130 @@
+// David Eberly, Geometric Tools, Redmond WA 98052
+// Copyright (c) 1998-2019
+// Distributed under the Boost Software License, Version 1.0.
+// http://www.boost.org/LICENSE_1_0.txt
+// http://www.geometrictools.com/License/Boost/LICENSE_1_0.txt
+// File Version: 3.0.0 (2016/06/19)
+
+#pragma once
+
+#include <GTEngine/Mathematics/GteIntpAkima1.h>
+
+namespace gte
+{
+
+template <typename Real>
+class IntpAkimaUniform1 : public IntpAkima1<Real>
+{
+public:
+    // Construction and destruction.  The interpolator is for uniformly
+    // spaced x-values.
+    virtual ~IntpAkimaUniform1();
+    IntpAkimaUniform1(int quantity, Real xMin, Real xSpacing, Real const* F);
+
+    // Member access.
+    inline virtual Real GetXMin() const override;
+    inline virtual Real GetXMax() const override;
+    inline Real GetXSpacing() const;
+
+protected:
+    virtual void Lookup(Real x, int& index, Real& dx) const override;
+
+    Real mXMin, mXMax, mXSpacing;
+};
+
+
+template <typename Real>
+IntpAkimaUniform1<Real>::~IntpAkimaUniform1()
+{
+}
+
+template <typename Real>
+IntpAkimaUniform1<Real>::IntpAkimaUniform1(int quantity, Real xMin,
+    Real xSpacing, Real const* F)
+    :
+    IntpAkima1<Real>(quantity, F),
+    mXMin(xMin),
+    mXSpacing(xSpacing)
+{
+    LogAssert(mXSpacing > (Real)0, "Spacing must be positive.");
+
+    mXMax = mXMin + mXSpacing * static_cast<Real>(quantity - 1);
+
+    // Compute slopes.
+    Real invDX = ((Real)1) / mXSpacing;
+    std::vector<Real> slope(quantity + 3);
+    int i, ip1, ip2;
+    for (i = 0, ip1 = 1, ip2 = 2; i < quantity - 1; ++i, ++ip1, ++ip2)
+    {
+        slope[ip2] = (this->mF[ip1] - this->mF[i]) * invDX;
+    }
+
+    slope[1] = ((Real)2) * slope[2] - slope[3];
+    slope[0] = ((Real)2) * slope[1] - slope[2];
+    slope[quantity + 1] = ((Real)2) * slope[quantity] - slope[quantity - 1];
+    slope[quantity + 2] = ((Real)2) * slope[quantity + 1] - slope[quantity];
+
+    // Construct derivatives.
+    std::vector<Real> FDer(quantity);
+    for (i = 0; i < quantity; ++i)
+    {
+        FDer[i] = this->ComputeDerivative(&slope[i]);
+    }
+
+    // Construct polynomials.
+    Real invDX2 = ((Real)1) / (mXSpacing * mXSpacing);
+    Real invDX3 = invDX2 / mXSpacing;
+    for (i = 0, ip1 = 1; i < quantity - 1; ++i, ++ip1)
+    {
+        auto& poly = this->mPoly[i];
+
+        Real F0 = F[i];
+        Real F1 = F[ip1];
+        Real df = F1 - F0;
+        Real FDer0 = FDer[i];
+        Real FDer1 = FDer[ip1];
+
+        poly[0] = F0;
+        poly[1] = FDer0;
+        poly[2] = (((Real)3) * df - mXSpacing * (FDer1 + ((Real)2) * FDer0)) * invDX2;
+        poly[3] = (mXSpacing * (FDer0 + FDer1) - ((Real)2) * df) * invDX3;
+    }
+}
+
+template <typename Real> inline
+Real IntpAkimaUniform1<Real>::GetXMin() const
+{
+    return mXMin;
+}
+
+template <typename Real> inline
+Real IntpAkimaUniform1<Real>::GetXMax() const
+{
+    return mXMax;
+}
+
+template <typename Real> inline
+Real IntpAkimaUniform1<Real>::GetXSpacing() const
+{
+    return mXSpacing;
+}
+
+template <typename Real>
+void IntpAkimaUniform1<Real>::Lookup(Real x, int& index, Real& dx) const
+{
+    // The caller has ensured that mXMin <= x <= mXMax.
+    for (index = 0; index + 1 < this->mQuantity; ++index)
+    {
+        if (x < mXMin + mXSpacing * (index + 1))
+        {
+            dx = x - (mXMin + mXSpacing * index);
+            return;
+        }
+    }
+
+    --index;
+    dx = x - (mXMin + mXSpacing * index);
+}
+
+
+}
diff --git a/Engine/Plugins/Experimental/GeometryProcessing/Source/GeometryAlgorithms/Private/GTEngine/Mathematics/GteIntpAkimaUniform2.h b/Engine/Plugins/Experimental/GeometryProcessing/Source/GeometryAlgorithms/Private/GTEngine/Mathematics/GteIntpAkimaUniform2.h
new file mode 100644
--- /dev/null
+++ b/Engine/Plugins/Experimental/GeometryProcessing/Source/GeometryAlgorithms/Private/GTEngine/Mathematics/GteIntpAkimaUniform2.h
@@ -0,0 +1,627 @@
+// David Eberly, Geometric Tools, Redmond WA 98052
+// Copyright (c) 1998-2019
+// Distributed under the Boost Software License, Version 1.0.
+// http://www.boost.org/LICENSE_1_0.txt
+// http://www.geometrictools.com/License/Boost/LICENSE_1_0.txt
+// File Version: 3.0.1 (2018/10/04)
+
+#pragma once
+
+#include <GTEngine/LowLevel/GteArray2.h>
+#include <GTEngine/LowLevel/GteLogger.h>
+#include <algorithm>
+#include <cmath>
+#include <cstring>
+
+// The interpolator is for uniformly spaced (x,y)-values.  The input samples
+// F must be stored in row-major order to represent f(x,y); that is,
+// F[c + xBound*r] corresponds to f(x,y), where c is the index corresponding
+// to x and r is the index corresponding to y.
+
+namespace gte
+{
+
+template <typename Real>
+class IntpAkimaUniform2
+{
+public:
+    // Construction and destruction.
+    ~IntpAkimaUniform2();
+    IntpAkimaUniform2(int xBound, int yBound, Real xMin, Real xSpacing,
+        Real yMin, Real ySpacing, Real const* F);
+
+    // Member access.
+    inline int GetXBound() const;
+    inline int GetYBound() const;
+    inline int GetQuantity() const;
+    inline Real const* GetF() const;
+    inline Real GetXMin() const;
+    inline Real GetXMax() const;
+    inline Real GetXSpacing() const;
+    inline Real GetYMin() const;
+    inline Real GetYMax() const;
+    inline Real GetYSpacing() const;
+
+    // Evaluate the function and its derivatives.  The functions clamp the
+    // inputs to xmin <= x <= xmax and ymin <= y <= ymax.  The first operator
+    // is for function evaluation.  The second operator is for function or
+    // derivative evaluations.  The xOrder argument is the order of the
+    // x-derivative and the yOrder argument is the order of the y-derivative.
+    // Both orders are zero to get the function value itself.
+    Real operator()(Real x, Real y) const;
+    Real operator()(int xOrder, int yOrder, Real x, Real y) const;
+
+private:
+    class Polynomial
+    {
+    public:
+        Polynomial();
+
+        // P(x,y) = (1,x,x^2,x^3)*A*(1,y,y^2,y^3).  The matrix term A[ix][iy]
+        // corresponds to the polynomial term x^{ix} y^{iy}.
+        Real& A(int ix, int iy);
+        Real operator()(Real x, Real y) const;
+        Real operator()(int xOrder, int yOrder, Real x, Real y) const;
+
+    private:
+        Real mCoeff[4][4];
+    };
+
+    // Support for construction.
+    void GetFX(Array2<Real> const& F, Array2<Real>& FX);
+    void GetFY(Array2<Real> const& F, Array2<Real>& FY);
+    void GetFXY(Array2<Real> const& F, Array2<Real>& FXY);
+    void GetPolynomials(Array2<Real> const& F, Array2<Real> const& FX,
+        Array2<Real> const& FY, Array2<Real> const& FXY);
+    Real ComputeDerivative(Real const* slope) const;
+    void Construct(Polynomial& poly, Real const F[2][2], Real const FX[2][2],
+        Real const FY[2][2], Real const FXY[2][2]);
+
+    // Support for evaluation.
+    void XLookup(Real x, int& xIndex, Real& dx) const;
+    void YLookup(Real y, int& yIndex, Real& dy) const;
+
+    int mXBound, mYBound, mQuantity;
+    Real mXMin, mXMax, mXSpacing;
+    Real mYMin, mYMax, mYSpacing;
+    Real const* mF;
+    Array2<Polynomial> mPoly;
+};
+
+
+template <typename Real>
+IntpAkimaUniform2<Real>::~IntpAkimaUniform2()
+{
+}
+
+template <typename Real>
+IntpAkimaUniform2<Real>::IntpAkimaUniform2(int xBound, int yBound, Real xMin,
+    Real xSpacing, Real yMin, Real ySpacing, Real const* F)
+    :
+    mXBound(xBound),
+    mYBound(yBound),
+    mQuantity(xBound * yBound),
+    mXMin(xMin),
+    mXSpacing(xSpacing),
+    mYMin(yMin),
+    mYSpacing(ySpacing),
+    mF(F),
+    mPoly(xBound - 1, mYBound - 1)
+{
+    // At least a 3x3 block of data points is needed to construct the
+    // estimates of the boundary derivatives.
+    LogAssert(mXBound >= 3 && mYBound >= 3 && mF, "Invalid input.");
+    LogAssert(mXSpacing > (Real)0 && mYSpacing > (Real)0, "Invalid input.");
+
+    mXMax = mXMin + mXSpacing * static_cast<Real>(mXBound - 1);
+    mYMax = mYMin + mYSpacing * static_cast<Real>(mYBound - 1);
+
+    // Create a 2D wrapper for the 1D samples.
+    Array2<Real> Fmap(mXBound, mYBound, const_cast<Real*>(mF));
+
+    // Construct first-order derivatives.
+    Array2<Real> FX(mXBound, mYBound), FY(mXBound, mYBound);
+    GetFX(Fmap, FX);
+    GetFY(Fmap, FY);
+
+    // Construct second-order derivatives.
+    Array2<Real> FXY(mXBound, mYBound);
+    GetFXY(Fmap, FXY);
+
+    // Construct polynomials.
+    GetPolynomials(Fmap, FX, FY, FXY);
+}
+
+template <typename Real> inline
+int IntpAkimaUniform2<Real>::GetXBound() const
+{
+    return mXBound;
+}
+
+template <typename Real> inline
+int IntpAkimaUniform2<Real>::GetYBound() const
+{
+    return mYBound;
+}
+
+template <typename Real> inline
+int IntpAkimaUniform2<Real>::GetQuantity() const
+{
+    return mQuantity;
+}
+
+template <typename Real> inline
+Real const* IntpAkimaUniform2<Real>::GetF() const
+{
+    return mF;
+}
+
+template <typename Real> inline
+Real IntpAkimaUniform2<Real>::GetXMin() const
+{
+    return mXMin;
+}
+
+template <typename Real> inline
+Real IntpAkimaUniform2<Real>::GetXMax() const
+{
+    return mXMax;
+}
+
+template <typename Real> inline
+Real IntpAkimaUniform2<Real>::GetXSpacing() const
+{
+    return mXSpacing;
+}
+
+template <typename Real> inline
+Real IntpAkimaUniform2<Real>::GetYMin() const
+{
+    return mYMin;
+}
+
+template <typename Real> inline
+Real IntpAkimaUniform2<Real>::GetYMax() const
+{
+    return mYMax;
+}
+
+template <typename Real> inline
+Real IntpAkimaUniform2<Real>::GetYSpacing() const
+{
+    return mYSpacing;
+}
+
+template <typename Real>
+Real IntpAkimaUniform2<Real>::operator()(Real x, Real y) const
+{
+    x = std::min(std::max(x, mXMin), mXMax);
+    y = std::min(std::max(y, mYMin), mYMax);
+    int ix, iy;
+    Real dx, dy;
+    XLookup(x, ix, dx);
+    YLookup(y, iy, dy);
+    return mPoly[iy][ix](dx, dy);
+}
+
+template <typename Real>
+Real IntpAkimaUniform2<Real>::operator()(int xOrder, int yOrder, Real x,
+    Real y) const
+{
+    x = std::min(std::max(x, mXMin), mXMax);
+    y = std::min(std::max(y, mYMin), mYMax);
+    int ix, iy;
+    Real dx, dy;
+    XLookup(x, ix, dx);
+    YLookup(y, iy, dy);
+    return mPoly[iy][ix](xOrder, yOrder, dx, dy);
+}
+
+template <typename Real>
+void IntpAkimaUniform2<Real>::GetFX(Array2<Real> const& F, Array2<Real>& FX)
+{
+    Array2<Real> slope(mXBound + 3, mYBound);
+    Real invDX = ((Real)1) / mXSpacing;
+    int ix, iy;
+    for (iy = 0; iy < mYBound; ++iy)
+    {
+        for (ix = 0; ix < mXBound - 1; ++ix)
+        {
+            slope[iy][ix + 2] = (F[iy][ix + 1] - F[iy][ix]) * invDX;
+        }
+
+        slope[iy][1] = ((Real)2) * slope[iy][2] - slope[iy][3];
+        slope[iy][0] = ((Real)2) * slope[iy][1] - slope[iy][2];
+        slope[iy][mXBound + 1] = ((Real)2) * slope[iy][mXBound] - slope[iy][mXBound - 1];
+        slope[iy][mXBound + 2] = ((Real)2) * slope[iy][mXBound + 1] - slope[iy][mXBound];
+    }
+
+    for (iy = 0; iy < mYBound; ++iy)
+    {
+        for (ix = 0; ix < mXBound; ++ix)
+        {
+            FX[iy][ix] = ComputeDerivative(slope[iy] + ix);
+        }
+    }
+}
+
+template <typename Real>
+void IntpAkimaUniform2<Real>::GetFY(Array2<Real> const& F, Array2<Real>& FY)
+{
+    Array2<Real> slope(mYBound + 3, mXBound);
+    Real invDY = ((Real)1) / mYSpacing;
+    int ix, iy;
+    for (ix = 0; ix < mXBound; ++ix)
+    {
+        for (iy = 0; iy < mYBound - 1; ++iy)
+        {
+            slope[ix][iy + 2] = (F[iy + 1][ix] - F[iy][ix]) * invDY;
+        }
+
+        slope[ix][1] = ((Real)2) * slope[ix][2] - slope[ix][3];
+        slope[ix][0] = ((Real)2) * slope[ix][1] - slope[ix][2];
+        slope[ix][mYBound + 1] = ((Real)2) * slope[ix][mYBound] - slope[ix][mYBound - 1];
+        slope[ix][mYBound + 2] = ((Real)2) * slope[ix][mYBound + 1] - slope[ix][mYBound];
+    }
+
+    for (ix = 0; ix < mXBound; ++ix)
+    {
+        for (iy = 0; iy < mYBound; ++iy)
+        {
+            FY[iy][ix] = ComputeDerivative(slope[ix] + iy);
+        }
+    }
+}
+
+template <typename Real>
+void IntpAkimaUniform2<Real>::GetFXY(Array2<Real> const& F, Array2<Real>& FXY)
+{
+    int xBoundM1 = mXBound - 1;
+    int yBoundM1 = mYBound - 1;
+    int ix0 = xBoundM1, ix1 = ix0 - 1, ix2 = ix1 - 1;
+    int iy0 = yBoundM1, iy1 = iy0 - 1, iy2 = iy1 - 1;
+    int ix, iy;
+
+    Real invDXDY = ((Real)1) / (mXSpacing * mYSpacing);
+
+    // corners
+    FXY[0][0] = ((Real)0.25)*invDXDY*(
+        ((Real)9)*F[0][0]
+        - ((Real)12)*F[0][1]
+        + ((Real)3)*F[0][2]
+        - ((Real)12)*F[1][0]
+        + ((Real)16)*F[1][1]
+        - ((Real)4)*F[1][2]
+        + ((Real)3)*F[2][0]
+        - ((Real)4)*F[2][1]
+        + F[2][2]);
+
+    FXY[0][xBoundM1] = ((Real)0.25)*invDXDY*(
+        ((Real)9)*F[0][ix0]
+        - ((Real)12)*F[0][ix1]
+        + ((Real)3)*F[0][ix2]
+        - ((Real)12)*F[1][ix0]
+        + ((Real)16)*F[1][ix1]
+        - ((Real)4)*F[1][ix2]
+        + ((Real)3)*F[2][ix0]
+        - ((Real)4)*F[2][ix1]
+        + F[2][ix2]);
+
+    FXY[yBoundM1][0] = ((Real)0.25)*invDXDY*(
+        ((Real)9)*F[iy0][0]
+        - ((Real)12)*F[iy0][1]
+        + ((Real)3)*F[iy0][2]
+        - ((Real)12)*F[iy1][0]
+        + ((Real)16)*F[iy1][1]
+        - ((Real)4)*F[iy1][2]
+        + ((Real)3)*F[iy2][0]
+        - ((Real)4)*F[iy2][1]
+        + F[iy2][2]);
+
+    FXY[yBoundM1][xBoundM1] = ((Real)0.25)*invDXDY*(
+        ((Real)9)*F[iy0][ix0]
+        - ((Real)12)*F[iy0][ix1]
+        + ((Real)3)*F[iy0][ix2]
+        - ((Real)12)*F[iy1][ix0]
+        + ((Real)16)*F[iy1][ix1]
+        - ((Real)4)*F[iy1][ix2]
+        + ((Real)3)*F[iy2][ix0]
+        - ((Real)4)*F[iy2][ix1]
+        + F[iy2][ix2]);
+
+    // x-edges
+    for (ix = 1; ix < xBoundM1; ++ix)
+    {
+        FXY[0][ix] = ((Real)0.25)*invDXDY*(
+            ((Real)3)*(F[0][ix - 1] - F[0][ix + 1])
+            - ((Real)4)*(F[1][ix - 1] - F[1][ix + 1])
+            + (F[2][ix - 1] - F[2][ix + 1]));
+
+        FXY[yBoundM1][ix] = ((Real)0.25)*invDXDY*(
+            ((Real)3)*(F[iy0][ix - 1] - F[iy0][ix + 1])
+            - ((Real)4)*(F[iy1][ix - 1] - F[iy1][ix + 1])
+            + (F[iy2][ix - 1] - F[iy2][ix + 1]));
+    }
+
+    // y-edges
+    for (iy = 1; iy < yBoundM1; ++iy)
+    {
+        FXY[iy][0] = ((Real)0.25)*invDXDY*(
+            ((Real)3)*(F[iy - 1][0] - F[iy + 1][0])
+            - ((Real)4)*(F[iy - 1][1] - F[iy + 1][1])
+            + (F[iy - 1][2] - F[iy + 1][2]));
+
+        FXY[iy][xBoundM1] = ((Real)0.25)*invDXDY*(
+            ((Real)3)*(F[iy - 1][ix0] - F[iy + 1][ix0])
+            - ((Real)4)*(F[iy - 1][ix1] - F[iy + 1][ix1])
+            + (F[iy - 1][ix2] - F[iy + 1][ix2]));
+    }
+
+    // interior
+    for (iy = 1; iy < yBoundM1; ++iy)
+    {
+        for (ix = 1; ix < xBoundM1; ++ix)
+        {
+            FXY[iy][ix] = ((Real)0.25)*invDXDY*(F[iy - 1][ix - 1] -
+                F[iy - 1][ix + 1] - F[iy + 1][ix - 1] + F[iy + 1][ix + 1]);
+        }
+    }
+}
+
+template <typename Real>
+void IntpAkimaUniform2<Real>::GetPolynomials(Array2<Real> const& F,
+    Array2<Real> const& FX, Array2<Real> const& FY, Array2<Real> const& FXY)
+{
+    int xBoundM1 = mXBound - 1;
+    int yBoundM1 = mYBound - 1;
+    for (int iy = 0; iy < yBoundM1; ++iy)
+    {
+        for (int ix = 0; ix < xBoundM1; ++ix)
+        {
+            // Note the 'transposing' of the 2x2 blocks (to match notation
+            // used in the polynomial definition).
+            Real G[2][2] =
+            {
+                { F[iy][ix], F[iy + 1][ix] },
+                { F[iy][ix + 1], F[iy + 1][ix + 1] }
+            };
+
+            Real GX[2][2] =
+            {
+                { FX[iy][ix], FX[iy + 1][ix] },
+                { FX[iy][ix + 1], FX[iy + 1][ix + 1] }
+            };
+
+            Real GY[2][2] =
+            {
+                { FY[iy][ix], FY[iy + 1][ix] },
+                { FY[iy][ix + 1], FY[iy + 1][ix + 1] }
+            };
+
+            Real GXY[2][2] =
+            {
+                { FXY[iy][ix], FXY[iy + 1][ix] },
+                { FXY[iy][ix + 1], FXY[iy + 1][ix + 1] }
+            };
+
+            Construct(mPoly[iy][ix], G, GX, GY, GXY);
+        }
+    }
+}
+
+template <typename Real>
+Real IntpAkimaUniform2<Real>::ComputeDerivative(Real const* slope) const
+{
+    if (slope[1] != slope[2])
+    {
+        if (slope[0] != slope[1])
+        {
+            if (slope[2] != slope[3])
+            {
+                Real ad0 = std::abs(slope[3] - slope[2]);
+                Real ad1 = std::abs(slope[0] - slope[1]);
+                return (ad0 * slope[1] + ad1 * slope[2]) / (ad0 + ad1);
+            }
+            else
+            {
+                return slope[2];
+            }
+        }
+        else
+        {
+            if (slope[2] != slope[3])
+            {
+                return slope[1];
+            }
+            else
+            {
+                return ((Real)0.5) * (slope[1] + slope[2]);
+            }
+        }
+    }
+    else
+    {
+        return slope[1];
+    }
+}
+
+template <typename Real>
+void IntpAkimaUniform2<Real>::Construct(Polynomial& poly, Real const F[2][2],
+    Real const FX[2][2], Real const FY[2][2], Real const FXY[2][2])
+{
+    Real dx = mXSpacing;
+    Real dy = mYSpacing;
+    Real invDX = ((Real)1) / dx, invDX2 = invDX*invDX;
+    Real invDY = ((Real)1) / dy, invDY2 = invDY*invDY;
+    Real b0, b1, b2, b3;
+
+    poly.A(0, 0) = F[0][0];
+    poly.A(1, 0) = FX[0][0];
+    poly.A(0, 1) = FY[0][0];
+    poly.A(1, 1) = FXY[0][0];
+
+    b0 = (F[1][0] - poly(0, 0, dx, (Real)0))*invDX2;
+    b1 = (FX[1][0] - poly(1, 0, dx, (Real)0))*invDX;
+    poly.A(2, 0) = ((Real)3)*b0 - b1;
+    poly.A(3, 0) = (-((Real)2)*b0 + b1)*invDX;
+
+    b0 = (F[0][1] - poly(0, 0, (Real)0, dy))*invDY2;
+    b1 = (FY[0][1] - poly(0, 1, (Real)0, dy))*invDY;
+    poly.A(0, 2) = ((Real)3)*b0 - b1;
+    poly.A(0, 3) = (-((Real)2)*b0 + b1)*invDY;
+
+    b0 = (FY[1][0] - poly(0, 1, dx, (Real)0))*invDX2;
+    b1 = (FXY[1][0] - poly(1, 1, dx, (Real)0))*invDX;
+    poly.A(2, 1) = ((Real)3)*b0 - b1;
+    poly.A(3, 1) = (-((Real)2)*b0 + b1)*invDX;
+
+    b0 = (FX[0][1] - poly(1, 0, (Real)0, dy))*invDY2;
+    b1 = (FXY[0][1] - poly(1, 1, (Real)0, dy))*invDY;
+    poly.A(1, 2) = ((Real)3)*b0 - b1;
+    poly.A(1, 3) = (-((Real)2)*b0 + b1)*invDY;
+
+    b0 = (F[1][1] - poly(0, 0, dx, dy))*invDX2*invDY2;
+    b1 = (FX[1][1] - poly(1, 0, dx, dy))*invDX*invDY2;
+    b2 = (FY[1][1] - poly(0, 1, dx, dy))*invDX2*invDY;
+    b3 = (FXY[1][1] - poly(1, 1, dx, dy))*invDX*invDY;
+    poly.A(2, 2) = ((Real)9)*b0 - ((Real)3)*b1 - ((Real)3)*b2 + b3;
+    poly.A(3, 2) = (-((Real)6)*b0 + ((Real)3)*b1 + ((Real)2)*b2 - b3)*invDX;
+    poly.A(2, 3) = (-((Real)6)*b0 + ((Real)2)*b1 + ((Real)3)*b2 - b3)*invDY;
+    poly.A(3, 3) = (((Real)4)*b0 - ((Real)2)*b1 - ((Real)2)*b2 + b3)*invDX*invDY;
+}
+
+template <typename Real>
+void IntpAkimaUniform2<Real>::XLookup(Real x, int& xIndex, Real& dx) const
+{
+    for (xIndex = 0; xIndex + 1 < mXBound; ++xIndex)
+    {
+        if (x < mXMin + mXSpacing*(xIndex + 1))
+        {
+            dx = x - (mXMin + mXSpacing*xIndex);
+            return;
+        }
+    }
+
+    --xIndex;
+    dx = x - (mXMin + mXSpacing*xIndex);
+}
+
+template <typename Real>
+void IntpAkimaUniform2<Real>::YLookup(Real y, int& yIndex, Real& dy) const
+{
+    for (yIndex = 0; yIndex + 1 < mYBound; ++yIndex)
+    {
+        if (y < mYMin + mYSpacing*(yIndex + 1))
+        {
+            dy = y - (mYMin + mYSpacing*yIndex);
+            return;
+        }
+    }
+
+    yIndex--;
+    dy = y - (mYMin + mYSpacing*yIndex);
+}
+
+template <typename Real>
+IntpAkimaUniform2<Real>::Polynomial::Polynomial()
+{
+    memset(&mCoeff[0][0], 0, 16 * sizeof(Real));
+}
+
+template <typename Real>
+Real& IntpAkimaUniform2<Real>::Polynomial::A(int ix, int iy)
+{
+    return mCoeff[ix][iy];
+}
+
+template <typename Real>
+Real IntpAkimaUniform2<Real>::Polynomial::operator()(Real x, Real y) const
+{
+    Real B[4];
+    for (int i = 0; i <= 3; ++i)
+    {
+        B[i] = mCoeff[i][0] + y*(mCoeff[i][1] + y*(mCoeff[i][2] + y*mCoeff[i][3]));
+    }
+
+    return B[0] + x*(B[1] + x*(B[2] + x*B[3]));
+}
+
+template <typename Real>
+Real IntpAkimaUniform2<Real>::Polynomial::operator()(int xOrder, int yOrder,
+    Real x, Real y) const
+{
+    Real xPow[4];
+    switch (xOrder)
+    {
+    case 0:
+        xPow[0] = (Real)1;
+        xPow[1] = x;
+        xPow[2] = x * x;
+        xPow[3] = x * x * x;
+        break;
+    case 1:
+        xPow[0] = (Real)0;
+        xPow[1] = (Real)1;
+        xPow[2] = ((Real)2) * x;
+        xPow[3] = ((Real)3) * x * x;
+        break;
+    case 2:
+        xPow[0] = (Real)0;
+        xPow[1] = (Real)0;
+        xPow[2] = (Real)2;
+        xPow[3] = ((Real)6) * x;
+        break;
+    case 3:
+        xPow[0] = (Real)0;
+        xPow[1] = (Real)0;
+        xPow[2] = (Real)0;
+        xPow[3] = (Real)6;
+        break;
+    default:
+        return (Real)0;
+    }
+
+    Real yPow[4];
+    switch (yOrder)
+    {
+    case 0:
+        yPow[0] = (Real)1;
+        yPow[1] = y;
+        yPow[2] = y * y;
+        yPow[3] = y * y * y;
+        break;
+    case 1:
+        yPow[0] = (Real)0;
+        yPow[1] = (Real)1;
+        yPow[2] = ((Real)2) * y;
+        yPow[3] = ((Real)3) * y * y;
+        break;
+    case 2:
+        yPow[0] = (Real)0;
+        yPow[1] = (Real)0;
+        yPow[2] = (Real)2;
+        yPow[3] = ((Real)6) * y;
+        break;
+    case 3:
+        yPow[0] = (Real)0;
+        yPow[1] = (Real)0;
+        yPow[2] = (Real)0;
+        yPow[3] = (Real)6;
+        break;
+    default:
+        return (Real)0;
+    }
+
+    Real p = (Real)0;
+    for (int iy = 0; iy <= 3; ++iy)
+    {
+        for (int ix = 0; ix <= 3; ++ix)
+        {
+            p += mCoeff[ix][iy] * xPow[ix] * yPow[iy];
+        }
+    }
+
+    return p;
+}
+
+}
