move GeometryProcessing plugin from NotForLicensees to Experimental

diff --git a/Engine/Plugins/Experimental/GeometryProcessing/Source/GeometryAlgorithms/Private/GTEngine/Mathematics/GteIntpAkimaUniform3.h b/Engine/Plugins/Experimental/GeometryProcessing/Source/GeometryAlgorithms/Private/GTEngine/Mathematics/GteIntpAkimaUniform3.h
new file mode 100644
--- /dev/null
+++ b/Engine/Plugins/Experimental/GeometryProcessing/Source/GeometryAlgorithms/Private/GTEngine/Mathematics/GteIntpAkimaUniform3.h
@@ -0,0 +1,1503 @@
+// David Eberly, Geometric Tools, Redmond WA 98052
+// Copyright (c) 1998-2019
+// Distributed under the Boost Software License, Version 1.0.
+// http://www.boost.org/LICENSE_1_0.txt
+// http://www.geometrictools.com/License/Boost/LICENSE_1_0.txt
+// File Version: 3.0.0 (2016/06/19)
+
+#pragma once
+
+#include <GTEngine/LowLevel/GteArray3.h>
+#include <GTEngine/LowLevel/GteLogger.h>
+#include <algorithm>
+#include <cmath>
+#include <cstring>
+
+// The interpolator is for uniformly spaced(x,y z)-values.  The input samples
+// must be stored in lexicographical order to represent f(x,y,z); that is,
+// F[c + xBound*(r + yBound*s)] corresponds to f(x,y,z), where c is the index
+// corresponding to x, r is the index corresponding to y, and s is the index
+// corresponding to z.
+
+namespace gte
+{
+
+template <typename Real>
+class IntpAkimaUniform3
+{
+public:
+    // Construction and destruction.
+    ~IntpAkimaUniform3();
+    IntpAkimaUniform3(int xBound, int yBound, int zBound, Real xMin,
+        Real xSpacing, Real yMin, Real ySpacing, Real zMin, Real zSpacing,
+        Real const* F);
+
+    // Member access.
+    inline int GetXBound() const;
+    inline int GetYBound() const;
+    inline int GetZBound() const;
+    inline int GetQuantity() const;
+    inline Real const* GetF() const;
+    inline Real GetXMin() const;
+    inline Real GetXMax() const;
+    inline Real GetXSpacing() const;
+    inline Real GetYMin() const;
+    inline Real GetYMax() const;
+    inline Real GetYSpacing() const;
+    inline Real GetZMin() const;
+    inline Real GetZMax() const;
+    inline Real GetZSpacing() const;
+
+    // Evaluate the function and its derivatives.  The functions clamp the
+    // inputs to xmin <= x <= xmax, ymin <= y <= ymax, and zmin <= z <= zmax.
+    // The first operator is for function evaluation.  The second operator is
+    // for function or derivative evaluations.  The xOrder argument is the
+    // order of the x-derivative, the yOrder argument is the order of the
+    // y-derivative, and the zOrder argument is the order of the z-derivative.
+    // All orders are zero to get the function value itself.
+    Real operator()(Real x, Real y, Real z) const;
+    Real operator()(int xOrder, int yOrder, int zOrder, Real x, Real y, Real z) const;
+
+private:
+    class Polynomial
+    {
+    public:
+        Polynomial();
+
+        // P(x,y,z) = sum_{i=0}^3 sum_{j=0}^3 sum_{k=0}^3 a_{ijk} x^i y^j z^k.
+        // The tensor term A[ix][iy][iz] corresponds to the polynomial term
+        // x^{ix} y^{iy} z^{iz}.
+        Real& A(int ix, int iy, int iz);
+        Real operator()(Real x, Real y, Real z) const;
+        Real operator()(int xOrder, int yOrder, int zOrder, Real x, Real y, Real z) const;
+
+    private:
+        Real mCoeff[4][4][4];
+    };
+
+    // Support for construction.
+    void GetFX(Array3<Real> const& F, Array3<Real>& FX);
+    void GetFY(Array3<Real> const& F, Array3<Real>& FY);
+    void GetFZ(Array3<Real> const& F, Array3<Real>& FZ);
+    void GetFXY(Array3<Real> const& F, Array3<Real>& FXY);
+    void GetFXZ(Array3<Real> const& F, Array3<Real>& FXZ);
+    void GetFYZ(Array3<Real> const& F, Array3<Real>& FYZ);
+    void GetFXYZ(Array3<Real> const& F, Array3<Real>& FXYZ);
+    void GetPolynomials(Array3<Real> const& F, Array3<Real> const& FX,
+        Array3<Real> const& FY, Array3<Real> const& FZ, Array3<Real> const& FXY,
+        Array3<Real> const& FXZ, Array3<Real> const& FYZ, Array3<Real> const& FXYZ);
+
+    Real ComputeDerivative(Real const* slope) const;
+    void Construct(Polynomial& poly, Real const F[2][2][2],
+        Real const FX[2][2][2], Real const FY[2][2][2],
+        Real const FZ[2][2][2], Real const FXY[2][2][2],
+        Real const FXZ[2][2][2], Real const FYZ[2][2][2],
+        Real const FXYZ[2][2][2]);
+
+    void XLookup(Real x, int& xIndex, Real& dx) const;
+    void YLookup(Real y, int& yIndex, Real& dy) const;
+    void ZLookup(Real z, int& zIndex, Real& dz) const;
+
+    int mXBound, mYBound, mZBound, mQuantity;
+    Real mXMin, mXMax, mXSpacing;
+    Real mYMin, mYMax, mYSpacing;
+    Real mZMin, mZMax, mZSpacing;
+    Real const* mF;
+    Array3<Polynomial> mPoly;
+};
+
+
+template <typename Real>
+IntpAkimaUniform3<Real>::~IntpAkimaUniform3()
+{
+}
+
+template <typename Real>
+IntpAkimaUniform3<Real>::IntpAkimaUniform3(int xBound, int yBound, int zBound,
+    Real xMin, Real xSpacing, Real yMin, Real ySpacing, Real zMin,
+    Real zSpacing, Real const* F)
+    :
+    mXBound(xBound),
+    mYBound(yBound),
+    mZBound(zBound),
+    mQuantity(xBound * yBound * zBound),
+    mXMin(xMin),
+    mXSpacing(xSpacing),
+    mYMin(yMin),
+    mYSpacing(ySpacing),
+    mZMin(zMin),
+    mZSpacing(zSpacing),
+    mF(F),
+    mPoly(xBound - 1, yBound - 1, zBound - 1)
+{
+    // At least a 3x3x3 block of data points is needed to construct the
+    // estimates of the boundary derivatives.
+    LogAssert(mXBound >= 3 && mYBound >= 3 && mZBound >= 3 && mF,
+        "Invalid input.");
+    LogAssert(mXSpacing > (Real)0 && mYSpacing > (Real)0 &&
+        mZSpacing > (Real)0, "Invalid input.");
+
+    mXMax = mXMin + mXSpacing * static_cast<Real>(mXBound - 1);
+    mYMax = mYMin + mYSpacing * static_cast<Real>(mYBound - 1);
+    mZMax = mZMin + mZSpacing * static_cast<Real>(mZBound - 1);
+
+    // Create a 3D wrapper for the 1D samples.
+    Array3<Real> Fmap(mXBound, mYBound, mZBound, const_cast<Real*>(mF));
+
+    // Construct first-order derivatives.
+    Array3<Real> FX(mXBound, mYBound, mZBound);
+    Array3<Real> FY(mXBound, mYBound, mZBound);
+    Array3<Real> FZ(mXBound, mYBound, mZBound);
+    GetFX(Fmap, FX);
+    GetFX(Fmap, FY);
+    GetFX(Fmap, FZ);
+
+    // Construct second-order derivatives.
+    Array3<Real> FXY(mXBound, mYBound, mZBound);
+    Array3<Real> FXZ(mXBound, mYBound, mZBound);
+    Array3<Real> FYZ(mXBound, mYBound, mZBound);
+    GetFX(Fmap, FXY);
+    GetFX(Fmap, FXZ);
+    GetFX(Fmap, FYZ);
+
+    // Construct third-order derivatives.
+    Array3<Real> FXYZ(mXBound, mYBound, mZBound);
+    GetFXYZ(Fmap, FXYZ);
+
+    // Construct polynomials.
+    GetPolynomials(Fmap, FX, FY, FZ, FXY, FXZ, FYZ, FXYZ);
+}
+
+template <typename Real> inline
+int IntpAkimaUniform3<Real>::GetXBound() const
+{
+    return mXBound;
+}
+
+template <typename Real> inline
+int IntpAkimaUniform3<Real>::GetYBound() const
+{
+    return mYBound;
+}
+
+template <typename Real> inline
+int IntpAkimaUniform3<Real>::GetZBound() const
+{
+    return mZBound;
+}
+
+template <typename Real> inline
+int IntpAkimaUniform3<Real>::GetQuantity() const
+{
+    return mQuantity;
+}
+
+template <typename Real> inline
+Real const* IntpAkimaUniform3<Real>::GetF() const
+{
+    return mF;
+}
+
+template <typename Real> inline
+Real IntpAkimaUniform3<Real>::GetXMin() const
+{
+    return mXMin;
+}
+
+template <typename Real> inline
+Real IntpAkimaUniform3<Real>::GetXMax() const
+{
+    return mXMax;
+}
+
+template <typename Real> inline
+Real IntpAkimaUniform3<Real>::GetXSpacing() const
+{
+    return mXSpacing;
+}
+
+template <typename Real> inline
+Real IntpAkimaUniform3<Real>::GetYMin() const
+{
+    return mYMin;
+}
+
+template <typename Real> inline
+Real IntpAkimaUniform3<Real>::GetYMax() const
+{
+    return mYMax;
+}
+
+template <typename Real> inline
+Real IntpAkimaUniform3<Real>::GetYSpacing() const
+{
+    return mYSpacing;
+}
+
+template <typename Real> inline
+Real IntpAkimaUniform3<Real>::GetZMin() const
+{
+    return mZMin;
+}
+
+template <typename Real> inline
+Real IntpAkimaUniform3<Real>::GetZMax() const
+{
+    return mZMax;
+}
+
+template <typename Real> inline
+Real IntpAkimaUniform3<Real>::GetZSpacing() const
+{
+    return mZSpacing;
+}
+
+template <typename Real>
+Real IntpAkimaUniform3<Real>::operator()(Real x, Real y, Real z) const
+{
+    x = std::min(std::max(x, mXMin), mXMax);
+    y = std::min(std::max(y, mYMin), mYMax);
+    z = std::min(std::max(z, mZMin), mZMax);
+    int ix, iy, iz;
+    Real dx, dy, dz;
+    XLookup(x, ix, dx);
+    YLookup(y, iy, dy);
+    ZLookup(z, iz, dz);
+    return mPoly[iz][iy][ix](dx, dy, dz);
+}
+
+template <typename Real>
+Real IntpAkimaUniform3<Real>::operator()(int xOrder, int yOrder, int zOrder,
+    Real x, Real y, Real z) const
+{
+    x = std::min(std::max(x, mXMin), mXMax);
+    y = std::min(std::max(y, mYMin), mYMax);
+    z = std::min(std::max(z, mZMin), mZMax);
+    int ix, iy, iz;
+    Real dx, dy, dz;
+    XLookup(x, ix, dx);
+    YLookup(y, iy, dy);
+    ZLookup(z, iz, dz);
+    return mPoly[iz][iy][ix](xOrder, yOrder, zOrder, dx, dy, dz);
+}
+
+template <typename Real>
+void IntpAkimaUniform3<Real>::GetFX(Array3<Real> const& F, Array3<Real>& FX)
+{
+    Array3<Real> slope(mXBound + 3, mYBound, mZBound);
+    Real invDX = ((Real)1) / mXSpacing;
+    int ix, iy, iz;
+    for (iz = 0; iz < mZBound; ++iz)
+    {
+        for (iy = 0; iy < mYBound; ++iy)
+        {
+            for (ix = 0; ix < mXBound - 1; ++ix)
+            {
+                slope[iz][iy][ix + 2] = (F[iz][iy][ix + 1] - F[iz][iy][ix]) * invDX;
+            }
+
+            slope[iz][iy][1] = ((Real)2) * slope[iz][iy][2] - slope[iz][iy][3];
+            slope[iz][iy][0] = ((Real)2) * slope[iz][iy][1] - slope[iz][iy][2];
+            slope[iz][iy][mXBound + 1] = ((Real)2) * slope[iz][iy][mXBound] - slope[iz][iy][mXBound - 1];
+            slope[iz][iy][mXBound + 2] = ((Real)2) * slope[iz][iy][mXBound + 1] - slope[iz][iy][mXBound];
+        }
+    }
+
+    for (iz = 0; iz < mZBound; ++iz)
+    {
+        for (iy = 0; iy < mYBound; ++iy)
+        {
+            for (ix = 0; ix < mXBound; ++ix)
+            {
+                FX[iz][iy][ix] = ComputeDerivative(slope[iz][iy] + ix);
+            }
+        }
+    }
+}
+
+template <typename Real>
+void IntpAkimaUniform3<Real>::GetFY(Array3<Real> const& F, Array3<Real>& FY)
+{
+    Array3<Real> slope(mYBound + 3, mXBound, mZBound);
+    Real invDY = ((Real)1) / mYSpacing;
+    int ix, iy, iz;
+    for (iz = 0; iz < mZBound; ++iz)
+    {
+        for (ix = 0; ix < mXBound; ++ix)
+        {
+            for (iy = 0; iy < mYBound - 1; ++iy)
+            {
+                slope[iz][ix][iy + 2] = (F[iz][iy + 1][ix] - F[iz][iy][ix]) * invDY;
+            }
+
+            slope[iz][ix][1] = ((Real)2) * slope[iz][ix][2] - slope[iz][ix][3];
+            slope[iz][ix][0] = ((Real)2) * slope[iz][ix][1] - slope[iz][ix][2];
+            slope[iz][ix][mYBound + 1] = ((Real)2) * slope[iz][ix][mYBound] - slope[iz][ix][mYBound - 1];
+            slope[iz][ix][mYBound + 2] = ((Real)2) * slope[iz][ix][mYBound + 1] - slope[iz][ix][mYBound];
+        }
+    }
+
+    for (iz = 0; iz < mZBound; ++iz)
+    {
+        for (ix = 0; ix < mXBound; ++ix)
+        {
+            for (iy = 0; iy < mYBound; ++iy)
+            {
+                FY[iz][iy][ix] = ComputeDerivative(slope[iz][ix] + iy);
+            }
+        }
+    }
+}
+
+template <typename Real>
+void IntpAkimaUniform3<Real>::GetFZ(Array3<Real> const& F, Array3<Real>& FZ)
+{
+    Array3<Real> slope(mZBound + 3, mXBound, mYBound);
+    Real invDZ = ((Real)1) / mZSpacing;
+    int ix, iy, iz;
+    for (iy = 0; iy < mYBound; ++iy)
+    {
+        for (ix = 0; ix < mXBound; ++ix)
+        {
+            for (iz = 0; iz < mZBound - 1; ++iz)
+            {
+                slope[iy][ix][iz + 2] = (F[iz + 1][iy][ix] - F[iz][iy][ix]) * invDZ;
+            }
+
+            slope[iy][ix][1] = ((Real)2) * slope[iy][ix][2] - slope[iy][ix][3];
+            slope[iy][ix][0] = ((Real)2) * slope[iy][ix][1] - slope[iy][ix][2];
+            slope[iy][ix][mZBound + 1] = ((Real)2) * slope[iy][ix][mZBound] - slope[iy][ix][mZBound - 1];
+            slope[iy][ix][mZBound + 2] = ((Real)2) * slope[iy][ix][mZBound + 1] - slope[iy][ix][mZBound];
+        }
+    }
+
+    for (iy = 0; iy < mYBound; ++iy)
+    {
+        for (ix = 0; ix < mXBound; ++ix)
+        {
+            for (iz = 0; iz < mZBound; ++iz)
+            {
+                FZ[iz][iy][ix] = ComputeDerivative(slope[iy][ix] + iz);
+            }
+        }
+    }
+}
+
+template <typename Real>
+void IntpAkimaUniform3<Real>::GetFXY(Array3<Real> const& F, Array3<Real>& FXY)
+{
+    int xBoundM1 = mXBound - 1;
+    int yBoundM1 = mYBound - 1;
+    int ix0 = xBoundM1, ix1 = ix0 - 1, ix2 = ix1 - 1;
+    int iy0 = yBoundM1, iy1 = iy0 - 1, iy2 = iy1 - 1;
+    int ix, iy, iz;
+
+    Real invDXDY = ((Real)1) / (mXSpacing * mYSpacing);
+    for (iz = 0; iz < mZBound; ++iz)
+    {
+        // corners of z-slice
+        FXY[iz][0][0] = ((Real)0.25)*invDXDY*(
+            ((Real)9)*F[iz][0][0]
+            - ((Real)12)*F[iz][0][1]
+            + ((Real)3)*F[iz][0][2]
+            - ((Real)12)*F[iz][1][0]
+            + ((Real)16)*F[iz][1][1]
+            - ((Real)4)*F[iz][1][2]
+            + ((Real)3)*F[iz][2][0]
+            - ((Real)4)*F[iz][2][1]
+            + F[iz][2][2]);
+
+        FXY[iz][0][xBoundM1] = ((Real)0.25)*invDXDY*(
+            ((Real)9)*F[iz][0][ix0]
+            - ((Real)12)*F[iz][0][ix1]
+            + ((Real)3)*F[iz][0][ix2]
+            - ((Real)12)*F[iz][1][ix0]
+            + ((Real)16)*F[iz][1][ix1]
+            - ((Real)4)*F[iz][1][ix2]
+            + ((Real)3)*F[iz][2][ix0]
+            - ((Real)4)*F[iz][2][ix1]
+            + F[iz][2][ix2]);
+
+        FXY[iz][yBoundM1][0] = ((Real)0.25)*invDXDY*(
+            ((Real)9)*F[iz][iy0][0]
+            - ((Real)12)*F[iz][iy0][1]
+            + ((Real)3)*F[iz][iy0][2]
+            - ((Real)12)*F[iz][iy1][0]
+            + ((Real)16)*F[iz][iy1][1]
+            - ((Real)4)*F[iz][iy1][2]
+            + ((Real)3)*F[iz][iy2][0]
+            - ((Real)4)*F[iz][iy2][1]
+            + F[iz][iy2][2]);
+
+        FXY[iz][yBoundM1][xBoundM1] = ((Real)0.25)*invDXDY*(
+            ((Real)9)*F[iz][iy0][ix0]
+            - ((Real)12)*F[iz][iy0][ix1]
+            + ((Real)3)*F[iz][iy0][ix2]
+            - ((Real)12)*F[iz][iy1][ix0]
+            + ((Real)16)*F[iz][iy1][ix1]
+            - ((Real)4)*F[iz][iy1][ix2]
+            + ((Real)3)*F[iz][iy2][ix0]
+            - ((Real)4)*F[iz][iy2][ix1]
+            + F[iz][iy2][ix2]);
+
+        // x-edges of z-slice
+        for (ix = 1; ix < xBoundM1; ++ix)
+        {
+            FXY[iz][0][ix] = ((Real)0.25)*invDXDY*(
+                ((Real)3)*(F[iz][0][ix - 1] - F[iz][0][ix + 1]) -
+                ((Real)4)*(F[iz][1][ix - 1] - F[iz][1][ix + 1]) +
+                (F[iz][2][ix - 1] - F[iz][2][ix + 1]));
+
+            FXY[iz][yBoundM1][ix] = ((Real)0.25)*invDXDY*(
+                ((Real)3)*(F[iz][iy0][ix - 1] - F[iz][iy0][ix + 1])
+                - ((Real)4)*(F[iz][iy1][ix - 1] - F[iz][iy1][ix + 1]) +
+                (F[iz][iy2][ix - 1] - F[iz][iy2][ix + 1]));
+        }
+
+        // y-edges of z-slice
+        for (iy = 1; iy < yBoundM1; ++iy)
+        {
+            FXY[iz][iy][0] = ((Real)0.25)*invDXDY*(
+                ((Real)3)*(F[iz][iy - 1][0] - F[iz][iy + 1][0]) -
+                ((Real)4)*(F[iz][iy - 1][1] - F[iz][iy + 1][1]) +
+                (F[iz][iy - 1][2] - F[iz][iy + 1][2]));
+
+            FXY[iz][iy][xBoundM1] = ((Real)0.25)*invDXDY*(
+                ((Real)3)*(F[iz][iy - 1][ix0] - F[iz][iy + 1][ix0])
+                - ((Real)4)*(F[iz][iy - 1][ix1] - F[iz][iy + 1][ix1]) +
+                (F[iz][iy - 1][ix2] - F[iz][iy + 1][ix2]));
+        }
+
+        // interior of z-slice
+        for (iy = 1; iy < yBoundM1; ++iy)
+        {
+            for (ix = 1; ix < xBoundM1; ++ix)
+            {
+                FXY[iz][iy][ix] = ((Real)0.25)*invDXDY*(
+                    F[iz][iy - 1][ix - 1] - F[iz][iy - 1][ix + 1] -
+                    F[iz][iy + 1][ix - 1] + F[iz][iy + 1][ix + 1]);
+            }
+        }
+    }
+}
+
+template <typename Real>
+void IntpAkimaUniform3<Real>::GetFXZ(Array3<Real> const& F, Array3<Real>& FXZ)
+{
+    int xBoundM1 = mXBound - 1;
+    int zBoundM1 = mZBound - 1;
+    int ix0 = xBoundM1, ix1 = ix0 - 1, ix2 = ix1 - 1;
+    int iz0 = zBoundM1, iz1 = iz0 - 1, iz2 = iz1 - 1;
+    int ix, iy, iz;
+
+    Real invDXDZ = ((Real)1) / (mXSpacing * mZSpacing);
+    for (iy = 0; iy < mYBound; ++iy)
+    {
+        // corners of z-slice
+        FXZ[0][iy][0] = ((Real)0.25)*invDXDZ*(
+            ((Real)9)*F[0][iy][0]
+            - ((Real)12)*F[0][iy][1]
+            + ((Real)3)*F[0][iy][2]
+            - ((Real)12)*F[1][iy][0]
+            + ((Real)16)*F[1][iy][1]
+            - ((Real)4)*F[1][iy][2]
+            + ((Real)3)*F[2][iy][0]
+            - ((Real)4)*F[2][iy][1]
+            + F[2][iy][2]);
+
+        FXZ[0][iy][xBoundM1] = ((Real)0.25)*invDXDZ*(
+            ((Real)9)*F[0][iy][ix0]
+            - ((Real)12)*F[0][iy][ix1]
+            + ((Real)3)*F[0][iy][ix2]
+            - ((Real)12)*F[1][iy][ix0]
+            + ((Real)16)*F[1][iy][ix1]
+            - ((Real)4)*F[1][iy][ix2]
+            + ((Real)3)*F[2][iy][ix0]
+            - ((Real)4)*F[2][iy][ix1]
+            + F[2][iy][ix2]);
+
+        FXZ[zBoundM1][iy][0] = ((Real)0.25)*invDXDZ*(
+            ((Real)9)*F[iz0][iy][0]
+            - ((Real)12)*F[iz0][iy][1]
+            + ((Real)3)*F[iz0][iy][2]
+            - ((Real)12)*F[iz1][iy][0]
+            + ((Real)16)*F[iz1][iy][1]
+            - ((Real)4)*F[iz1][iy][2]
+            + ((Real)3)*F[iz2][iy][0]
+            - ((Real)4)*F[iz2][iy][1]
+            + F[iz2][iy][2]);
+
+        FXZ[zBoundM1][iy][xBoundM1] = ((Real)0.25)*invDXDZ*(
+            ((Real)9)*F[iz0][iy][ix0]
+            - ((Real)12)*F[iz0][iy][ix1]
+            + ((Real)3)*F[iz0][iy][ix2]
+            - ((Real)12)*F[iz1][iy][ix0]
+            + ((Real)16)*F[iz1][iy][ix1]
+            - ((Real)4)*F[iz1][iy][ix2]
+            + ((Real)3)*F[iz2][iy][ix0]
+            - ((Real)4)*F[iz2][iy][ix1]
+            + F[iz2][iy][ix2]);
+
+        // x-edges of y-slice
+        for (ix = 1; ix < xBoundM1; ++ix)
+        {
+            FXZ[0][iy][ix] = ((Real)0.25)*invDXDZ*(
+                ((Real)3)*(F[0][iy][ix - 1] - F[0][iy][ix + 1]) -
+                ((Real)4)*(F[1][iy][ix - 1] - F[1][iy][ix + 1]) +
+                (F[2][iy][ix - 1] - F[2][iy][ix + 1]));
+
+            FXZ[zBoundM1][iy][ix] = ((Real)0.25)*invDXDZ*(
+                ((Real)3)*(F[iz0][iy][ix - 1] - F[iz0][iy][ix + 1])
+                - ((Real)4)*(F[iz1][iy][ix - 1] - F[iz1][iy][ix + 1]) +
+                (F[iz2][iy][ix - 1] - F[iz2][iy][ix + 1]));
+        }
+
+        // z-edges of y-slice
+        for (iz = 1; iz < zBoundM1; ++iz)
+        {
+            FXZ[iz][iy][0] = ((Real)0.25)*invDXDZ*(
+                ((Real)3)*(F[iz - 1][iy][0] - F[iz + 1][iy][0]) -
+                ((Real)4)*(F[iz - 1][iy][1] - F[iz + 1][iy][1]) +
+                (F[iz - 1][iy][2] - F[iz + 1][iy][2]));
+
+            FXZ[iz][iy][xBoundM1] = ((Real)0.25)*invDXDZ*(
+                ((Real)3)*(F[iz - 1][iy][ix0] - F[iz + 1][iy][ix0])
+                - ((Real)4)*(F[iz - 1][iy][ix1] - F[iz + 1][iy][ix1]) +
+                (F[iz - 1][iy][ix2] - F[iz + 1][iy][ix2]));
+        }
+
+        // interior of y-slice
+        for (iz = 1; iz < zBoundM1; ++iz)
+        {
+            for (ix = 1; ix < xBoundM1; ++ix)
+            {
+                FXZ[iz][iy][ix] = ((Real)0.25)*invDXDZ*(
+                    F[iz - 1][iy][ix - 1] - F[iz - 1][iy][ix + 1] -
+                    F[iz + 1][iy][ix - 1] + F[iz + 1][iy][ix + 1]);
+            }
+        }
+    }
+}
+
+template <typename Real>
+void IntpAkimaUniform3<Real>::GetFYZ(Array3<Real> const& F, Array3<Real>& FYZ)
+{
+    int yBoundM1 = mYBound - 1;
+    int zBoundM1 = mZBound - 1;
+    int iy0 = yBoundM1, iy1 = iy0 - 1, iy2 = iy1 - 1;
+    int iz0 = zBoundM1, iz1 = iz0 - 1, iz2 = iz1 - 1;
+    int ix, iy, iz;
+
+    Real invDYDZ = ((Real)1) / (mYSpacing * mZSpacing);
+    for (ix = 0; ix < mXBound; ++ix)
+    {
+        // corners of x-slice
+        FYZ[0][0][ix] = ((Real)0.25)*invDYDZ*(
+            ((Real)9)*F[0][0][ix]
+            - ((Real)12)*F[0][1][ix]
+            + ((Real)3)*F[0][2][ix]
+            - ((Real)12)*F[1][0][ix]
+            + ((Real)16)*F[1][1][ix]
+            - ((Real)4)*F[1][2][ix]
+            + ((Real)3)*F[2][0][ix]
+            - ((Real)4)*F[2][1][ix]
+            + F[2][2][ix]);
+
+        FYZ[0][yBoundM1][ix] = ((Real)0.25)*invDYDZ*(
+            ((Real)9)*F[0][iy0][ix]
+            - ((Real)12)*F[0][iy1][ix]
+            + ((Real)3)*F[0][iy2][ix]
+            - ((Real)12)*F[1][iy0][ix]
+            + ((Real)16)*F[1][iy1][ix]
+            - ((Real)4)*F[1][iy2][ix]
+            + ((Real)3)*F[2][iy0][ix]
+            - ((Real)4)*F[2][iy1][ix]
+            + F[2][iy2][ix]);
+
+        FYZ[zBoundM1][0][ix] = ((Real)0.25)*invDYDZ*(
+            ((Real)9)*F[iz0][0][ix]
+            - ((Real)12)*F[iz0][1][ix]
+            + ((Real)3)*F[iz0][2][ix]
+            - ((Real)12)*F[iz1][0][ix]
+            + ((Real)16)*F[iz1][1][ix]
+            - ((Real)4)*F[iz1][2][ix]
+            + ((Real)3)*F[iz2][0][ix]
+            - ((Real)4)*F[iz2][1][ix]
+            + F[iz2][2][ix]);
+
+        FYZ[zBoundM1][yBoundM1][ix] = ((Real)0.25)*invDYDZ*(
+            ((Real)9)*F[iz0][iy0][ix]
+            - ((Real)12)*F[iz0][iy1][ix]
+            + ((Real)3)*F[iz0][iy2][ix]
+            - ((Real)12)*F[iz1][iy0][ix]
+            + ((Real)16)*F[iz1][iy1][ix]
+            - ((Real)4)*F[iz1][iy2][ix]
+            + ((Real)3)*F[iz2][iy0][ix]
+            - ((Real)4)*F[iz2][iy1][ix]
+            + F[iz2][iy2][ix]);
+
+        // y-edges of x-slice
+        for (iy = 1; iy < yBoundM1; ++iy)
+        {
+            FYZ[0][iy][ix] = ((Real)0.25)*invDYDZ*(
+                ((Real)3)*(F[0][iy - 1][ix] - F[0][iy + 1][ix]) -
+                ((Real)4)*(F[1][iy - 1][ix] - F[1][iy + 1][ix]) +
+                (F[2][iy - 1][ix] - F[2][iy + 1][ix]));
+
+            FYZ[zBoundM1][iy][ix] = ((Real)0.25)*invDYDZ*(
+                ((Real)3)*(F[iz0][iy - 1][ix] - F[iz0][iy + 1][ix])
+                - ((Real)4)*(F[iz1][iy - 1][ix] - F[iz1][iy + 1][ix]) +
+                (F[iz2][iy - 1][ix] - F[iz2][iy + 1][ix]));
+        }
+
+        // z-edges of x-slice
+        for (iz = 1; iz < zBoundM1; ++iz)
+        {
+            FYZ[iz][0][ix] = ((Real)0.25)*invDYDZ*(
+                ((Real)3)*(F[iz - 1][0][ix] - F[iz + 1][0][ix]) -
+                ((Real)4)*(F[iz - 1][1][ix] - F[iz + 1][1][ix]) +
+                (F[iz - 1][2][ix] - F[iz + 1][2][ix]));
+
+            FYZ[iz][yBoundM1][ix] = ((Real)0.25)*invDYDZ*(
+                ((Real)3)*(F[iz - 1][iy0][ix] - F[iz + 1][iy0][ix])
+                - ((Real)4)*(F[iz - 1][iy1][ix] - F[iz + 1][iy1][ix]) +
+                (F[iz - 1][iy2][ix] - F[iz + 1][iy2][ix]));
+        }
+
+        // interior of x-slice
+        for (iz = 1; iz < zBoundM1; ++iz)
+        {
+            for (iy = 1; iy < yBoundM1; ++iy)
+            {
+                FYZ[iz][iy][ix] = ((Real)0.25)*invDYDZ*(
+                    F[iz - 1][iy - 1][ix] - F[iz - 1][iy + 1][ix] -
+                    F[iz + 1][iy - 1][ix] + F[iz + 1][iy + 1][ix]);
+            }
+        }
+    }
+}
+
+template <typename Real>
+void IntpAkimaUniform3<Real>::GetFXYZ(Array3<Real> const& F, Array3<Real>& FXYZ)
+{
+    int xBoundM1 = mXBound - 1;
+    int yBoundM1 = mYBound - 1;
+    int zBoundM1 = mZBound - 1;
+    int ix, iy, iz, ix0, iy0, iz0;
+
+    Real invDXDYDZ = ((Real)1) / (mXSpacing * mYSpacing * mZSpacing);
+
+    // convolution masks
+    //   centered difference, O(h^2)
+    Real CDer[3] = { -(Real)0.5, (Real)0, (Real)0.5 };
+    //   one-sided difference, O(h^2)
+    Real ODer[3] = { -(Real)1.5, (Real)2, -(Real)0.5 };
+    Real mask;
+
+    // corners
+    FXYZ[0][0][0] = (Real)0;
+    FXYZ[0][0][xBoundM1] = (Real)0;
+    FXYZ[0][yBoundM1][0] = (Real)0;
+    FXYZ[0][yBoundM1][xBoundM1] = (Real)0;
+    FXYZ[zBoundM1][0][0] = (Real)0;
+    FXYZ[zBoundM1][0][xBoundM1] = (Real)0;
+    FXYZ[zBoundM1][yBoundM1][0] = (Real)0;
+    FXYZ[zBoundM1][yBoundM1][xBoundM1] = (Real)0;
+    for (iz = 0; iz <= 2; ++iz)
+    {
+        for (iy = 0; iy <= 2; ++iy)
+        {
+            for (ix = 0; ix <= 2; ++ix)
+            {
+                mask = invDXDYDZ*ODer[ix] * ODer[iy] * ODer[iz];
+                FXYZ[0][0][0] += mask * F[iz][iy][ix];
+                FXYZ[0][0][xBoundM1] += mask * F[iz][iy][xBoundM1 - ix];
+                FXYZ[0][yBoundM1][0] += mask * F[iz][yBoundM1 - iy][ix];
+                FXYZ[0][yBoundM1][xBoundM1] += mask * F[iz][yBoundM1 - iy][xBoundM1 - ix];
+                FXYZ[zBoundM1][0][0] += mask * F[zBoundM1 - iz][iy][ix];
+                FXYZ[zBoundM1][0][xBoundM1] += mask * F[zBoundM1 - iz][iy][xBoundM1 - ix];
+                FXYZ[zBoundM1][yBoundM1][0] += mask * F[zBoundM1 - iz][yBoundM1 - iy][ix];
+                FXYZ[zBoundM1][yBoundM1][xBoundM1] += mask * F[zBoundM1 - iz][yBoundM1 - iy][xBoundM1 - ix];
+            }
+        }
+    }
+
+    // x-edges
+    for (ix0 = 1; ix0 < xBoundM1; ++ix0)
+    {
+        FXYZ[0][0][ix0] = (Real)0;
+        FXYZ[0][yBoundM1][ix0] = (Real)0;
+        FXYZ[zBoundM1][0][ix0] = (Real)0;
+        FXYZ[zBoundM1][yBoundM1][ix0] = (Real)0;
+        for (iz = 0; iz <= 2; ++iz)
+        {
+            for (iy = 0; iy <= 2; ++iy)
+            {
+                for (ix = 0; ix <= 2; ++ix)
+                {
+                    mask = invDXDYDZ*CDer[ix] * ODer[iy] * ODer[iz];
+                    FXYZ[0][0][ix0] += mask * F[iz][iy][ix0 + ix - 1];
+                    FXYZ[0][yBoundM1][ix0] += mask * F[iz][yBoundM1 - iy][ix0 + ix - 1];
+                    FXYZ[zBoundM1][0][ix0] += mask * F[zBoundM1 - iz][iy][ix0 + ix - 1];
+                    FXYZ[zBoundM1][yBoundM1][ix0] += mask * F[zBoundM1 - iz][yBoundM1 - iy][ix0 + ix - 1];
+                }
+            }
+        }
+    }
+
+    // y-edges
+    for (iy0 = 1; iy0 < yBoundM1; ++iy0)
+    {
+        FXYZ[0][iy0][0] = (Real)0;
+        FXYZ[0][iy0][xBoundM1] = (Real)0;
+        FXYZ[zBoundM1][iy0][0] = (Real)0;
+        FXYZ[zBoundM1][iy0][xBoundM1] = (Real)0;
+        for (iz = 0; iz <= 2; ++iz)
+        {
+            for (iy = 0; iy <= 2; ++iy)
+            {
+                for (ix = 0; ix <= 2; ++ix)
+                {
+                    mask = invDXDYDZ*ODer[ix] * CDer[iy] * ODer[iz];
+                    FXYZ[0][iy0][0] += mask * F[iz][iy0 + iy - 1][ix];
+                    FXYZ[0][iy0][xBoundM1] += mask * F[iz][iy0 + iy - 1][xBoundM1 - ix];
+                    FXYZ[zBoundM1][iy0][0] += mask * F[zBoundM1 - iz][iy0 + iy - 1][ix];
+                    FXYZ[zBoundM1][iy0][xBoundM1] += mask * F[zBoundM1 - iz][iy0 + iy - 1][xBoundM1 - ix];
+                }
+            }
+        }
+    }
+
+    // z-edges
+    for (iz0 = 1; iz0 < zBoundM1; ++iz0)
+    {
+        FXYZ[iz0][0][0] = (Real)0;
+        FXYZ[iz0][0][xBoundM1] = (Real)0;
+        FXYZ[iz0][yBoundM1][0] = (Real)0;
+        FXYZ[iz0][yBoundM1][xBoundM1] = (Real)0;
+        for (iz = 0; iz <= 2; ++iz)
+        {
+            for (iy = 0; iy <= 2; ++iy)
+            {
+                for (ix = 0; ix <= 2; ++ix)
+                {
+                    mask = invDXDYDZ*ODer[ix] * ODer[iy] * CDer[iz];
+                    FXYZ[iz0][0][0] += mask * F[iz0 + iz - 1][iy][ix];
+                    FXYZ[iz0][0][xBoundM1] += mask * F[iz0 + iz - 1][iy][xBoundM1 - ix];
+                    FXYZ[iz0][yBoundM1][0] += mask * F[iz0 + iz - 1][yBoundM1 - iy][ix];
+                    FXYZ[iz0][yBoundM1][xBoundM1] += mask * F[iz0 + iz - 1][yBoundM1 - iy][xBoundM1 - ix];
+                }
+            }
+        }
+    }
+
+    // xy-faces
+    for (iy0 = 1; iy0 < yBoundM1; ++iy0)
+    {
+        for (ix0 = 1; ix0 < xBoundM1; ++ix0)
+        {
+            FXYZ[0][iy0][ix0] = (Real)0;
+            FXYZ[zBoundM1][iy0][ix0] = (Real)0;
+            for (iz = 0; iz <= 2; ++iz)
+            {
+                for (iy = 0; iy <= 2; ++iy)
+                {
+                    for (ix = 0; ix <= 2; ++ix)
+                    {
+                        mask = invDXDYDZ*CDer[ix] * CDer[iy] * ODer[iz];
+                        FXYZ[0][iy0][ix0] += mask * F[iz][iy0 + iy - 1][ix0 + ix - 1];
+                        FXYZ[zBoundM1][iy0][ix0] += mask * F[zBoundM1 - iz][iy0 + iy - 1][ix0 + ix - 1];
+                    }
+                }
+            }
+        }
+    }
+
+    // xz-faces
+    for (iz0 = 1; iz0 < zBoundM1; ++iz0)
+    {
+        for (ix0 = 1; ix0 < xBoundM1; ++ix0)
+        {
+            FXYZ[iz0][0][ix0] = (Real)0;
+            FXYZ[iz0][yBoundM1][ix0] = (Real)0;
+            for (iz = 0; iz <= 2; ++iz)
+            {
+                for (iy = 0; iy <= 2; ++iy)
+                {
+                    for (ix = 0; ix <= 2; ++ix)
+                    {
+                        mask = invDXDYDZ*CDer[ix] * ODer[iy] * CDer[iz];
+                        FXYZ[iz0][0][ix0] += mask * F[iz0 + iz - 1][iy][ix0 + ix - 1];
+                        FXYZ[iz0][yBoundM1][ix0] += mask * F[iz0 + iz - 1][yBoundM1 - iy][ix0 + ix - 1];
+                    }
+                }
+            }
+        }
+    }
+
+    // yz-faces
+    for (iz0 = 1; iz0 < zBoundM1; ++iz0)
+    {
+        for (iy0 = 1; iy0 < yBoundM1; ++iy0)
+        {
+            FXYZ[iz0][iy0][0] = (Real)0;
+            FXYZ[iz0][iy0][xBoundM1] = (Real)0;
+            for (iz = 0; iz <= 2; ++iz)
+            {
+                for (iy = 0; iy <= 2; ++iy)
+                {
+                    for (ix = 0; ix <= 2; ++ix)
+                    {
+                        mask = invDXDYDZ*ODer[ix] * CDer[iy] * CDer[iz];
+                        FXYZ[iz0][iy0][0] += mask * F[iz0 + iz - 1][iy0 + iy - 1][ix];
+                        FXYZ[iz0][iy0][xBoundM1] += mask * F[iz0 + iz - 1][iy0 + iy - 1][xBoundM1 - ix];
+                    }
+                }
+            }
+        }
+    }
+
+    // interiors
+    for (iz0 = 1; iz0 < zBoundM1; ++iz0)
+    {
+        for (iy0 = 1; iy0 < yBoundM1; ++iy0)
+        {
+            for (ix0 = 1; ix0 < xBoundM1; ++ix0)
+            {
+                FXYZ[iz0][iy0][ix0] = (Real)0;
+
+                for (iz = 0; iz <= 2; ++iz)
+                {
+                    for (iy = 0; iy <= 2; ++iy)
+                    {
+                        for (ix = 0; ix <= 2; ++ix)
+                        {
+                            mask = invDXDYDZ*CDer[ix] * CDer[iy] * CDer[iz];
+                            FXYZ[iz0][iy0][ix0] += mask * F[iz0 + iz - 1][iy0 + iy - 1][ix0 + ix - 1];
+                        }
+                    }
+                }
+            }
+        }
+    }
+}
+
+template <typename Real>
+void IntpAkimaUniform3<Real>::GetPolynomials(Array3<Real> const& F, Array3<Real> const& FX,
+    Array3<Real> const& FY, Array3<Real> const& FZ, Array3<Real> const& FXY,
+    Array3<Real> const& FXZ, Array3<Real> const& FYZ, Array3<Real> const& FXYZ)
+{
+    int xBoundM1 = mXBound - 1;
+    int yBoundM1 = mYBound - 1;
+    int zBoundM1 = mZBound - 1;
+    for (int iz = 0; iz < zBoundM1; ++iz)
+    {
+        for (int iy = 0; iy < yBoundM1; ++iy)
+        {
+            for (int ix = 0; ix < xBoundM1; ++ix)
+            {
+                // Note the 'transposing' of the 2x2x2 blocks (to match
+                // notation used in the polynomial definition).
+                Real G[2][2][2] =
+                {
+                    {
+                        {
+                            F[iz][iy][ix],
+                            F[iz + 1][iy][ix]
+                        },
+                        {
+                            F[iz][iy + 1][ix],
+                            F[iz + 1][iy + 1][ix]
+                        }
+                    },
+                    {
+                        {
+                            F[iz][iy][ix + 1],
+                            F[iz + 1][iy][ix + 1]
+                        },
+                        {
+                            F[iz][iy + 1][ix + 1],
+                            F[iz + 1][iy + 1][ix + 1]
+                        }
+                    }
+                };
+
+                Real GX[2][2][2] =
+                {
+                    {
+                        {
+                            FX[iz][iy][ix],
+                            FX[iz + 1][iy][ix]
+                        },
+                        {
+                            FX[iz][iy + 1][ix],
+                            FX[iz + 1][iy + 1][ix]
+                        }
+                    },
+                    {
+                        {
+                            FX[iz][iy][ix + 1],
+                            FX[iz + 1][iy][ix + 1]
+                        },
+                        {
+                            FX[iz][iy + 1][ix + 1],
+                            FX[iz + 1][iy + 1][ix + 1]
+                        }
+                    }
+                };
+
+                Real GY[2][2][2] =
+                {
+                    {
+                        {
+                            FY[iz][iy][ix],
+                            FY[iz + 1][iy][ix]
+                        },
+                        {
+                            FY[iz][iy + 1][ix],
+                            FY[iz + 1][iy + 1][ix]
+                        }
+                    },
+                    {
+                        {
+                            FY[iz][iy][ix + 1],
+                            FY[iz + 1][iy][ix + 1]
+                        },
+                        {
+                            FY[iz][iy + 1][ix + 1],
+                            FY[iz + 1][iy + 1][ix + 1]
+                        }
+                    }
+                };
+
+                Real GZ[2][2][2] =
+                {
+                    {
+                        {
+                            FZ[iz][iy][ix],
+                            FZ[iz + 1][iy][ix]
+                        },
+                        {
+                            FZ[iz][iy + 1][ix],
+                            FZ[iz + 1][iy + 1][ix]
+                        }
+                    },
+                    {
+                        {
+                            FZ[iz][iy][ix + 1],
+                            FZ[iz + 1][iy][ix + 1]
+                        },
+                        {
+                            FZ[iz][iy + 1][ix + 1],
+                            FZ[iz + 1][iy + 1][ix + 1]
+                        }
+                    }
+                };
+
+                Real GXY[2][2][2] =
+                {
+                    {
+                        {
+                            FXY[iz][iy][ix],
+                            FXY[iz + 1][iy][ix]
+                        },
+                        {
+                            FXY[iz][iy + 1][ix],
+                            FXY[iz + 1][iy + 1][ix]
+                        }
+                    },
+                    {
+                        {
+                            FXY[iz][iy][ix + 1],
+                            FXY[iz + 1][iy][ix + 1]
+                        },
+                        {
+                            FXY[iz][iy + 1][ix + 1],
+                            FXY[iz + 1][iy + 1][ix + 1]
+                        }
+                    }
+                };
+
+                Real GXZ[2][2][2] =
+                {
+                    {
+                        {
+                            FXZ[iz][iy][ix],
+                            FXZ[iz + 1][iy][ix]
+                        },
+                        {
+                            FXZ[iz][iy + 1][ix],
+                            FXZ[iz + 1][iy + 1][ix]
+                        }
+                    },
+                    {
+                        {
+                            FXZ[iz][iy][ix + 1],
+                            FXZ[iz + 1][iy][ix + 1]
+                        },
+                        {
+                            FXZ[iz][iy + 1][ix + 1],
+                            FXZ[iz + 1][iy + 1][ix + 1]
+                        }
+                    }
+                };
+
+                Real GYZ[2][2][2] =
+                {
+                    {
+                        {
+                            FYZ[iz][iy][ix],
+                            FYZ[iz + 1][iy][ix]
+                        },
+                        {
+                            FYZ[iz][iy + 1][ix],
+                            FYZ[iz + 1][iy + 1][ix]
+                        }
+                    },
+                    {
+                        {
+                            FYZ[iz][iy][ix + 1],
+                            FYZ[iz + 1][iy][ix + 1]
+                        },
+                        {
+                            FYZ[iz][iy + 1][ix + 1],
+                            FYZ[iz + 1][iy + 1][ix + 1]
+                        }
+                    }
+                };
+
+                Real GXYZ[2][2][2] =
+                {
+                    {
+                        {
+                            FXYZ[iz][iy][ix],
+                            FXYZ[iz + 1][iy][ix]
+                        },
+                        {
+                            FXYZ[iz][iy + 1][ix],
+                            FXYZ[iz + 1][iy + 1][ix]
+                        }
+                    },
+                    {
+                        {
+                            FXYZ[iz][iy][ix + 1],
+                            FXYZ[iz + 1][iy][ix + 1]
+                        },
+                        {
+                            FXYZ[iz][iy + 1][ix + 1],
+                            FXYZ[iz + 1][iy + 1][ix + 1]
+                        }
+                    }
+                };
+
+                Construct(mPoly[iz][iy][ix], G, GX, GY, GZ, GXY, GXZ, GYZ, GXYZ);
+            }
+        }
+    }
+}
+
+template <typename Real>
+Real IntpAkimaUniform3<Real>::ComputeDerivative(Real const* slope) const
+{
+    if (slope[1] != slope[2])
+    {
+        if (slope[0] != slope[1])
+        {
+            if (slope[2] != slope[3])
+            {
+                Real ad0 = std::abs(slope[3] - slope[2]);
+                Real ad1 = std::abs(slope[0] - slope[1]);
+                return (ad0 * slope[1] + ad1 * slope[2]) / (ad0 + ad1);
+            }
+            else
+            {
+                return slope[2];
+            }
+        }
+        else
+        {
+            if (slope[2] != slope[3])
+            {
+                return slope[1];
+            }
+            else
+            {
+                return ((Real)0.5) * (slope[1] + slope[2]);
+            }
+        }
+    }
+    else
+    {
+        return slope[1];
+    }
+}
+
+template <typename Real>
+void IntpAkimaUniform3<Real>::Construct(Polynomial& poly,
+    Real const F[2][2][2], Real const FX[2][2][2], Real const FY[2][2][2],
+    Real const FZ[2][2][2], Real const FXY[2][2][2], Real const FXZ[2][2][2],
+    Real const FYZ[2][2][2], Real const FXYZ[2][2][2])
+{
+    Real dx = mXSpacing, dy = mYSpacing, dz = mZSpacing;
+    Real invDX = ((Real)1) / dx, invDX2 = invDX * invDX;
+    Real invDY = ((Real)1) / dy, invDY2 = invDY * invDY;
+    Real invDZ = ((Real)1) / dz, invDZ2 = invDZ * invDZ;
+    Real b0, b1, b2, b3, b4, b5, b6, b7;
+
+    poly.A(0, 0, 0) = F[0][0][0];
+    poly.A(1, 0, 0) = FX[0][0][0];
+    poly.A(0, 1, 0) = FY[0][0][0];
+    poly.A(0, 0, 1) = FZ[0][0][0];
+    poly.A(1, 1, 0) = FXY[0][0][0];
+    poly.A(1, 0, 1) = FXZ[0][0][0];
+    poly.A(0, 1, 1) = FYZ[0][0][0];
+    poly.A(1, 1, 1) = FXYZ[0][0][0];
+
+    // solve for Aij0
+    b0 = (F[1][0][0] - poly(0, 0, 0, dx, (Real)0, (Real)0))*invDX2;
+    b1 = (FX[1][0][0] - poly(1, 0, 0, dx, (Real)0, (Real)0))*invDX;
+    poly.A(2, 0, 0) = ((Real)3)*b0 - b1;
+    poly.A(3, 0, 0) = (-((Real)2)*b0 + b1)*invDX;
+
+    b0 = (F[0][1][0] - poly(0, 0, 0, (Real)0, dy, (Real)0))*invDY2;
+    b1 = (FY[0][1][0] - poly(0, 1, 0, (Real)0, dy, (Real)0))*invDY;
+    poly.A(0, 2, 0) = ((Real)3)*b0 - b1;
+    poly.A(0, 3, 0) = (-((Real)2)*b0 + b1)*invDY;
+
+    b0 = (FY[1][0][0] - poly(0, 1, 0, dx, (Real)0, (Real)0))*invDX2;
+    b1 = (FXY[1][0][0] - poly(1, 1, 0, dx, (Real)0, (Real)0))*invDX;
+    poly.A(2, 1, 0) = ((Real)3)*b0 - b1;
+    poly.A(3, 1, 0) = (-((Real)2)*b0 + b1)*invDX;
+
+    b0 = (FX[0][1][0] - poly(1, 0, 0, (Real)0, dy, (Real)0))*invDY2;
+    b1 = (FXY[0][1][0] - poly(1, 1, 0, (Real)0, dy, (Real)0))*invDY;
+    poly.A(1, 2, 0) = ((Real)3)*b0 - b1;
+    poly.A(1, 3, 0) = (-((Real)2)*b0 + b1)*invDY;
+
+    b0 = (F[1][1][0] - poly(0, 0, 0, dx, dy, (Real)0))*invDX2*invDY2;
+    b1 = (FX[1][1][0] - poly(1, 0, 0, dx, dy, (Real)0))*invDX*invDY2;
+    b2 = (FY[1][1][0] - poly(0, 1, 0, dx, dy, (Real)0))*invDX2*invDY;
+    b3 = (FXY[1][1][0] - poly(1, 1, 0, dx, dy, (Real)0))*invDX*invDY;
+    poly.A(2, 2, 0) = ((Real)9)*b0 - ((Real)3)*b1 - ((Real)3)*b2 + b3;
+    poly.A(3, 2, 0) = (-((Real)6)*b0 + ((Real)3)*b1 + ((Real)2)*b2 - b3)*invDX;
+    poly.A(2, 3, 0) = (-((Real)6)*b0 + ((Real)2)*b1 + ((Real)3)*b2 - b3)*invDY;
+    poly.A(3, 3, 0) = (((Real)4)*b0 - ((Real)2)*b1 - ((Real)2)*b2 + b3)*invDX*invDY;
+
+    // solve for Ai0k
+    b0 = (F[0][0][1] - poly(0, 0, 0, (Real)0, (Real)0, dz))*invDZ2;
+    b1 = (FZ[0][0][1] - poly(0, 0, 1, (Real)0, (Real)0, dz))*invDZ;
+    poly.A(0, 0, 2) = ((Real)3)*b0 - b1;
+    poly.A(0, 0, 3) = (-((Real)2)*b0 + b1)*invDZ;
+
+    b0 = (FZ[1][0][0] - poly(0, 0, 1, dx, (Real)0, (Real)0))*invDX2;
+    b1 = (FXZ[1][0][0] - poly(1, 0, 1, dx, (Real)0, (Real)0))*invDX;
+    poly.A(2, 0, 1) = ((Real)3)*b0 - b1;
+    poly.A(3, 0, 1) = (-((Real)2)*b0 + b1)*invDX;
+
+    b0 = (FX[0][0][1] - poly(1, 0, 0, (Real)0, (Real)0, dz))*invDZ2;
+    b1 = (FXZ[0][0][1] - poly(1, 0, 1, (Real)0, (Real)0, dz))*invDZ;
+    poly.A(1, 0, 2) = ((Real)3)*b0 - b1;
+    poly.A(1, 0, 3) = (-((Real)2)*b0 + b1)*invDZ;
+
+    b0 = (F[1][0][1] - poly(0, 0, 0, dx, (Real)0, dz))*invDX2*invDZ2;
+    b1 = (FX[1][0][1] - poly(1, 0, 0, dx, (Real)0, dz))*invDX*invDZ2;
+    b2 = (FZ[1][0][1] - poly(0, 0, 1, dx, (Real)0, dz))*invDX2*invDZ;
+    b3 = (FXZ[1][0][1] - poly(1, 0, 1, dx, (Real)0, dz))*invDX*invDZ;
+    poly.A(2, 0, 2) = ((Real)9)*b0 - ((Real)3)*b1 - ((Real)3)*b2 + b3;
+    poly.A(3, 0, 2) = (-((Real)6)*b0 + ((Real)3)*b1 + ((Real)2)*b2 - b3)*invDX;
+    poly.A(2, 0, 3) = (-((Real)6)*b0 + ((Real)2)*b1 + ((Real)3)*b2 - b3)*invDZ;
+    poly.A(3, 0, 3) = (((Real)4)*b0 - ((Real)2)*b1 - ((Real)2)*b2 + b3)*invDX*invDZ;
+
+    // solve for A0jk
+    b0 = (FZ[0][1][0] - poly(0, 0, 1, (Real)0, dy, (Real)0))*invDY2;
+    b1 = (FYZ[0][1][0] - poly(0, 1, 1, (Real)0, dy, (Real)0))*invDY;
+    poly.A(0, 2, 1) = ((Real)3)*b0 - b1;
+    poly.A(0, 3, 1) = (-((Real)2)*b0 + b1)*invDY;
+
+    b0 = (FY[0][0][1] - poly(0, 1, 0, (Real)0, (Real)0, dz))*invDZ2;
+    b1 = (FYZ[0][0][1] - poly(0, 1, 1, (Real)0, (Real)0, dz))*invDZ;
+    poly.A(0, 1, 2) = ((Real)3)*b0 - b1;
+    poly.A(0, 1, 3) = (-((Real)2)*b0 + b1)*invDZ;
+
+    b0 = (F[0][1][1] - poly(0, 0, 0, (Real)0, dy, dz))*invDY2*invDZ2;
+    b1 = (FY[0][1][1] - poly(0, 1, 0, (Real)0, dy, dz))*invDY*invDZ2;
+    b2 = (FZ[0][1][1] - poly(0, 0, 1, (Real)0, dy, dz))*invDY2*invDZ;
+    b3 = (FYZ[0][1][1] - poly(0, 1, 1, (Real)0, dy, dz))*invDY*invDZ;
+    poly.A(0, 2, 2) = ((Real)9)*b0 - ((Real)3)*b1 - ((Real)3)*b2 + b3;
+    poly.A(0, 3, 2) = (-((Real)6)*b0 + ((Real)3)*b1 + ((Real)2)*b2 - b3)*invDY;
+    poly.A(0, 2, 3) = (-((Real)6)*b0 + ((Real)2)*b1 + ((Real)3)*b2 - b3)*invDZ;
+    poly.A(0, 3, 3) = (((Real)4)*b0 - ((Real)2)*b1 - ((Real)2)*b2 + b3)*invDY*invDZ;
+
+    // solve for Aij1
+    b0 = (FYZ[1][0][0] - poly(0, 1, 1, dx, (Real)0, (Real)0))*invDX2;
+    b1 = (FXYZ[1][0][0] - poly(1, 1, 1, dx, (Real)0, (Real)0))*invDX;
+    poly.A(2, 1, 1) = ((Real)3)*b0 - b1;
+    poly.A(3, 1, 1) = (-((Real)2)*b0 + b1)*invDX;
+
+    b0 = (FXZ[0][1][0] - poly(1, 0, 1, (Real)0, dy, (Real)0))*invDY2;
+    b1 = (FXYZ[0][1][0] - poly(1, 1, 1, (Real)0, dy, (Real)0))*invDY;
+    poly.A(1, 2, 1) = ((Real)3)*b0 - b1;
+    poly.A(1, 3, 1) = (-((Real)2)*b0 + b1)*invDY;
+
+    b0 = (FZ[1][1][0] - poly(0, 0, 1, dx, dy, (Real)0))*invDX2*invDY2;
+    b1 = (FXZ[1][1][0] - poly(1, 0, 1, dx, dy, (Real)0))*invDX*invDY2;
+    b2 = (FYZ[1][1][0] - poly(0, 1, 1, dx, dy, (Real)0))*invDX2*invDY;
+    b3 = (FXYZ[1][1][0] - poly(1, 1, 1, dx, dy, (Real)0))*invDX*invDY;
+    poly.A(2, 2, 1) = ((Real)9)*b0 - ((Real)3)*b1 - ((Real)3)*b2 + b3;
+    poly.A(3, 2, 1) = (-((Real)6)*b0 + ((Real)3)*b1 + ((Real)2)*b2 - b3)*invDX;
+    poly.A(2, 3, 1) = (-((Real)6)*b0 + ((Real)2)*b1 + ((Real)3)*b2 - b3)*invDY;
+    poly.A(3, 3, 1) = (((Real)4)*b0 - ((Real)2)*b1 - ((Real)2)*b2 + b3)*invDX*invDY;
+
+    // solve for Ai1k
+    b0 = (FXY[0][0][1] - poly(1, 1, 0, (Real)0, (Real)0, dz))*invDZ2;
+    b1 = (FXYZ[0][0][1] - poly(1, 1, 1, (Real)0, (Real)0, dz))*invDZ;
+    poly.A(1, 1, 2) = ((Real)3)*b0 - b1;
+    poly.A(1, 1, 3) = (-((Real)2)*b0 + b1)*invDZ;
+
+    b0 = (FY[1][0][1] - poly(0, 1, 0, dx, (Real)0, dz))*invDX2*invDZ2;
+    b1 = (FXY[1][0][1] - poly(1, 1, 0, dx, (Real)0, dz))*invDX*invDZ2;
+    b2 = (FYZ[1][0][1] - poly(0, 1, 1, dx, (Real)0, dz))*invDX2*invDZ;
+    b3 = (FXYZ[1][0][1] - poly(1, 1, 1, dx, (Real)0, dz))*invDX*invDZ;
+    poly.A(2, 1, 2) = ((Real)9)*b0 - ((Real)3)*b1 - ((Real)3)*b2 + b3;
+    poly.A(3, 1, 2) = (-((Real)6)*b0 + ((Real)3)*b1 + ((Real)2)*b2 - b3)*invDX;
+    poly.A(2, 1, 3) = (-((Real)6)*b0 + ((Real)2)*b1 + ((Real)3)*b2 - b3)*invDZ;
+    poly.A(3, 1, 3) = (((Real)4)*b0 - ((Real)2)*b1 - ((Real)2)*b2 + b3)*invDX*invDZ;
+
+    // solve for A1jk
+    b0 = (FX[0][1][1] - poly(1, 0, 0, (Real)0, dy, dz))*invDY2*invDZ2;
+    b1 = (FXY[0][1][1] - poly(1, 1, 0, (Real)0, dy, dz))*invDY*invDZ2;
+    b2 = (FXZ[0][1][1] - poly(1, 0, 1, (Real)0, dy, dz))*invDY2*invDZ;
+    b3 = (FXYZ[0][1][1] - poly(1, 1, 1, (Real)0, dy, dz))*invDY*invDZ;
+    poly.A(1, 2, 2) = ((Real)9)*b0 - ((Real)3)*b1 - ((Real)3)*b2 + b3;
+    poly.A(1, 3, 2) = (-((Real)6)*b0 + ((Real)3)*b1 + ((Real)2)*b2 - b3)*invDY;
+    poly.A(1, 2, 3) = (-((Real)6)*b0 + ((Real)2)*b1 + ((Real)3)*b2 - b3)*invDZ;
+    poly.A(1, 3, 3) = (((Real)4)*b0 - ((Real)2)*b1 - ((Real)2)*b2 + b3)*invDY*invDZ;
+
+    // solve for remaining Aijk with i >= 2, j >= 2, k >= 2
+    b0 = (F[1][1][1] - poly(0, 0, 0, dx, dy, dz))*invDX2*invDY2*invDZ2;
+    b1 = (FX[1][1][1] - poly(1, 0, 0, dx, dy, dz))*invDX*invDY2*invDZ2;
+    b2 = (FY[1][1][1] - poly(0, 1, 0, dx, dy, dz))*invDX2*invDY*invDZ2;
+    b3 = (FZ[1][1][1] - poly(0, 0, 1, dx, dy, dz))*invDX2*invDY2*invDZ;
+    b4 = (FXY[1][1][1] - poly(1, 1, 0, dx, dy, dz))*invDX*invDY*invDZ2;
+    b5 = (FXZ[1][1][1] - poly(1, 0, 1, dx, dy, dz))*invDX*invDY2*invDZ;
+    b6 = (FYZ[1][1][1] - poly(0, 1, 1, dx, dy, dz))*invDX2*invDY*invDZ;
+    b7 = (FXYZ[1][1][1] - poly(1, 1, 1, dx, dy, dz))*invDX*invDY*invDZ;
+    poly.A(2, 2, 2) = ((Real)27)*b0 - ((Real)9)*b1 - ((Real)9)*b2 -
+        ((Real)9)*b3 + ((Real)3)*b4 + ((Real)3)*b5 + ((Real)3)*b6 - b7;
+    poly.A(3, 2, 2) = (((Real)-18)*b0 + ((Real)9)*b1 + ((Real)6)*b2 +
+        ((Real)6)*b3 - ((Real)3)*b4 - ((Real)3)*b5 - ((Real)2)*b6 + b7)*invDX;
+    poly.A(2, 3, 2) = (((Real)-18)*b0 + ((Real)6)*b1 + ((Real)9)*b2 +
+        ((Real)6)*b3 - ((Real)3)*b4 - ((Real)2)*b5 - ((Real)3)*b6 + b7)*invDY;
+    poly.A(2, 2, 3) = (((Real)-18)*b0 + ((Real)6)*b1 + ((Real)6)*b2 +
+        ((Real)9)*b3 - ((Real)2)*b4 - ((Real)3)*b5 - ((Real)3)*b6 + b7)*invDZ;
+    poly.A(3, 3, 2) = (((Real)12)*b0 - ((Real)6)*b1 - ((Real)6)*b2 -
+        ((Real)4)*b3 + ((Real)3)*b4 + ((Real)2)*b5 + ((Real)2)*b6 - b7)*
+        invDX*invDY;
+    poly.A(3, 2, 3) = (((Real)12)*b0 - ((Real)6)*b1 - ((Real)4)*b2 -
+        ((Real)6)*b3 + ((Real)2)*b4 + ((Real)3)*b5 + ((Real)2)*b6 - b7)*
+        invDX*invDZ;
+    poly.A(2, 3, 3) = (((Real)12)*b0 - ((Real)4)*b1 - ((Real)6)*b2 -
+        ((Real)6)*b3 + ((Real)2)*b4 + ((Real)2)*b5 + ((Real)3)*b6 - b7)*
+        invDY*invDZ;
+    poly.A(3, 3, 3) = (((Real)-8)*b0 + ((Real)4)*b1 + ((Real)4)*b2 +
+        ((Real)4)*b3 - ((Real)2)*b4 - ((Real)2)*b5 - ((Real)2)*b6 + b7)*
+        invDX*invDY*invDZ;
+}
+
+template <typename Real>
+void IntpAkimaUniform3<Real>::XLookup(Real x, int& xIndex, Real& dx) const
+{
+    for (xIndex = 0; xIndex + 1 < mXBound; ++xIndex)
+    {
+        if (x < mXMin + mXSpacing*(xIndex + 1))
+        {
+            dx = x - (mXMin + mXSpacing*xIndex);
+            return;
+        }
+    }
+
+    --xIndex;
+    dx = x - (mXMin + mXSpacing*xIndex);
+}
+
+template <typename Real>
+void IntpAkimaUniform3<Real>::YLookup(Real y, int& yIndex, Real& dy) const
+{
+    for (yIndex = 0; yIndex + 1 < mYBound; ++yIndex)
+    {
+        if (y < mYMin + mYSpacing*(yIndex + 1))
+        {
+            dy = y - (mYMin + mYSpacing*yIndex);
+            return;
+        }
+    }
+
+    --yIndex;
+    dy = y - (mYMin + mYSpacing*yIndex);
+}
+
+template <typename Real>
+void IntpAkimaUniform3<Real>::ZLookup(Real z, int& zIndex, Real& dz) const
+{
+    for (zIndex = 0; zIndex + 1 < mZBound; ++zIndex)
+    {
+        if (z < mZMin + mZSpacing*(zIndex + 1))
+        {
+            dz = z - (mZMin + mZSpacing*zIndex);
+            return;
+        }
+    }
+
+    --zIndex;
+    dz = z - (mZMin + mZSpacing*zIndex);
+}
+
+template <typename Real>
+IntpAkimaUniform3<Real>::Polynomial::Polynomial()
+{
+    memset(&mCoeff[0][0][0], 0, 64 * sizeof(Real));
+}
+
+template <typename Real>
+Real& IntpAkimaUniform3<Real>::Polynomial::A(int ix, int iy, int iz)
+{
+    return mCoeff[ix][iy][iz];
+}
+
+template <typename Real>
+Real IntpAkimaUniform3<Real>::Polynomial::operator()(Real x, Real y, Real z)
+const
+{
+    Real xPow[4] = { (Real)1, x, x * x, x * x * x };
+    Real yPow[4] = { (Real)1, y, y * y, y * y * y };
+    Real zPow[4] = { (Real)1, z, z * z, z * z * z };
+
+    Real p = (Real)0;
+    for (int iz = 0; iz <= 3; ++iz)
+    {
+        for (int iy = 0; iy <= 3; ++iy)
+        {
+            for (int ix = 0; ix <= 3; ++ix)
+            {
+                p += mCoeff[ix][iy][iz] * xPow[ix] * yPow[iy] * zPow[iz];
+            }
+        }
+    }
+
+    return p;
+}
+
+template <typename Real>
+Real IntpAkimaUniform3<Real>::Polynomial::operator()(int xOrder, int yOrder,
+    int zOrder, Real x, Real y, Real z) const
+{
+    Real xPow[4];
+    switch (xOrder)
+    {
+    case 0:
+        xPow[0] = (Real)1;
+        xPow[1] = x;
+        xPow[2] = x * x;
+        xPow[3] = x * x * x;
+        break;
+    case 1:
+        xPow[0] = (Real)0;
+        xPow[1] = (Real)1;
+        xPow[2] = ((Real)2) * x;
+        xPow[3] = ((Real)3) * x * x;
+        break;
+    case 2:
+        xPow[0] = (Real)0;
+        xPow[1] = (Real)0;
+        xPow[2] = (Real)2;
+        xPow[3] = ((Real)6) * x;
+        break;
+    case 3:
+        xPow[0] = (Real)0;
+        xPow[1] = (Real)0;
+        xPow[2] = (Real)0;
+        xPow[3] = (Real)6;
+        break;
+    default:
+        return (Real)0;
+    }
+
+    Real yPow[4];
+    switch (yOrder)
+    {
+    case 0:
+        yPow[0] = (Real)1;
+        yPow[1] = y;
+        yPow[2] = y * y;
+        yPow[3] = y * y * y;
+        break;
+    case 1:
+        yPow[0] = (Real)0;
+        yPow[1] = (Real)1;
+        yPow[2] = ((Real)2) * y;
+        yPow[3] = ((Real)3) * y * y;
+        break;
+    case 2:
+        yPow[0] = (Real)0;
+        yPow[1] = (Real)0;
+        yPow[2] = (Real)2;
+        yPow[3] = ((Real)6) * y;
+        break;
+    case 3:
+        yPow[0] = (Real)0;
+        yPow[1] = (Real)0;
+        yPow[2] = (Real)0;
+        yPow[3] = (Real)6;
+        break;
+    default:
+        return (Real)0;
+    }
+
+    Real zPow[4];
+    switch (zOrder)
+    {
+    case 0:
+        zPow[0] = (Real)1;
+        zPow[1] = z;
+        zPow[2] = z * z;
+        zPow[3] = z * z * z;
+        break;
+    case 1:
+        zPow[0] = (Real)0;
+        zPow[1] = (Real)1;
+        zPow[2] = ((Real)2) * z;
+        zPow[3] = ((Real)3) * z * z;
+        break;
+    case 2:
+        zPow[0] = (Real)0;
+        zPow[1] = (Real)0;
+        zPow[2] = (Real)2;
+        zPow[3] = ((Real)6) * z;
+        break;
+    case 3:
+        zPow[0] = (Real)0;
+        zPow[1] = (Real)0;
+        zPow[2] = (Real)0;
+        zPow[3] = (Real)6;
+        break;
+    default:
+        return (Real)0;
+    }
+
+    Real p = (Real)0;
+
+    for (int iz = 0; iz <= 3; ++iz)
+    {
+        for (int iy = 0; iy <= 3; ++iy)
+        {
+            for (int ix = 0; ix <= 3; ++ix)
+            {
+                p += mCoeff[ix][iy][iz] * xPow[ix] * yPow[iy] * zPow[iz];
+            }
+        }
+    }
+
+    return p;
+}
+
+}
