move GeometryProcessing plugin from NotForLicensees to Experimental

diff --git a/Engine/Plugins/Experimental/GeometryProcessing/Source/GeometryAlgorithms/Private/GTEngine/Mathematics/GteIntpBSplineUniform.h b/Engine/Plugins/Experimental/GeometryProcessing/Source/GeometryAlgorithms/Private/GTEngine/Mathematics/GteIntpBSplineUniform.h
new file mode 100644
--- /dev/null
+++ b/Engine/Plugins/Experimental/GeometryProcessing/Source/GeometryAlgorithms/Private/GTEngine/Mathematics/GteIntpBSplineUniform.h
@@ -0,0 +1,1455 @@
+// David Eberly, Geometric Tools, Redmond WA 98052
+// Copyright (c) 1998-2019
+// Distributed under the Boost Software License, Version 1.0.
+// http://www.boost.org/LICENSE_1_0.txt
+// http://www.geometrictools.com/License/Boost/LICENSE_1_0.txt
+// File Version: 3.13.2 (2018/10/05)
+
+#pragma once
+
+#include <GTEngine/Mathematics/GtePolynomial1.h>
+#include <array>
+#include <cmath>
+
+// IntpBSplineUniform is the class for B-spline interpolation of uniformly
+// spaced N-dimensional data.  The algorithm is described in
+//   https://www.geometrictools.com/Documentation/BSplineInterpolation.pdf
+// A sample application for this topic is
+//   GeometricTools/GTEngine/Samples/Imagics/BSplineInterpolation
+//
+// The Controls adapter allows access to your control points without regard
+// to how you organize your data.  You can even defer the computation of a
+// control point until it is needed via the operator()(...) calls that
+// Controls must provide, and you can cache the points according to your own
+// needs.  The minimal interface for a Controls adapter is
+//
+// struct Controls
+// {
+//   // The control_point_type is of your choosing.  It must support
+//   // assignment, scalar multiplication and addition.  Specifically, if
+//   // C0, C1 and C2 are control points and s is a scalar, the interpolator
+//   // needs to perform operations
+//   //   C1 = C0;
+//   //   C1 = C0 * s;
+//   //   C2 = C0 + C1;
+//   typedef control_point_type Type;
+//
+//   // The number of elements in the specified dimension.
+//   int GetSize(int dimension) const;
+//
+//   // Get a control point based on an n-tuple lookup.  The interpolator
+//   // does not need to know your organization; all it needs is the
+//   // desired control point.  The 'tuple' input must have n elements.
+//   Type operator() (int const* tuple) const;
+//
+//   // If you choose to use the specialized interpolators for dimensions
+//   // 1, 2 or 3, you must also provide the following accessor, where
+//   // the input is an n-tuple listed component-by-component (1, 2 or 3
+//   // components).
+//   Type operator() (int i0, int i1, ..., int inm1) const;
+
+namespace gte
+{
+    template <typename Real, typename Controls>
+    class IntpBSplineUniformShared
+    {
+    protected:
+        // Abstract base class construction.  A virtual destructor is not
+        // provided because there are no required side effects in the base
+        // class when destroying objects from the derived classes.
+        IntpBSplineUniformShared(int numDimensions, int const* degrees,
+            Controls const& controls, typename Controls::Type ctZero, int cacheMode)
+            :
+            mNumDimensions(numDimensions),
+            mDegree(numDimensions),
+            mControls(&controls),
+            mCTZero(ctZero),
+            mCacheMode(cacheMode),
+            mNumLocalControls(0),
+            mDegreeP1(numDimensions),
+            mNumControls(numDimensions),
+            mTMin(numDimensions),
+            mTMax(numDimensions),
+            mBlender(numDimensions),
+            mDCoefficient(numDimensions),
+            mLMax(numDimensions),
+            mPowerDSDT(numDimensions),
+            mITuple(numDimensions),
+            mJTuple(numDimensions),
+            mKTuple(numDimensions),
+            mLTuple(numDimensions),
+            mSumIJTuple(numDimensions),
+            mUTuple(numDimensions),
+            mPTuple(numDimensions),
+            mValid(IsValid(degrees))
+        {
+            if (!mValid)
+            {
+                return;
+            }
+
+            mNumLocalControls = 1;
+            for (int dim = 0; dim < mNumDimensions; ++dim)
+            {
+                mDegree[dim] = degrees[dim];
+                mDegreeP1[dim] = degrees[dim] + 1;
+                mNumLocalControls *= mDegreeP1[dim];
+                mNumControls[dim] = controls.GetSize(dim);
+                mTMin[dim] = (Real)-0.5;
+                mTMax[dim] = static_cast<Real>(mNumControls[dim]) - (Real)0.5;
+                ComputeBlendingMatrix(mDegree[dim], mBlender[dim]);
+                ComputeDCoefficients(mDegree[dim], mDCoefficient[dim], mLMax[dim]);
+                ComputePowers(mDegree[dim], mNumControls[dim], mTMin[dim], mTMax[dim], mPowerDSDT[dim]);
+            }
+
+            if (mCacheMode == NO_CACHING)
+            {
+                mPhi.resize(mNumDimensions);
+                for (int dim = 0; dim < mNumDimensions; ++dim)
+                {
+                    mPhi[dim].resize(mDegreeP1[dim]);
+                }
+            }
+            else
+            {
+                InitializeTensors();
+            }
+        }
+
+#if !defined(GTE_INTP_BSPLINE_UNIFORM_NO_SPECIALIZATION)
+        IntpBSplineUniformShared()
+            :
+            mNumDimensions(0),
+            mControls(nullptr),
+            mCTZero(),
+            mCacheMode(0),
+            mNumLocalControls(0)
+        {
+        }
+#endif
+
+    public:
+        // Support for caching the intermediate tensor product of control
+        // points with the blending matrices.  A precached container has all
+        // elements precomputed before any Evaluate(...) calls.  The 'bool'
+        // flags are all set to 'true'.  A cached container fills the elements
+        // on demand.  The 'bool' flags are initially 'false', indicating the
+        // EvalType component has not yet been computed.  After it is computed
+        // and stored, the flag is set to 'true'.
+        enum
+        {
+            NO_CACHING,
+            PRE_CACHING,
+            ON_DEMAND_CACHING
+        };
+
+        // Member access.
+        inline int GetDegree(int dim) const { return mDegree[dim]; }
+        inline int GetNumControls(int dim) const { return mNumControls[dim]; }
+        inline Real GetTMin(int dim) const { return mTMin[dim]; }
+        inline Real GetTMax(int dim) const { return mTMax[dim]; }
+        inline int GetCacheMode() const { return mCacheMode; }
+
+    protected:
+        // Disallow copying and moving.
+        IntpBSplineUniformShared(IntpBSplineUniformShared const&) = delete;
+        IntpBSplineUniformShared& operator=(IntpBSplineUniformShared const&) = delete;
+        IntpBSplineUniformShared(IntpBSplineUniformShared&&) = delete;
+        IntpBSplineUniformShared& operator=(IntpBSplineUniformShared&&) = delete;
+
+        // Verify that the constraints for number of controls and degrees
+        // are satisfied.
+        bool IsValid(int const* degrees)
+        {
+            // The condition c+1 > d+1 is required so that when s = c+1-d, its
+            // maximum value, we have at least two s-knots (d and d + 1).
+            for (int dim = 0; dim < mNumDimensions; ++dim)
+            {
+                if (mControls->GetSize(dim) <= degrees[dim] + 1)
+                {
+                    LogError("Incompatible degree and number of controls.");
+                    return false;
+                }
+            }
+            return true;
+        }
+
+        // ComputeBlendingMatrix, ComputeDCoefficients, ComputePowers and
+        // GetKey are used by the general-dimension derived classes and by
+        // the specializations for dimensions 1, 2 and 3.
+
+        // Compute the blending matrix that combines the control points and
+        // the polynomial vector.  The matrix A is stored in row-major order.
+        static void ComputeBlendingMatrix(int degree, std::vector<Real>& A)
+        {
+            int const degreeP1 = degree + 1;
+            A.resize(degreeP1 * degreeP1);
+
+            if (degree == 0)
+            {
+                A[0] = (Real)1;
+                return;
+            }
+
+            // P_{0,0}(s)
+            std::vector<Polynomial1<Real>> P(degreeP1);
+            P[0][0] = (Real)1;
+
+            // L0 = s/j
+            Polynomial1<Real> L0(1);
+            L0[0] = (Real)0;
+
+            // L1(s) = (j + 1 - s)/j
+            Polynomial1<Real> L1(1);
+
+            // s-1 is used in computing translated P_{j-1,k-1}(s-1)
+            Polynomial1<Real> sm1 = { (Real)-1, (Real)1 };
+
+            // Compute
+            //   P_{j,k}(s) = L0(s)*P_{j-1,k}(s) + L1(s)*P_{j-1,k-1}(s-1)
+            // for 0 <= k <= j where 1 <= j <= degree.  When k = 0,
+            // P_{j-1,-1}(s) = 0, so P_{j,0}(s) = L0(s)*P_{j-1,0}(s).  When
+            // k = j, P_{j-1,j}(s) = 0, so P_{j,j}(s) = L1(s)*P_{j-1,j-1}(s).
+            // The polynomials at level j-1 are currently stored in P[0]
+            // through P[j-1].  The polynomials at level j are computed and
+            // stored in P[0] through P[j]; that is, they are computed in
+            // place to reduce memory usage and copying.  This requires
+            // computing P[k] (level j) from P[k] (level j-1) and P[k-1]
+            // (level j-1), which means we have to process k = j down to
+            // k = 0.
+            for (int j = 1; j <= degree; ++j)
+            {
+                Real invJ = (Real)1 / (Real)j;
+                L0[1] = invJ;
+                L1[0] = (Real)1 + invJ;
+                L1[1] = -invJ;
+
+                for (int k = j; k >= 0; --k)
+                {
+                    Polynomial1<Real> result = { (Real)0 };
+
+                    if (k > 0)
+                    {
+                        result += L1 * P[k - 1].GetTranslation((Real)1);
+                    }
+
+                    if (k < j)
+                    {
+                        result += L0 * P[k];
+                    }
+
+                    P[k] = result;
+                }
+            }
+
+            // Compute Q_{d,k}(s) = P_{d,k}(s + k).
+            std::vector<Polynomial1<Real>> Q(degreeP1);
+            for (int k = 0; k <= degree; ++k)
+            {
+                Q[k] = P[k].GetTranslation(static_cast<Real>(-k));
+            }
+
+            // Extract the matrix A from the Q-polynomials.  Row r of A
+            // contains the coefficients of Q_{d,d-r}(s).
+            for (int k = 0, row = degree; k <= degree; ++k, --row)
+            {
+                for (int col = 0; col <= degree; ++col)
+                {
+                    A[col + degreeP1 * row] = Q[k][col];
+                }
+            }
+        }
+
+        // Compute the coefficients for the derivative polynomial terms.
+        static void ComputeDCoefficients(int degree, std::vector<Real>& dCoefficients,
+            std::vector<int>& ellMax)
+        {
+            int numDCoefficients = (degree + 1) * (degree + 2) / 2;
+            dCoefficients.resize(numDCoefficients);
+            for (int i = 0; i < numDCoefficients; ++i)
+            {
+                dCoefficients[i] = 1;
+            }
+
+            for (int order = 1, col0 = 0, col1 = degree + 1; order <= degree; ++order)
+            {
+                ++col0;
+                for (int c = order, m = 1; c <= degree; ++c, ++m, ++col0, ++col1)
+                {
+                    dCoefficients[col1] = dCoefficients[col0] * m;
+                }
+            }
+
+            ellMax.resize(degree + 1);
+            ellMax[0] = degree;
+            for (int i0 = 0, i1 = 1; i1 <= degree; i0 = i1++)
+            {
+                ellMax[i1] = ellMax[i0] + degree - i0;
+            }
+        }
+
+        // Compute powers of ds/dt.
+        void ComputePowers(int degree, int numControls, Real tmin, Real tmax,
+            std::vector<Real>& powerDSDT)
+        {
+            Real dsdt = static_cast<Real>(numControls - degree) / (tmax - tmin);
+            powerDSDT.resize(degree + 1);
+            powerDSDT[0] = (Real)1;
+            powerDSDT[1] = dsdt;
+            for (int i = 2; i <= degree; ++i)
+            {
+                powerDSDT[i] = powerDSDT[i - 1] * dsdt;
+            }
+        }
+
+        // Determine the interval [index,index+1) corresponding to the
+        // specified value of t and compute u in that interval.
+        static void GetKey(Real t, Real tmin, Real tmax, Real dsdt, int numControls,
+            int degree, int& index, Real& u)
+        {
+            // Compute s - d = ((c + 1 - d)/(c + 1))(t + 1/2), the index for
+            // which d + index <= s < d + index + 1.  Let u = s - d - index so
+            // that 0 <= u < 1.
+            if (t > tmin)
+            {
+                if (t < tmax)
+                {
+                    Real smd = dsdt * (t - tmin);
+                    index = static_cast<uint32_t>(std::floor(smd));
+                    u = smd - static_cast<float>(index);
+                }
+                else
+                {
+                    // In the evaluation, s = c + 1 - d and i = c - d.  This
+                    // causes s-d-i to be 1 in G_c(c+1-d).  Effectively, the
+                    // selection of i extends the s-domain [d,c+1) to its
+                    // support [d,c+1].
+                    index = numControls - 1 - degree;
+                    u = (Real)1;
+                }
+            }
+            else
+            {
+                index = 0;
+                u = (Real)0;
+            }
+        }
+
+        // The remaining functions are used only by the general-dimension
+        // derived classes when caching is enabled.
+
+        // For the multidimensional tensor Phi{iTuple, kTuple), compute the
+        // portion of the 1-dimensional index that corresponds to iTuple.
+        int GetRowIndex(std::vector<int> const& i) const
+        {
+            int rowIndex = i[mNumDimensions - 1];
+            int j1 = 2 * mNumDimensions - 2;
+            for (int j0 = mNumDimensions - 2; j0 >= 0; --j0, --j1)
+            {
+                rowIndex = mTBound[j1] * rowIndex + i[j0];
+            }
+            rowIndex = mTBound[j1] * rowIndex;
+            return rowIndex;
+        }
+
+        // For the multidimensional tensor Phi{iTuple, kTuple), combine the
+        // GetRowIndex(...) output with kTuple to produce the full
+        // 1-dimensional index.
+        int GetIndex(int rowIndex, std::vector<int> const& k) const
+        {
+            int index = rowIndex + k[mNumDimensions - 1];
+            for (int j = mNumDimensions - 2; j >= 0; --j)
+            {
+                index = mTBound[j] * index + k[j];
+            }
+            return index;
+        }
+
+        // Compute Phi(iTuple, kTuple).  The 'index' value is an already
+        // computed 1-dimensional index for the tensor.
+        void ComputeTensor(int const* i, int const* k, int index)
+        {
+            auto element = mCTZero;
+            for (int dim = 0; dim < mNumDimensions; ++dim)
+            {
+                mComputeJTuple[dim] = 0;
+            }
+            for (int iterate = 0; iterate < mNumLocalControls; ++iterate)
+            {
+                Real blend(1);
+                for (int dim = 0; dim < mNumDimensions; ++dim)
+                {
+                    blend *= mBlender[dim][k[dim] + mDegreeP1[dim] * mComputeJTuple[dim]];
+                    mComputeSumIJTuple[dim] = i[dim] + mComputeJTuple[dim];
+                }
+                element = element + (*mControls)(mComputeSumIJTuple.data()) * blend;
+
+                for (int dim = 0; dim < mNumDimensions; ++dim)
+                {
+                    if (++mComputeJTuple[dim] < mDegreeP1[dim])
+                    {
+                        break;
+                    }
+                    mComputeJTuple[dim] = 0;
+                }
+            }
+            mTensor[index] = element;
+        }
+
+        // Allocate the containers used for caching and fill in the tensor
+        // for precaching when that mode is selected.
+        void InitializeTensors()
+        {
+            mTBound.resize(2 * mNumDimensions);
+            mComputeJTuple.resize(mNumDimensions);
+            mComputeSumIJTuple.resize(mNumDimensions);
+            mDegreeMinusOrder.resize(mNumDimensions);
+            mTerm.resize(mNumDimensions);
+
+            int current = 0;
+            int numCached = 1;
+            for (int dim = 0; dim < mNumDimensions; ++dim, ++current)
+            {
+                mTBound[current] = mDegreeP1[dim];
+                numCached *= mTBound[current];
+            }
+            for (int dim = 0; dim < mNumDimensions; ++dim, ++current)
+            {
+                mTBound[current] = mNumControls[dim] - mDegree[dim];
+                numCached *= mTBound[current];
+            }
+            mTensor.resize(numCached);
+            mCached.resize(numCached);
+            if (mCacheMode == PRE_CACHING)
+            {
+                std::vector<int> tuple(2 * mNumDimensions, 0);
+                for (int index = 0; index < numCached; ++index)
+                {
+                    ComputeTensor(&tuple[mNumDimensions], &tuple[0], index);
+                    for (int i = 0; i < 2 * mNumDimensions; ++i)
+                    {
+                        if (++tuple[i] < mTBound[i])
+                        {
+                            break;
+                        }
+                        tuple[i] = 0;
+                    }
+                }
+                std::fill(mCached.begin(), mCached.end(), true);
+            }
+            else
+            {
+                std::fill(mCached.begin(), mCached.end(), false);
+            }
+        }
+
+        // Evaluate the interpolator.  Each element of 'order' indicates the
+        // order of the derivative you want to compute.  For the function
+        // value itself, pass in 'order' that has all 0 elements.
+        typename Controls::Type EvaluateNoCaching(int const* order, Real const* t)
+        {
+            auto result = mCTZero;
+            if (!mValid)
+            {
+                return result;
+            }
+
+            for (int dim = 0; dim < mNumDimensions; ++dim)
+            {
+                if (order[dim] < 0 || order[dim] > mDegree[dim])
+                {
+                    return result;
+                }
+            }
+
+            for (int dim = 0; dim < mNumDimensions; ++dim)
+            {
+                GetKey(t[dim], mTMin[dim], mTMax[dim], mPowerDSDT[dim][1],
+                    mNumControls[dim], mDegree[dim], mITuple[dim], mUTuple[dim]);
+            }
+
+            for (int dim = 0; dim < mNumDimensions; ++dim)
+            {
+                int jIndex = 0;
+                for (int j = 0; j <= mDegree[dim]; ++j)
+                {
+                    int kjIndex = mDegree[dim] + jIndex;
+                    int ell = mLMax[dim][order[dim]];
+                    mPhi[dim][j] = (Real)0;
+                    for (int k = mDegree[dim]; k >= order[dim]; --k)
+                    {
+                        mPhi[dim][j] = mPhi[dim][j] * mUTuple[dim] +
+                            mBlender[dim][kjIndex--] * mDCoefficient[dim][ell--];
+                    }
+                    jIndex += mDegreeP1[dim];
+                }
+            }
+
+            for (int dim = 0; dim < mNumDimensions; ++dim)
+            {
+                mJTuple[dim] = 0;
+                mSumIJTuple[dim] = mITuple[dim];
+                mPTuple[dim] = mPhi[dim][0];
+            }
+            for (int iterate = 0; iterate < mNumLocalControls; ++iterate)
+            {
+                Real product(1);
+                for (int dim = 0; dim < mNumDimensions; ++dim)
+                {
+                    product *= mPTuple[dim];
+                }
+
+                result = result + (*mControls)(mSumIJTuple.data()) * product;
+
+                for (int dim = 0; dim < mNumDimensions; ++dim)
+                {
+                    if (++mJTuple[dim] <= mDegree[dim])
+                    {
+                        mSumIJTuple[dim] = mITuple[dim] + mJTuple[dim];
+                        mPTuple[dim] = mPhi[dim][mJTuple[dim]];
+                        break;
+                    }
+                    mJTuple[dim] = 0;
+                    mSumIJTuple[dim] = mITuple[dim];
+                    mPTuple[dim] = mPhi[dim][0];
+                }
+            }
+
+            Real adjust(1);
+            for (int dim = 0; dim < mNumDimensions; ++dim)
+            {
+                adjust *= mPowerDSDT[dim][order[dim]];
+            }
+            result = result * adjust;
+            return result;
+        }
+
+        typename Controls::Type EvaluateCaching(int const* order, Real const* t)
+        {
+            if (!mValid)
+            {
+                return mCTZero;
+            }
+
+            int numIterates = 1;
+            for (int dim = 0; dim < mNumDimensions; ++dim)
+            {
+                mDegreeMinusOrder[dim] = mDegree[dim] - order[dim];
+                if (mDegreeMinusOrder[dim] < 0 || mDegreeMinusOrder[dim] > mDegree[dim])
+                {
+                    return mCTZero;
+                }
+                numIterates *= mDegreeMinusOrder[dim] + 1;
+            }
+
+            for (int dim = 0; dim < mNumDimensions; ++dim)
+            {
+                GetKey(t[dim], mTMin[dim], mTMax[dim], mPowerDSDT[dim][1],
+                    mNumControls[dim], mDegree[dim], mITuple[dim], mUTuple[dim]);
+            }
+
+            int rowIndex = GetRowIndex(mITuple);
+            for (int dim = 0; dim < mNumDimensions; ++dim)
+            {
+                mJTuple[dim] = 0;
+                mKTuple[dim] = mDegree[dim];
+                mLTuple[dim] = mLMax[dim][order[dim]];
+                mTerm[dim] = mCTZero;
+            }
+            for (int iterate = 0; iterate < numIterates; ++iterate)
+            {
+                int index = GetIndex(rowIndex, mKTuple);
+                if (mCacheMode == ON_DEMAND_CACHING && !mCached[index])
+                {
+                    ComputeTensor(mITuple.data(), mKTuple.data(), index);
+                    mCached[index] = true;
+                }
+                mTerm[0] = mTerm[0] * mUTuple[0] + mTensor[index] * mDCoefficient[0][mLTuple[0]];
+                for (int dim = 0; dim < mNumDimensions; ++dim)
+                {
+                    if (++mJTuple[dim] <= mDegreeMinusOrder[dim])
+                    {
+                        --mKTuple[dim];
+                        --mLTuple[dim];
+                        break;
+                    }
+                    int dimp1 = dim + 1;
+                    if (dimp1 < mNumDimensions)
+                    {
+                        mTerm[dimp1] = mTerm[dimp1] * mUTuple[dimp1] + mTerm[dim] * mDCoefficient[dimp1][mLTuple[dimp1]];
+                        mTerm[dim] = mCTZero;
+                        mJTuple[dim] = 0;
+                        mKTuple[dim] = mDegree[dim];
+                        mLTuple[dim] = mLMax[dim][order[dim]];
+                    }
+                }
+            }
+            auto result = mTerm[mNumDimensions - 1];
+
+            Real adjust(1);
+            for (int dim = 0; dim < mNumDimensions; ++dim)
+            {
+                adjust *= mPowerDSDT[dim][order[dim]];
+            }
+            result = result * adjust;
+            return result;
+        }
+
+        // Constructor inputs.
+        int const mNumDimensions;  // N
+        std::vector<int> mDegree;  // degree[N]
+        Controls const* mControls;
+        typename Controls::Type const mCTZero;
+        int const mCacheMode;
+
+        // Parameters for B-spline evaluation.  All std::vector containers
+        // have N elements.
+        int mNumLocalControls;  // product of (degree[]+1)
+        std::vector<int> mDegreeP1;
+        std::vector<int> mNumControls;
+        std::vector<Real> mTMin;
+        std::vector<Real> mTMax;
+        std::vector<std::vector<Real>> mBlender;
+        std::vector<std::vector<Real>> mDCoefficient;
+        std::vector<std::vector<int>> mLMax;
+        std::vector<std::vector<Real>> mPowerDSDT;
+        std::vector<int> mITuple;
+        std::vector<int> mJTuple;
+        std::vector<int> mKTuple;
+        std::vector<int> mLTuple;
+        std::vector<int> mSumIJTuple;
+        std::vector<Real> mUTuple;
+        std::vector<Real> mPTuple;
+
+        // Support for no-cached B-spline evaluation.  The std::vector
+        // container has N elements.
+        std::vector<std::vector<Real>> mPhi;
+
+        // Support for cached B-spline evaluation.
+        std::vector<int> mTBound;  // tbound[2*N]
+        std::vector<int> mComputeJTuple;  // computejtuple[N]
+        std::vector<int> mComputeSumIJTuple;  // computesumijtuple[N]
+        std::vector<int> mDegreeMinusOrder;  // degreeminusorder[N]
+        std::vector<typename Controls::Type> mTerm;  // mTerm[N]
+        std::vector<typename Controls::Type> mTensor;  // depends on numcontrols
+        std::vector<bool> mCached;  // same size as mTensor
+
+        bool mValid;
+    };
+}
+
+// Implementation for B-spline interpolation whose dimension is known at
+// compile time.
+namespace gte
+{
+    template <typename Real, typename Controls, int N = 0>
+    class IntpBSplineUniform : public IntpBSplineUniformShared<Real, Controls>
+    {
+    public:
+        // The caller is responsible for ensuring that the IntpBSplineUniform
+        // object persists as long as the input 'controls' exists.
+        IntpBSplineUniform(std::array<int, N> const& degrees, Controls const& controls,
+            typename Controls::Type ctZero, int cacheMode)
+            :
+            IntpBSplineUniformShared<Real, Controls>(N, degrees.data(), controls,
+                ctZero, cacheMode)
+        {
+            if (!this->mValid)
+            {
+                return;
+            }
+        }
+
+        // Disallow copying and moving.
+        IntpBSplineUniform(IntpBSplineUniform const&) = delete;
+        IntpBSplineUniform& operator=(IntpBSplineUniform const&) = delete;
+        IntpBSplineUniform(IntpBSplineUniform&&) = delete;
+        IntpBSplineUniform& operator=(IntpBSplineUniform&&) = delete;
+
+        // Evaluate the interpolator.  Each element of 'order' indicates the
+        // order of the derivative you want to compute.  For the function
+        // value itself, pass in 'order' that has all 0 elements.
+        typename Controls::Type Evaluate(std::array<int, N> const& order,
+            std::array<Real, N> const& t)
+        {
+            if (this->mValid)
+            {
+                if (this->mCacheMode == this->NO_CACHING)
+                {
+                    return this->EvaluateNoCaching(order.data(), t.data());
+                }
+                else
+                {
+                    return this->EvaluateCaching(order.data(), t.data());
+                }
+            }
+            else
+            {
+                return this->mCTZero;
+            }
+        }
+    };
+}
+
+// Implementation for B-spline interpolation whose dimension is known only
+// at run time.
+namespace gte
+{
+    template <typename Real, typename Controls>
+    class IntpBSplineUniform<Real, Controls, 0> : public IntpBSplineUniformShared<Real, Controls>
+    {
+    public:
+        // The caller is responsible for ensuring that the IntpBSplineUniform
+        // object persists as long as the input 'controls' exists.
+        IntpBSplineUniform(std::vector<int> const& degrees, Controls const& controls,
+            typename Controls::Type ctZero, int cacheMode)
+            :
+            IntpBSplineUniformShared<Real, Controls>(static_cast<int>(degrees.size()),
+                degrees.data(), controls, ctZero, cacheMode)
+        {
+            if (!this->mValid)
+            {
+                return;
+            }
+        }
+
+        // Disallow copying and moving.
+        IntpBSplineUniform(IntpBSplineUniform const&) = delete;
+        IntpBSplineUniform& operator=(IntpBSplineUniform const&) = delete;
+        IntpBSplineUniform(IntpBSplineUniform&&) = delete;
+        IntpBSplineUniform& operator=(IntpBSplineUniform&&) = delete;
+
+        // Evaluate the interpolator.  Each element of 'order' indicates the
+        // order of the derivative you want to compute.  For the function
+        // value itself, pass in 'order' that has all 0 elements.
+        typename Controls::Type Evaluate(std::vector<int> const& order,
+            std::vector<Real> const& t)
+        {
+            if (this->mValid
+                && static_cast<int>(order.size()) >= this->mNumDimensions
+                && static_cast<int>(t.size()) >= this->mNumDimensions)
+            {
+                if (this->mCacheMode == this->NO_CACHING)
+                {
+                    return this->EvaluateNoCaching(order.data(), t.data());
+                }
+                else
+                {
+                    return this->EvaluateCaching(order.data(), t.data());
+                }
+            }
+            else
+            {
+                return this->mCTZero;
+            }
+        }
+    };
+}
+
+// To use only the N-dimensional template code above, define the symbol
+// GTE_INTP_BSPLINE_UNIFORM_NO_SPECIALIZATION to disable the specializations
+// that occur below.  Notice that the interfaces are different between the
+// specializations and the general code.
+#if !defined(GTE_INTP_BSPLINE_UNIFORM_NO_SPECIALIZATION)
+
+// Specialization for 1-dimensional data.
+namespace gte
+{
+    template <typename Real, typename Controls>
+    class IntpBSplineUniform<Real, Controls, 1> : public IntpBSplineUniformShared<Real, Controls>
+    {
+    public:
+        // The caller is responsible for ensuring that the IntpBSplineUniform
+        // object persists as long as the input 'controls' exists.
+        IntpBSplineUniform(int degree, Controls const& controls,
+            typename Controls::Type ctZero, int cacheMode)
+            :
+            IntpBSplineUniformShared<Real, Controls>(),
+            mDegree(degree),
+            mControls(&controls),
+            mCTZero(ctZero),
+            mCacheMode(cacheMode),
+            mValid(IsValid())
+        {
+            if (!mValid)
+            {
+                return;
+            }
+
+            mDegreeP1 = mDegree + 1;
+            mNumControls = mControls->GetSize(0);
+            mTMin = (Real)-0.5;
+            mTMax = static_cast<Real>(mNumControls) - (Real)0.5;
+            mNumTRows = 0;
+            mNumTCols = 0;
+            this->ComputeBlendingMatrix(mDegree, mBlender);
+            this->ComputeDCoefficients(mDegree, mDCoefficient, mLMax);
+            this->ComputePowers(mDegree, mNumControls, mTMin, mTMax, mPowerDSDT);
+            if (mCacheMode == this->NO_CACHING)
+            {
+                mPhi.resize(mDegreeP1);
+            }
+            else
+            {
+                InitializeTensors();
+            }
+        }
+
+        // Disallow copying and moving.
+        IntpBSplineUniform(IntpBSplineUniform const&) = delete;
+        IntpBSplineUniform& operator=(IntpBSplineUniform const&) = delete;
+        IntpBSplineUniform(IntpBSplineUniform&&) = delete;
+        IntpBSplineUniform& operator=(IntpBSplineUniform&&) = delete;
+
+        // Member access.
+        inline int GetDegree(int) const { return mDegree; }
+        inline int GetNumControls(int) const { return mNumControls; }
+        inline Real GetTMin(int) const { return mTMin; }
+        inline Real GetTMax(int) const { return mTMax; }
+        inline int GetCacheMode() const { return mCacheMode; }
+
+        // Evaluate the interpolator.  The order is 0 when you want the B-spline
+        // function value itself.  The order is 1 for the first derivative of the
+        // function, and so on.
+        typename Controls::Type Evaluate(std::array<int, 1> const& order,
+            std::array<Real,1> const& t)
+        {
+            auto result = mCTZero;
+            if (mValid
+                && 0 <= order[0] && order[0] <= mDegree)
+            {
+                int i;
+                Real u;
+                this->GetKey(t[0], mTMin, mTMax, mPowerDSDT[1], mNumControls, mDegree, i, u);
+
+                if (mCacheMode == this->NO_CACHING)
+                {
+                    int jIndex = 0;
+                    for (int j = 0; j <= mDegree; ++j)
+                    {
+                        int kjIndex = mDegree + jIndex;
+                        int ell = mLMax[order[0]];
+                        mPhi[j] = (Real)0;
+                        for (int k = mDegree; k >= order[0]; --k)
+                        {
+                            mPhi[j] = mPhi[j] * u + mBlender[kjIndex--] * mDCoefficient[ell--];
+                        }
+                        jIndex += mDegreeP1;
+                    }
+
+                    for (int j = 0; j <= mDegree; ++j)
+                    {
+                        result = result + (*mControls)(i + j) * mPhi[j];
+                    }
+                }
+                else
+                {
+                    int iIndex = mNumTCols * i;
+                    int kiIndex = mDegree + iIndex;
+                    int ell = mLMax[order[0]];
+                    for (int k = mDegree; k >= order[0]; --k)
+                    {
+                        if (mCacheMode == this->ON_DEMAND_CACHING && !mCached[kiIndex])
+                        {
+                            ComputeTensor(i, k, kiIndex);
+                            mCached[kiIndex] = true;
+                        }
+
+                        result = result * u + mTensor[kiIndex--] * mDCoefficient[ell--];
+                    }
+                }
+
+                result = result * mPowerDSDT[order[0]];
+            }
+            return result;
+        }
+
+    protected:
+        // Verify that the constraints for number of controls and degrees
+        // are satisfied.
+        bool IsValid()
+        {
+            // The condition c+1 > d+1 is required so that when s = c+1-d, its
+            // maximum value, we have at least two s-knots (d and d + 1).
+            if (mControls->GetSize(0) <= mDegree + 1)
+            {
+                LogError("Incompatible degree and number of controls.");
+                return false;
+            }
+            return true;
+        }
+
+        void ComputeTensor(int r, int c, int index)
+        {
+            auto element = mCTZero;
+            for (int j = 0; j <= mDegree; ++j)
+            {
+                element = element + (*mControls)(r + j) * mBlender[c + mDegreeP1 * j];
+            }
+            mTensor[index] = element;
+        }
+
+        void InitializeTensors()
+        {
+            mNumTRows = mNumControls - mDegree;
+            mNumTCols = mDegreeP1;
+            int numCached = mNumTRows * mNumTCols;
+            mTensor.resize(numCached);
+            mCached.resize(numCached);
+            if (mCacheMode == this->PRE_CACHING)
+            {
+                for (int r = 0, index = 0; r < mNumTRows; ++r)
+                {
+                    for (int c = 0; c < mNumTCols; ++c, ++index)
+                    {
+                        ComputeTensor(r, c, index);
+                    }
+                }
+                std::fill(mCached.begin(), mCached.end(), true);
+            }
+            else
+            {
+                std::fill(mCached.begin(), mCached.end(), false);
+            }
+        }
+
+        // Constructor inputs.
+        int mDegree;
+        Controls const* mControls;
+        typename Controls::Type mCTZero;
+        int mCacheMode;
+
+        // Parameters for B-spline evaluation.
+        int mDegreeP1;
+        int mNumControls;
+        Real mTMin, mTMax;
+        std::vector<Real> mBlender;
+        std::vector<Real> mDCoefficient;
+        std::vector<int> mLMax;
+        std::vector<Real> mPowerDSDT;
+
+        // Support for no-cached B-spline evaluation.
+        std::vector<Real> mPhi;
+
+        // Support for cached B-spline evaluation.
+        int mNumTRows, mNumTCols;
+        std::vector<typename Controls::Type> mTensor;
+        std::vector<bool> mCached;
+
+        bool mValid;
+    };
+}
+
+// Specialization for 2-dimensional data.
+namespace gte
+{
+    template <typename Real, typename Controls>
+    class IntpBSplineUniform<Real, Controls, 2> : public IntpBSplineUniformShared<Real, Controls>
+    {
+    public:
+        // The caller is responsible for ensuring that the IntpBSplineUniform2
+        // object persists as long as the input 'controls' exists.
+        IntpBSplineUniform(std::array<int, 2> const& degrees, Controls const& controls,
+            typename Controls::Type ctZero, int cacheMode)
+            :
+            IntpBSplineUniformShared<Real, Controls>(),
+            mDegree(degrees),
+            mControls(&controls),
+            mCTZero(ctZero),
+            mCacheMode(cacheMode),
+            mValid(IsValid())
+        {
+            if (!mValid)
+            {
+                return;
+            }
+
+            for (int dim = 0; dim < 2; ++dim)
+            {
+                mDegreeP1[dim] = mDegree[dim] + 1;
+                mNumControls[dim] = mControls->GetSize(dim);
+                mTMin[dim] = (Real)-0.5;
+                mTMax[dim] = static_cast<Real>(mNumControls[dim]) - (Real)0.5;
+                mNumTRows[dim] = 0;
+                mNumTCols[dim] = 0;
+                this->ComputeBlendingMatrix(mDegree[dim], mBlender[dim]);
+                this->ComputeDCoefficients(mDegree[dim], mDCoefficient[dim], mLMax[dim]);
+                this->ComputePowers(mDegree[dim], mNumControls[dim], mTMin[dim], mTMax[dim], mPowerDSDT[dim]);
+            }
+
+            if (mCacheMode == this->NO_CACHING)
+            {
+                for (int dim = 0; dim < 2; ++dim)
+                {
+                    mPhi[dim].resize(mDegreeP1[dim]);
+                }
+            }
+            else
+            {
+                InitializeTensors();
+            }
+        }
+
+        // Disallow copying and moving.
+        IntpBSplineUniform(IntpBSplineUniform const&) = delete;
+        IntpBSplineUniform& operator=(IntpBSplineUniform const&) = delete;
+        IntpBSplineUniform(IntpBSplineUniform&&) = delete;
+        IntpBSplineUniform& operator=(IntpBSplineUniform&&) = delete;
+
+        // Member access.
+        inline int GetDegree(int dim) const { return mDegree[dim]; }
+        inline int GetNumControls(int dim) const { return mNumControls[dim]; }
+        inline Real GetTMin(int dim) const { return mTMin[dim]; }
+        inline Real GetTMax(int dim) const { return mTMax[dim]; }
+        inline int GetCacheMode() const { return mCacheMode; }
+
+        // Evaluate the interpolator.  The order is (0,0) when you want the
+        // B-spline function value itself.  The order0 is 1 for the first
+        // derivative with respect to t0 and the order1 is 1 for the first
+        // derivative with respect to t1.  Higher-order derivatives in other
+        // t-inputs are computed similarly.
+        typename Controls::Type Evaluate(std::array<int, 2> const& order,
+            std::array<Real, 2> const& t)
+        {
+            auto result = mCTZero;
+            if (mValid
+                && 0 <= order[0] && order[0] <= mDegree[0]
+                && 0 <= order[1] && order[1] <= mDegree[1])
+            {
+                std::array<int, 2> i;
+                std::array<Real, 2> u;
+                for (int dim = 0; dim < 2; ++dim)
+                {
+                    this->GetKey(t[dim], mTMin[dim], mTMax[dim], mPowerDSDT[dim][1],
+                        mNumControls[dim], mDegree[dim], i[dim], u[dim]);
+                }
+
+                if (mCacheMode == this->NO_CACHING)
+                {
+                    for (int dim = 0; dim < 2; ++dim)
+                    {
+                        int jIndex = 0;
+                        for (int j = 0; j <= mDegree[dim]; ++j)
+                        {
+                            int kjIndex = mDegree[dim] + jIndex;
+                            int ell = mLMax[dim][order[dim]];
+                            mPhi[dim][j] = (Real)0;
+                            for (int k = mDegree[dim]; k >= order[dim]; --k)
+                            {
+                                mPhi[dim][j] = mPhi[dim][j] * u[dim] +
+                                    mBlender[dim][kjIndex--] * mDCoefficient[dim][ell--];
+                            }
+                            jIndex += mDegreeP1[dim];
+                        }
+                    }
+
+                    for (int j1 = 0; j1 <= mDegree[1]; ++j1)
+                    {
+                        Real phi1 = mPhi[1][j1];
+                        for (int j0 = 0; j0 <= mDegree[0]; ++j0)
+                        {
+                            Real phi0 = mPhi[0][j0];
+                            Real phi01 = phi0 * phi1;
+                            result = result + (*mControls)(i[0] + j0, i[1] + j1) * phi01;
+                        }
+                    }
+                }
+                else
+                {
+                    int i0i1Index = mNumTCols[1] * (i[0] + mNumTRows[0] * i[1]);
+                    int k1i0i1Index = mDegree[1] + i0i1Index;
+                    int ell1 = mLMax[1][order[1]];
+                    for (int k1 = mDegree[1]; k1 >= order[1]; --k1)
+                    {
+                        int k0k1i0i1Index = mDegree[0] + mNumTCols[0] * k1i0i1Index;
+                        int ell0 = mLMax[0][order[0]];
+                        auto term = mCTZero;
+                        for (int k0 = mDegree[0]; k0 >= order[0]; --k0)
+                        {
+                            if (mCacheMode == this->ON_DEMAND_CACHING && !mCached[k0k1i0i1Index])
+                            {
+                                ComputeTensor(i[0], i[1], k0, k1, k0k1i0i1Index);
+                                mCached[k0k1i0i1Index] = true;
+                            }
+                            term = term * u[0] + mTensor[k0k1i0i1Index--] * mDCoefficient[0][ell0--];
+                        }
+                        result = result * u[1] + term * mDCoefficient[1][ell1--];
+                        --k1i0i1Index;
+                    }
+                }
+
+                Real adjust(1);
+                for (int dim = 0; dim < 2; ++dim)
+                {
+                    adjust *= mPowerDSDT[dim][order[dim]];
+                }
+                result = result * adjust;
+            }
+            return result;
+        }
+
+        // Verify that the constraints for number of controls and degrees
+        // are satisfied.
+        bool IsValid()
+        {
+            // The condition c+1 > d+1 is required so that when s = c+1-d, its
+            // maximum value, we have at least two s-knots (d and d + 1).
+            for (int dim = 0; dim < 2; ++dim)
+            {
+                if (mControls->GetSize(dim) <= mDegree[dim] + 1)
+                {
+                    LogError("Incompatible degree and number of controls.");
+                    return false;
+                }
+            }
+            return true;
+        }
+
+        void ComputeTensor(int r0, int r1, int c0, int c1, int index)
+        {
+            auto element = mCTZero;
+            for (int j1 = 0; j1 <= mDegree[1]; ++j1)
+            {
+                Real blend1 = mBlender[1][c1 + mDegreeP1[1] * j1];
+                for (int j0 = 0; j0 <= mDegree[0]; ++j0)
+                {
+                    Real blend0 = mBlender[0][c0 + mDegreeP1[0] * j0];
+                    Real blend01 = blend0 * blend1;
+                    element = element + (*mControls)(r0 + j0, r1 + j1) * blend01;
+                }
+            }
+            mTensor[index] = element;
+        }
+
+        void InitializeTensors()
+        {
+            int numCached = 1;
+            for (int dim = 0; dim < 2; ++dim)
+            {
+                mNumTRows[dim] = mNumControls[dim] - mDegree[dim];
+                mNumTCols[dim] = mDegreeP1[dim];
+                numCached *= mNumTRows[dim] * mNumTCols[dim];
+            }
+            mTensor.resize(numCached);
+            mCached.resize(numCached);
+            if (mCacheMode == this->PRE_CACHING)
+            {
+                for (int r1 = 0, index = 0; r1 < mNumTRows[1]; ++r1)
+                {
+                    for (int r0 = 0; r0 < mNumTRows[0]; ++r0)
+                    {
+                        for (int c1 = 0; c1 < mNumTCols[1]; ++c1)
+                        {
+                            for (int c0 = 0; c0 < mNumTCols[0]; ++c0, ++index)
+                            {
+                                ComputeTensor(r0, r1, c0, c1, index);
+                            }
+                        }
+                    }
+                }
+                std::fill(mCached.begin(), mCached.end(), true);
+            }
+            else
+            {
+                std::fill(mCached.begin(), mCached.end(), false);
+            }
+        }
+
+        // Constructor inputs.
+        std::array<int, 2> mDegree;
+        Controls const* mControls;
+        typename Controls::Type mCTZero;
+        int mCacheMode;
+
+        // Parameters for B-spline evaluation.
+        std::array<int, 2> mDegreeP1;
+        std::array<int, 2> mNumControls;
+        std::array<Real, 2> mTMin, mTMax;
+        std::array<std::vector<Real>, 2> mBlender;
+        std::array<std::vector<Real>, 2> mDCoefficient;
+        std::array<std::vector<int>, 2> mLMax;
+        std::array<std::vector<Real>, 2> mPowerDSDT;
+
+        // Support for no-cached B-spline evaluation.
+        std::array<std::vector<Real>, 2> mPhi;
+
+        // Support for cached B-spline evaluation.
+        std::array<int, 2> mNumTRows, mNumTCols;
+        std::vector<typename Controls::Type> mTensor;
+        std::vector<bool> mCached;
+
+        bool mValid;
+    };
+}
+
+// Specialization for 3-dimensional data.
+namespace gte
+{
+    template <typename Real, typename Controls>
+    class IntpBSplineUniform<Real, Controls, 3> : public IntpBSplineUniformShared<Real, Controls>
+    {
+    public:
+        // The caller is responsible for ensuring that the IntpBSplineUniform3
+        // object persists as long as the input 'controls' exists.
+        IntpBSplineUniform(std::array<int, 3> const& degrees, Controls const& controls,
+            typename Controls::Type ctZero, int cacheMode)
+            :
+            IntpBSplineUniformShared<Real, Controls>(),
+            mDegree(degrees),
+            mControls(&controls),
+            mCTZero(ctZero),
+            mCacheMode(cacheMode),
+            mValid(IsValid())
+        {
+            if (!mValid)
+            {
+                return;
+            }
+
+            for (int dim = 0; dim < 3; ++dim)
+            {
+                mDegreeP1[dim] = mDegree[dim] + 1;
+                mNumControls[dim] = mControls->GetSize(dim);
+                mTMin[dim] = (Real)-0.5;
+                mTMax[dim] = static_cast<Real>(mNumControls[dim]) - (Real)0.5;
+                mNumTRows[dim] = 0;
+                mNumTCols[dim] = 0;
+                this->ComputeBlendingMatrix(mDegree[dim], mBlender[dim]);
+                this->ComputeDCoefficients(mDegree[dim], mDCoefficient[dim], mLMax[dim]);
+                this->ComputePowers(mDegree[dim], mNumControls[dim], mTMin[dim], mTMax[dim], mPowerDSDT[dim]);
+            }
+
+            if (mCacheMode == this->NO_CACHING)
+            {
+                for (int dim = 0; dim < 3; ++dim)
+                {
+                    mPhi[dim].resize(mDegreeP1[dim]);
+                }
+            }
+            else
+            {
+                InitializeTensors();
+            }
+
+        }
+
+        // Disallow copying and moving.
+        IntpBSplineUniform(IntpBSplineUniform const&) = delete;
+        IntpBSplineUniform& operator=(IntpBSplineUniform const&) = delete;
+        IntpBSplineUniform(IntpBSplineUniform&&) = delete;
+        IntpBSplineUniform& operator=(IntpBSplineUniform&&) = delete;
+
+        // Member access.  The input i specifies the dimension (0, 1, 2).
+        inline int GetDegree(int dim) const { return mDegree[dim]; }
+        inline int GetNumControls(int dim) const { return mNumControls[dim]; }
+        inline Real GetTMin(int dim) const { return mTMin[dim]; }
+        inline Real GetTMax(int dim) const { return mTMax[dim]; }
+
+        // Evaluate the interpolator.  The order is (0,0,0) when you want the
+        // B-spline function value itself.  The order0 is 1 for the first
+        // derivative with respect to t0, the order1 is 1 for the first
+        // derivative with respect to t1 or the order2 is 1 for the first
+        // derivative with respect to t2.  Higher-order derivatives in other
+        // t-inputs are computed similarly.
+        typename Controls::Type Evaluate(std::array<int, 3> const& order,
+            std::array<Real, 3> const& t)
+        {
+            auto result = mCTZero;
+            if (mValid
+                && 0 <= order[0] && order[0] <= mDegree[0]
+                && 0 <= order[1] && order[1] <= mDegree[1]
+                && 0 <= order[2] && order[2] <= mDegree[2])
+            {
+                std::array<int, 3> i;
+                std::array<Real, 3> u;
+                for (int dim = 0; dim < 3; ++dim)
+                {
+                    this->GetKey(t[dim], mTMin[dim], mTMax[dim], mPowerDSDT[dim][1],
+                        mNumControls[dim], mDegree[dim], i[dim], u[dim]);
+                }
+
+                if (mCacheMode == this->NO_CACHING)
+                {
+                    for (int dim = 0; dim < 3; ++dim)
+                    {
+                        int jIndex = 0;
+                        for (int j = 0; j <= mDegree[dim]; ++j)
+                        {
+                            int kjIndex = mDegree[dim] + jIndex;
+                            int ell = mLMax[dim][order[dim]];
+                            mPhi[dim][j] = (Real)0;
+                            for (int k = mDegree[dim]; k >= order[dim]; --k)
+                            {
+                                mPhi[dim][j] = mPhi[dim][j] * u[dim] +
+                                    mBlender[dim][kjIndex--] * mDCoefficient[dim][ell--];
+                            }
+                            jIndex += mDegreeP1[dim];
+                        }
+                    }
+
+                    for (int j2 = 0; j2 <= mDegree[2]; ++j2)
+                    {
+                        Real phi2 = mPhi[2][j2];
+                        for (int j1 = 0; j1 <= mDegree[1]; ++j1)
+                        {
+                            Real phi1 = mPhi[1][j1];
+                            Real phi12 = phi1 * phi2;
+                            for (int j0 = 0; j0 <= mDegree[0]; ++j0)
+                            {
+                                Real phi0 = mPhi[0][j0];
+                                Real phi012 = phi0 * phi12;
+                                result = result + (*mControls)(i[0] + j0, i[1] + j1, i[2] + j2) * phi012;
+                            }
+                        }
+                    }
+                }
+                else
+                {
+                    int i0i1i2Index = mNumTCols[2] * (i[0] + mNumTRows[0] * (i[1] + mNumTRows[1] * i[2]));
+                    int k2i0i1i2Index = mDegree[2] + i0i1i2Index;
+                    int ell2 = mLMax[2][order[2]];
+                    for (int k2 = mDegree[2]; k2 >= order[2]; --k2)
+                    {
+                        int k1k2i0i1i2Index = mDegree[1] + mNumTCols[1] * k2i0i1i2Index;
+                        int ell1 = mLMax[1][order[1]];
+                        auto term1 = mCTZero;
+                        for (int k1 = mDegree[1]; k1 >= order[1]; --k1)
+                        {
+                            int k0k1k2i0i1i2Index = mDegree[0] + mNumTCols[0] * k1k2i0i1i2Index;
+                            int ell0 = mLMax[0][order[0]];
+                            auto term0 = mCTZero;
+                            for (int k0 = mDegree[0]; k0 >= order[0]; --k0)
+                            {
+                                if (mCacheMode == this->ON_DEMAND_CACHING && !mCached[k0k1k2i0i1i2Index])
+                                {
+                                    ComputeTensor(i[0], i[1], i[2], k0, k1, k2, k0k1k2i0i1i2Index);
+                                    mCached[k0k1k2i0i1i2Index] = true;
+                                }
+
+                                term0 = term0 * u[0] + mTensor[k0k1k2i0i1i2Index--] * mDCoefficient[0][ell0--];
+                            }
+                            term1 = term1 * u[1] + term0 * mDCoefficient[1][ell1--];
+                            --k1k2i0i1i2Index;
+                        }
+                        result = result * u[2] + term1 * mDCoefficient[2][ell2--];
+                        --k2i0i1i2Index;
+                    }
+                }
+
+                Real adjust(1);
+                for (int dim = 0; dim < 3; ++dim)
+                {
+                    adjust *= mPowerDSDT[dim][order[dim]];
+                }
+                result = result * adjust;
+            }
+            return result;
+        }
+
+    protected:
+        // Verify that the constraints for number of controls and degrees
+        // are satisfied.
+        bool IsValid()
+        {
+            // The condition c+1 > d+1 is required so that when s = c+1-d, its
+            // maximum value, we have at least two s-knots (d and d + 1).
+            for (int dim = 0; dim < 3; ++dim)
+            {
+                if (mControls->GetSize(dim) <= mDegree[dim] + 1)
+                {
+                    LogError("Incompatible degree and number of controls.");
+                    return false;
+                }
+            }
+            return true;
+        }
+
+        void ComputeTensor(int r0, int r1, int r2, int c0, int c1, int c2, int index)
+        {
+            auto element = mCTZero;
+            for (int j2 = 0; j2 <= mDegree[2]; ++j2)
+            {
+                Real blend2 = mBlender[2][c2 + mDegreeP1[2] * j2];
+                for (int j1 = 0; j1 <= mDegree[1]; ++j1)
+                {
+                    Real blend1 = mBlender[1][c1 + mDegreeP1[1] * j1];
+                    Real blend12 = blend1 * blend2;
+                    for (int j0 = 0; j0 <= mDegree[0]; ++j0)
+                    {
+                        Real blend0 = mBlender[0][c0 + mDegreeP1[0] * j0];
+                        Real blend012 = blend0 * blend12;
+                        element = element + (*mControls)(r0 + j0, r1 + j1, r2 + j2) * blend012;
+                    }
+                }
+            }
+            mTensor[index] = element;
+        }
+
+        void InitializeTensors()
+        {
+            int numCached = 1;
+            for (int dim = 0; dim < 3; ++dim)
+            {
+                mNumTRows[dim] = mNumControls[dim] - mDegree[dim];
+                mNumTCols[dim] = mDegreeP1[dim];
+                numCached *= mNumTRows[dim] * mNumTCols[dim];
+            }
+            mTensor.resize(numCached);
+            mCached.resize(numCached);
+            if (mCacheMode == this->PRE_CACHING)
+            {
+                for (int r2 = 0, index = 0; r2 < mNumTRows[2]; ++r2)
+                {
+                    for (int r1 = 0; r1 < mNumTRows[1]; ++r1)
+                    {
+                        for (int r0 = 0; r0 < mNumTRows[0]; ++r0)
+                        {
+                            for (int c2 = 0; c2 < mNumTCols[2]; ++c2)
+                            {
+                                for (int c1 = 0; c1 < mNumTCols[1]; ++c1)
+                                {
+                                    for (int c0 = 0; c0 < mNumTCols[0]; ++c0, ++index)
+                                    {
+                                        ComputeTensor(r0, r1, r2, c0, c1, c2, index);
+                                    }
+                                }
+                            }
+                        }
+                    }
+                }
+                std::fill(mCached.begin(), mCached.end(), true);
+            }
+            else
+            {
+                std::fill(mCached.begin(), mCached.end(), false);
+            }
+        }
+
+        // Constructor inputs.
+        std::array<int, 3> mDegree;
+        Controls const* mControls;
+        typename Controls::Type mCTZero;
+        int mCacheMode;
+
+        // Parameters for B-spline evaluation.
+        std::array<int, 3> mDegreeP1;
+        std::array<int, 3> mNumControls;
+        std::array<Real, 3> mTMin, mTMax;
+        std::array<std::vector<Real>, 3> mBlender;
+        std::array<std::vector<Real>, 3> mDCoefficient;
+        std::array<std::vector<int>, 3> mLMax;
+        std::array<std::vector<Real>, 3> mPowerDSDT;
+
+        // Support for no-cached B-spline evaluation.
+        std::array<std::vector<Real>, 3> mPhi;
+
+        // Support for cached B-spline evaluation.
+        std::array<int, 3> mNumTRows, mNumTCols;
+        std::vector<typename Controls::Type> mTensor;
+        std::vector<bool> mCached;
+
+        bool mValid;
+    };
+}
+
+#endif
diff --git a/Engine/Plugins/Experimental/GeometryProcessing/Source/GeometryAlgorithms/Private/GTEngine/Mathematics/GteIntpBicubic2.h b/Engine/Plugins/Experimental/GeometryProcessing/Source/GeometryAlgorithms/Private/GTEngine/Mathematics/GteIntpBicubic2.h
new file mode 100644
--- /dev/null
+++ b/Engine/Plugins/Experimental/GeometryProcessing/Source/GeometryAlgorithms/Private/GTEngine/Mathematics/GteIntpBicubic2.h
@@ -0,0 +1,425 @@
+// David Eberly, Geometric Tools, Redmond WA 98052
+// Copyright (c) 1998-2019
+// Distributed under the Boost Software License, Version 1.0.
+// http://www.boost.org/LICENSE_1_0.txt
+// http://www.geometrictools.com/License/Boost/LICENSE_1_0.txt
+// File Version: 3.0.0 (2016/06/19)
+
+#pragma once
+
+#include <GTEngine/LowLevel/GteLogger.h>
+
+// The interpolator is for uniformly spaced (x,y)-values.  The input samples
+// F must be stored in row-major order to represent f(x,y); that is,
+// F[c + xBound*r] corresponds to f(x,y), where c is the index corresponding
+// to x and r is the index corresponding to y.  Exact interpolation is
+// achieved by setting catmullRom to 'true', giving you the Catmull-Rom
+// blending matrix.  If a smooth interpolation is desired, set catmullRom to
+// 'false' to obtain B-spline blending.
+
+namespace gte
+{
+
+template <typename Real>
+class IntpBicubic2
+{
+public:
+    // Construction.
+    IntpBicubic2(int xBound, int yBound, Real xMin, Real xSpacing,
+        Real yMin, Real ySpacing, Real const* F, bool catmullRom);
+
+    // Member access.
+    inline int GetXBound() const;
+    inline int GetYBound() const;
+    inline int GetQuantity() const;
+    inline Real const* GetF() const;
+    inline Real GetXMin() const;
+    inline Real GetXMax() const;
+    inline Real GetXSpacing() const;
+    inline Real GetYMin() const;
+    inline Real GetYMax() const;
+    inline Real GetYSpacing() const;
+
+    // Evaluate the function and its derivatives.  The functions clamp the
+    // inputs to xmin <= x <= xmax and ymin <= y <= ymax.  The first operator
+    // is for function evaluation.  The second operator is for function or
+    // derivative evaluations.  The xOrder argument is the order of the
+    // x-derivative and the yOrder argument is the order of the y-derivative.
+    // Both orders are zero to get the function value itself.
+    Real operator()(Real x, Real y) const;
+    Real operator()(int xOrder, int yOrder, Real x, Real y) const;
+
+private:
+    int mXBound, mYBound, mQuantity;
+    Real mXMin, mXMax, mXSpacing, mInvXSpacing;
+    Real mYMin, mYMax, mYSpacing, mInvYSpacing;
+    Real const* mF;
+    Real mBlend[4][4];
+};
+
+
+template <typename Real>
+IntpBicubic2<Real>::IntpBicubic2(int xBound, int yBound, Real xMin,
+    Real xSpacing, Real yMin, Real ySpacing, Real const* F, bool catmullRom)
+    :
+    mXBound(xBound),
+    mYBound(yBound),
+    mQuantity(xBound * yBound),
+    mXMin(xMin),
+    mXSpacing(xSpacing),
+    mYMin(yMin),
+    mYSpacing(ySpacing),
+    mF(F)
+{
+    // At least a 3x3 block of data points are needed to construct the
+    // estimates of the boundary derivatives.
+    LogAssert(mXBound >= 3 && mYBound >= 3 && mF, "Invalid input.");
+    LogAssert(mXSpacing > (Real)0 && mYSpacing > (Real)0, "Invalid input.");
+
+    mXMax = mXMin + mXSpacing * static_cast<Real>(mXBound - 1);
+    mInvXSpacing = ((Real)1) / mXSpacing;
+    mYMax = mYMin + mYSpacing * static_cast<Real>(mYBound - 1);
+    mInvYSpacing = ((Real)1) / mYSpacing;
+
+    if (catmullRom)
+    {
+        mBlend[0][0] = (Real)0;
+        mBlend[0][1] = (Real)-0.5;
+        mBlend[0][2] = (Real)1;
+        mBlend[0][3] = (Real)-0.5;
+        mBlend[1][0] = (Real)1;
+        mBlend[1][1] = (Real)0;
+        mBlend[1][2] = (Real)-2.5;
+        mBlend[1][3] = (Real)1.5;
+        mBlend[2][0] = (Real)0;
+        mBlend[2][1] = (Real)0.5;
+        mBlend[2][2] = (Real)2;
+        mBlend[2][3] = (Real)-1.5;
+        mBlend[3][0] = (Real)0;
+        mBlend[3][1] = (Real)0;
+        mBlend[3][2] = (Real)-0.5;
+        mBlend[3][3] = (Real)0.5;
+    }
+    else
+    {
+        mBlend[0][0] = (Real)1 / (Real)6;
+        mBlend[0][1] = (Real)-3 / (Real)6;
+        mBlend[0][2] = (Real)3 / (Real)6;
+        mBlend[0][3] = (Real)-1 / (Real)6;;
+        mBlend[1][0] = (Real)4 / (Real)6;
+        mBlend[1][1] = (Real)0 / (Real)6;
+        mBlend[1][2] = (Real)-6 / (Real)6;
+        mBlend[1][3] = (Real)3 / (Real)6;
+        mBlend[2][0] = (Real)1 / (Real)6;
+        mBlend[2][1] = (Real)3 / (Real)6;
+        mBlend[2][2] = (Real)3 / (Real)6;
+        mBlend[2][3] = (Real)-3 / (Real)6;
+        mBlend[3][0] = (Real)0 / (Real)6;
+        mBlend[3][1] = (Real)0 / (Real)6;
+        mBlend[3][2] = (Real)0 / (Real)6;
+        mBlend[3][3] = (Real)1 / (Real)6;
+    }
+}
+
+template <typename Real> inline
+int IntpBicubic2<Real>::GetXBound() const
+{
+    return mXBound;
+}
+
+template <typename Real> inline
+int IntpBicubic2<Real>::GetYBound() const
+{
+    return mYBound;
+}
+
+template <typename Real> inline
+int IntpBicubic2<Real>::GetQuantity() const
+{
+    return mQuantity;
+}
+
+template <typename Real> inline
+Real const* IntpBicubic2<Real>::GetF() const
+{
+    return mF;
+}
+
+template <typename Real> inline
+Real IntpBicubic2<Real>::GetXMin() const
+{
+    return mXMin;
+}
+
+template <typename Real> inline
+Real IntpBicubic2<Real>::GetXMax() const
+{
+    return mXMax;
+}
+
+template <typename Real> inline
+Real IntpBicubic2<Real>::GetXSpacing() const
+{
+    return mXSpacing;
+}
+
+template <typename Real> inline
+Real IntpBicubic2<Real>::GetYMin() const
+{
+    return mYMin;
+}
+
+template <typename Real> inline
+Real IntpBicubic2<Real>::GetYMax() const
+{
+    return mYMax;
+}
+
+template <typename Real> inline
+Real IntpBicubic2<Real>::GetYSpacing() const
+{
+    return mYSpacing;
+}
+
+template <typename Real>
+Real IntpBicubic2<Real>::operator()(Real x, Real y) const
+{
+    // Compute x-index and clamp to image.
+    Real xIndex = (x - mXMin) * mInvXSpacing;
+    int ix = static_cast<int>(xIndex);
+    if (ix < 0)
+    {
+        ix = 0;
+    }
+    else if (ix >= mXBound)
+    {
+        ix = mXBound - 1;
+    }
+
+    // Compute y-index and clamp to image.
+    Real yIndex = (y - mYMin) * mInvYSpacing;
+    int iy = static_cast<int>(yIndex);
+    if (iy < 0)
+    {
+        iy = 0;
+    }
+    else if (iy >= mYBound)
+    {
+        iy = mYBound - 1;
+    }
+
+    Real U[4];
+    U[0] = (Real)1;
+    U[1] = xIndex - ix;
+    U[2] = U[1] * U[1];
+    U[3] = U[1] * U[2];
+
+    Real V[4];
+    V[0] = (Real)1;
+    V[1] = yIndex - iy;
+    V[2] = V[1] * V[1];
+    V[3] = V[1] * V[2];
+
+    // Compute P = M*U and Q = M*V.
+    Real P[4], Q[4];
+    for (int row = 0; row < 4; ++row)
+    {
+        P[row] = (Real)0;
+        Q[row] = (Real)0;
+        for (int col = 0; col < 4; ++col)
+        {
+            P[row] += mBlend[row][col] * U[col];
+            Q[row] += mBlend[row][col] * V[col];
+        }
+    }
+
+    // Compute (M*U)^t D (M*V) where D is the 4x4 subimage containing (x,y).
+    --ix;
+    --iy;
+    Real result = (Real)0;
+    for (int row = 0; row < 4; ++row)
+    {
+        int yClamp = iy + row;
+        if (yClamp < 0)
+        {
+            yClamp = 0;
+        }
+        else if (yClamp > mYBound - 1)
+        {
+            yClamp = mYBound - 1;
+        }
+
+        for (int col = 0; col < 4; ++col)
+        {
+            int xClamp = ix + col;
+            if (xClamp < 0)
+            {
+                xClamp = 0;
+            }
+            else if (xClamp > mXBound - 1)
+            {
+                xClamp = mXBound - 1;
+            }
+
+            result += P[col] * Q[row] * mF[xClamp + mXBound * yClamp];
+        }
+    }
+
+    return result;
+}
+
+template <typename Real>
+Real IntpBicubic2<Real>::operator()(int xOrder, int yOrder, Real x, Real y)
+const
+{
+    // Compute x-index and clamp to image.
+    Real xIndex = (x - mXMin) * mInvXSpacing;
+    int ix = static_cast<int>(xIndex);
+    if (ix < 0)
+    {
+        ix = 0;
+    }
+    else if (ix >= mXBound)
+    {
+        ix = mXBound - 1;
+    }
+
+    // Compute y-index and clamp to image.
+    Real yIndex = (y - mYMin) * mInvYSpacing;
+    int iy = static_cast<int>(yIndex);
+    if (iy < 0)
+    {
+        iy = 0;
+    }
+    else if (iy >= mYBound)
+    {
+        iy = mYBound - 1;
+    }
+
+    Real U[4], dx, xMult;
+    switch (xOrder)
+    {
+    case 0:
+        dx = xIndex - ix;
+        U[0] = (Real)1;
+        U[1] = dx;
+        U[2] = dx * U[1];
+        U[3] = dx * U[2];
+        xMult = (Real)1;
+        break;
+    case 1:
+        dx = xIndex - ix;
+        U[0] = (Real)0;
+        U[1] = (Real)1;
+        U[2] = ((Real)2) * dx;
+        U[3] = ((Real)3) * dx * dx;
+        xMult = mInvXSpacing;
+        break;
+    case 2:
+        dx = xIndex - ix;
+        U[0] = (Real)0;
+        U[1] = (Real)0;
+        U[2] = (Real)2;
+        U[3] = (Real)6 * dx;
+        xMult = mInvXSpacing * mInvXSpacing;
+        break;
+    case 3:
+        U[0] = (Real)0;
+        U[1] = (Real)0;
+        U[2] = (Real)0;
+        U[3] = (Real)6;
+        xMult = mInvXSpacing * mInvXSpacing * mInvXSpacing;
+        break;
+    default:
+        return (Real)0;
+    }
+
+    Real V[4], dy, yMult;
+    switch (yOrder)
+    {
+    case 0:
+        dy = yIndex - iy;
+        V[0] = (Real)1;
+        V[1] = dy;
+        V[2] = dy * V[1];
+        V[3] = dy * V[2];
+        yMult = (Real)1;
+        break;
+    case 1:
+        dy = yIndex - iy;
+        V[0] = (Real)0;
+        V[1] = (Real)1;
+        V[2] = ((Real)2) * dy;
+        V[3] = ((Real)3) * dy * dy;
+        yMult = mInvYSpacing;
+        break;
+    case 2:
+        dy = yIndex - iy;
+        V[0] = (Real)0;
+        V[1] = (Real)0;
+        V[2] = (Real)2;
+        V[3] = ((Real)6) * dy;
+        yMult = mInvYSpacing * mInvYSpacing;
+        break;
+    case 3:
+        V[0] = (Real)0;
+        V[1] = (Real)0;
+        V[2] = (Real)0;
+        V[3] = (Real)6;
+        yMult = mInvYSpacing * mInvYSpacing * mInvYSpacing;
+        break;
+    default:
+        return (Real)0;
+    }
+
+    // Compute P = M*U and Q = M*V.
+    Real P[4], Q[4];
+    for (int row = 0; row < 4; ++row)
+    {
+        P[row] = (Real)0;
+        Q[row] = (Real)0;
+        for (int col = 0; col < 4; ++col)
+        {
+            P[row] += mBlend[row][col] * U[col];
+            Q[row] += mBlend[row][col] * V[col];
+        }
+    }
+
+    // Compute (M*U)^t D (M*V) where D is the 4x4 subimage containing (x,y).
+    --ix;
+    --iy;
+    Real result = (Real)0;
+    for (int row = 0; row < 4; ++row)
+    {
+        int yClamp = iy + row;
+        if (yClamp < 0)
+        {
+            yClamp = 0;
+        }
+        else if (yClamp > mYBound - 1)
+        {
+            yClamp = mYBound - 1;
+        }
+
+        for (int col = 0; col < 4; ++col)
+        {
+            int xClamp = ix + col;
+            if (xClamp < 0)
+            {
+                xClamp = 0;
+            }
+            else if (xClamp > mXBound - 1)
+            {
+                xClamp = mXBound - 1;
+            }
+
+            result += P[col] * Q[row] * mF[xClamp + mXBound * yClamp];
+        }
+    }
+    result *= xMult * yMult;
+
+    return result;
+}
+
+
+}
diff --git a/Engine/Plugins/Experimental/GeometryProcessing/Source/GeometryAlgorithms/Private/GTEngine/Mathematics/GteIntpBilinear2.h b/Engine/Plugins/Experimental/GeometryProcessing/Source/GeometryAlgorithms/Private/GTEngine/Mathematics/GteIntpBilinear2.h
new file mode 100644
--- /dev/null
+++ b/Engine/Plugins/Experimental/GeometryProcessing/Source/GeometryAlgorithms/Private/GTEngine/Mathematics/GteIntpBilinear2.h
@@ -0,0 +1,325 @@
+// David Eberly, Geometric Tools, Redmond WA 98052
+// Copyright (c) 1998-2019
+// Distributed under the Boost Software License, Version 1.0.
+// http://www.boost.org/LICENSE_1_0.txt
+// http://www.geometrictools.com/License/Boost/LICENSE_1_0.txt
+// File Version: 3.0.0 (2016/06/19)
+
+#pragma once
+
+#include <GTEngine/LowLevel/GteLogger.h>
+
+// The interpolator is for uniformly spaced (x,y)-values.  The input samples
+// F must be stored in row-major order to represent f(x,y); that is,
+// F[c + xBound*r] corresponds to f(x,y), where c is the index corresponding
+// to x and r is the index corresponding to y.
+
+namespace gte
+{
+
+template <typename Real>
+class IntpBilinear2
+{
+public:
+    // Construction.
+    IntpBilinear2(int xBound, int yBound, Real xMin, Real xSpacing,
+        Real yMin, Real ySpacing, Real const* F);
+
+    // Member access.
+    inline int GetXBound() const;
+    inline int GetYBound() const;
+    inline int GetQuantity() const;
+    inline Real const* GetF() const;
+    inline Real GetXMin() const;
+    inline Real GetXMax() const;
+    inline Real GetXSpacing() const;
+    inline Real GetYMin() const;
+    inline Real GetYMax() const;
+    inline Real GetYSpacing() const;
+
+    // Evaluate the function and its derivatives.  The functions clamp the
+    // inputs to xmin <= x <= xmax and ymin <= y <= ymax.  The first operator
+    // is for function evaluation.  The second operator is for function or
+    // derivative evaluations.  The xOrder argument is the order of the
+    // x-derivative and the yOrder argument is the order of the y-derivative.
+    // Both orders are zero to get the function value itself.
+    Real operator()(Real x, Real y) const;
+    Real operator()(int xOrder, int yOrder, Real x, Real y) const;
+
+private:
+    int mXBound, mYBound, mQuantity;
+    Real mXMin, mXMax, mXSpacing, mInvXSpacing;
+    Real mYMin, mYMax, mYSpacing, mInvYSpacing;
+    Real const* mF;
+    Real mBlend[2][2];
+};
+
+
+template <typename Real>
+IntpBilinear2<Real>::IntpBilinear2(int xBound, int yBound, Real xMin,
+    Real xSpacing, Real yMin, Real ySpacing, Real const* F)
+    :
+    mXBound(xBound),
+    mYBound(yBound),
+    mQuantity(xBound * yBound),
+    mXMin(xMin),
+    mXSpacing(xSpacing),
+    mYMin(yMin),
+    mYSpacing(ySpacing),
+    mF(F)
+{
+    // At least a 2x2 block of data points are needed.
+    LogAssert(mXBound >= 2 && mYBound >= 2 && mF, "Invalid input.");
+    LogAssert(mXSpacing > (Real)0 && mYSpacing > (Real)0, "Invalid input.");
+
+    mXMax = mXMin + mXSpacing * static_cast<Real>(mXBound - 1);
+    mInvXSpacing = ((Real)1) / mXSpacing;
+    mYMax = mYMin + mYSpacing * static_cast<Real>(mYBound - 1);
+    mInvYSpacing = ((Real)1) / mYSpacing;
+
+    mBlend[0][0] = (Real)1;
+    mBlend[0][1] = (Real)-1;
+    mBlend[1][0] = (Real)0;
+    mBlend[1][1] = (Real)1;
+}
+
+template <typename Real> inline
+int IntpBilinear2<Real>::GetXBound() const
+{
+    return mXBound;
+}
+
+template <typename Real> inline
+int IntpBilinear2<Real>::GetYBound() const
+{
+    return mYBound;
+}
+
+template <typename Real> inline
+int IntpBilinear2<Real>::GetQuantity() const
+{
+    return mQuantity;
+}
+
+template <typename Real> inline
+Real const* IntpBilinear2<Real>::GetF() const
+{
+    return mF;
+}
+
+template <typename Real> inline
+Real IntpBilinear2<Real>::GetXMin() const
+{
+    return mXMin;
+}
+
+template <typename Real> inline
+Real IntpBilinear2<Real>::GetXMax() const
+{
+    return mXMax;
+}
+
+template <typename Real> inline
+Real IntpBilinear2<Real>::GetXSpacing() const
+{
+    return mXSpacing;
+}
+
+template <typename Real> inline
+Real IntpBilinear2<Real>::GetYMin() const
+{
+    return mYMin;
+}
+
+template <typename Real> inline
+Real IntpBilinear2<Real>::GetYMax() const
+{
+    return mYMax;
+}
+
+template <typename Real> inline
+Real IntpBilinear2<Real>::GetYSpacing() const
+{
+    return mYSpacing;
+}
+
+template <typename Real>
+Real IntpBilinear2<Real>::operator()(Real x, Real y) const
+{
+    // Compute x-index and clamp to image.
+    Real xIndex = (x - mXMin) * mInvXSpacing;
+    int ix = static_cast<int>(xIndex);
+    if (ix < 0)
+    {
+        ix = 0;
+    }
+    else if (ix >= mXBound)
+    {
+        ix = mXBound - 1;
+    }
+
+    // Compute y-index and clamp to image.
+    Real yIndex = (y - mYMin) * mInvYSpacing;
+    int iy = static_cast<int>(yIndex);
+    if (iy < 0)
+    {
+        iy = 0;
+    }
+    else if (iy >= mYBound)
+    {
+        iy = mYBound - 1;
+    }
+
+    Real U[2];
+    U[0] = (Real)1;
+    U[1] = xIndex - ix;
+
+    Real V[2];
+    V[0] = (Real)1.0;
+    V[1] = yIndex - iy;
+
+    // Compute P = M*U and Q = M*V.
+    Real P[2], Q[2];
+    for (int row = 0; row < 2; ++row)
+    {
+        P[row] = (Real)0;
+        Q[row] = (Real)0;
+        for (int col = 0; col < 2; ++col)
+        {
+            P[row] += mBlend[row][col] * U[col];
+            Q[row] += mBlend[row][col] * V[col];
+        }
+    }
+
+    // Compute (M*U)^t D (M*V) where D is the 2x2 subimage containing (x,y).
+    Real result = (Real)0;
+    for (int row = 0; row < 2; ++row)
+    {
+        int yClamp = iy + row;
+        if (yClamp >= mYBound)
+        {
+            yClamp = mYBound - 1;
+        }
+
+        for (int col = 0; col < 2; ++col)
+        {
+            int xClamp = ix + col;
+            if (xClamp >= mXBound)
+            {
+                xClamp = mXBound - 1;
+            }
+
+            result += P[col] * Q[row] * mF[xClamp + mXBound * yClamp];
+        }
+    }
+
+    return result;
+}
+
+template <typename Real>
+Real IntpBilinear2<Real>::operator()(int xOrder, int yOrder, Real x, Real y)
+const
+{
+    // Compute x-index and clamp to image.
+    Real xIndex = (x - mXMin) * mInvXSpacing;
+    int ix = static_cast<int>(xIndex);
+    if (ix < 0)
+    {
+        ix = 0;
+    }
+    else if (ix >= mXBound)
+    {
+        ix = mXBound - 1;
+    }
+
+    // Compute y-index and clamp to image.
+    Real yIndex = (y - mYMin) * mInvYSpacing;
+    int iy = static_cast<int>(yIndex);
+    if (iy < 0)
+    {
+        iy = 0;
+    }
+    else if (iy >= mYBound)
+    {
+        iy = mYBound - 1;
+    }
+
+    Real U[2], dx, xMult;
+    switch (xOrder)
+    {
+    case 0:
+        dx = xIndex - ix;
+        U[0] = (Real)1;
+        U[1] = dx;
+        xMult = (Real)1;
+        break;
+    case 1:
+        dx = xIndex - ix;
+        U[0] = (Real)0;
+        U[1] = (Real)1;
+        xMult = mInvXSpacing;
+        break;
+    default:
+        return (Real)0;
+    }
+
+    Real V[2], dy, yMult;
+    switch (yOrder)
+    {
+    case 0:
+        dy = yIndex - iy;
+        V[0] = (Real)1;
+        V[1] = dy;
+        yMult = (Real)1;
+        break;
+    case 1:
+        dy = yIndex - iy;
+        V[0] = (Real)0;
+        V[1] = (Real)1;
+        yMult = mInvYSpacing;
+        break;
+    default:
+        return (Real)0;
+    }
+
+    // Compute P = M*U and Q = M*V.
+    Real P[2], Q[2];
+    for (int row = 0; row < 2; ++row)
+    {
+        P[row] = (Real)0;
+        Q[row] = (Real)0;
+        for (int col = 0; col < 2; ++col)
+        {
+            P[row] += mBlend[row][col] * U[col];
+            Q[row] += mBlend[row][col] * V[col];
+        }
+    }
+
+    // Compute (M*U)^t D (M*V) where D is the 2x2 subimage containing (x,y).
+    Real result = (Real)0;
+    for (int row = 0; row < 2; ++row)
+    {
+        int yClamp = iy + row;
+        if (yClamp >= mYBound)
+        {
+            yClamp = mYBound - 1;
+        }
+
+        for (int col = 0; col < 2; ++col)
+        {
+            int xClamp = ix + col;
+            if (xClamp >= mXBound)
+            {
+                xClamp = mXBound - 1;
+            }
+
+            result += P[col] * Q[row] * mF[xClamp + mXBound * yClamp];
+        }
+    }
+    result *= xMult * yMult;
+
+    return result;
+}
+
+
+}
diff --git a/Engine/Plugins/Experimental/GeometryProcessing/Source/GeometryAlgorithms/Private/GTEngine/Mathematics/GteIntpLinearNonuniform2.h b/Engine/Plugins/Experimental/GeometryProcessing/Source/GeometryAlgorithms/Private/GTEngine/Mathematics/GteIntpLinearNonuniform2.h
new file mode 100644
--- /dev/null
+++ b/Engine/Plugins/Experimental/GeometryProcessing/Source/GeometryAlgorithms/Private/GTEngine/Mathematics/GteIntpLinearNonuniform2.h
@@ -0,0 +1,84 @@
+// David Eberly, Geometric Tools, Redmond WA 98052
+// Copyright (c) 1998-2019
+// Distributed under the Boost Software License, Version 1.0.
+// http://www.boost.org/LICENSE_1_0.txt
+// http://www.geometrictools.com/License/Boost/LICENSE_1_0.txt
+// File Version: 3.0.0 (2016/06/19)
+
+#pragma once
+
+// Linear interpolation of a network of triangles whose vertices are of the
+// form (x,y,f(x,y)).  The function samples are F[i] and represent
+// f(x[i],y[i]), where i is the index of the input vertex (x[i],y[i]) to
+// Delaunay2.
+//
+// The TriangleMesh interface must support the following:
+//   bool GetIndices(int, std::array<int, 3>&) const;
+//   bool GetBarycentrics(int, Vector2<Real> const&,
+//       std::array<Real, 3>&) const;
+//   int GetContainingTriangle(Vector2<Real> const&) const;
+
+#include <GTEngine/LowLevel/GteLogger.h>
+#include <GTEngine/Mathematics/GteVector2.h>
+
+namespace gte
+{
+
+template <typename Real, typename TriangleMesh>
+class IntpLinearNonuniform2
+{
+public:
+    // Construction.
+    IntpLinearNonuniform2(TriangleMesh const& mesh, Real const* F);
+
+    // Linear interpolation.  The return value is 'true' if and only if the
+    // input point P is in the convex hull of the input vertices, in which
+    // case the interpolation is valid.
+    bool operator()(Vector2<Real> const& P, Real& F) const;
+
+private:
+    TriangleMesh const* mMesh;
+    Real const* mF;
+};
+
+
+template <typename Real, typename TriangleMesh>
+IntpLinearNonuniform2<Real, TriangleMesh>::IntpLinearNonuniform2(
+    TriangleMesh const& mesh, Real const* F)
+    :
+    mMesh(&mesh),
+    mF(F)
+{
+    LogAssert(mF != nullptr, "Invalid input.");
+}
+
+template <typename Real, typename TriangleMesh>
+bool IntpLinearNonuniform2<Real, TriangleMesh>::operator()(
+    Vector2<Real> const& P, Real& F) const
+{
+    int t = mMesh->GetContainingTriangle(P);
+    if (t == -1)
+    {
+        // The point is outside the triangulation.
+        return false;
+    }
+
+    // Get the barycentric coordinates of P with respect to the triangle,
+    // P = b0*V0 + b1*V1 + b2*V2, where b0 + b1 + b2 = 1.
+    std::array<Real, 3> bary;
+    if (!mMesh->GetBarycentrics(t, P, bary))
+    {
+        LogWarning("P is in a needle-like or degenerate triangle.");
+        return false;
+    }
+
+    // The result is a barycentric combination of function values.
+    std::array<int, 3> indices;
+    mMesh->GetIndices(t, indices);
+    F = bary[0] * mF[indices[0]] + bary[1] * mF[indices[1]] +
+        bary[2] * mF[indices[2]];
+    return true;
+}
+
+
+}
diff --git a/Engine/Plugins/Experimental/GeometryProcessing/Source/GeometryAlgorithms/Private/GTEngine/Mathematics/GteIntpLinearNonuniform3.h b/Engine/Plugins/Experimental/GeometryProcessing/Source/GeometryAlgorithms/Private/GTEngine/Mathematics/GteIntpLinearNonuniform3.h
new file mode 100644
--- /dev/null
+++ b/Engine/Plugins/Experimental/GeometryProcessing/Source/GeometryAlgorithms/Private/GTEngine/Mathematics/GteIntpLinearNonuniform3.h
@@ -0,0 +1,83 @@
+// David Eberly, Geometric Tools, Redmond WA 98052
+// Copyright (c) 1998-2019
+// Distributed under the Boost Software License, Version 1.0.
+// http://www.boost.org/LICENSE_1_0.txt
+// http://www.geometrictools.com/License/Boost/LICENSE_1_0.txt
+// File Version: 3.0.0 (2016/06/19)
+
+#pragma once
+
+// Linear interpolation of a network of triangles whose vertices are of the
+// form (x,y,z,f(x,y,z)).  The function samples are F[i] and represent
+// f(x[i],y[i],z[i]), where i is the index of the input vertex
+// (x[i],y[i],z[i]) to Delaunay3.
+//
+// The TetrahedronMesh interface must support the following:
+//   int GetContainingTetrahedron(Vector3<Real> const&) const;
+//   bool GetIndices(int, std::array<int, 4>&) const;
+//   bool GetBarycentrics(int, Vector3<Real> const&, Real[4]) const;
+
+#include <GTEngine/LowLevel/GteLogger.h>
+#include <GTEngine/Mathematics/GteVector3.h>
+
+namespace gte
+{
+
+template <typename Real, typename TetrahedronMesh>
+class IntpLinearNonuniform3
+{
+public:
+    // Construction.
+    IntpLinearNonuniform3(TetrahedronMesh const& mesh, Real const* F);
+
+    // Linear interpolation.  The return value is 'true' if and only if the
+    // input point is in the convex hull of the input vertices, in which case
+    // the interpolation is valid.
+    bool operator()(Vector3<Real> const& P, Real& F) const;
+
+private:
+    TetrahedronMesh const* mMesh;
+    Real const* mF;
+};
+
+
+template <typename Real, typename TetrahedronMesh>
+IntpLinearNonuniform3<Real, TetrahedronMesh>::IntpLinearNonuniform3(
+    TetrahedronMesh const& mesh, Real const* F)
+    :
+    mMesh(&mesh),
+    mF(F)
+{
+    LogAssert(mF != nullptr, "Invalid input.");
+}
+
+template <typename Real, typename TetrahedronMesh>
+bool IntpLinearNonuniform3<Real, TetrahedronMesh>::operator()(
+    Vector3<Real> const& P, Real& F) const
+{
+    int t = mMesh->GetContainingTetrahedron(P);
+    if (t == -1)
+    {
+        // The point is outside the tetrahedralization.
+        return false;
+    }
+
+    // Get the barycentric coordinates of P with respect to the tetrahedron,
+    // P = b0*V0 + b1*V1 + b2*V2 + b3*V3, where b0 + b1 + b2 + b3 = 1.
+    std::array<Real, 4> bary;
+    if (!mMesh->GetBarycentrics(t, P, bary))
+    {
+        LogWarning("P is in a needle-like, flat, or degenerate tetrahedron.");
+        return false;
+    }
+
+    // The result is a barycentric combination of function values.
+    std::array<int, 4> indices;
+    mMesh->GetIndices(t, indices);
+    F = bary[0] * mF[indices[0]] + bary[1] * mF[indices[1]] +
+        bary[2] * mF[indices[2]] + bary[3] * mF[indices[4]];
+    return true;
+}
+
+
+}
