move GeometryProcessing plugin from NotForLicensees to Experimental

diff --git a/Engine/Plugins/Experimental/GeometryProcessing/Source/GeometryAlgorithms/Private/GTEngine/Mathematics/GteIntpQuadraticNonuniform2.h b/Engine/Plugins/Experimental/GeometryProcessing/Source/GeometryAlgorithms/Private/GTEngine/Mathematics/GteIntpQuadraticNonuniform2.h
new file mode 100644
--- /dev/null
+++ b/Engine/Plugins/Experimental/GeometryProcessing/Source/GeometryAlgorithms/Private/GTEngine/Mathematics/GteIntpQuadraticNonuniform2.h
@@ -0,0 +1,488 @@
+// David Eberly, Geometric Tools, Redmond WA 98052
+// Copyright (c) 1998-2019
+// Distributed under the Boost Software License, Version 1.0.
+// http://www.boost.org/LICENSE_1_0.txt
+// http://www.geometrictools.com/License/Boost/LICENSE_1_0.txt
+// File Version: 3.0.0 (2016/06/19)
+
+#pragma once
+
+#include <GTEngine/Mathematics/GteDelaunay2.h>
+#include <GTEngine/Mathematics/GteContScribeCircle2.h>
+#include <GTEngine/Mathematics/GteDistPointAlignedBox.h>
+
+// Quadratic interpolation of a network of triangles whose vertices are of
+// the form (x,y,f(x,y)).  This code is an implementation of the algorithm
+// found in
+//
+//   Zoltan J. Cendes and Steven H. Wong,
+//   C1 quadratic interpolation over arbitrary point sets,
+//   IEEE Computer Graphics & Applications,
+//   pp. 8-16, 1987
+//
+// The TriangleMesh interface must support the following:
+//   int GetNumVertices() const;
+//   int GetNumTriangles() const;
+//   Vector2<Real> const* GetVertices() const;
+//   int const* GetIndices() const;
+//   bool GetVertices(int, std::array<Vector2<Real>, 3>&) const;
+//   bool GetIndices(int, std::array<int, 3>&) const;
+//   bool GetAdjacencies(int, std::array<int, 3>&) const;
+//   bool GetBarycentrics(int, Vector2<Real> const&,
+//      std::array<Real, 3>&) const;
+//   int GetContainingTriangle(Vector2<Real> const&) const;
+
+namespace gte
+{
+
+template <typename Real, typename TriangleMesh>
+class IntpQuadraticNonuniform2
+{
+public:
+    // Construction.
+    //
+    // The first constructor requires only F and a measure of the rate of
+    // change of the function values relative to changes in the spatial
+    // variables.  The df/dx and df/dy values are estimated at the sample
+    // points using mesh normals and spatialDelta.
+    //
+    // The second constructor requires you to specify function values F and
+    // first-order partial derivative values df/dx and df/dy.
+
+    IntpQuadraticNonuniform2(TriangleMesh const& mesh, Real const* F,
+        Real spatialDelta);
+
+    IntpQuadraticNonuniform2(TriangleMesh const& mesh, Real const* F,
+        Real const* FX, Real const* FY);
+
+    // Quadratic interpolation.  The return value is 'true' if and only if the
+    // input point is in the convex hull of the input vertices, in which case
+    // the interpolation is valid.
+    bool operator()(Vector2<Real> const & P, Real& F, Real& FX, Real& FY)
+        const;
+
+private:
+    void EstimateDerivatives(Real spatialDelta);
+    void ProcessTriangles();
+    void ComputeCrossEdgeIntersections(int t);
+    void ComputeCoefficients(int t);
+
+    class TriangleData
+    {
+    public:
+        Vector2<Real> center;
+        Vector2<Real> intersect[3];
+        Real coeff[19];
+    };
+
+    class Jet
+    {
+    public:
+        Real F, FX, FY;
+    };
+
+    TriangleMesh const* mMesh;
+    Real const* mF;
+    Real const* mFX;
+    Real const* mFY;
+    std::vector<Real> mFXStorage;
+    std::vector<Real> mFYStorage;
+    std::vector<TriangleData> mTData;
+};
+
+
+template <typename Real, typename TriangleMesh>
+IntpQuadraticNonuniform2<Real, TriangleMesh>::IntpQuadraticNonuniform2(
+    TriangleMesh const& mesh, Real const* F, Real spatialDelta)
+    :
+    mMesh(&mesh),
+    mF(F),
+    mFX(nullptr),
+    mFY(nullptr)
+{
+    EstimateDerivatives(spatialDelta);
+    ProcessTriangles();
+}
+
+template <typename Real, typename TriangleMesh>
+IntpQuadraticNonuniform2<Real, TriangleMesh>::IntpQuadraticNonuniform2(
+    TriangleMesh const& mesh, Real const* F, Real const* FX, Real const* FY)
+    :
+    mMesh(&mesh),
+    mF(F),
+    mFX(FX),
+    mFY(FY)
+{
+    ProcessTriangles();
+}
+
+template <typename Real, typename TriangleMesh>
+bool IntpQuadraticNonuniform2<Real, TriangleMesh>::operator()(
+    Vector2<Real> const& P, Real& F, Real& FX, Real& FY) const
+{
+    int t = mMesh->GetContainingTriangle(P);
+    if (t == -1)
+    {
+        // The point is outside the triangulation.
+        return false;
+    }
+
+    // Get the vertices of the triangle.
+    std::array<Vector2<Real>, 3> V;
+    mMesh->GetVertices(t, V);
+
+    // Get the additional information for the triangle.
+    TriangleData const& tData = mTData[t];
+
+    // Determine which of the six subtriangles contains the target point.
+    // Theoretically, P must be in one of these subtriangles.
+    Vector2<Real> sub0 = tData.center;
+    Vector2<Real> sub1;
+    Vector2<Real> sub2 = tData.intersect[2];
+    Vector3<Real> bary;
+    int index;
+
+    Real const zero = (Real)0, one = (Real)1;
+    AlignedBox3<Real> barybox({ zero, zero, zero }, { one, one, one });
+    DCPQuery<Real, Vector3<Real>, AlignedBox3<Real>> pbQuery;
+    int minIndex = 0;
+    Real minDistance = (Real)-1;
+    Vector3<Real> minBary;
+    Vector2<Real> minSub0, minSub1, minSub2;
+
+    for (index = 1; index <= 6; ++index)
+    {
+        sub1 = sub2;
+        if (index % 2)
+        {
+            sub2 = V[index / 2];
+        }
+        else
+        {
+            sub2 = tData.intersect[index / 2 - 1];
+        }
+
+        bool valid = ComputeBarycentrics(P, sub0, sub1, sub2, &bary[0]);
+        if (valid
+            && zero <= bary[0] && bary[0] <= one
+            && zero <= bary[1] && bary[1] <= one
+            && zero <= bary[2] && bary[2] <= one)
+        {
+            // P is in triangle <Sub0,Sub1,Sub2>
+            break;
+        }
+
+        // When computing with floating-point arithmetic, rounding errors
+        // can cause us to reach this code when, theoretically, the point
+        // is in the subtriangle.  Keep track of the (b0,b1,b2) that is
+        // closest to the barycentric cube [0,1]^3 and choose the triangle
+        // corresponding to it when all 6 tests previously fail.
+        Real distance = pbQuery(bary, barybox).distance;
+        if (minIndex == 0 || distance < minDistance)
+        {
+            minDistance = distance;
+            minIndex = index;
+            minBary = bary;
+            minSub0 = sub0;
+            minSub1 = sub1;
+            minSub2 = sub2;
+        }
+    }
+
+    // If the subtriangle was not found, rounding errors caused problems.
+    // Choose the barycentric point closest to the box.
+    if (index > 6)
+    {
+        index = minIndex;
+        bary = minBary;
+        sub0 = minSub0;
+        sub1 = minSub1;
+        sub2 = minSub2;
+    }
+
+    // Fetch Bezier control points.
+    Real bez[6] =
+    {
+        tData.coeff[0],
+        tData.coeff[12 + index],
+        tData.coeff[13 + (index % 6)],
+        tData.coeff[index],
+        tData.coeff[6 + index],
+        tData.coeff[1 + (index % 6)]
+    };
+
+    // Evaluate Bezier quadratic.
+    F = bary[0] * (bez[0] * bary[0] + bez[1] * bary[1] + bez[2] * bary[2]) +
+        bary[1] * (bez[1] * bary[0] + bez[3] * bary[1] + bez[4] * bary[2]) +
+        bary[2] * (bez[2] * bary[0] + bez[4] * bary[1] + bez[5] * bary[2]);
+
+    // Evaluate barycentric derivatives of F.
+    Real FU = ((Real)2)*(bez[0] * bary[0] + bez[1] * bary[1] +
+        bez[2] * bary[2]);
+    Real FV = ((Real)2)*(bez[1] * bary[0] + bez[3] * bary[1] +
+        bez[4] * bary[2]);
+    Real FW = ((Real)2)*(bez[2] * bary[0] + bez[4] * bary[1] +
+        bez[5] * bary[2]);
+    Real duw = FU - FW;
+    Real dvw = FV - FW;
+
+    // Convert back to (x,y) coordinates.
+    Real m00 = sub0[0] - sub2[0];
+    Real m10 = sub0[1] - sub2[1];
+    Real m01 = sub1[0] - sub2[0];
+    Real m11 = sub1[1] - sub2[1];
+    Real inv = ((Real)1) / (m00*m11 - m10*m01);
+
+    FX = inv*(m11*duw - m10*dvw);
+    FY = inv*(m00*dvw - m01*duw);
+    return true;
+}
+
+template <typename Real, typename TriangleMesh>
+void IntpQuadraticNonuniform2<Real, TriangleMesh>::EstimateDerivatives(
+    Real spatialDelta)
+{
+    int numVertices = mMesh->GetNumVertices();
+    Vector2<Real> const* vertices = mMesh->GetVertices();
+    int numTriangles = mMesh->GetNumTriangles();
+    int const* indices = mMesh->GetIndices();
+
+    mFXStorage.resize(numVertices);
+    mFYStorage.resize(numVertices);
+    std::vector<Real> FZ(numVertices);
+    std::fill(mFXStorage.begin(), mFXStorage.end(), (Real)0);
+    std::fill(mFYStorage.begin(), mFYStorage.end(), (Real)0);
+    std::fill(FZ.begin(), FZ.end(), (Real)0);
+
+    mFX = &mFXStorage[0];
+    mFY = &mFYStorage[0];
+
+    // Accumulate normals at spatial locations (averaging process).
+    for (int t = 0; t < numTriangles; ++t)
+    {
+        // Get three vertices of triangle.
+        int v0 = *indices++;
+        int v1 = *indices++;
+        int v2 = *indices++;
+
+        // Compute normal vector of triangle (with positive z-component).
+        Real dx1 = vertices[v1][0] - vertices[v0][0];
+        Real dy1 = vertices[v1][1] - vertices[v0][1];
+        Real dz1 = mF[v1] - mF[v0];
+        Real dx2 = vertices[v2][0] - vertices[v0][0];
+        Real dy2 = vertices[v2][1] - vertices[v0][1];
+        Real dz2 = mF[v2] - mF[v0];
+        Real nx = dy1*dz2 - dy2*dz1;
+        Real ny = dz1*dx2 - dz2*dx1;
+        Real nz = dx1*dy2 - dx2*dy1;
+        if (nz < (Real)0)
+        {
+            nx = -nx;
+            ny = -ny;
+            nz = -nz;
+        }
+
+        mFXStorage[v0] += nx;  mFYStorage[v0] += ny;  FZ[v0] += nz;
+        mFXStorage[v1] += nx;  mFYStorage[v1] += ny;  FZ[v1] += nz;
+        mFXStorage[v2] += nx;  mFYStorage[v2] += ny;  FZ[v2] += nz;
+    }
+
+    // Scale the normals to form (x,y,-1).
+    for (int i = 0; i < numVertices; ++i)
+    {
+        if (FZ[i] != (Real)0)
+        {
+            Real inv = -spatialDelta / FZ[i];
+            mFXStorage[i] *= inv;
+            mFYStorage[i] *= inv;
+        }
+        else
+        {
+            mFXStorage[i] = (Real)0;
+            mFYStorage[i] = (Real)0;
+        }
+    }
+}
+
+template <typename Real, typename TriangleMesh>
+void IntpQuadraticNonuniform2<Real, TriangleMesh>::ProcessTriangles()
+{
+    // Add degenerate triangles to boundary triangles so that interpolation
+    // at the boundary can be treated in the same way as interpolation in
+    // the interior.
+
+    // Compute centers of inscribed circles for triangles.
+    Vector2<Real> const* vertices = mMesh->GetVertices();
+    int numTriangles = mMesh->GetNumTriangles();
+    int const* indices = mMesh->GetIndices();
+    mTData.resize(numTriangles);
+    int t;
+    for (t = 0; t < numTriangles; ++t)
+    {
+        int v0 = *indices++;
+        int v1 = *indices++;
+        int v2 = *indices++;
+        Circle2<Real> circle;
+        Inscribe(vertices[v0], vertices[v1], vertices[v2], circle);
+        mTData[t].center = circle.center;
+    }
+
+    // Compute cross-edge intersections.
+    for (t = 0; t < numTriangles; ++t)
+    {
+        ComputeCrossEdgeIntersections(t);
+    }
+
+    // Compute Bezier coefficients.
+    for (t = 0; t < numTriangles; ++t)
+    {
+        ComputeCoefficients(t);
+    }
+}
+
+template <typename Real, typename TriangleMesh>
+void IntpQuadraticNonuniform2<Real, TriangleMesh>::
+ComputeCrossEdgeIntersections(int t)
+{
+    // Get the vertices of the triangle.
+    std::array<Vector2<Real>, 3> V;
+    mMesh->GetVertices(t, V);
+
+    // Get the centers of adjacent triangles.
+    TriangleData& tData = mTData[t];
+    std::array<int, 3> adjacencies;
+    mMesh->GetAdjacencies(t, adjacencies);
+    for (int j0 = 2, j1 = 0; j1 < 3; j0 = j1++)
+    {
+        int a = adjacencies[j0];
+        if (a >= 0)
+        {
+            // Get center of adjacent triangle's inscribing circle.
+            Vector2<Real> U = mTData[a].center;
+            Real m00 = V[j0][1] - V[j1][1];
+            Real m01 = V[j1][0] - V[j0][0];
+            Real m10 = tData.center[1] - U[1];
+            Real m11 = U[0] - tData.center[0];
+            Real r0 = m00*V[j0][0] + m01*V[j0][1];
+            Real r1 = m10*tData.center[0] + m11*tData.center[1];
+            Real invDet = ((Real)1) / (m00*m11 - m01*m10);
+            tData.intersect[j0][0] = (m11*r0 - m01*r1)*invDet;
+            tData.intersect[j0][1] = (m00*r1 - m10*r0)*invDet;
+        }
+        else
+        {
+            // No adjacent triangle, use center of edge.
+            tData.intersect[j0] = ((Real)0.5)*(V[j0] + V[j1]);
+        }
+    }
+}
+
+template <typename Real, typename TriangleMesh>
+void IntpQuadraticNonuniform2<Real, TriangleMesh>::ComputeCoefficients(int t)
+{
+    // Get the vertices of the triangle.
+    std::array<Vector2<Real>, 3> V;
+    mMesh->GetVertices(t, V);
+
+    // Get the additional information for the triangle.
+    TriangleData& tData = mTData[t];
+
+    // Get the sample data at main triangle vertices.
+    std::array<int, 3> indices;
+    mMesh->GetIndices(t, indices);
+    Jet jet[3];
+    for (int j = 0; j < 3; ++j)
+    {
+        int k = indices[j];
+        jet[j].F = mF[k];
+        jet[j].FX = mFX[k];
+        jet[j].FY = mFY[k];
+    }
+
+    // Get centers of adjacent triangles.
+    std::array<int, 3> adjacencies;
+    mMesh->GetAdjacencies(t, adjacencies);
+    Vector2<Real> U[3];
+    for (int j0 = 2, j1 = 0; j1 < 3; j0 = j1++)
+    {
+        int a = adjacencies[j0];
+        if (a >= 0)
+        {
+            // Get center of adjacent triangle's circumscribing circle.
+            U[j0] = mTData[a].center;
+        }
+        else
+        {
+            // No adjacent triangle, use center of edge.
+            U[j0] = ((Real)0.5)*(V[j0] + V[j1]);
+        }
+    }
+
+    // Compute intermediate terms.
+    std::array<Real, 3> cenT, cen0, cen1, cen2;
+    mMesh->GetBarycentrics(t, tData.center, cenT);
+    mMesh->GetBarycentrics(t, U[0], cen0);
+    mMesh->GetBarycentrics(t, U[1], cen1);
+    mMesh->GetBarycentrics(t, U[2], cen2);
+
+    Real alpha = (cenT[1] * cen1[0] - cenT[0] * cen1[1])
+        / (cen1[0] - cenT[0]);
+    Real beta = (cenT[2] * cen2[1] - cenT[1] * cen2[2])
+        / (cen2[1] - cenT[1]);
+    Real gamma = (cenT[0] * cen0[2] - cenT[2] * cen0[0])
+        / (cen0[2] - cenT[2]);
+    Real oneMinusAlpha = (Real)1 - alpha;
+    Real oneMinusBeta = (Real)1 - beta;
+    Real oneMinusGamma = (Real)1 - gamma;
+
+    Real tmp, A[9], B[9];
+
+    tmp = cenT[0] * V[0][0] + cenT[1] * V[1][0] + cenT[2] * V[2][0];
+    A[0] = ((Real)0.5)*(tmp - V[0][0]);
+    A[1] = ((Real)0.5)*(tmp - V[1][0]);
+    A[2] = ((Real)0.5)*(tmp - V[2][0]);
+    A[3] = ((Real)0.5)*beta*(V[2][0] - V[0][0]);
+    A[4] = ((Real)0.5)*oneMinusGamma*(V[1][0] - V[0][0]);
+    A[5] = ((Real)0.5)*gamma*(V[0][0] - V[1][0]);
+    A[6] = ((Real)0.5)*oneMinusAlpha*(V[2][0] - V[1][0]);
+    A[7] = ((Real)0.5)*alpha*(V[1][0] - V[2][0]);
+    A[8] = ((Real)0.5)*oneMinusBeta*(V[0][0] - V[2][0]);
+
+    tmp = cenT[0] * V[0][1] + cenT[1] * V[1][1] + cenT[2] * V[2][1];
+    B[0] = ((Real)0.5)*(tmp - V[0][1]);
+    B[1] = ((Real)0.5)*(tmp - V[1][1]);
+    B[2] = ((Real)0.5)*(tmp - V[2][1]);
+    B[3] = ((Real)0.5)*beta*(V[2][1] - V[0][1]);
+    B[4] = ((Real)0.5)*oneMinusGamma*(V[1][1] - V[0][1]);
+    B[5] = ((Real)0.5)*gamma*(V[0][1] - V[1][1]);
+    B[6] = ((Real)0.5)*oneMinusAlpha*(V[2][1] - V[1][1]);
+    B[7] = ((Real)0.5)*alpha*(V[1][1] - V[2][1]);
+    B[8] = ((Real)0.5)*oneMinusBeta*(V[0][1] - V[2][1]);
+
+    // Compute Bezier coefficients.
+    tData.coeff[2] = jet[0].F;
+    tData.coeff[4] = jet[1].F;
+    tData.coeff[6] = jet[2].F;
+
+    tData.coeff[14] = jet[0].F + A[0] * jet[0].FX + B[0] * jet[0].FY;
+    tData.coeff[7] = jet[0].F + A[3] * jet[0].FX + B[3] * jet[0].FY;
+    tData.coeff[8] = jet[0].F + A[4] * jet[0].FX + B[4] * jet[0].FY;
+    tData.coeff[16] = jet[1].F + A[1] * jet[1].FX + B[1] * jet[1].FY;
+    tData.coeff[9] = jet[1].F + A[5] * jet[1].FX + B[5] * jet[1].FY;
+    tData.coeff[10] = jet[1].F + A[6] * jet[1].FX + B[6] * jet[1].FY;
+    tData.coeff[18] = jet[2].F + A[2] * jet[2].FX + B[2] * jet[2].FY;
+    tData.coeff[11] = jet[2].F + A[7] * jet[2].FX + B[7] * jet[2].FY;
+    tData.coeff[12] = jet[2].F + A[8] * jet[2].FX + B[8] * jet[2].FY;
+
+    tData.coeff[5] = alpha*tData.coeff[10] + oneMinusAlpha*tData.coeff[11];
+    tData.coeff[17] = alpha*tData.coeff[16] + oneMinusAlpha*tData.coeff[18];
+    tData.coeff[1] = beta*tData.coeff[12] + oneMinusBeta*tData.coeff[7];
+    tData.coeff[13] = beta*tData.coeff[18] + oneMinusBeta*tData.coeff[14];
+    tData.coeff[3] = gamma*tData.coeff[8] + oneMinusGamma*tData.coeff[9];
+    tData.coeff[15] = gamma*tData.coeff[14] + oneMinusGamma*tData.coeff[16];
+    tData.coeff[0] = cenT[0] * tData.coeff[14] + cenT[1] * tData.coeff[16] +
+        cenT[2] * tData.coeff[18];
+}
+
+
+}
diff --git a/Engine/Plugins/Experimental/GeometryProcessing/Source/GeometryAlgorithms/Private/GTEngine/Mathematics/GteIntpSphere2.h b/Engine/Plugins/Experimental/GeometryProcessing/Source/GeometryAlgorithms/Private/GTEngine/Mathematics/GteIntpSphere2.h
new file mode 100644
--- /dev/null
+++ b/Engine/Plugins/Experimental/GeometryProcessing/Source/GeometryAlgorithms/Private/GTEngine/Mathematics/GteIntpSphere2.h
@@ -0,0 +1,137 @@
+// David Eberly, Geometric Tools, Redmond WA 98052
+// Copyright (c) 1998-2019
+// Distributed under the Boost Software License, Version 1.0.
+// http://www.boost.org/LICENSE_1_0.txt
+// http://www.geometrictools.com/License/Boost/LICENSE_1_0.txt
+// File Version: 3.0.1 (2018/10/05)
+
+#pragma once
+
+#include <GTEngine/Mathematics/GteDelaunay2Mesh.h>
+#include <GTEngine/Mathematics/GteMath.h>
+#include <GTEngine/Mathematics/GteIntpQuadraticNonuniform2.h>
+#include <memory>
+
+// Interpolation of a scalar-valued function defined on a sphere.  Although
+// the sphere lives in 3D, the interpolation is a 2D method whose input
+// points are angles (theta,phi) from spherical coordinates.  The domains of
+// the angles are -pi <= theta <= pi and 0 <= phi <= pi.
+
+namespace gte
+{
+
+template <typename InputType, typename ComputeType, typename RationalType>
+class IntpSphere2
+{
+public:
+    // Construction and destruction.  For complete spherical coverage, include
+    // the two antipodal (theta,phi) points (-pi,0,F(-pi,0)) and
+    // (-pi,pi,F(-pi,pi)) in the input data.  These correspond to the sphere
+    // poles x = 0, y = 0, and |z| = 1.
+    ~IntpSphere2();
+    IntpSphere2(int numPoints, InputType const* theta, InputType const* phi,
+        InputType const* F);
+
+    // Spherical coordinates are
+    //   x = cos(theta)*sin(phi)
+    //   y = sin(theta)*sin(phi)
+    //   z = cos(phi)
+    // for -pi <= theta <= pi, 0 <= phi <= pi.  The application can use this
+    // function to convert unit length vectors (x,y,z) to (theta,phi).
+    static void GetSphericalCoordinates(InputType x, InputType y, InputType z,
+        InputType& theta, InputType& phi);
+
+    // The return value is 'true' if and only if the input point is in the
+    // convex hull of the input (theta,pi) array, in which case the
+    // interpolation is valid.
+    bool operator()(InputType theta, InputType phi, InputType& F) const;
+
+private:
+    typedef Delaunay2Mesh<InputType, ComputeType, RationalType> TriangleMesh;
+
+    std::vector<Vector2<InputType>> mWrapAngles;
+    Delaunay2<InputType, ComputeType> mDelaunay;
+    TriangleMesh mMesh;
+    std::vector<InputType> mWrapF;
+    std::unique_ptr<IntpQuadraticNonuniform2<InputType, TriangleMesh>> mInterp;
+};
+
+
+template <typename InputType, typename ComputeType, typename RationalType>
+IntpSphere2<InputType, ComputeType, RationalType>::~IntpSphere2()
+{
+}
+
+template <typename InputType, typename ComputeType, typename RationalType>
+IntpSphere2<InputType, ComputeType, RationalType>::IntpSphere2(int numPoints,
+    InputType const* theta, InputType const* phi, InputType const* F)
+    :
+    mMesh(mDelaunay)
+{
+    // Copy the input data.  The larger arrays are used to support wrap-around
+    // in the Delaunay triangulation for the interpolator.
+    int totalPoints = 3 * numPoints;
+    mWrapAngles.resize(totalPoints);
+    mWrapF.resize(totalPoints);
+    for (int i = 0; i < numPoints; ++i)
+    {
+        mWrapAngles[i][0] = theta[i];
+        mWrapAngles[i][1] = phi[i];
+        mWrapF[i] = F[i];
+    }
+
+    // Use periodicity to get wrap-around in the Delaunay triangulation.
+    int i0 = 0, i1 = numPoints, i2 = 2 * numPoints;
+    for (/**/; i0 < numPoints; ++i0, ++i1, ++i2)
+    {
+        mWrapAngles[i1][0] = mWrapAngles[i0][0] + (InputType)GTE_C_TWO_PI;
+        mWrapAngles[i2][0] = mWrapAngles[i0][0] - (InputType)GTE_C_TWO_PI;
+        mWrapAngles[i1][1] = mWrapAngles[i0][1];
+        mWrapAngles[i2][1] = mWrapAngles[i0][1];
+        mWrapF[i1] = mWrapF[i0];
+        mWrapF[i2] = mWrapF[i0];
+    }
+
+    mDelaunay(totalPoints, &mWrapAngles[0], (ComputeType)0);
+    mInterp = std::make_unique<IntpQuadraticNonuniform2<InputType, TriangleMesh>>(
+        mMesh, &mWrapF[0], (InputType)1);
+}
+
+template <typename InputType, typename ComputeType, typename RationalType>
+void IntpSphere2<InputType, ComputeType, RationalType>::
+GetSphericalCoordinates(InputType x, InputType y, InputType z,
+InputType& theta, InputType& phi)
+{
+    // Assumes (x,y,z) is unit length.  Returns -pi <= theta <= pi and
+    // 0 <= phiAngle <= pi.
+
+    if (z < (InputType)1)
+    {
+        if (z > -(InputType)1)
+        {
+            theta = std::atan2(y, x);
+            phi = std::acos(z);
+        }
+        else
+        {
+            theta = -(InputType)GTE_C_PI;
+            phi = (InputType)GTE_C_PI;
+        }
+    }
+    else
+    {
+        theta = -(InputType)GTE_C_PI;
+        phi = (InputType)0;
+    }
+}
+
+template <typename InputType, typename ComputeType, typename RationalType>
+bool IntpSphere2<InputType, ComputeType, RationalType>::operator()(
+    InputType theta, InputType phi, InputType& F) const
+{
+    Vector2<InputType> angles{ theta, phi };
+    InputType thetaDeriv, phiDeriv;
+    return (*mInterp)(angles, F, thetaDeriv, phiDeriv);
+}
+
+}
diff --git a/Engine/Plugins/Experimental/GeometryProcessing/Source/GeometryAlgorithms/Private/GTEngine/Mathematics/GteIntpThinPlateSpline2.h b/Engine/Plugins/Experimental/GeometryProcessing/Source/GeometryAlgorithms/Private/GTEngine/Mathematics/GteIntpThinPlateSpline2.h
new file mode 100644
--- /dev/null
+++ b/Engine/Plugins/Experimental/GeometryProcessing/Source/GeometryAlgorithms/Private/GTEngine/Mathematics/GteIntpThinPlateSpline2.h
@@ -0,0 +1,295 @@
+// David Eberly, Geometric Tools, Redmond WA 98052
+// Copyright (c) 1998-2019
+// Distributed under the Boost Software License, Version 1.0.
+// http://www.boost.org/LICENSE_1_0.txt
+// http://www.geometrictools.com/License/Boost/LICENSE_1_0.txt
+// File Version: 3.0.1 (2018/10/05)
+
+#pragma once
+
+#include <GTEngine/Mathematics/GteGMatrix.h>
+#include <algorithm>
+#include <cmath>
+#include <limits>
+
+// WARNING.  The implementation allows you to transform the inputs (x,y) to
+// the unit square and perform the interpolation in that space.  The idea is
+// to keep the floating-point numbers to order 1 for numerical stability of
+// the algorithm.  The classical thin-plate spline algorithm does not include
+// this transformation.  The interpolation is invariant to translations and
+// rotations of (x,y) but not to scaling.
+
+namespace gte
+{
+
+template <typename Real>
+class IntpThinPlateSpline2
+{
+public:
+    // Construction.  Data points are (x,y,f(x,y)).  The smoothing parameter
+    // must be nonnegative.
+    IntpThinPlateSpline2(int numPoints, Real const* X, Real const* Y,
+        Real const* F, Real smooth, bool transformToUnitSquare);
+
+    // Check this after the constructor call to see whether the thin plate
+    // spline coefficients were successfully computed.  If so, then calls to
+    // operator()(Real,Real) will work properly.
+    inline bool IsInitialized() const;
+
+    // Evaluate the interpolator.  If IsInitialized() returns 'false', the
+    // operator will return std::numeric_limits<Real>::max().
+    Real operator()(Real x, Real y) const;
+
+    // Compute the functional value a^T*M*a when lambda is zero or
+    // lambda*w^T*(M+lambda*I)*w when lambda is positive.  See the thin plate
+    // splines PDF for a description of these quantities.
+    Real ComputeFunctional() const;
+
+private:
+    // Kernel(t) = t^2 * log(t^2)
+    static Real Kernel(Real t);
+
+    // Input data.
+    int mNumPoints;
+    std::vector<Real> mX;
+    std::vector<Real> mY;
+    Real mSmooth;
+
+    // Thin plate spline coefficients.  The A[] coefficients are associated
+    // with the Green's functions G(x,y,*) and the B[] coefficients are
+    // associated with the affine term B[0] + B[1]*x + B[2]*y.
+    std::vector<Real> mA;  // mNumPoints elements
+    Real mB[3];
+
+    // Extent of input data.
+    Real mXMin, mXMax, mXInvRange;
+    Real mYMin, mYMax, mYInvRange;
+
+    bool mInitialized;
+};
+
+
+template <typename Real>
+IntpThinPlateSpline2<Real>::IntpThinPlateSpline2(int numPoints, Real const* X,
+    Real const* Y, Real const* F, Real smooth, bool transformToUnitSquare)
+    :
+    mNumPoints(numPoints),
+    mX(numPoints),
+    mY(numPoints),
+    mSmooth(smooth),
+    mA(numPoints),
+    mInitialized(false)
+{
+    if (numPoints < 3 || !X || !Y || !F || smooth < (Real)0)
+    {
+        LogError("Invalid input.");
+        return;
+    }
+
+    int i, row, col;
+
+    if (transformToUnitSquare)
+    {
+        // Map input (x,y) to unit square.  This is not part of the classical
+        // thin-plate spline algorithm because the interpolation is not
+        // invariant to scalings.
+        auto extreme = std::minmax_element(X, X + mNumPoints);
+        mXMin = *extreme.first;
+        mXMax = *extreme.second;
+        mXInvRange = ((Real)1) / (mXMax - mXMin);
+        for (i = 0; i < mNumPoints; ++i)
+        {
+            mX[i] = (X[i] - mXMin) * mXInvRange;
+        }
+
+        extreme = std::minmax_element(Y, Y + mNumPoints);
+        mYMin = *extreme.first;
+        mYMax = *extreme.second;
+        mYInvRange = ((Real)1) / (mYMax - mYMin);
+        for (i = 0; i < mNumPoints; ++i)
+        {
+            mY[i] = (Y[i] - mYMin) * mYInvRange;
+        }
+    }
+    else
+    {
+        // The classical thin-plate spline uses the data as is.  The values
+        // mXMax and mYMax are not used, but they are initialized anyway
+        // (to irrelevant numbers).
+        mXMin = (Real)0;
+        mXMax = (Real)1;
+        mXInvRange = (Real)1;
+        mYMin = (Real)0;
+        mYMax = (Real)1;
+        mYInvRange = (Real)1;
+        std::copy(X, X + mNumPoints, mX.begin());
+        std::copy(Y, Y + mNumPoints, mY.begin());
+    }
+
+    // Compute matrix A = M + lambda*I [NxN matrix].
+    GMatrix<Real> AMat(mNumPoints, mNumPoints);
+    for (row = 0; row < mNumPoints; ++row)
+    {
+        for (col = 0; col < mNumPoints; ++col)
+        {
+            if (row == col)
+            {
+                AMat(row, col) = mSmooth;
+            }
+            else
+            {
+                Real dx = mX[row] - mX[col];
+                Real dy = mY[row] - mY[col];
+                Real t = std::sqrt(dx*dx + dy*dy);
+                AMat(row, col) = Kernel(t);
+            }
+        }
+    }
+
+    // Compute matrix B [Nx3 matrix].
+    GMatrix<Real> BMat(mNumPoints, 3);
+    for (row = 0; row < mNumPoints; ++row)
+    {
+        BMat(row, 0) = (Real)1;
+        BMat(row, 1) = mX[row];
+        BMat(row, 2) = mY[row];
+    }
+
+    // Compute A^{-1}.
+    bool invertible;
+    GMatrix<Real> invAMat = Inverse(AMat, &invertible);
+    if (!invertible)
+    {
+        return;
+    }
+
+    // Compute P = B^T A^{-1}  [3xN matrix].
+    GMatrix<Real> PMat = MultiplyATB(BMat, invAMat);
+
+    // Compute Q = P B = B^T A^{-1} B  [3x3 matrix].
+    GMatrix<Real> QMat = PMat * BMat;
+
+    // Compute Q^{-1}.
+    GMatrix<Real> invQMat = Inverse(QMat, &invertible);
+    if (!invertible)
+    {
+        return;
+    }
+
+    // Compute P*z.
+    Real prod[3];
+    for (row = 0; row < 3; ++row)
+    {
+        prod[row] = (Real)0;
+        for (i = 0; i < mNumPoints; ++i)
+        {
+            prod[row] += PMat(row, i) * F[i];
+        }
+    }
+
+    // Compute 'b' vector for smooth thin plate spline.
+    for (row = 0; row < 3; ++row)
+    {
+        mB[row] = (Real)0;
+        for (i = 0; i < 3; ++i)
+        {
+            mB[row] += invQMat(row, i) * prod[i];
+        }
+    }
+
+    // Compute z-B*b.
+    std::vector<Real> tmp(mNumPoints);
+    for (row = 0; row < mNumPoints; ++row)
+    {
+        tmp[row] = F[row];
+        for (i = 0; i < 3; ++i)
+        {
+            tmp[row] -= BMat(row, i) * mB[i];
+        }
+    }
+
+    // Compute 'a' vector for smooth thin plate spline.
+    for (row = 0; row < mNumPoints; ++row)
+    {
+        mA[row] = (Real)0;
+        for (i = 0; i < mNumPoints; ++i)
+        {
+            mA[row] += invAMat(row, i) * tmp[i];
+        }
+    }
+
+    mInitialized = true;
+}
+
+template <typename Real> inline
+bool IntpThinPlateSpline2<Real>::IsInitialized() const
+{
+    return mInitialized;
+}
+
+template <typename Real>
+Real IntpThinPlateSpline2<Real>::operator()(Real x, Real y) const
+{
+    if (mInitialized)
+    {
+        // Map (x,y) to the unit square.
+        x = (x - mXMin) * mXInvRange;
+        y = (y - mYMin) * mYInvRange;
+
+        Real result = mB[0] + mB[1] * x + mB[2] * y;
+        for (int i = 0; i < mNumPoints; ++i)
+        {
+            Real dx = x - mX[i];
+            Real dy = y - mY[i];
+            Real t = std::sqrt(dx*dx + dy*dy);
+            result += mA[i] * Kernel(t);
+        }
+        return result;
+    }
+
+    return std::numeric_limits<Real>::max();
+}
+
+template <typename Real>
+Real IntpThinPlateSpline2<Real>::ComputeFunctional() const
+{
+    Real functional = (Real)0;
+    for (int row = 0; row < mNumPoints; ++row)
+    {
+        for (int col = 0; col < mNumPoints; ++col)
+        {
+            if (row == col)
+            {
+                functional += mSmooth * mA[row] * mA[col];
+            }
+            else
+            {
+                Real dx = mX[row] - mX[col];
+                Real dy = mY[row] - mY[col];
+                Real t = std::sqrt(dx * dx + dy * dy);
+                functional += Kernel(t) * mA[row] * mA[col];
+            }
+        }
+    }
+
+    if (mSmooth > (Real)0)
+    {
+        functional *= mSmooth;
+    }
+
+    return functional;
+}
+
+template <typename Real>
+Real IntpThinPlateSpline2<Real>::Kernel(Real t)
+{
+    if (t > (Real)0)
+    {
+        Real t2 = t * t;
+        return t2 * std::log(t2);
+    }
+    return (Real)0;
+}
+
+
+}
diff --git a/Engine/Plugins/Experimental/GeometryProcessing/Source/GeometryAlgorithms/Private/GTEngine/Mathematics/GteIntpThinPlateSpline3.h b/Engine/Plugins/Experimental/GeometryProcessing/Source/GeometryAlgorithms/Private/GTEngine/Mathematics/GteIntpThinPlateSpline3.h
new file mode 100644
--- /dev/null
+++ b/Engine/Plugins/Experimental/GeometryProcessing/Source/GeometryAlgorithms/Private/GTEngine/Mathematics/GteIntpThinPlateSpline3.h
@@ -0,0 +1,312 @@
+// David Eberly, Geometric Tools, Redmond WA 98052
+// Copyright (c) 1998-2019
+// Distributed under the Boost Software License, Version 1.0.
+// http://www.boost.org/LICENSE_1_0.txt
+// http://www.geometrictools.com/License/Boost/LICENSE_1_0.txt
+// File Version: 3.0.0 (2016/06/19)
+
+#pragma once
+
+#include <GTEngine/Mathematics/GteGMatrix.h>
+#include <algorithm>
+#include <cmath>
+#include <limits>
+
+// WARNING.  The implementation allows you to transform the inputs (x,y,z) to
+// the unit cube and perform the interpolation in that space.  The idea is
+// to keep the floating-point numbers to order 1 for numerical stability of
+// the algorithm.  The classical thin-plate spline algorithm does not include
+// this transformation.  The interpolation is invariant to translations and
+// rotations of (x,y,z) but not to scaling.
+
+namespace gte
+{
+
+template <typename Real>
+class IntpThinPlateSpline3
+{
+public:
+    // Construction.  Data points are (x,y,z,f(x,y,z)).  The smoothing parameter
+    // must be nonnegative
+    IntpThinPlateSpline3(int numPoints, Real const* X, Real const* Y,
+        Real const* Z, Real const* F, Real smooth, bool transformToUnitCube);
+
+    // Check this after the constructor call to see whether the thin plate
+    // spline coefficients were successfully computed.  If so, then calls to
+    // operator()(Real,Real,Real) will work properly.
+    inline bool IsInitialized() const;
+
+    // Evaluate the interpolator.  If IsInitialized()returns 'false', the
+    // operator will return std::numeric_limits<Real>::max().
+    Real operator()(Real x, Real y, Real z) const;
+
+    // Compute the functional value a^T*M*a when lambda is zero or
+    // lambda*w^T*(M+lambda*I)*w when lambda is positive.  See the thin plate
+    // splines PDF for a description of these quantities.
+    Real ComputeFunctional() const;
+
+private:
+    // Kernel(t) = -|t|
+    static Real Kernel(Real t);
+
+    // Input data.
+    int mNumPoints;
+    std::vector<Real> mX;
+    std::vector<Real> mY;
+    std::vector<Real> mZ;
+    Real mSmooth;
+
+    // Thin plate spline coefficients.  The A[] coefficients are associated
+    // with the Green's functions G(x,y,z,*) and the B[] coefficients are
+    // associated with the affine term B[0] + B[1]*x + B[2]*y + B[3]*z.
+    std::vector<Real> mA;  // mNumPoints elements
+    Real mB[4];
+
+    // Extent of input data.
+    Real mXMin, mXMax, mXInvRange;
+    Real mYMin, mYMax, mYInvRange;
+    Real mZMin, mZMax, mZInvRange;
+
+    bool mInitialized;
+};
+
+
+template <typename Real>
+IntpThinPlateSpline3<Real>::IntpThinPlateSpline3(int numPoints, Real const* X,
+    Real const* Y, Real const* Z, Real const* F, Real smooth,
+    bool transformToUnitCube)
+    :
+    mNumPoints(numPoints),
+    mX(numPoints),
+    mY(numPoints),
+    mZ(numPoints),
+    mSmooth(smooth),
+    mA(numPoints),
+    mInitialized(false)
+{
+    if (numPoints < 4 || !X || !Y || !Z || !F || smooth < (Real)0)
+    {
+        LogError("Invalid input.");
+        return;
+    }
+
+    int i, row, col;
+
+    if (transformToUnitCube)
+    {
+        // Map input (x,y,z) to unit cube.  This is not part of the classical
+        // thin-plate spline algorithm, because the interpolation is not
+        // invariant to scalings.
+        auto extreme = std::minmax_element(X, X + mNumPoints);
+        mXMin = *extreme.first;
+        mXMax = *extreme.second;
+        mXInvRange = ((Real)1) / (mXMax - mXMin);
+        for (i = 0; i < mNumPoints; ++i)
+        {
+            mX[i] = (X[i] - mXMin) * mXInvRange;
+        }
+
+        extreme = std::minmax_element(Y, Y + mNumPoints);
+        mYMin = *extreme.first;
+        mYMax = *extreme.second;
+        mYInvRange = ((Real)1) / (mYMax - mYMin);
+        for (i = 0; i < mNumPoints; ++i)
+        {
+            mY[i] = (Y[i] - mYMin) * mYInvRange;
+        }
+
+        extreme = std::minmax_element(Z, Z + mNumPoints);
+        mZMin = *extreme.first;
+        mZMax = *extreme.second;
+        mZInvRange = ((Real)1) / (mZMax - mZMin);
+        for (i = 0; i < mNumPoints; ++i)
+        {
+            mZ[i] = (Z[i] - mZMin) * mZInvRange;
+        }
+    }
+    else
+    {
+        // The classical thin-plate spline uses the data as is.  The values
+        // mXMax, mYMax, and mZMax are not used, but they are initialized
+        // anyway (to irrelevant numbers).
+        mXMin = (Real)0;
+        mXMax = (Real)1;
+        mXInvRange = (Real)1;
+        mYMin = (Real)0;
+        mYMax = (Real)1;
+        mYInvRange = (Real)1;
+        mZMin = (Real)0;
+        mZMax = (Real)1;
+        mZInvRange = (Real)1;
+        std::copy(X, X + mNumPoints, mX.begin());
+        std::copy(Y, Y + mNumPoints, mY.begin());
+        std::copy(Z, Z + mNumPoints, mZ.begin());
+    }
+
+    // Compute matrix A = M + lambda*I [NxN matrix].
+    GMatrix<Real> AMat(mNumPoints, mNumPoints);
+    for (row = 0; row < mNumPoints; ++row)
+    {
+        for (col = 0; col < mNumPoints; ++col)
+        {
+            if (row == col)
+            {
+                AMat(row, col) = mSmooth;
+            }
+            else
+            {
+                Real dx = mX[row] - mX[col];
+                Real dy = mY[row] - mY[col];
+                Real dz = mZ[row] - mZ[col];
+                Real t = std::sqrt(dx * dx + dy * dy + dz * dz);
+                AMat(row, col) = Kernel(t);
+            }
+        }
+    }
+
+    // Compute matrix B [Nx4 matrix].
+    GMatrix<Real> BMat(mNumPoints, 4);
+    for (row = 0; row < mNumPoints; ++row)
+    {
+        BMat(row, 0) = (Real)1;
+        BMat(row, 1) = mX[row];
+        BMat(row, 2) = mY[row];
+        BMat(row, 3) = mZ[row];
+    }
+
+    // Compute A^{-1}.
+    bool invertible;
+    GMatrix<Real> invAMat = Inverse(AMat, &invertible);
+    if (!invertible)
+    {
+        return;
+    }
+
+    // Compute P = B^t A^{-1} [4xN matrix].
+    GMatrix<Real> PMat = MultiplyATB(BMat, invAMat);
+
+    // Compute Q = P B = B^t A^{-1} B  [4x4 matrix].
+    GMatrix<Real> QMat = PMat * BMat;
+
+    // Compute Q^{-1}.
+    GMatrix<Real> invQMat = Inverse(QMat, &invertible);
+    if (!invertible)
+    {
+        return;
+    }
+
+    // Compute P*w.
+    Real prod[4];
+    for (row = 0; row < 4; ++row)
+    {
+        prod[row] = (Real)0;
+        for (i = 0; i < mNumPoints; ++i)
+        {
+            prod[row] += PMat(row, i) * F[i];
+        }
+    }
+
+    // Compute 'b' vector for smooth thin plate spline.
+    for (row = 0; row < 4; ++row)
+    {
+        mB[row] = (Real)0;
+        for (i = 0; i < 4; ++i)
+        {
+            mB[row] += invQMat(row, i) * prod[i];
+        }
+    }
+
+    // Compute w-B*b.
+    std::vector<Real> tmp(mNumPoints);
+    for (row = 0; row < mNumPoints; ++row)
+    {
+        tmp[row] = F[row];
+        for (i = 0; i < 4; ++i)
+        {
+            tmp[row] -= BMat(row, i) * mB[i];
+        }
+    }
+
+    // Compute 'a' vector for smooth thin plate spline.
+    for (row = 0; row < mNumPoints; ++row)
+    {
+        mA[row] = (Real)0;
+        for (i = 0; i < mNumPoints; ++i)
+        {
+            mA[row] += invAMat(row, i) * tmp[i];
+        }
+    }
+
+    mInitialized = true;
+}
+
+template <typename Real>
+bool IntpThinPlateSpline3<Real>::IsInitialized() const
+{
+    return mInitialized;
+}
+
+template <typename Real>
+Real IntpThinPlateSpline3<Real>::operator()(Real x, Real y, Real z) const
+{
+    if (mInitialized)
+    {
+        // Map (x,y,z) to the unit cube.
+        x = (x - mXMin) * mXInvRange;
+        y = (y - mYMin) * mYInvRange;
+        z = (z - mZMin) * mZInvRange;
+
+        Real result = mB[0] + mB[1] * x + mB[2] * y + mB[3] * z;
+        for (int i = 0; i < mNumPoints; ++i)
+        {
+            Real dx = x - mX[i];
+            Real dy = y - mY[i];
+            Real dz = z - mZ[i];
+            Real t = std::sqrt(dx*dx + dy*dy + dz*dz);
+            result += mA[i] * Kernel(t);
+        }
+        return result;
+    }
+
+    return std::numeric_limits<Real>::max();
+}
+
+template <typename Real>
+Real IntpThinPlateSpline3<Real>::ComputeFunctional() const
+{
+    Real functional = (Real)0;
+    for (int row = 0; row < mNumPoints; ++row)
+    {
+        for (int col = 0; col < mNumPoints; ++col)
+        {
+            if (row == col)
+            {
+                functional += mSmooth * mA[row] * mA[col];
+            }
+            else
+            {
+                Real dx = mX[row] - mX[col];
+                Real dy = mY[row] - mY[col];
+                Real dz = mZ[row] - mZ[col];
+                Real t = std::sqrt(dx * dx + dy * dy + dz * dz);
+                functional += Kernel(t) * mA[row] * mA[col];
+            }
+        }
+    }
+
+    if (mSmooth > (Real)0)
+    {
+        functional *= mSmooth;
+    }
+
+    return functional;
+}
+
+template <typename Real>
+Real IntpThinPlateSpline3<Real>::Kernel(Real t)
+{
+    return -std::abs(t);
+}
+
+
+}
diff --git a/Engine/Plugins/Experimental/GeometryProcessing/Source/GeometryAlgorithms/Private/GTEngine/Mathematics/GteIntpTricubic3.h b/Engine/Plugins/Experimental/GeometryProcessing/Source/GeometryAlgorithms/Private/GTEngine/Mathematics/GteIntpTricubic3.h
new file mode 100644
--- /dev/null
+++ b/Engine/Plugins/Experimental/GeometryProcessing/Source/GeometryAlgorithms/Private/GTEngine/Mathematics/GteIntpTricubic3.h
@@ -0,0 +1,570 @@
+// David Eberly, Geometric Tools, Redmond WA 98052
+// Copyright (c) 1998-2019
+// Distributed under the Boost Software License, Version 1.0.
+// http://www.boost.org/LICENSE_1_0.txt
+// http://www.geometrictools.com/License/Boost/LICENSE_1_0.txt
+// File Version: 3.0.0 (2016/06/19)
+
+#pragma once
+
+#include <GTEngine/LowLevel/GteLogger.h>
+
+// The interpolator is for uniformly spaced(x,y z)-values.  The input samples
+// must be stored in lexicographical order to represent f(x,y,z); that is,
+// F[c + xBound*(r + yBound*s)] corresponds to f(x,y,z), where c is the index
+// corresponding to x, r is the index corresponding to y, and s is the index
+// corresponding to z.  Exact interpolation is achieved by setting catmullRom
+// to 'true', giving you the Catmull-Rom blending matrix.  If a smooth
+// interpolation is desired, set catmullRom to 'false' to obtain B-spline
+// blending.
+
+namespace gte
+{
+
+template <typename Real>
+class IntpTricubic3
+{
+public:
+    // Construction.
+    IntpTricubic3(int xBound, int yBound, int zBound, Real xMin,
+        Real xSpacing, Real yMin, Real ySpacing, Real zMin, Real zSpacing,
+        Real const* F, bool catmullRom);
+
+    // Member access.
+    inline int GetXBound() const;
+    inline int GetYBound() const;
+    inline int GetZBound() const;
+    inline int GetQuantity() const;
+    inline Real const* GetF() const;
+    inline Real GetXMin() const;
+    inline Real GetXMax() const;
+    inline Real GetXSpacing() const;
+    inline Real GetYMin() const;
+    inline Real GetYMax() const;
+    inline Real GetYSpacing() const;
+    inline Real GetZMin() const;
+    inline Real GetZMax() const;
+    inline Real GetZSpacing() const;
+
+    // Evaluate the function and its derivatives.  The functions clamp the
+    // inputs to xmin <= x <= xmax, ymin <= y <= ymax, and zmin <= z <= zmax.
+    // The first operator is for function evaluation.  The second operator is
+    // for function or derivative evaluations.  The xOrder argument is the
+    // order of the x-derivative, the yOrder argument is the order of the
+    // y-derivative, and the zOrder argument is the order of the z-derivative.
+    // All orders are zero to get the function value itself.
+    Real operator()(Real x, Real y, Real z) const;
+    Real operator()(int xOrder, int yOrder, int zOrder, Real x, Real y,
+        Real z) const;
+
+private:
+    int mXBound, mYBound, mZBound, mQuantity;
+    Real mXMin, mXMax, mXSpacing, mInvXSpacing;
+    Real mYMin, mYMax, mYSpacing, mInvYSpacing;
+    Real mZMin, mZMax, mZSpacing, mInvZSpacing;
+    Real const* mF;
+    Real mBlend[4][4];
+};
+
+
+template <typename Real>
+IntpTricubic3<Real>::IntpTricubic3(int xBound, int yBound, int zBound,
+    Real xMin, Real xSpacing, Real yMin, Real ySpacing, Real zMin,
+    Real zSpacing, Real const* F, bool catmullRom)
+    :
+    mXBound(xBound),
+    mYBound(yBound),
+    mZBound(zBound),
+    mQuantity(xBound * yBound * zBound),
+    mXMin(xMin),
+    mXSpacing(xSpacing),
+    mYMin(yMin),
+    mYSpacing(ySpacing),
+    mZMin(zMin),
+    mZSpacing(zSpacing),
+    mF(F)
+{
+    // At least a 4x4x4 block of data points are needed to construct the
+    // tricubic interpolation.
+    LogAssert(mXBound >= 4 && mYBound >= 4 && mZBound >= 4 && mF,
+        "Invalid input.");
+    LogAssert(mXSpacing > (Real)0 && mYSpacing > (Real)0 &&
+        mZSpacing > (Real)0, "Invalid input.");
+
+    mXMax = mXMin + mXSpacing * static_cast<Real>(mXBound - 1);
+    mInvXSpacing = ((Real)1) / mXSpacing;
+    mYMax = mYMin + mYSpacing * static_cast<Real>(mYBound - 1);
+    mInvYSpacing = ((Real)1) / mYSpacing;
+    mZMax = mZMin + mZSpacing * static_cast<Real>(mZBound - 1);
+    mInvZSpacing = ((Real)1) / mZSpacing;
+
+    if (catmullRom)
+    {
+        mBlend[0][0] = (Real)0;
+        mBlend[0][1] = (Real)-0.5;
+        mBlend[0][2] = (Real)1;
+        mBlend[0][3] = (Real)-0.5;
+        mBlend[1][0] = (Real)1;
+        mBlend[1][1] = (Real)0;
+        mBlend[1][2] = (Real)-2.5;
+        mBlend[1][3] = (Real)1.5;
+        mBlend[2][0] = (Real)0;
+        mBlend[2][1] = (Real)0.5;
+        mBlend[2][2] = (Real)2;
+        mBlend[2][3] = (Real)-1.5;
+        mBlend[3][0] = (Real)0;
+        mBlend[3][1] = (Real)0;
+        mBlend[3][2] = (Real)-0.5;
+        mBlend[3][3] = (Real)0.5;
+    }
+    else
+    {
+        mBlend[0][0] = (Real)1 / (Real)6;
+        mBlend[0][1] = (Real)-3 / (Real)6;
+        mBlend[0][2] = (Real)3 / (Real)6;
+        mBlend[0][3] = (Real)-1 / (Real)6;;
+        mBlend[1][0] = (Real)4 / (Real)6;
+        mBlend[1][1] = (Real)0 / (Real)6;
+        mBlend[1][2] = (Real)-6 / (Real)6;
+        mBlend[1][3] = (Real)3 / (Real)6;
+        mBlend[2][0] = (Real)1 / (Real)6;
+        mBlend[2][1] = (Real)3 / (Real)6;
+        mBlend[2][2] = (Real)3 / (Real)6;
+        mBlend[2][3] = (Real)-3 / (Real)6;
+        mBlend[3][0] = (Real)0 / (Real)6;
+        mBlend[3][1] = (Real)0 / (Real)6;
+        mBlend[3][2] = (Real)0 / (Real)6;
+        mBlend[3][3] = (Real)1 / (Real)6;
+    }
+}
+
+template <typename Real> inline
+int IntpTricubic3<Real>::GetXBound() const
+{
+    return mXBound;
+}
+
+template <typename Real> inline
+int IntpTricubic3<Real>::GetYBound() const
+{
+    return mYBound;
+}
+
+template <typename Real> inline
+int IntpTricubic3<Real>::GetZBound() const
+{
+    return mZBound;
+}
+
+template <typename Real> inline
+int IntpTricubic3<Real>::GetQuantity() const
+{
+    return mQuantity;
+}
+
+template <typename Real> inline
+Real const* IntpTricubic3<Real>::GetF() const
+{
+    return mF;
+}
+
+template <typename Real> inline
+Real IntpTricubic3<Real>::GetXMin() const
+{
+    return mXMin;
+}
+
+template <typename Real> inline
+Real IntpTricubic3<Real>::GetXMax() const
+{
+    return mXMax;
+}
+
+template <typename Real> inline
+Real IntpTricubic3<Real>::GetXSpacing() const
+{
+    return mXSpacing;
+}
+
+template <typename Real> inline
+Real IntpTricubic3<Real>::GetYMin() const
+{
+    return mYMin;
+}
+
+template <typename Real> inline
+Real IntpTricubic3<Real>::GetYMax() const
+{
+    return mYMax;
+}
+
+template <typename Real> inline
+Real IntpTricubic3<Real>::GetYSpacing() const
+{
+    return mYSpacing;
+}
+
+template <typename Real> inline
+Real IntpTricubic3<Real>::GetZMin() const
+{
+    return mZMin;
+}
+
+template <typename Real> inline
+Real IntpTricubic3<Real>::GetZMax() const
+{
+    return mZMax;
+}
+
+template <typename Real> inline
+Real IntpTricubic3<Real>::GetZSpacing() const
+{
+    return mZSpacing;
+}
+
+template <typename Real>
+Real IntpTricubic3<Real>::operator()(Real x, Real y, Real z) const
+{
+    // Compute x-index and clamp to image.
+    Real xIndex = (x - mXMin) * mInvXSpacing;
+    int ix = static_cast<int>(xIndex);
+    if (ix < 0)
+    {
+        ix = 0;
+    }
+    else if (ix >= mXBound)
+    {
+        ix = mXBound - 1;
+    }
+
+    // Compute y-index and clamp to image.
+    Real yIndex = (y - mYMin) * mInvYSpacing;
+    int iy = static_cast<int>(yIndex);
+    if (iy < 0)
+    {
+        iy = 0;
+    }
+    else if (iy >= mYBound)
+    {
+        iy = mYBound - 1;
+    }
+
+    // Compute z-index and clamp to image.
+    Real zIndex = (z - mZMin) * mInvZSpacing;
+    int iz = static_cast<int>(zIndex);
+    if (iz < 0)
+    {
+        iz = 0;
+    }
+    else if (iz >= mZBound)
+    {
+        iz = mZBound - 1;
+    }
+
+    Real U[4];
+    U[0] = (Real)1;
+    U[1] = xIndex - ix;
+    U[2] = U[1] * U[1];
+    U[3] = U[1] * U[2];
+
+    Real V[4];
+    V[0] = (Real)1;
+    V[1] = yIndex - iy;
+    V[2] = V[1] * V[1];
+    V[3] = V[1] * V[2];
+
+    Real W[4];
+    W[0] = (Real)1;
+    W[1] = zIndex - iz;
+    W[2] = W[1] * W[1];
+    W[3] = W[1] * W[2];
+
+    // Compute P = M*U, Q = M*V, R = M*W.
+    Real P[4], Q[4], R[4];
+    for (int row = 0; row < 4; ++row)
+    {
+        P[row] = (Real)0;
+        Q[row] = (Real)0;
+        R[row] = (Real)0;
+        for (int col = 0; col < 4; ++col)
+        {
+            P[row] += mBlend[row][col] * U[col];
+            Q[row] += mBlend[row][col] * V[col];
+            R[row] += mBlend[row][col] * W[col];
+        }
+    }
+
+    // Compute the tensor product (M*U)(M*V)(M*W)*D where D is the 4x4x4
+    // subimage containing (x,y,z).
+    --ix;
+    --iy;
+    --iz;
+    Real result = (Real)0;
+    for (int slice = 0; slice < 4; ++slice)
+    {
+        int zClamp = iz + slice;
+        if (zClamp < 0)
+        {
+            zClamp = 0;
+        }
+        else if (zClamp > mZBound - 1)
+        {
+            zClamp = mZBound - 1;
+        }
+
+        for (int row = 0; row < 4; ++row)
+        {
+            int yClamp = iy + row;
+            if (yClamp < 0)
+            {
+                yClamp = 0;
+            }
+            else if (yClamp > mYBound - 1)
+            {
+                yClamp = mYBound - 1;
+            }
+
+            for (int col = 0; col < 4; ++col)
+            {
+                int xClamp = ix + col;
+                if (xClamp < 0)
+                {
+                    xClamp = 0;
+                }
+                else if (xClamp > mXBound - 1)
+                {
+                    xClamp = mXBound - 1;
+                }
+
+                result += P[col] * Q[row] * R[slice] *
+                    mF[xClamp + mXBound * (yClamp + mYBound * zClamp)];
+            }
+        }
+    }
+
+    return result;
+}
+
+template <typename Real>
+Real IntpTricubic3<Real>::operator()(int xOrder, int yOrder, int zOrder,
+    Real x, Real y, Real z) const
+{
+    // Compute x-index and clamp to image.
+    Real xIndex = (x - mXMin) * mInvXSpacing;
+    int ix = static_cast<int>(xIndex);
+    if (ix < 0)
+    {
+        ix = 0;
+    }
+    else if (ix >= mXBound)
+    {
+        ix = mXBound - 1;
+    }
+
+    // Compute y-index and clamp to image.
+    Real yIndex = (y - mYMin) * mInvYSpacing;
+    int iy = static_cast<int>(yIndex);
+    if (iy < 0)
+    {
+        iy = 0;
+    }
+    else if (iy >= mYBound)
+    {
+        iy = mYBound - 1;
+    }
+
+    // Compute z-index and clamp to image.
+    Real zIndex = (z - mZMin) * mInvZSpacing;
+    int iz = static_cast<int>(zIndex);
+    if (iz < 0)
+    {
+        iz = 0;
+    }
+    else if (iz >= mZBound)
+    {
+        iz = mZBound - 1;
+    }
+
+    Real U[4], dx, xMult;
+    switch (xOrder)
+    {
+    case 0:
+        dx = xIndex - ix;
+        U[0] = (Real)1;
+        U[1] = dx;
+        U[2] = dx * U[1];
+        U[3] = dx * U[2];
+        xMult = (Real)1;
+        break;
+    case 1:
+        dx = xIndex - ix;
+        U[0] = (Real)0;
+        U[1] = (Real)1;
+        U[2] = ((Real)2) * dx;
+        U[3] = ((Real)3) * dx * dx;
+        xMult = mInvXSpacing;
+        break;
+    case 2:
+        dx = xIndex - ix;
+        U[0] = (Real)0;
+        U[1] = (Real)0;
+        U[2] = (Real)2;
+        U[3] = ((Real)6) * dx;
+        xMult = mInvXSpacing * mInvXSpacing;
+        break;
+    case 3:
+        U[0] = (Real)0;
+        U[1] = (Real)0;
+        U[2] = (Real)0;
+        U[3] = (Real)6;
+        xMult = mInvXSpacing * mInvXSpacing * mInvXSpacing;
+        break;
+    default:
+        return (Real)0;
+    }
+
+    Real V[4], dy, yMult;
+    switch (yOrder)
+    {
+    case 0:
+        dy = yIndex - iy;
+        V[0] = (Real)1;
+        V[1] = dy;
+        V[2] = dy * V[1];
+        V[3] = dy * V[2];
+        yMult = (Real)1;
+        break;
+    case 1:
+        dy = yIndex - iy;
+        V[0] = (Real)0;
+        V[1] = (Real)1;
+        V[2] = ((Real)2) * dy;
+        V[3] = ((Real)3) * dy * dy;
+        yMult = mInvYSpacing;
+        break;
+    case 2:
+        dy = yIndex - iy;
+        V[0] = (Real)0;
+        V[1] = (Real)0;
+        V[2] = (Real)2;
+        V[3] = ((Real)6) * dy;
+        yMult = mInvYSpacing * mInvYSpacing;
+        break;
+    case 3:
+        V[0] = (Real)0;
+        V[1] = (Real)0;
+        V[2] = (Real)0;
+        V[3] = (Real)6;
+        yMult = mInvYSpacing * mInvYSpacing * mInvYSpacing;
+        break;
+    default:
+        return (Real)0;
+    }
+
+    Real W[4], dz, zMult;
+    switch (zOrder)
+    {
+    case 0:
+        dz = zIndex - iz;
+        W[0] = (Real)1;
+        W[1] = dz;
+        W[2] = dz * W[1];
+        W[3] = dz * W[2];
+        zMult = (Real)1;
+        break;
+    case 1:
+        dz = zIndex - iz;
+        W[0] = (Real)0;
+        W[1] = (Real)1;
+        W[2] = ((Real)2) * dz;
+        W[3] = ((Real)3) * dz * dz;
+        zMult = mInvZSpacing;
+        break;
+    case 2:
+        dz = zIndex - iz;
+        W[0] = (Real)0;
+        W[1] = (Real)0;
+        W[2] = (Real)2;
+        W[3] = ((Real)6) * dz;
+        zMult = mInvZSpacing * mInvZSpacing;
+        break;
+    case 3:
+        W[0] = (Real)0;
+        W[1] = (Real)0;
+        W[2] = (Real)0;
+        W[3] = (Real)6;
+        zMult = mInvZSpacing * mInvZSpacing * mInvZSpacing;
+        break;
+    default:
+        return (Real)0;
+    }
+
+    // Compute P = M*U, Q = M*V, and R = M*W.
+    Real P[4], Q[4], R[4];
+    for (int row = 0; row < 4; ++row)
+    {
+        P[row] = (Real)0;
+        Q[row] = (Real)0;
+        R[row] = (Real)0;
+        for (int col = 0; col < 4; ++col)
+        {
+            P[row] += mBlend[row][col] * U[col];
+            Q[row] += mBlend[row][col] * V[col];
+            R[row] += mBlend[row][col] * W[col];
+        }
+    }
+
+    // Compute the tensor product (M*U)(M*V)(M*W)*D where D is the 4x4x4
+    // subimage containing (x,y,z).
+    --ix;
+    --iy;
+    --iz;
+    Real result = (Real)0;
+    for (int slice = 0; slice < 4; ++slice)
+    {
+        int zClamp = iz + slice;
+        if (zClamp < 0)
+        {
+            zClamp = 0;
+        }
+        else if (zClamp > mZBound - 1)
+        {
+            zClamp = mZBound - 1;
+        }
+
+        for (int row = 0; row < 4; ++row)
+        {
+            int yClamp = iy + row;
+            if (yClamp < 0)
+            {
+                yClamp = 0;
+            }
+            else if (yClamp > mYBound - 1)
+            {
+                yClamp = mYBound - 1;
+            }
+
+            for (int col = 0; col < 4; ++col)
+            {
+                int xClamp = ix + col;
+                if (xClamp < 0)
+                {
+                    xClamp = 0;
+                }
+                else if (xClamp > mXBound - 1)
+                {
+                    xClamp = mXBound - 1;
+                }
+
+                result += P[col] * Q[row] * R[slice] *
+                    mF[xClamp + mXBound * (yClamp + mYBound * zClamp)];
+            }
+        }
+    }
+    result *= xMult * yMult * zMult;
+
+    return result;
+}
+
+
+}
diff --git a/Engine/Plugins/Experimental/GeometryProcessing/Source/GeometryAlgorithms/Private/GTEngine/Mathematics/GteIntpTrilinear3.h b/Engine/Plugins/Experimental/GeometryProcessing/Source/GeometryAlgorithms/Private/GTEngine/Mathematics/GteIntpTrilinear3.h
new file mode 100644
--- /dev/null
+++ b/Engine/Plugins/Experimental/GeometryProcessing/Source/GeometryAlgorithms/Private/GTEngine/Mathematics/GteIntpTrilinear3.h
@@ -0,0 +1,440 @@
+// David Eberly, Geometric Tools, Redmond WA 98052
+// Copyright (c) 1998-2019
+// Distributed under the Boost Software License, Version 1.0.
+// http://www.boost.org/LICENSE_1_0.txt
+// http://www.geometrictools.com/License/Boost/LICENSE_1_0.txt
+// File Version: 3.0.0 (2016/06/19)
+
+#pragma once
+
+#include <GTEngine/LowLevel/GteLogger.h>
+
+// The interpolator is for uniformly spaced(x,y z)-values.  The input samples
+// must be stored in lexicographical order to represent f(x,y,z); that is,
+// F[c + xBound*(r + yBound*s)] corresponds to f(x,y,z), where c is the index
+// corresponding to x, r is the index corresponding to y, and s is the index
+// corresponding to z.
+
+namespace gte
+{
+
+template <typename Real>
+class IntpTrilinear3
+{
+public:
+    // Construction.
+    IntpTrilinear3(int xBound, int yBound, int zBound, Real xMin,
+        Real xSpacing, Real yMin, Real ySpacing, Real zMin, Real zSpacing,
+        Real const* F);
+
+    // Member access.
+    inline int GetXBound() const;
+    inline int GetYBound() const;
+    inline int GetZBound() const;
+    inline int GetQuantity() const;
+    inline Real const* GetF() const;
+    inline Real GetXMin() const;
+    inline Real GetXMax() const;
+    inline Real GetXSpacing() const;
+    inline Real GetYMin() const;
+    inline Real GetYMax() const;
+    inline Real GetYSpacing() const;
+    inline Real GetZMin() const;
+    inline Real GetZMax() const;
+    inline Real GetZSpacing() const;
+
+    // Evaluate the function and its derivatives.  The functions clamp the
+    // inputs to xmin <= x <= xmax, ymin <= y <= ymax, and zmin <= z <= zmax.
+    // The first operator is for function evaluation.  The second operator is
+    // for function or derivative evaluations.  The xOrder argument is the
+    // order of the x-derivative, the yOrder argument is the order of the
+    // y-derivative, and the zOrder argument is the order of the z-derivative.
+    // All orders are zero to get the function value itself.
+    Real operator()(Real x, Real y, Real z) const;
+    Real operator()(int xOrder, int yOrder, int zOrder, Real x, Real y,
+        Real z) const;
+
+private:
+    int mXBound, mYBound, mZBound, mQuantity;
+    Real mXMin, mXMax, mXSpacing, mInvXSpacing;
+    Real mYMin, mYMax, mYSpacing, mInvYSpacing;
+    Real mZMin, mZMax, mZSpacing, mInvZSpacing;
+    Real const* mF;
+    Real mBlend[2][2];
+};
+
+
+template <typename Real>
+IntpTrilinear3<Real>::IntpTrilinear3(int xBound, int yBound, int zBound,
+    Real xMin, Real xSpacing, Real yMin, Real ySpacing, Real zMin,
+    Real zSpacing, Real const* F)
+    :
+    mXBound(xBound),
+    mYBound(yBound),
+    mZBound(zBound),
+    mQuantity(xBound * yBound * zBound),
+    mXMin(xMin),
+    mXSpacing(xSpacing),
+    mYMin(yMin),
+    mYSpacing(ySpacing),
+    mZMin(zMin),
+    mZSpacing(zSpacing),
+    mF(F)
+{
+    // At least a 2x2x2 block of data points are needed to construct the
+    // trilinear interpolation.
+    LogAssert(mXBound >= 2 && mYBound >= 2 && mZBound >= 2 && mF,
+        "Invalid input.");
+    LogAssert(mXSpacing > (Real)0 && mYSpacing > (Real)0 &&
+        mZSpacing > (Real)0, "Invalid input.");
+
+    mXMax = mXMin + mXSpacing * static_cast<Real>(mXBound - 1);
+    mInvXSpacing = ((Real)1) / mXSpacing;
+    mYMax = mYMin + mYSpacing * static_cast<Real>(mYBound - 1);
+    mInvYSpacing = ((Real)1) / mYSpacing;
+    mZMax = mZMin + mZSpacing * static_cast<Real>(mZBound - 1);
+    mInvZSpacing = ((Real)1) / mZSpacing;
+
+    mBlend[0][0] = (Real)1;
+    mBlend[0][1] = (Real)-1;
+    mBlend[1][0] = (Real)0;
+    mBlend[1][1] = (Real)1;
+}
+
+template <typename Real> inline
+int IntpTrilinear3<Real>::GetXBound() const
+{
+    return mXBound;
+}
+
+template <typename Real> inline
+int IntpTrilinear3<Real>::GetYBound() const
+{
+    return mYBound;
+}
+
+template <typename Real> inline
+int IntpTrilinear3<Real>::GetZBound() const
+{
+    return mZBound;
+}
+
+template <typename Real> inline
+int IntpTrilinear3<Real>::GetQuantity() const
+{
+    return mQuantity;
+}
+
+template <typename Real> inline
+Real const* IntpTrilinear3<Real>::GetF() const
+{
+    return mF;
+}
+
+template <typename Real> inline
+Real IntpTrilinear3<Real>::GetXMin() const
+{
+    return mXMin;
+}
+
+template <typename Real> inline
+Real IntpTrilinear3<Real>::GetXMax() const
+{
+    return mXMax;
+}
+
+template <typename Real> inline
+Real IntpTrilinear3<Real>::GetXSpacing() const
+{
+    return mXSpacing;
+}
+
+template <typename Real> inline
+Real IntpTrilinear3<Real>::GetYMin() const
+{
+    return mYMin;
+}
+
+template <typename Real> inline
+Real IntpTrilinear3<Real>::GetYMax() const
+{
+    return mYMax;
+}
+
+template <typename Real> inline
+Real IntpTrilinear3<Real>::GetYSpacing() const
+{
+    return mYSpacing;
+}
+
+template <typename Real> inline
+Real IntpTrilinear3<Real>::GetZMin() const
+{
+    return mZMin;
+}
+
+template <typename Real> inline
+Real IntpTrilinear3<Real>::GetZMax() const
+{
+    return mZMax;
+}
+
+template <typename Real> inline
+Real IntpTrilinear3<Real>::GetZSpacing() const
+{
+    return mZSpacing;
+}
+
+template <typename Real>
+Real IntpTrilinear3<Real>::operator()(Real x, Real y, Real z) const
+{
+    // Compute x-index and clamp to image.
+    Real xIndex = (x - mXMin) * mInvXSpacing;
+    int ix = static_cast<int>(xIndex);
+    if (ix < 0)
+    {
+        ix = 0;
+    }
+    else if (ix >= mXBound)
+    {
+        ix = mXBound - 1;
+    }
+
+    // Compute y-index and clamp to image.
+    Real yIndex = (y - mYMin) * mInvYSpacing;
+    int iy = static_cast<int>(yIndex);
+    if (iy < 0)
+    {
+        iy = 0;
+    }
+    else if (iy >= mYBound)
+    {
+        iy = mYBound - 1;
+    }
+
+    // Compute z-index and clamp to image.
+    Real zIndex = (z - mZMin) * mInvZSpacing;
+    int iz = static_cast<int>(zIndex);
+    if (iz < 0)
+    {
+        iz = 0;
+    }
+    else if (iz >= mZBound)
+    {
+        iz = mZBound - 1;
+    }
+
+    Real U[2];
+    U[0] = (Real)1;
+    U[1] = xIndex - ix;
+
+    Real V[2];
+    V[0] = (Real)1;
+    V[1] = yIndex - iy;
+
+    Real W[2];
+    W[0] = (Real)1;
+    W[1] = zIndex - iz;
+
+    // Compute P = M*U, Q = M*V, R = M*W.
+    Real P[2], Q[2], R[2];
+    for (int row = 0; row < 2; ++row)
+    {
+        P[row] = (Real)0;
+        Q[row] = (Real)0;
+        R[row] = (Real)0;
+        for (int col = 0; col < 2; ++col)
+        {
+            P[row] += mBlend[row][col] * U[col];
+            Q[row] += mBlend[row][col] * V[col];
+            R[row] += mBlend[row][col] * W[col];
+        }
+    }
+
+    // compute the tensor product (M*U)(M*V)(M*W)*D where D is the 2x2x2
+    // subimage containing (x,y,z)
+    Real result = (Real)0;
+    for (int slice = 0; slice < 2; ++slice)
+    {
+        int zClamp = iz + slice;
+        if (zClamp >= mZBound)
+        {
+            zClamp = mZBound - 1;
+        }
+
+        for (int row = 0; row < 2; ++row)
+        {
+            int yClamp = iy + row;
+            if (yClamp >= mYBound)
+            {
+                yClamp = mYBound - 1;
+            }
+
+            for (int col = 0; col < 2; ++col)
+            {
+                int xClamp = ix + col;
+                if (xClamp >= mXBound)
+                {
+                    xClamp = mXBound - 1;
+                }
+
+                result += P[col] * Q[row] * R[slice] *
+                    mF[xClamp + mXBound * (yClamp + mYBound * zClamp)];
+            }
+        }
+    }
+
+    return result;
+}
+
+template <typename Real>
+Real IntpTrilinear3<Real>::operator()(int xOrder, int yOrder, int zOrder,
+    Real x, Real y, Real z) const
+{
+    // Compute x-index and clamp to image.
+    Real xIndex = (x - mXMin) * mInvXSpacing;
+    int ix = static_cast<int>(xIndex);
+    if (ix < 0)
+    {
+        ix = 0;
+    }
+    else if (ix >= mXBound)
+    {
+        ix = mXBound - 1;
+    }
+
+    // Compute y-index and clamp to image.
+    Real yIndex = (y - mYMin) * mInvYSpacing;
+    int iy = static_cast<int>(yIndex);
+    if (iy < 0)
+    {
+        iy = 0;
+    }
+    else if (iy >= mYBound)
+    {
+        iy = mYBound - 1;
+    }
+
+    // Compute z-index and clamp to image.
+    Real zIndex = (z - mZMin) * mInvZSpacing;
+    int iz = static_cast<int>(zIndex);
+    if (iz < 0)
+    {
+        iz = 0;
+    }
+    else if (iz >= mZBound)
+    {
+        iz = mZBound - 1;
+    }
+
+    Real U[2], dx, xMult;
+    switch (xOrder)
+    {
+    case 0:
+        dx = xIndex - ix;
+        U[0] = (Real)1;
+        U[1] = dx;
+        xMult = (Real)1;
+        break;
+    case 1:
+        dx = xIndex - ix;
+        U[0] = (Real)0;
+        U[1] = (Real)1;
+        xMult = mInvXSpacing;
+        break;
+    default:
+        return (Real)0;
+    }
+
+    Real V[2], dy, yMult;
+    switch (yOrder)
+    {
+    case 0:
+        dy = yIndex - iy;
+        V[0] = (Real)1;
+        V[1] = dy;
+        yMult = (Real)1;
+        break;
+    case 1:
+        dy = yIndex - iy;
+        V[0] = (Real)0;
+        V[1] = (Real)1;
+        yMult = mInvYSpacing;
+        break;
+    default:
+        return (Real)0;
+    }
+
+    Real W[2], dz, zMult;
+    switch (zOrder)
+    {
+    case 0:
+        dz = zIndex - iz;
+        W[0] = (Real)1;
+        W[1] = dz;
+        zMult = (Real)1;
+        break;
+    case 1:
+        dz = zIndex - iz;
+        W[0] = (Real)0;
+        W[1] = (Real)1;
+        zMult = mInvZSpacing;
+        break;
+    default:
+        return (Real)0;
+    }
+
+    // Compute P = M*U, Q = M*V, and R = M*W.
+    Real P[2], Q[2], R[2];
+    for (int row = 0; row < 2; ++row)
+    {
+        P[row] = (Real)0;
+        Q[row] = (Real)0;
+        R[row] = (Real)0;
+        for (int col = 0; col < 2; ++col)
+        {
+            P[row] += mBlend[row][col] * U[col];
+            Q[row] += mBlend[row][col] * V[col];
+            R[row] += mBlend[row][col] * W[col];
+        }
+    }
+
+    // Compute the tensor product (M*U)(M*V)(M*W)*D where D is the 2x2x2
+    // subimage containing (x,y,z).
+    Real result = (Real)0;
+    for (int slice = 0; slice < 2; ++slice)
+    {
+        int zClamp = iz + slice;
+        if (zClamp >= mZBound)
+        {
+            zClamp = mZBound - 1;
+        }
+
+        for (int row = 0; row < 2; ++row)
+        {
+            int yClamp = iy + row;
+            if (yClamp >= mYBound)
+            {
+                yClamp = mYBound - 1;
+            }
+
+            for (int col = 0; col < 2; ++col)
+            {
+                int xClamp = ix + col;
+                if (xClamp >= mXBound)
+                {
+                    xClamp = mXBound - 1;
+                }
+
+                result += P[col] * Q[row] * R[slice] *
+                    mF[xClamp + mXBound * (yClamp + mYBound * zClamp)];
+            }
+        }
+    }
+    result *= xMult * yMult * zMult;
+
+    return result;
+}
+
+
+}
diff --git a/Engine/Plugins/Experimental/GeometryProcessing/Source/GeometryAlgorithms/Private/GTEngine/Mathematics/GteIntpVectorField2.h b/Engine/Plugins/Experimental/GeometryProcessing/Source/GeometryAlgorithms/Private/GTEngine/Mathematics/GteIntpVectorField2.h
new file mode 100644
--- /dev/null
+++ b/Engine/Plugins/Experimental/GeometryProcessing/Source/GeometryAlgorithms/Private/GTEngine/Mathematics/GteIntpVectorField2.h
@@ -0,0 +1,93 @@
+// David Eberly, Geometric Tools, Redmond WA 98052
+// Copyright (c) 1998-2019
+// Distributed under the Boost Software License, Version 1.0.
+// http://www.boost.org/LICENSE_1_0.txt
+// http://www.geometrictools.com/License/Boost/LICENSE_1_0.txt
+// File Version: 3.0.0 (2016/06/19)
+
+#pragma once
+
+#include <GTEngine/Mathematics/GteDelaunay2Mesh.h>
+#include <GTEngine/Mathematics/GteIntpQuadraticNonuniform2.h>
+#include <memory>
+
+// Given points (x0[i],y0[i]) which are mapped to (x1[i],y1[i]) for
+// 0 <= i < N, interpolate positions (xIn,yIn) to (xOut,yOut).
+
+namespace gte
+{
+
+template <typename InputType, typename ComputeType, typename RationalType>
+class IntpVectorField2
+{
+public:
+    // Construction and destruction.
+    ~IntpVectorField2();
+    IntpVectorField2(int numPoints, Vector2<InputType> const* domain,
+        Vector2<InputType> const* range);
+
+    // The return value is 'true' if and only if (xIn,yIn) is in the convex
+    // hull of the input domain points, in which case the interpolation is
+    // valid.
+    bool operator()(Vector2<InputType> const& input,
+        Vector2<InputType>& output) const;
+
+protected:
+    typedef Delaunay2Mesh<InputType, ComputeType, RationalType> TriangleMesh;
+    Delaunay2<InputType, ComputeType> mDelaunay;
+    TriangleMesh mMesh;
+
+    std::vector<InputType> mXRange;
+    std::vector<InputType> mYRange;
+    std::unique_ptr<IntpQuadraticNonuniform2<InputType, TriangleMesh>> mXInterp;
+    std::unique_ptr<IntpQuadraticNonuniform2<InputType, TriangleMesh>> mYInterp;
+};
+
+
+template <typename InputType, typename ComputeType, typename RationalType>
+IntpVectorField2<InputType, ComputeType, RationalType>::~IntpVectorField2()
+{
+}
+
+template <typename InputType, typename ComputeType, typename RationalType>
+IntpVectorField2<InputType, ComputeType, RationalType>::IntpVectorField2(
+    int numPoints, Vector2<InputType> const* domain,
+    Vector2<InputType> const* range)
+    :
+    mMesh(mDelaunay)
+{
+    // Repackage the output vectors into individual components.  This is
+    // required because of the format that the quadratic interpolator expects
+    // for its input data.
+    mXRange.resize(numPoints);
+    mYRange.resize(numPoints);
+    for (int i = 0; i < numPoints; ++i)
+    {
+        mXRange[i] = range[i][0];
+        mYRange[i] = range[i][1];
+    }
+
+    // Common triangulator for the interpolators.
+    mDelaunay(numPoints, &domain[0], (ComputeType)0);
+
+    // Create interpolator for x-coordinate of vector field.
+    mXInterp = std::make_unique<IntpQuadraticNonuniform2<InputType, TriangleMesh>>(
+        mMesh, &mXRange[0], (InputType)1);
+
+    // Create interpolator for y-coordinate of vector field, but share the
+    // already created triangulation for the x-interpolator.
+    mYInterp = std::make_unique<IntpQuadraticNonuniform2<InputType, TriangleMesh>>(
+        mMesh, &mYRange[0], (InputType)1);
+}
+
+template <typename InputType, typename ComputeType, typename RationalType>
+bool IntpVectorField2<InputType, ComputeType, RationalType>::operator()(
+    Vector2<InputType> const& input, Vector2<InputType>& output) const
+{
+    InputType xDeriv, yDeriv;
+    return (*mXInterp)(input, output[0], xDeriv, yDeriv)
+        && (*mYInterp)(input, output[1], xDeriv, yDeriv);
+}
+
+
+}
diff --git a/Engine/Plugins/Experimental/GeometryProcessing/Source/GeometryAlgorithms/Private/GTEngine/Mathematics/GteIntrAlignedBox2AlignedBox2.h b/Engine/Plugins/Experimental/GeometryProcessing/Source/GeometryAlgorithms/Private/GTEngine/Mathematics/GteIntrAlignedBox2AlignedBox2.h
new file mode 100644
--- /dev/null
+++ b/Engine/Plugins/Experimental/GeometryProcessing/Source/GeometryAlgorithms/Private/GTEngine/Mathematics/GteIntrAlignedBox2AlignedBox2.h
@@ -0,0 +1,109 @@
+// David Eberly, Geometric Tools, Redmond WA 98052
+// Copyright (c) 1998-2019
+// Distributed under the Boost Software License, Version 1.0.
+// http://www.boost.org/LICENSE_1_0.txt
+// http://www.geometrictools.com/License/Boost/LICENSE_1_0.txt
+// File Version: 3.0.0 (2016/06/19)
+
+#pragma once
+
+#include <GTEngine/Mathematics/GteAlignedBox.h>
+#include <GTEngine/Mathematics/GteFIQuery.h>
+#include <GTEngine/Mathematics/GteTIQuery.h>
+
+// The queries consider the box to be a solid.
+//
+// The aligned-aligned queries use simple min-max comparisions.  The
+// interesection of aligned boxes is an aligned box, possibly degenerate,
+// where min[d] == max[d] for at least one dimension d.
+
+namespace gte
+{
+
+template <typename Real>
+class TIQuery<Real, AlignedBox2<Real>, AlignedBox2<Real>>
+{
+public:
+    struct Result
+    {
+        bool intersect;
+    };
+
+    Result operator()(AlignedBox2<Real> const& box0,
+        AlignedBox2<Real> const& box1);
+};
+
+template <typename Real>
+class FIQuery<Real, AlignedBox2<Real>, AlignedBox2<Real>>
+{
+public:
+    struct Result
+    {
+        bool intersect;
+        AlignedBox2<Real> box;
+    };
+
+    Result operator()(AlignedBox2<Real> const& box0,
+        AlignedBox2<Real> const& box1);
+};
+
+
+template <typename Real>
+typename TIQuery<Real, AlignedBox2<Real>, AlignedBox2<Real>>::Result
+TIQuery<Real, AlignedBox2<Real>, AlignedBox2<Real>>::operator()(
+    AlignedBox2<Real> const& box0, AlignedBox2<Real> const& box1)
+{
+    Result result;
+    for (int i = 0; i < 2; i++)
+    {
+        if (box0.max[i] < box1.min[i] || box0.min[i] > box1.max[i])
+        {
+            result.intersect = false;
+            return result;
+        }
+    }
+    result.intersect = true;
+    return result;
+}
+
+template <typename Real>
+typename FIQuery<Real, AlignedBox2<Real>, AlignedBox2<Real>>::Result
+FIQuery<Real, AlignedBox2<Real>, AlignedBox2<Real>>::operator()(
+    AlignedBox2<Real> const& box0, AlignedBox2<Real> const& box1)
+{
+    Result result;
+    for (int i = 0; i < 2; i++)
+    {
+        if (box0.max[i] < box1.min[i] || box0.min[i] > box1.max[i])
+        {
+            result.intersect = false;
+            return result;
+        }
+    }
+
+    for (int i = 0; i < 2; i++)
+    {
+        if (box0.max[i] <= box1.max[i])
+        {
+            result.box.max[i] = box0.max[i];
+        }
+        else
+        {
+            result.box.max[i] = box1.max[i];
+        }
+
+        if (box0.min[i] <= box1.min[i])
+        {
+            result.box.min[i] = box1.min[i];
+        }
+        else
+        {
+            result.box.min[i] = box0.min[i];
+        }
+    }
+    result.intersect = true;
+    return result;
+}
+
+
+}
diff --git a/Engine/Plugins/Experimental/GeometryProcessing/Source/GeometryAlgorithms/Private/GTEngine/Mathematics/GteIntrAlignedBox2Circle2.h b/Engine/Plugins/Experimental/GeometryProcessing/Source/GeometryAlgorithms/Private/GTEngine/Mathematics/GteIntrAlignedBox2Circle2.h
new file mode 100644
--- /dev/null
+++ b/Engine/Plugins/Experimental/GeometryProcessing/Source/GeometryAlgorithms/Private/GTEngine/Mathematics/GteIntrAlignedBox2Circle2.h
@@ -0,0 +1,348 @@
+// David Eberly, Geometric Tools, Redmond WA 98052
+// Copyright (c) 1998-2019
+// Distributed under the Boost Software License, Version 1.0.
+// http://www.boost.org/LICENSE_1_0.txt
+// http://www.geometrictools.com/License/Boost/LICENSE_1_0.txt
+// File Version: 3.15.3 (2018/10/05)
+
+#pragma once
+
+#include <GTEngine/Mathematics/GteVector2.h>
+#include <GTEngine/Mathematics/GteHypersphere.h>
+#include <GTEngine/Mathematics/GteDistPointAlignedBox.h>
+#include <GTEngine/Mathematics/GteFIQuery.h>
+#include <GTEngine/Mathematics/GteTIQuery.h>
+
+// The find-intersection query is based on the document
+// https://www.geometrictools.com/Documentation/IntersectionMovingCircleRectangle.pdf
+
+namespace gte
+{
+    template <typename Real>
+    class TIQuery<Real, AlignedBox2<Real>, Circle2<Real>>
+    {
+    public:
+        // The intersection query considers the box and circle to be solids;
+        // that is, the circle object includes the region inside the circular
+        // boundary and the box object includes the region inside the
+        // rectangular boundary.  If the circle object and box object
+        // overlap, the objects intersect.
+        struct Result
+        {
+            bool intersect;
+        };
+
+        Result operator()(AlignedBox2<Real> const& box, Circle2<Real> const& circle)
+        {
+            DCPQuery<Real, Vector2<Real>, AlignedBox2<Real>> pbQuery;
+            auto pbResult = pbQuery(circle.center, box);
+            Result result;
+            result.intersect = (pbResult.sqrDistance <= circle.radius * circle.radius);
+            return result;
+        }
+    };
+
+    template <typename Real>
+    class FIQuery<Real, AlignedBox2<Real>, Circle2<Real>>
+    {
+    public:
+        // Currently, only a dynamic query is supported.  A static query will
+        // need to compute the intersection set of (solid) box and circle.
+        struct Result
+        {
+            // The cases are
+            // 1. Objects initially overlapping.  The contactPoint is only one
+            //    of infinitely many points in the overlap.
+            //      intersectionType = -1
+            //      contactTime = 0
+            //      contactPoint = circle.center
+            // 2. Objects initially separated but do not intersect later.  The
+            //      contactTime and contactPoint are invalid.
+            //      intersectionType = 0
+            //      contactTime = 0
+            //      contactPoint = (0,0)
+            // 3. Objects initially separated but intersect later.
+            //      intersectionType = +1
+            //      contactTime = first time T > 0
+            //      contactPoint = corresponding first contact
+            int intersectionType;
+            Real contactTime;
+            Vector2<Real> contactPoint;
+
+            // TODO: To support arbitrary precision for the contactTime,
+            // return q0, q1 and q2 where contactTime = (q0 - sqrt(q1)) / q2.
+            // The caller can compute contactTime to desired number of digits
+            // of precision.  These are valid when intersectionType is +1 but
+            // are set to zero (invalid) in the other cases.  Do the same for
+            // the contactPoint.
+        };
+
+        Result operator()(AlignedBox2<Real> const& box, Vector2<Real> const& boxVelocity,
+            Circle2<Real> const& circle, Vector2<Real> const& circleVelocity)
+        {
+            Result result = { 0, (Real)0, { (Real)0, (Real)0 } };
+
+            // Translate the circle and box so that the box center becomes
+            // the origin.  Compute the velocity of the circle relative to
+            // the box.
+            Vector2<Real> boxCenter = (box.max + box.min) * (Real)0.5;
+            Vector2<Real> extent = (box.max - box.min) * (Real)0.5;
+            Vector2<Real> C = circle.center - boxCenter;
+            Vector2<Real> V = circleVelocity - boxVelocity;
+
+            // Change signs on components, if necessary, to transform C to the
+            // first quadrant.  Adjust the velocity accordingly.
+            Real sign[2];
+            for (int i = 0; i < 2; ++i)
+            {
+                if (C[i] >= (Real)0)
+                {
+                    sign[i] = (Real)1;
+                }
+                else
+                {
+                    C[i] = -C[i];
+                    V[i] = -V[i];
+                    sign[i] = (Real)-1;
+                }
+            }
+
+            DoQuery(extent, C, circle.radius, V, result);
+
+            if (result.intersectionType != 0)
+            {
+                // Translate back to the original coordinate system.
+                for (int i = 0; i < 2; ++i)
+                {
+                    if (sign[i] < (Real)0)
+                    {
+                        result.contactPoint[i] = -result.contactPoint[i];
+                    }
+                }
+
+                result.contactPoint += boxCenter;
+            }
+            return result;
+        }
+
+    protected:
+        void DoQuery(Vector2<Real> const& K, Vector2<Real> const& C,
+            Real radius, Vector2<Real> const& V, Result& result)
+        {
+            Vector2<Real> delta = C - K;
+            if (delta[1] <= radius)
+            {
+                if (delta[0] <= radius)
+                {
+                    if (delta[1] <= (Real)0)
+                    {
+                        if (delta[0] <= (Real)0)
+                        {
+                            InteriorOverlap(C, result);
+                        }
+                        else
+                        {
+                            EdgeOverlap(0, 1, K, C, delta, radius, result);
+                        }
+                    }
+                    else
+                    {
+                        if (delta[0] <= (Real)0)
+                        {
+                            EdgeOverlap(1, 0, K, C, delta, radius, result);
+                        }
+                        else
+                        {
+                            if (Dot(delta, delta) <= radius * radius)
+                            {
+                                VertexOverlap(K, delta, radius, result);
+                            }
+                            else
+                            {
+                                VertexSeparated(K, delta, V, radius, result);
+                            }
+                        }
+
+                    }
+                }
+                else
+                {
+                    EdgeUnbounded(0, 1, K, C, radius, delta, V, result);
+                }
+            }
+            else
+            {
+                if (delta[0] <= radius)
+                {
+                    EdgeUnbounded(1, 0, K, C, radius, delta, V, result);
+                }
+                else
+                {
+                    VertexUnbounded(K, C, radius, delta, V, result);
+                }
+            }
+        }
+
+    private:
+        void InteriorOverlap(Vector2<Real> const& C, Result& result)
+        {
+            result.intersectionType = -1;
+            result.contactTime = (Real)0;
+            result.contactPoint = C;
+        }
+
+        void EdgeOverlap(int i0, int i1, Vector2<Real> const& K, Vector2<Real> const& C,
+            Vector2<Real> const& delta, Real radius, Result& result)
+        {
+            result.intersectionType = (delta[i0] < radius ? -1 : 1);
+            result.contactTime = (Real)0;
+            result.contactPoint[i0] = K[i0];
+            result.contactPoint[i1] = C[i1];
+        }
+
+        void VertexOverlap(Vector2<Real> const& K0, Vector2<Real> const& delta,
+            Real radius, Result& result)
+        {
+            Real sqrDistance = delta[0] * delta[0] + delta[1] * delta[1];
+            Real sqrRadius = radius * radius;
+            result.intersectionType = (sqrDistance < sqrRadius ? -1 : 1);
+            result.contactTime = (Real)0;
+            result.contactPoint = K0;
+        }
+
+        void VertexSeparated(Vector2<Real> const& K0, Vector2<Real> const& delta0,
+            Vector2<Real> const& V, Real radius, Result& result)
+        {
+            Real q0 = -Dot(V, delta0);
+            if (q0 > (Real)0)
+            {
+                Real dotVPerpD0 = Dot(V, Perp(delta0));
+                Real q2 = Dot(V, V);
+                Real q1 = radius * radius * q2 - dotVPerpD0 * dotVPerpD0;
+                if (q1 >= (Real)0)
+                {
+                    IntersectsVertex(0, 1, K0, q0, q1, q2, result);
+                }
+            }
+        }
+
+        void EdgeUnbounded(int i0, int i1, Vector2<Real> const& K0, Vector2<Real> const& C,
+            Real radius, Vector2<Real> const& delta0, Vector2<Real> const& V, Result& result)
+        {
+            if (V[i0] < (Real)0)
+            {
+                Real dotVPerpD0 = V[i0] * delta0[i1] - V[i1] * delta0[i0];
+                if (radius * V[i1] + dotVPerpD0 >= (Real)0)
+                {
+                    Vector2<Real> K1, delta1;
+                    K1[i0] = K0[i0];
+                    K1[i1] = -K0[i1];
+                    delta1[i0] = C[i0] - K1[i0];
+                    delta1[i1] = C[i1] - K1[i1];
+                    Real dotVPerpD1 = V[i0] * delta1[i1] - V[i1] * delta1[i0];
+                    if (radius * V[i1] + dotVPerpD1 <= (Real)0)
+                    {
+                        IntersectsEdge(i0, i1, K0, C, radius, V, result);
+                    }
+                    else
+                    {
+                        Real q2 = Dot(V, V);
+                        Real q1 = radius * radius * q2 - dotVPerpD1 * dotVPerpD1;
+                        if (q1 >= (Real)0)
+                        {
+                            Real q0 = -(V[i0] * delta1[i0] + V[i1] * delta1[i1]);
+                            IntersectsVertex(i0, i1, K1, q0, q1, q2, result);
+                        }
+                    }
+                }
+                else
+                {
+                    Real q2 = Dot(V, V);
+                    Real q1 = radius * radius * q2 - dotVPerpD0 * dotVPerpD0;
+                    if (q1 >= (Real)0)
+                    {
+                        Real q0 = -(V[i0] * delta0[i0] + V[i1] * delta0[i1]);
+                        IntersectsVertex(i0, i1, K0, q0, q1, q2, result);
+                    }
+                }
+            }
+        }
+
+        void VertexUnbounded(Vector2<Real> const& K0, Vector2<Real> const& C, Real radius,
+            Vector2<Real> const& delta0, Vector2<Real> const& V, Result& result)
+        {
+            if (V[0] < (Real)0 && V[1] < (Real)0)
+            {
+                Real dotVPerpD0 = Dot(V, Perp(delta0));
+                if (radius * V[0] - dotVPerpD0 <= (Real)0)
+                {
+                    if (-radius * V[1] - dotVPerpD0 >= (Real)0)
+                    {
+                        Real q2 = Dot(V, V);
+                        Real q1 = radius * radius * q2 - dotVPerpD0 * dotVPerpD0;
+                        Real q0 = -Dot(V, delta0);
+                        IntersectsVertex(0, 1, K0, q0, q1, q2, result);
+                    }
+                    else
+                    {
+                        Vector2<Real> K1{ K0[0], -K0[1] };
+                        Vector2<Real> delta1 = C - K1;
+                        Real dotVPerpD1 = Dot(V, Perp(delta1));
+                        if (-radius * V[1] - dotVPerpD1 >= (Real)0)
+                        {
+                            IntersectsEdge(0, 1, K0, C, radius, V, result);
+                        }
+                        else
+                        {
+                            Real q2 = Dot(V, V);
+                            Real q1 = radius * radius * q2 - dotVPerpD1 * dotVPerpD1;
+                            if (q1 >= (Real)0)
+                            {
+                                Real q0 = -Dot(V, delta1);
+                                IntersectsVertex(0, 1, K1, q0, q1, q2, result);
+                            }
+                        }
+                    }
+                }
+                else
+                {
+                    Vector2<Real> K2{ -K0[0], K0[1] };
+                    Vector2<Real> delta2 = C - K2;
+                    Real dotVPerpD2 = Dot(V, Perp(delta2));
+                    if (radius * V[0] - dotVPerpD2 <= (Real)0)
+                    {
+                        IntersectsEdge(1, 0, K0, C, radius, V, result);
+                    }
+                    else
+                    {
+                        Real q2 = Dot(V, V);
+                        Real q1 = radius * radius * q2 - dotVPerpD2 * dotVPerpD2;
+                        if (q1 >= (Real)0)
+                        {
+                            Real q0 = -Dot(V, delta2);
+                            IntersectsVertex(1, 0, K2, q0, q1, q2, result);
+                        }
+                    }
+                }
+            }
+        }
+
+        void IntersectsVertex(int i0, int i1, Vector2<Real> const& K,
+            Real q0, Real q1, Real q2, Result& result)
+        {
+            result.intersectionType = +1;
+            result.contactTime = (q0 - std::sqrt(q1)) / q2;
+            result.contactPoint[i0] = K[i0];
+            result.contactPoint[i1] = K[i1];
+        }
+
+        void IntersectsEdge(int i0, int i1, Vector2<Real> const& K0, Vector2<Real> const& C,
+            Real radius, Vector2<Real> const& V, Result& result)
+        {
+            result.intersectionType = +1;
+            result.contactTime = (K0[i0] + radius - C[i0]) / V[i0];
+            result.contactPoint[i0] = K0[i0];
+            result.contactPoint[i1] = C[i1] + result.contactTime * V[i1];
+        }
+    };
+}
diff --git a/Engine/Plugins/Experimental/GeometryProcessing/Source/GeometryAlgorithms/Private/GTEngine/Mathematics/GteIntrAlignedBox2OrientedBox2.h b/Engine/Plugins/Experimental/GeometryProcessing/Source/GeometryAlgorithms/Private/GTEngine/Mathematics/GteIntrAlignedBox2OrientedBox2.h
new file mode 100644
--- /dev/null
+++ b/Engine/Plugins/Experimental/GeometryProcessing/Source/GeometryAlgorithms/Private/GTEngine/Mathematics/GteIntrAlignedBox2OrientedBox2.h
@@ -0,0 +1,110 @@
+// David Eberly, Geometric Tools, Redmond WA 98052
+// Copyright (c) 1998-2019
+// Distributed under the Boost Software License, Version 1.0.
+// http://www.boost.org/LICENSE_1_0.txt
+// http://www.geometrictools.com/License/Boost/LICENSE_1_0.txt
+// File Version: 3.0.0 (2016/06/19)
+
+#pragma once
+
+#include <GTEngine/Mathematics/GteVector2.h>
+#include <GTEngine/Mathematics/GteAlignedBox.h>
+#include <GTEngine/Mathematics/GteOrientedBox.h>
+#include <GTEngine/Mathematics/GteFIQuery.h>
+#include <GTEngine/Mathematics/GteTIQuery.h>
+
+// The queries consider the box to be a solid.
+//
+// The test-intersection query uses the method of separating axes.  The set of
+// potential separating directions includes the 2 edge normals of box0 and the
+// 2 edge normals of box1.  The integer 'separating' identifies the axis that
+// reported separation; there may be more than one but only one is reported.
+// The value is 0 when box0.axis[0] separates, 1 when box0.axis[1] separates,
+// 2 when box1.axis[0] separates, or 3 when box1.axis[1] separates.
+
+namespace gte
+{
+
+template <typename Real>
+class TIQuery<Real, AlignedBox2<Real>, OrientedBox2<Real>>
+{
+public:
+    struct Result
+    {
+        bool intersect;
+        int separating;
+    };
+
+    Result operator()(AlignedBox2<Real> const& box0,
+        OrientedBox2<Real> const& box1);
+};
+
+
+template <typename Real>
+typename TIQuery<Real, AlignedBox2<Real>, OrientedBox2<Real>>::Result
+TIQuery<Real, AlignedBox2<Real>, OrientedBox2<Real>>::operator()(
+    AlignedBox2<Real> const& box0, OrientedBox2<Real> const& box1)
+{
+    Result result;
+
+    // Get the centered form of the aligned box.  The axes are implicitly
+    // A0[0] = (1,0) and A0[1] = (0,1).
+    Vector2<Real> C0, E0;
+    box0.GetCenteredForm(C0, E0);
+
+    // Convenience variables.
+    Vector2<Real> const& C1 = box1.center;
+    Vector2<Real> const* A1 = &box1.axis[0];
+    Vector2<Real> const& E1 = box1.extent;
+
+    // Compute difference of box centers.
+    Vector2<Real> D = C1 - C0;
+
+    Real absDot01[2][2], rSum;
+
+    // Test box0.axis[0] = (1,0).
+    absDot01[0][0] = std::abs(A1[0][0]);
+    absDot01[0][1] = std::abs(A1[1][0]);
+    rSum = E0[0] + E1[0] * absDot01[0][0] + E1[1] * absDot01[0][1];
+    if (std::abs(D[0]) > rSum)
+    {
+        result.intersect = false;
+        result.separating = 0;
+        return result;
+    }
+
+    // Test axis box0.axis[1] = (0,1).
+    absDot01[1][0] = std::abs(A1[0][1]);
+    absDot01[1][1] = std::abs(A1[1][1]);
+    rSum = E0[1] + E1[0] * absDot01[1][0] + E1[1] * absDot01[1][1];
+    if (std::abs(D[1]) > rSum)
+    {
+        result.intersect = false;
+        result.separating = 1;
+        return result;
+    }
+
+    // Test axis box1.axis[0].
+    rSum = E1[0] + E0[0] * absDot01[0][0] + E0[1] * absDot01[1][0];
+    if (std::abs(Dot(A1[0], D)) > rSum)
+    {
+        result.intersect = false;
+        result.separating = 2;
+        return result;
+    }
+
+    // Test axis box1.axis[1].
+    rSum = E1[1] + E0[0] * absDot01[0][1] + E0[1] * absDot01[1][1];
+    if (std::abs(Dot(A1[1], D)) > rSum)
+    {
+        result.intersect = false;
+        result.separating = 3;
+        return result;
+    }
+
+    result.intersect = true;
+    return result;
+}
+
+
+}
diff --git a/Engine/Plugins/Experimental/GeometryProcessing/Source/GeometryAlgorithms/Private/GTEngine/Mathematics/GteIntrAlignedBox3AlignedBox3.h b/Engine/Plugins/Experimental/GeometryProcessing/Source/GeometryAlgorithms/Private/GTEngine/Mathematics/GteIntrAlignedBox3AlignedBox3.h
new file mode 100644
--- /dev/null
+++ b/Engine/Plugins/Experimental/GeometryProcessing/Source/GeometryAlgorithms/Private/GTEngine/Mathematics/GteIntrAlignedBox3AlignedBox3.h
@@ -0,0 +1,109 @@
+// David Eberly, Geometric Tools, Redmond WA 98052
+// Copyright (c) 1998-2019
+// Distributed under the Boost Software License, Version 1.0.
+// http://www.boost.org/LICENSE_1_0.txt
+// http://www.geometrictools.com/License/Boost/LICENSE_1_0.txt
+// File Version: 3.0.0 (2016/06/19)
+
+#pragma once
+
+#include <GTEngine/Mathematics/GteAlignedBox.h>
+#include <GTEngine/Mathematics/GteFIQuery.h>
+#include <GTEngine/Mathematics/GteTIQuery.h>
+
+// The queries consider the box to be a solid.
+//
+// The aligned-aligned queries use simple min-max comparisions.  The
+// interesection of aligned boxes is an aligned box, possibly degenerate,
+// where min[d] == max[d] for at least one dimension d.
+
+namespace gte
+{
+
+template <typename Real>
+class TIQuery<Real, AlignedBox3<Real>, AlignedBox3<Real>>
+{
+public:
+    struct Result
+    {
+        bool intersect;
+    };
+
+    Result operator()(AlignedBox3<Real> const& box0,
+        AlignedBox3<Real> const& box1);
+};
+
+template <typename Real>
+class FIQuery<Real, AlignedBox3<Real>, AlignedBox3<Real>>
+{
+public:
+    struct Result
+    {
+        bool intersect;
+        AlignedBox3<Real> box;
+    };
+
+    Result operator()(AlignedBox3<Real> const& box0,
+        AlignedBox3<Real> const& box1);
+};
+
+
+template <typename Real>
+typename TIQuery<Real, AlignedBox3<Real>, AlignedBox3<Real>>::Result
+TIQuery<Real, AlignedBox3<Real>, AlignedBox3<Real>>::operator()(
+    AlignedBox3<Real> const& box0, AlignedBox3<Real> const& box1)
+{
+    Result result;
+    for (int i = 0; i < 3; i++)
+    {
+        if (box0.max[i] < box1.min[i] || box0.min[i] > box1.max[i])
+        {
+            result.intersect = false;
+            return result;
+        }
+    }
+    result.intersect = true;
+    return result;
+}
+
+template <typename Real>
+typename FIQuery<Real, AlignedBox3<Real>, AlignedBox3<Real>>::Result
+FIQuery<Real, AlignedBox3<Real>, AlignedBox3<Real>>::operator()(
+    AlignedBox3<Real> const& box0, AlignedBox3<Real> const& box1)
+{
+    Result result;
+    for (int i = 0; i < 3; i++)
+    {
+        if (box0.max[i] < box1.min[i] || box0.min[i] > box1.max[i])
+        {
+            result.intersect = false;
+            return result;
+        }
+    }
+
+    for (int i = 0; i < 3; i++)
+    {
+        if (box0.max[i] <= box1.max[i])
+        {
+            result.box.max[i] = box0.max[i];
+        }
+        else
+        {
+            result.box.max[i] = box1.max[i];
+        }
+
+        if (box0.min[i] <= box1.min[i])
+        {
+            result.box.min[i] = box1.min[i];
+        }
+        else
+        {
+            result.box.min[i] = box0.min[i];
+        }
+    }
+    result.intersect = true;
+    return result;
+}
+
+
+}
