move GeometryProcessing plugin from NotForLicensees to Experimental

diff --git a/Engine/Plugins/Experimental/GeometryProcessing/Source/GeometryAlgorithms/Private/GTEngine/Mathematics/GteIntrAlignedBox3Cone3.h b/Engine/Plugins/Experimental/GeometryProcessing/Source/GeometryAlgorithms/Private/GTEngine/Mathematics/GteIntrAlignedBox3Cone3.h
new file mode 100644
--- /dev/null
+++ b/Engine/Plugins/Experimental/GeometryProcessing/Source/GeometryAlgorithms/Private/GTEngine/Mathematics/GteIntrAlignedBox3Cone3.h
@@ -0,0 +1,984 @@
+// David Eberly, Geometric Tools, Redmond WA 98052
+// Copyright (c) 1998-2019
+// Distributed under the Boost Software License, Version 1.0.
+// http://www.boost.org/LICENSE_1_0.txt
+// http://www.geometrictools.com/License/Boost/LICENSE_1_0.txt
+// File Version: 3.0.2 (2018/11/29)
+
+#pragma once
+
+#include <GTEngine/Mathematics/GteAlignedBox.h>
+#include <GTEngine/Mathematics/GteCone.h>
+#include <GTEngine/Mathematics/GteIntrRay3AlignedBox3.h>
+#include <GTEngine/Mathematics/GteIntrSegment3AlignedBox3.h>
+
+// Test for intersection of a box and a cone.  The cone can be infinite
+//   0 <= minHeight < maxHeight = std::numeric_limits<Real>::max()
+// or finite (cone frustum)
+//   0 <= minHeight < maxHeight < std::numeric_limits<Real>::max().
+// The algorithm is described in
+//   https://www.geometrictools.com/Documentation/IntersectionBoxCone.pdf
+// and reports an intersection only when the intersection set has positive
+// volume.  For example, let the box be outside the cone.  If the box is
+// below the minHeight plane at the cone vertex and just touches the cone
+// vertex, no intersection is reported.  If the box is above the maxHeight
+// plane and just touches the disk capping the cone, either at a single
+// point, a line segment of points or a polygon of points, no intersection
+// is reported.
+
+namespace gte
+{
+    template <typename Real>
+    class TIQuery<Real, AlignedBox<3, Real>, Cone<3, Real>>
+    {
+    public:
+        struct Result
+        {
+            bool intersect;
+        };
+
+        TIQuery()
+            :
+            mNumCandidateEdges(0)
+        {
+            // An edge is { v0, v1 }, where v0 and v1 are relative to mVertices
+            // with v0 < v1.
+            mEdges[ 0] = { 0, 1 };
+            mEdges[ 1] = { 1, 3 };
+            mEdges[ 2] = { 2, 3 };
+            mEdges[ 3] = { 0, 2 };
+            mEdges[ 4] = { 4, 5 };
+            mEdges[ 5] = { 5, 7 };
+            mEdges[ 6] = { 6, 7 };
+            mEdges[ 7] = { 4, 6 };
+            mEdges[ 8] = { 0, 4 };
+            mEdges[ 9] = { 1, 5 };
+            mEdges[10] = { 3, 7 };
+            mEdges[11] = { 2, 6 };
+
+            // A face is { { v0, v1, v2, v3 }, { e0, e1, e2, e3 } }, where
+            // { v0, v1, v2, v3 } are relative to mVertices with
+            // v0 = min(v0,v1,v2,v3) and where { e0, e1, e2, e3 } are relative
+            // to mEdges.  For example, mFaces[0] has vertices { 0, 4, 6, 2 }.
+            // The edge { 0, 4 } is mEdges[8], the edge { 4, 6 } is mEdges[7],
+            // the edge { 6, 2 } is mEdges[11] and the edge { 2, 0 } is
+            // mEdges[3]; thus, the edge indices are { 8, 7, 11, 3 }.
+            mFaces[0] = { { 0, 4, 6, 2 }, {  8,  7, 11,  3 } };
+            mFaces[1] = { { 1, 3, 7, 5 }, {  1, 10,  5,  9 } };
+            mFaces[2] = { { 0, 1, 5, 4 }, {  0,  9,  4,  8 } };
+            mFaces[3] = { { 2, 6, 7, 3 }, { 11,  6, 10,  2 } };
+            mFaces[4] = { { 0, 2, 3, 1 }, {  3,  2,  1,  0 } };
+            mFaces[5] = { { 4, 5, 7, 6 }, {  4,  5,  6,  7 } };
+
+            // Clear the edges.
+            std::array<size_t, 2> ezero = { 0, 0 };
+            mCandidateEdges.fill(ezero);
+            for (size_t r = 0; r < MAX_VERTICES; ++r)
+            {
+                mAdjacencyMatrix[r].fill(0);
+            }
+
+            mConfiguration[ 0] = &TIQuery::NNNN_0;
+            mConfiguration[ 1] = &TIQuery::NNNZ_1;
+            mConfiguration[ 2] = &TIQuery::NNNP_2;
+            mConfiguration[ 3] = &TIQuery::NNZN_3;
+            mConfiguration[ 4] = &TIQuery::NNZZ_4;
+            mConfiguration[ 5] = &TIQuery::NNZP_5;
+            mConfiguration[ 6] = &TIQuery::NNPN_6;
+            mConfiguration[ 7] = &TIQuery::NNPZ_7;
+            mConfiguration[ 8] = &TIQuery::NNPP_8;
+            mConfiguration[ 9] = &TIQuery::NZNN_9;
+            mConfiguration[10] = &TIQuery::NZNZ_10;
+            mConfiguration[11] = &TIQuery::NZNP_11;
+            mConfiguration[12] = &TIQuery::NZZN_12;
+            mConfiguration[13] = &TIQuery::NZZZ_13;
+            mConfiguration[14] = &TIQuery::NZZP_14;
+            mConfiguration[15] = &TIQuery::NZPN_15;
+            mConfiguration[16] = &TIQuery::NZPZ_16;
+            mConfiguration[17] = &TIQuery::NZPP_17;
+            mConfiguration[18] = &TIQuery::NPNN_18;
+            mConfiguration[19] = &TIQuery::NPNZ_19;
+            mConfiguration[20] = &TIQuery::NPNP_20;
+            mConfiguration[21] = &TIQuery::NPZN_21;
+            mConfiguration[22] = &TIQuery::NPZZ_22;
+            mConfiguration[23] = &TIQuery::NPZP_23;
+            mConfiguration[24] = &TIQuery::NPPN_24;
+            mConfiguration[25] = &TIQuery::NPPZ_25;
+            mConfiguration[26] = &TIQuery::NPPP_26;
+            mConfiguration[27] = &TIQuery::ZNNN_27;
+            mConfiguration[28] = &TIQuery::ZNNZ_28;
+            mConfiguration[29] = &TIQuery::ZNNP_29;
+            mConfiguration[30] = &TIQuery::ZNZN_30;
+            mConfiguration[31] = &TIQuery::ZNZZ_31;
+            mConfiguration[32] = &TIQuery::ZNZP_32;
+            mConfiguration[33] = &TIQuery::ZNPN_33;
+            mConfiguration[34] = &TIQuery::ZNPZ_34;
+            mConfiguration[35] = &TIQuery::ZNPP_35;
+            mConfiguration[36] = &TIQuery::ZZNN_36;
+            mConfiguration[37] = &TIQuery::ZZNZ_37;
+            mConfiguration[38] = &TIQuery::ZZNP_38;
+            mConfiguration[39] = &TIQuery::ZZZN_39;
+            mConfiguration[40] = &TIQuery::ZZZZ_40;
+            mConfiguration[41] = &TIQuery::ZZZP_41;
+            mConfiguration[42] = &TIQuery::ZZPN_42;
+            mConfiguration[43] = &TIQuery::ZZPZ_43;
+            mConfiguration[44] = &TIQuery::ZZPP_44;
+            mConfiguration[45] = &TIQuery::ZPNN_45;
+            mConfiguration[46] = &TIQuery::ZPNZ_46;
+            mConfiguration[47] = &TIQuery::ZPNP_47;
+            mConfiguration[48] = &TIQuery::ZPZN_48;
+            mConfiguration[49] = &TIQuery::ZPZZ_49;
+            mConfiguration[50] = &TIQuery::ZPZP_50;
+            mConfiguration[51] = &TIQuery::ZPPN_51;
+            mConfiguration[52] = &TIQuery::ZPPZ_52;
+            mConfiguration[53] = &TIQuery::ZPPP_53;
+            mConfiguration[54] = &TIQuery::PNNN_54;
+            mConfiguration[55] = &TIQuery::PNNZ_55;
+            mConfiguration[56] = &TIQuery::PNNP_56;
+            mConfiguration[57] = &TIQuery::PNZN_57;
+            mConfiguration[58] = &TIQuery::PNZZ_58;
+            mConfiguration[59] = &TIQuery::PNZP_59;
+            mConfiguration[60] = &TIQuery::PNPN_60;
+            mConfiguration[61] = &TIQuery::PNPZ_61;
+            mConfiguration[62] = &TIQuery::PNPP_62;
+            mConfiguration[63] = &TIQuery::PZNN_63;
+            mConfiguration[64] = &TIQuery::PZNZ_64;
+            mConfiguration[65] = &TIQuery::PZNP_65;
+            mConfiguration[66] = &TIQuery::PZZN_66;
+            mConfiguration[67] = &TIQuery::PZZZ_67;
+            mConfiguration[68] = &TIQuery::PZZP_68;
+            mConfiguration[69] = &TIQuery::PZPN_69;
+            mConfiguration[70] = &TIQuery::PZPZ_70;
+            mConfiguration[71] = &TIQuery::PZPP_71;
+            mConfiguration[72] = &TIQuery::PPNN_72;
+            mConfiguration[73] = &TIQuery::PPNZ_73;
+            mConfiguration[74] = &TIQuery::PPNP_74;
+            mConfiguration[75] = &TIQuery::PPZN_75;
+            mConfiguration[76] = &TIQuery::PPZZ_76;
+            mConfiguration[77] = &TIQuery::PPZP_77;
+            mConfiguration[78] = &TIQuery::PPPN_78;
+            mConfiguration[79] = &TIQuery::PPPZ_79;
+            mConfiguration[80] = &TIQuery::PPPP_80;
+        }
+
+        Result operator()(AlignedBox<3, Real> const& box, Cone<3, Real> const& cone)
+        {
+            Result result;
+
+            // Quick-rejectance test.  Determine whether the box is outside
+            // the slab bounded by the minimum and maximum height planes.
+            // When outside the slab, the box vertices are not required by the
+            // cone-box intersection query, so the vertices are not yet
+            // computed.
+            Real boxMinHeight(0), boxMaxHeight(0);
+            ComputeBoxHeightInterval(box, cone, boxMinHeight, boxMaxHeight);
+            if (boxMaxHeight <= cone.minHeight || boxMinHeight >= cone.maxHeight)
+            {
+                // There is no volumetric overlap of the box and the cone. The
+                // box is clipped entirely.
+                result.intersect = false;
+                return result;
+            }
+
+            // Quick-acceptance test.  Determine whether the cone axis
+            // intersects the box.
+            if (ConeAxisIntersectsBox(box, cone))
+            {
+                result.intersect = true;
+                return result;
+            }
+
+            // Test for box fully inside the slab.  When inside the slab, the
+            // box vertices are required by the cone-box intersection query,
+            // so they are computed here; they are also required in the
+            // remaining cases.  Also when inside the slab, the box edges are
+            // the candidates, so they are copied to mCandidateEdges.
+            if (BoxFullyInConeSlab(box, boxMinHeight, boxMaxHeight, cone))
+            {
+                result.intersect = CandidatesHavePointInsideCone(cone);
+                return result;
+            }
+
+            // Clear the candidates array and adjacency matrix.
+            ClearCandidates();
+
+            // The box intersects at least one plane.  Compute the box-plane
+            // edge-interior intersection points.  Insert the box subedges into
+            // the array of candidate edges.
+            ComputeCandidatesOnBoxEdges(cone);
+
+            // Insert any relevant box face-interior clipped edges into the array
+            // of candidate edges.
+            ComputeCandidatesOnBoxFaces();
+
+            result.intersect = CandidatesHavePointInsideCone(cone);
+            return result;
+        }
+
+    protected:
+        // The constants here are described in the comments below.
+        enum
+        {
+            NUM_BOX_VERTICES = 8,
+            NUM_BOX_EDGES = 12,
+            NUM_BOX_FACES = 6,
+            MAX_VERTICES = 32,
+            VERTEX_MIN_BASE = 8,
+            VERTEX_MAX_BASE = 20,
+            MAX_CANDIDATE_EDGES = 496,
+            NUM_CONFIGURATIONS = 81
+        };
+
+        // The box topology is that of a cube whose vertices have components
+        // in {0,1}.  The cube vertices are indexed by
+        //   0: (0,0,0), 1: (1,0,0), 2: (1,1,0), 3: (0,1,0)
+        //   4: (0,0,1), 5: (1,0,1), 6: (1,1,1), 7: (0,1,1)
+
+        // The first 8 vertices are the box corners, the next 12 vertices are
+        // reserved for hmin-edge points and the final 12 vertices are reserved
+        // for the hmax-edge points.  The conservative upper bound of the number
+        // of vertices is 8 + 12 + 12 = 32.
+        std::array<Vector3<Real>, MAX_VERTICES> mVertices;
+
+        // The box has 12 edges stored in mEdges.  An edge is mEdges[i] =
+        // { v0, v1 }, where the indices v0 and v1 are relative to mVertices
+        // with v0 < v1.
+        std::array<std::array<size_t, 2>, NUM_BOX_EDGES> mEdges;
+
+        // The box has 6 faces stored in mFaces.  A face is mFaces[i] =
+        // { { v0, v1, v2, v3 }, { e0, e1, e2, e3 } }, where the face corner
+        // vertices are { v0, v1, v2, v3 }.  These indices are relative to
+        // mVertices.  The indices { e0, e1, e2, e3 } are relative to mEdges.
+        // The index e0 refers to edge { v0, v1 }, the index e1 refers to edge
+        // { v1, v2 }, the index e2 refers to edge { v2, v3 } and the index e3
+        // refers to edge { v3, v0 }.  The ordering of vertices for the faces
+        // is/ counterclockwise when viewed from outside the box.  The choice
+        // of initial vertex affects how you implement the graph data
+        // structure.  In this implementation, the initial vertex has minimum
+        // index for all vertices of that face.  The faces themselves are
+        // listed as -x face, +x face, -y face, +y face, -z face and +z face.
+        struct Face
+        {
+            std::array<size_t, 4> v, e;
+        };
+        std::array<Face, NUM_BOX_FACES> mFaces;
+
+        // Store the signed distances from the minimum and maximum height
+        // planes for the cone to the projection of the box vertices onto the
+        // cone axis.
+        std::array<Real, NUM_BOX_VERTICES> mProjectionMin, mProjectionMax;
+
+        // The mCandidateEdges array stores the edges of the clipped box that
+        // are candidates for containing the optimizing point.  The maximum
+        // number of candidate edges is 1 + 2 + ... + 31 = 496, which is a
+        // conservative bound because not all pairs of vertices form edges on
+        // box faces.  The candidate edges are stored as (v0,v1) with v0 < v1.
+        // The implementation is designed so that during a single query, the
+        // number of candidate edges can only grow.
+        size_t mNumCandidateEdges;
+        std::array<std::array<size_t, 2>, MAX_CANDIDATE_EDGES> mCandidateEdges;
+
+        // The mAdjancencyMatrix is a simple representation of edges in the
+        // graph G = (V,E) that represents the (wireframe) clipped box.  An
+        // edge (r,c) does not exist when mAdjancencyMatrix[r][c] = 0.  If an
+        // edge (r,c) does exist, it is appended to mCandidateEdges at index k
+        // and the adjacency matrix is set to mAdjacencyMatrix[r][c] = 1.
+        // This allows for a fast edge-in-graph test and a fast and efficient
+        // clear of mCandidateEdges and mAdjacencyMatrix.
+        std::array<std::array<size_t, MAX_VERTICES>, MAX_VERTICES> mAdjacencyMatrix;
+
+        typedef void (TIQuery::*ConfigurationFunction)(size_t, Face const&);
+        std::array<ConfigurationFunction, NUM_CONFIGURATIONS> mConfiguration;
+
+        static void ComputeBoxHeightInterval(AlignedBox<3, Real> const& box, Cone<3, Real> const& cone,
+            Real& boxMinHeight, Real& boxMaxHeight)
+        {
+            Vector<3, Real> C, e;
+            box.GetCenteredForm(C, e);
+            Vector<3, Real> const& V = cone.ray.origin;
+            Vector<3, Real> const& U = cone.ray.direction;
+            Vector<3, Real> CmV = C - V;
+            Real DdCmV = Dot(U, CmV);
+            Real radius = e[0] * std::abs(U[0]) + e[1] * std::abs(U[1]) + e[2] * std::abs(U[2]);
+            boxMinHeight = DdCmV - radius;
+            boxMaxHeight = DdCmV + radius;
+        }
+
+        static bool ConeAxisIntersectsBox(AlignedBox<3, Real> const& box, Cone<3, Real> const& cone)
+        {
+            if (cone.maxHeight < std::numeric_limits<Real>::max())
+            {
+                Segment<3, Real> segment;
+                segment.p[0] = cone.ray.origin + cone.minHeight * cone.ray.direction;
+                segment.p[1] = cone.ray.origin + cone.maxHeight * cone.ray.direction;
+                auto sbResult = TIQuery<Real, Segment<3, Real>, AlignedBox<3, Real>>()(segment, box);
+                if (sbResult.intersect)
+                {
+                    return true;
+                }
+            }
+            else
+            {
+                Ray<3, Real> ray;
+                ray.origin = cone.ray.origin + cone.minHeight * cone.ray.direction;
+                ray.direction = cone.ray.direction;
+                auto rbResult = TIQuery<Real, Ray<3, Real>, AlignedBox<3, Real>>()(ray, box);
+                if (rbResult.intersect)
+                {
+                    return true;
+                }
+            }
+            return false;
+        }
+
+        bool BoxFullyInConeSlab(AlignedBox<3, Real> const& box, Real boxMinHeight, Real boxMaxHeight, Cone<3, Real> const& cone)
+        {
+            // Compute the box vertices relative to cone vertex as origin.
+            mVertices[0] = { box.min[0], box.min[1], box.min[2] };
+            mVertices[1] = { box.max[0], box.min[1], box.min[2] };
+            mVertices[2] = { box.min[0], box.max[1], box.min[2] };
+            mVertices[3] = { box.max[0], box.max[1], box.min[2] };
+            mVertices[4] = { box.min[0], box.min[1], box.max[2] };
+            mVertices[5] = { box.max[0], box.min[1], box.max[2] };
+            mVertices[6] = { box.min[0], box.max[1], box.max[2] };
+            mVertices[7] = { box.max[0], box.max[1], box.max[2] };
+            for (int i = 0; i < NUM_BOX_VERTICES; ++i)
+            {
+                mVertices[i] -= cone.ray.origin;
+            }
+
+            if (cone.minHeight <= boxMinHeight && boxMaxHeight <= cone.maxHeight)
+            {
+                // The box is fully inside, so no clipping is necessary.
+                std::copy(mEdges.begin(), mEdges.end(), mCandidateEdges.begin());
+                mNumCandidateEdges = 12;
+                return true;
+            }
+            return false;
+        }
+
+        static bool HasPointInsideCone(Vector<3, Real> const& P0, Vector<3, Real> const& P1,
+            Cone<3, Real> const& cone)
+        {
+            // Define F(X) = Dot(U,X - V)/|X - V|, where U is the unit-length
+            // cone axis direction and V is the cone vertex.  The incoming
+            // points P0 and P1 are relative to V; that is, the original
+            // points are X0 = P0 + V and X1 = P1 + V.  The segment <P0,P1>
+            // and cone intersect when a segment point X is inside the cone;
+            // that is, when F(X) > cosAngle.  The comparison is converted to
+            // an equivalent one that does not involve divisions in order to
+            // avoid a division by zero if a vertex or edge contain (0,0,0).
+            // The function is G(X) = Dot(U,X-V) - cosAngle*Length(X-V).
+            Vector<3, Real> const& U = cone.ray.direction;
+
+            // Test whether P0 or P1 is inside the cone.
+            Real g = Dot(U, P0) - cone.cosAngle * Length(P0);
+            if (g > (Real)0)
+            {
+                // X0 = P0 + V is inside the cone.
+                return true;
+            }
+
+            g = Dot(U, P1) - cone.cosAngle * Length(P1);
+            if (g > (Real)0)
+            {
+                // X1 = P1 + V is inside the cone.
+                return true;
+            }
+
+            // Test whether an interior segment point is inside the cone.
+            Vector<3, Real> E = P1 - P0;
+            Vector<3, Real> crossP0U = Cross(P0, U);
+            Vector<3, Real> crossP0E = Cross(P0, E);
+            Real dphi0 = Dot(crossP0E, crossP0U);
+            if (dphi0 > (Real)0)
+            {
+                Vector3<Real> crossP1U = Cross(P1, U);
+                Real dphi1 = Dot(crossP0E, crossP1U);
+                if (dphi1 < (Real)0)
+                {
+                    Real t = dphi0 / (dphi0 - dphi1);
+                    Vector<3, Real> PMax = P0 + t * E;
+                    g = Dot(U, PMax) - cone.cosAngle * Length(PMax);
+                    if (g > (Real)0)
+                    {
+                        // The edge point XMax = Pmax + V is inside the cone.
+                        return true;
+                    }
+                }
+            }
+
+            return false;
+        }
+
+        bool CandidatesHavePointInsideCone(Cone<3, Real> const& cone) const
+        {
+            for (size_t i = 0; i < mNumCandidateEdges; ++i)
+            {
+                auto const& edge = mCandidateEdges[i];
+                Vector<3, Real> const& P0 = mVertices[edge[0]];
+                Vector<3, Real> const& P1 = mVertices[edge[1]];
+                if (HasPointInsideCone(P0, P1, cone))
+                {
+                    return true;
+                }
+            }
+            return false;
+        }
+
+        void ComputeCandidatesOnBoxEdges(Cone<3, Real> const& cone)
+        {
+            for (size_t i = 0; i < NUM_BOX_VERTICES; ++i)
+            {
+                Real h = Dot(cone.ray.direction, mVertices[i]);
+                mProjectionMin[i] = cone.minHeight - h;
+                mProjectionMax[i] = h - cone.maxHeight;
+            }
+
+            size_t v0 = VERTEX_MIN_BASE, v1 = VERTEX_MAX_BASE;
+            for (size_t i = 0; i < NUM_BOX_EDGES; ++i, ++v0, ++v1)
+            {
+                auto const& edge = mEdges[i];
+
+                // In the next blocks, the sign comparisons can be expressed
+                // instead as "s0 * s1 < 0". The multiplication could lead to
+                // floating-point underflow where the product becomes 0, so I
+                // avoid that approach.
+
+                // Process the hmin-plane.
+                Real p0Min = mProjectionMin[edge[0]];
+                Real p1Min = mProjectionMin[edge[1]];
+                bool clipMin = (p0Min < (Real)0 && p1Min > (Real)0) || (p0Min > (Real)0 && p1Min < (Real)0);
+                if (clipMin)
+                {
+                    mVertices[v0] = (p1Min * mVertices[edge[0]] - p0Min * mVertices[edge[1]]) / (p1Min - p0Min);
+                }
+
+                // Process the hmax-plane.
+                Real p0Max = mProjectionMax[edge[0]];
+                Real p1Max = mProjectionMax[edge[1]];
+                bool clipMax = (p0Max < (Real)0 && p1Max > (Real)0) || (p0Max > (Real)0 && p1Max < (Real)0);
+                if (clipMax)
+                {
+                    mVertices[v1] = (p1Max * mVertices[edge[0]] - p0Max * mVertices[edge[1]]) / (p1Max - p0Max);
+                }
+
+                // Get the candidate edges that are contained by the box edges.
+                if (clipMin)
+                {
+                    if (clipMax)
+                    {
+                        InsertEdge(v0, v1);
+                    }
+                    else
+                    {
+                        if (p0Min < (Real)0)
+                        {
+                            InsertEdge(edge[0], v0);
+                        }
+                        else  // p1Min < 0
+                        {
+                            InsertEdge(edge[1], v0);
+                        }
+                    }
+                }
+                else if (clipMax)
+                {
+                    if (p0Max < (Real)0)
+                    {
+                        InsertEdge(edge[0], v1);
+                    }
+                    else  // p1Max < 0
+                    {
+                        InsertEdge(edge[1], v1);
+                    }
+                }
+                else
+                {
+                    // No clipping has occurred.  If the edge is inside the box,
+                    // it is a candidate edge.  To be inside the box, the p*min
+                    // and p*max values must be nonpositive.
+                    if (p0Min <= (Real)0 && p1Min <= (Real)0 && p0Max <= (Real)0 && p1Max <= (Real)0)
+                    {
+                        InsertEdge(edge[0], edge[1]);
+                    }
+                }
+            }
+        }
+
+        void ComputeCandidatesOnBoxFaces()
+        {
+            Real p0, p1, p2, p3;
+            size_t b0, b1, b2, b3, index;
+            for (size_t i = 0; i < NUM_BOX_FACES; ++i)
+            {
+                auto const& face = mFaces[i];
+
+                // Process the hmin-plane.
+                p0 = mProjectionMin[face.v[0]];
+                p1 = mProjectionMin[face.v[1]];
+                p2 = mProjectionMin[face.v[2]];
+                p3 = mProjectionMin[face.v[3]];
+                b0 = (p0 < (Real)0 ? 0 : (p0 > (Real)0 ? 2 : 1));
+                b1 = (p1 < (Real)0 ? 0 : (p1 > (Real)0 ? 2 : 1));
+                b2 = (p2 < (Real)0 ? 0 : (p2 > (Real)0 ? 2 : 1));
+                b3 = (p3 < (Real)0 ? 0 : (p3 > (Real)0 ? 2 : 1));
+                index = b3 + 3 * (b2 + 3 * (b1 + 3 * b0));
+                (this->*mConfiguration[index])(VERTEX_MIN_BASE, face);
+
+                // Process the hmax-plane.
+                p0 = mProjectionMax[face.v[0]];
+                p1 = mProjectionMax[face.v[1]];
+                p2 = mProjectionMax[face.v[2]];
+                p3 = mProjectionMax[face.v[3]];
+                b0 = (p0 < (Real)0 ? 0 : (p0 > (Real)0 ? 2 : 1));
+                b1 = (p1 < (Real)0 ? 0 : (p1 > (Real)0 ? 2 : 1));
+                b2 = (p2 < (Real)0 ? 0 : (p2 > (Real)0 ? 2 : 1));
+                b3 = (p3 < (Real)0 ? 0 : (p3 > (Real)0 ? 2 : 1));
+                index = b3 + 3 * (b2 + 3 * (b1 + 3 * b0));
+                (this->*mConfiguration[index])(VERTEX_MAX_BASE, face);
+            }
+        }
+
+        void ClearCandidates()
+        {
+            for (size_t i = 0; i < mNumCandidateEdges; ++i)
+            {
+                auto const& edge = mCandidateEdges[i];
+                mAdjacencyMatrix[edge[0]][edge[1]] = 0;
+                mAdjacencyMatrix[edge[1]][edge[0]] = 0;
+            }
+            mNumCandidateEdges = 0;
+        }
+
+        void InsertEdge(size_t v0, size_t v1)
+        {
+            if (mAdjacencyMatrix[v0][v1] == 0)
+            {
+                mAdjacencyMatrix[v0][v1] = 1;
+                mAdjacencyMatrix[v1][v0] = 1;
+                mCandidateEdges[mNumCandidateEdges] = { v0, v1 };
+                ++mNumCandidateEdges;
+            }
+        }
+
+        // The 81 possible configurations for a box face.  The N stands for a
+        // '-', the Z stands for '0' and the P stands for '+'.  These are
+        // listed in the order that maps to the array mConfiguration.  Thus,
+        // NNNN maps to mConfiguration[0], NNNZ maps to mConfiguration[1], and
+        // so on.
+        void NNNN_0(size_t, Face const&)
+        {
+        }
+
+        void NNNZ_1(size_t, Face const&)
+        {
+        }
+
+        void NNNP_2(size_t base, Face const& face)
+        {
+            InsertEdge(base + face.e[2], base + face.e[3]);
+        }
+
+        void NNZN_3(size_t, Face const&)
+        {
+        }
+
+        void NNZZ_4(size_t, Face const&)
+        {
+        }
+
+        void NNZP_5(size_t base, Face const& face)
+        {
+            InsertEdge(face.v[2], base + face.e[3]);
+        }
+
+        void NNPN_6(size_t base, Face const& face)
+        {
+            InsertEdge(base + face.e[1], base + face.e[2]);
+        }
+
+        void NNPZ_7(size_t base, Face const& face)
+        {
+            InsertEdge(base + face.e[1], face.v[3]);
+        }
+
+        void NNPP_8(size_t base, Face const& face)
+        {
+            InsertEdge(base + face.e[1], base + face.e[3]);
+        }
+
+        void NZNN_9(size_t, Face const&)
+        {
+        }
+
+        void NZNZ_10(size_t, Face const&)
+        {
+        }
+
+        void NZNP_11(size_t base, Face const& face)
+        {
+            InsertEdge(base + face.e[2], face.v[3]);
+            InsertEdge(base + face.e[3], face.v[3]);
+        }
+
+        void NZZN_12(size_t, Face const&)
+        {
+        }
+
+        void NZZZ_13(size_t, Face const&)
+        {
+        }
+
+        void NZZP_14(size_t base, Face const& face)
+        {
+            InsertEdge(face.v[2], face.v[3]);
+            InsertEdge(base + face.e[3], face.v[3]);
+        }
+
+        void NZPN_15(size_t base, Face const& face)
+        {
+            InsertEdge(base + face.e[2], face.v[1]);
+        }
+
+        void NZPZ_16(size_t, Face const& face)
+        {
+            InsertEdge(face.v[1], face.v[3]);
+        }
+
+        void NZPP_17(size_t base, Face const& face)
+        {
+            InsertEdge(base + face.e[3], face.v[1]);
+        }
+
+        void NPNN_18(size_t base, Face const& face)
+        {
+            InsertEdge(base + face.e[0], base + face.e[1]);
+        }
+
+        void NPNZ_19(size_t base, Face const& face)
+        {
+            InsertEdge(base + face.e[0], face.v[1]);
+            InsertEdge(base + face.e[1], face.v[1]);
+        }
+
+        void NPNP_20(size_t base, Face const& face)
+        {
+            InsertEdge(base + face.e[0], face.v[1]);
+            InsertEdge(base + face.e[1], face.v[1]);
+            InsertEdge(base + face.e[2], face.v[3]);
+            InsertEdge(base + face.e[3], face.v[3]);
+        }
+
+        void NPZN_21(size_t base, Face const& face)
+        {
+            InsertEdge(base + face.e[0], face.v[2]);
+        }
+
+        void NPZZ_22(size_t base, Face const& face)
+        {
+            InsertEdge(base + face.e[0], face.v[1]);
+            InsertEdge(face.v[1], face.v[2]);
+        }
+
+        void NPZP_23(size_t base, Face const& face)
+        {
+            InsertEdge(base + face.e[0], face.v[1]);
+            InsertEdge(face.v[1], face.v[2]);
+            InsertEdge(base + face.e[3], face.v[2]);
+            InsertEdge(face.v[2], face.v[3]);
+        }
+
+        void NPPN_24(size_t base, Face const& face)
+        {
+            InsertEdge(base + face.e[0], base + face.e[2]);
+        }
+
+        void NPPZ_25(size_t base, Face const& face)
+        {
+            InsertEdge(base + face.e[0], face.v[3]);
+        }
+
+        void NPPP_26(size_t base, Face const& face)
+        {
+            InsertEdge(base + face.e[0], base + face.e[3]);
+        }
+
+        void ZNNN_27(size_t, Face const&)
+        {
+        }
+
+        void ZNNZ_28(size_t, Face const&)
+        {
+        }
+
+        void ZNNP_29(size_t base, Face const& face)
+        {
+            InsertEdge(base + face.e[2], face.v[0]);
+        }
+
+        void ZNZN_30(size_t, Face const&)
+        {
+        }
+
+        void ZNZZ_31(size_t, Face const&)
+        {
+        }
+
+        void ZNZP_32(size_t, Face const& face)
+        {
+            InsertEdge(face.v[0], face.v[2]);
+        }
+
+        void ZNPN_33(size_t base, Face const& face)
+        {
+            InsertEdge(base + face.e[1], face.v[2]);
+            InsertEdge(base + face.e[2], face.v[2]);
+        }
+
+        void ZNPZ_34(size_t base, Face const& face)
+        {
+            InsertEdge(base + face.e[1], face.v[2]);
+            InsertEdge(face.v[2], face.v[3]);
+        }
+
+        void ZNPP_35(size_t base, Face const& face)
+        {
+            InsertEdge(face.v[0], base + face.e[1]);
+        }
+
+        void ZZNN_36(size_t, Face const&)
+        {
+        }
+
+        void ZZNZ_37(size_t, Face const&)
+        {
+        }
+
+        void ZZNP_38(size_t base, Face const& face)
+        {
+            InsertEdge(face.v[0], face.v[3]);
+            InsertEdge(face.v[3], base + face.e[2]);
+        }
+
+        void ZZZN_39(size_t, Face const&)
+        {
+        }
+
+        void ZZZZ_40(size_t, Face const&)
+        {
+        }
+
+        void ZZZP_41(size_t, Face const& face)
+        {
+            InsertEdge(face.v[0], face.v[3]);
+            InsertEdge(face.v[3], face.v[2]);
+        }
+
+        void ZZPN_42(size_t base, Face const& face)
+        {
+            InsertEdge(face.v[1], face.v[2]);
+            InsertEdge(face.v[2], base + face.e[2]);
+        }
+
+        void ZZPZ_43(size_t, Face const& face)
+        {
+            InsertEdge(face.v[1], face.v[2]);
+            InsertEdge(face.v[2], face.v[3]);
+        }
+
+        void ZZPP_44(size_t, Face const&)
+        {
+        }
+
+        void ZPNN_45(size_t base, Face const& face)
+        {
+            InsertEdge(face.v[0], base + face.e[1]);
+        }
+
+        void ZPNZ_46(size_t base, Face const& face)
+        {
+            InsertEdge(face.v[0], face.v[1]);
+            InsertEdge(face.v[1], base + face.e[1]);
+        }
+
+        void ZPNP_47(size_t base, Face const& face)
+        {
+            InsertEdge(face.v[0], face.v[1]);
+            InsertEdge(face.v[1], base + face.e[1]);
+            InsertEdge(base + face.e[2], face.v[3]);
+            InsertEdge(face.v[3], face.v[0]);
+        }
+
+        void ZPZN_48(size_t, Face const& face)
+        {
+            InsertEdge(face.v[0], face.v[2]);
+        }
+
+        void ZPZZ_49(size_t, Face const& face)
+        {
+            InsertEdge(face.v[0], face.v[1]);
+            InsertEdge(face.v[1], face.v[2]);
+        }
+
+        void ZPZP_50(size_t, Face const&)
+        {
+        }
+
+        void ZPPN_51(size_t base, Face const& face)
+        {
+            InsertEdge(face.v[0], base + face.e[2]);
+        }
+
+        void ZPPZ_52(size_t, Face const&)
+        {
+        }
+
+        void ZPPP_53(size_t, Face const&)
+        {
+        }
+
+        void PNNN_54(size_t base, Face const& face)
+        {
+            InsertEdge(base + face.e[3], base + face.e[0]);
+        }
+
+        void PNNZ_55(size_t base, Face const& face)
+        {
+            InsertEdge(face.v[3], base + face.e[0]);
+        }
+
+        void PNNP_56(size_t base, Face const& face)
+        {
+            InsertEdge(base + face.e[2], base + face.e[0]);
+        }
+
+        void PNZN_57(size_t base, Face const& face)
+        {
+            InsertEdge(base + face.e[3], face.v[0]);
+            InsertEdge(face.v[0], base + face.e[0]);
+        }
+
+        void PNZZ_58(size_t base, Face const& face)
+        {
+            InsertEdge(face.v[3], face.v[0]);
+            InsertEdge(face.v[0], base + face.e[0]);
+        }
+
+        void PNZP_59(size_t base, Face const& face)
+        {
+            InsertEdge(face.v[2], base + face.e[0]);
+        }
+
+        void PNPN_60(size_t base, Face const& face)
+        {
+            InsertEdge(base + face.e[3], face.v[0]);
+            InsertEdge(face.v[0], base + face.e[0]);
+            InsertEdge(base + face.e[1], face.v[2]);
+            InsertEdge(face.v[2], base + face.e[2]);
+        }
+
+        void PNPZ_61(size_t base, Face const& face)
+        {
+            InsertEdge(face.v[3], face.v[0]);
+            InsertEdge(face.v[0], base + face.e[0]);
+            InsertEdge(base + face.e[1], face.v[2]);
+            InsertEdge(face.v[2], face.v[3]);
+        }
+
+        void PNPP_62(size_t base, Face const& face)
+        {
+            InsertEdge(base + face.e[0], base + face.e[1]);
+        }
+
+        void PZNN_63(size_t base, Face const& face)
+        {
+            InsertEdge(base + face.e[3], face.v[1]);
+        }
+
+        void PZNZ_64(size_t, Face const& face)
+        {
+            InsertEdge(face.v[3], face.v[1]);
+        }
+
+        void PZNP_65(size_t base, Face const& face)
+        {
+            InsertEdge(base + face.e[2], face.v[1]);
+        }
+
+        void PZZN_66(size_t base, Face const& face)
+        {
+            InsertEdge(base + face.e[3], face.v[0]);
+            InsertEdge(face.v[0], face.v[1]);
+        }
+
+        void PZZZ_67(size_t, Face const&)
+        {
+        }
+
+        void PZZP_68(size_t, Face const&)
+        {
+        }
+
+        void PZPN_69(size_t base, Face const& face)
+        {
+            InsertEdge(base + face.e[3], face.v[0]);
+            InsertEdge(face.v[0], face.v[1]);
+            InsertEdge(face.v[1], face.v[2]);
+            InsertEdge(face.v[2], base + face.e[2]);
+        }
+
+        void PZPZ_70(size_t, Face const&)
+        {
+        }
+
+        void PZPP_71(size_t, Face const&)
+        {
+        }
+
+        void PPNN_72(size_t base, Face const& face)
+        {
+            InsertEdge(base + face.e[3], base + face.e[1]);
+        }
+
+        void PPNZ_73(size_t base, Face const& face)
+        {
+            InsertEdge(face.v[3], base + face.e[1]);
+        }
+
+        void PPNP_74(size_t base, Face const& face)
+        {
+            InsertEdge(base + face.e[2], base + face.e[1]);
+        }
+
+        void PPZN_75(size_t base, Face const& face)
+        {
+            InsertEdge(base + face.e[2], face.v[2]);
+        }
+
+        void PPZZ_76(size_t, Face const&)
+        {
+        }
+
+        void PPZP_77(size_t, Face const&)
+        {
+        }
+
+        void PPPN_78(size_t base, Face const& face)
+        {
+            InsertEdge(base + face.e[3], base + face.e[2]);
+        }
+
+        void PPPZ_79(size_t, Face const&)
+        {
+        }
+
+        void PPPP_80(size_t, Face const&)
+        {
+        }
+    };
+
+    // Template alias for convenience.
+    template <typename Real>
+    using TIAlignedBox3Cone3 = TIQuery<Real, AlignedBox<3, Real>, Cone<3, Real>>;
+}
diff --git a/Engine/Plugins/Experimental/GeometryProcessing/Source/GeometryAlgorithms/Private/GTEngine/Mathematics/GteIntrAlignedBox3Cylinder3.h b/Engine/Plugins/Experimental/GeometryProcessing/Source/GeometryAlgorithms/Private/GTEngine/Mathematics/GteIntrAlignedBox3Cylinder3.h
new file mode 100644
--- /dev/null
+++ b/Engine/Plugins/Experimental/GeometryProcessing/Source/GeometryAlgorithms/Private/GTEngine/Mathematics/GteIntrAlignedBox3Cylinder3.h
@@ -0,0 +1,126 @@
+// David Eberly, Geometric Tools, Redmond WA 98052
+// Copyright (c) 1998-2019
+// Distributed under the Boost Software License, Version 1.0.
+// http://www.boost.org/LICENSE_1_0.txt
+// http://www.geometrictools.com/License/Boost/LICENSE_1_0.txt
+// File Version: 3.4.0 (2016/11/08)
+
+#pragma once
+
+#include <GTEngine/Mathematics/GteTIQuery.h>
+#include <GTEngine/Mathematics/GteAlignedBox.h>
+#include <GTEngine/Mathematics/GteCylinder3.h>
+#include <GTEngine/Mathematics/GteLCPSolver.h>
+#include <GTEngine/Mathematics/GteMatrix3x3.h>
+
+// The query considers the cylinder and box to be solids.
+
+namespace gte
+{
+
+template <typename Real>
+class TIQuery<Real, AlignedBox3<Real>, Cylinder3<Real>>
+{
+public:
+    struct Result
+    {
+        bool intersect;
+    };
+
+    Result operator()(AlignedBox3<Real> const& box, Cylinder3<Real> const& cylinder);
+
+private:
+};
+
+
+template <typename Real>
+typename TIQuery<Real, AlignedBox3<Real>, Cylinder3<Real>>::Result
+TIQuery<Real, AlignedBox3<Real>, Cylinder3<Real>>::operator()(
+    AlignedBox3<Real> const& box, Cylinder3<Real> const& cylinder)
+{
+    Result result;
+
+    // Translate the box and cylinder so that the box is in the first octant
+    // where all points in the box have nonnegative components.
+    Vector3<Real> corner = box.max - box.min;
+    Vector3<Real> origin = cylinder.axis.origin - box.min;
+    Vector3<Real> direction = cylinder.axis.direction;
+
+    // Compute quantities to initialize q and M in the LCP.
+    Real halfHeight = cylinder.height * (Real)0.5;
+    Matrix3x3<Real> P = (Matrix3x3<Real>::Identity() - OuterProduct(direction, direction));
+    Vector3<Real> C = -(P * origin);
+    Real originDotDirection = Dot(origin, direction);
+
+    Matrix<5, 3, Real> A;
+    A.SetRow(0, { (Real)-1, (Real) 0, (Real) 0 });
+    A.SetRow(1, { (Real) 0, (Real)-1, (Real) 0 });
+    A.SetRow(2, { (Real) 0, (Real) 0, (Real)-1 });
+    A.SetRow(3, direction);
+    A.SetRow(4, -direction);
+
+    Vector<5, Real> B =
+    {
+        -corner[0],
+        -corner[1],
+        -corner[2],
+        originDotDirection - halfHeight,
+        -originDotDirection - halfHeight
+    };
+
+    std::array<std::array<Real, 8>, 8> M;
+    for (int r = 0; r < 3; ++r)
+    {
+        for (int c = 0; c < 3; ++c)
+        {
+            M[r][c] = P(r, c);
+        }
+
+        for (int c = 3, i = 0; c < 8; ++c, ++i)
+        {
+            M[r][c] = -A(i, r);
+        }
+    }
+
+    for (int r = 3, i = 0; r < 8; ++r, ++i)
+    {
+        for (int c = 0; c < 3; ++c)
+        {
+            M[r][c] = A(i, c);
+        }
+
+        for (int c = 3; c < 8; ++c)
+        {
+            M[r][c] = (Real)0;
+        }
+    }
+
+    std::array<Real, 8> q;
+    for (int r = 0; r < 3; ++r)
+    {
+        q[r] = C[r];
+    }
+
+    for (int r = 3, i = 0; r < 8; ++r, ++i)
+    {
+        q[r] = -B[i];
+    }
+
+    std::array<Real, 8> w, z;
+    LCPSolver<Real, 8> LCP;
+    if (LCP.Solve(q, M, w, z))
+    {
+        Vector3<Real> zSolution{ z[0], z[1], z[2] };
+        Vector3<Real> diff = zSolution - origin;
+        Real qform = Dot(diff, P * diff);
+        result.intersect = (qform <= cylinder.radius * cylinder.radius);
+    }
+    else
+    {
+        result.intersect = false;
+    }
+
+    return result;
+}
+
+}
diff --git a/Engine/Plugins/Experimental/GeometryProcessing/Source/GeometryAlgorithms/Private/GTEngine/Mathematics/GteIntrAlignedBox3OrientedBox3.h b/Engine/Plugins/Experimental/GeometryProcessing/Source/GeometryAlgorithms/Private/GTEngine/Mathematics/GteIntrAlignedBox3OrientedBox3.h
new file mode 100644
--- /dev/null
+++ b/Engine/Plugins/Experimental/GeometryProcessing/Source/GeometryAlgorithms/Private/GTEngine/Mathematics/GteIntrAlignedBox3OrientedBox3.h
@@ -0,0 +1,336 @@
+// David Eberly, Geometric Tools, Redmond WA 98052
+// Copyright (c) 1998-2019
+// Distributed under the Boost Software License, Version 1.0.
+// http://www.boost.org/LICENSE_1_0.txt
+// http://www.geometrictools.com/License/Boost/LICENSE_1_0.txt
+// File Version: 3.0.0 (2016/06/19)
+
+#pragma once
+
+#include <GTEngine/Mathematics/GteVector3.h>
+#include <GTEngine/Mathematics/GteAlignedBox.h>
+#include <GTEngine/Mathematics/GteOrientedBox.h>
+#include <GTEngine/Mathematics/GteFIQuery.h>
+#include <GTEngine/Mathematics/GteTIQuery.h>
+
+// The queries consider the box to be a solid.
+//
+// The test-intersection query uses the method of separating axes.  The set of
+// potential separating directions includes the 3 face normals of box0, the 3
+// face normals of box1, and 9 directions, each of which is the cross product
+// of an edge of box0 and and an edge of box1.
+//
+// The separating axes involving cross products of edges has numerical
+// robustness problems when the two edges are nearly parallel.  The cross
+// product of the edges is nearly the zero vector, so normalization of the
+// cross product may produce unit-length directions that are not close to the
+// true direction.  Such a pair of edges occurs when a box0 face normal N0 and
+// a box1 face normal N1 are nearly parallel.  In this case, you may skip the
+// edge-edge directions, which is equivalent to projecting the boxes onto the
+// plane with normal N0 and applying a 2D separating axis test.  The ability
+// to do so involves choosing a small nonnegative epsilon .  It is used to
+// determine whether two face normals, one from each box, are nearly parallel:
+// |Dot(N0,N1)| >= 1 - epsilon.  If the input is negative, it is clamped to
+// zero.
+//
+// The pair of integers 'separating', say, (i0,i1), identify the axis that
+// reported separation; there may be more than one but only one is
+// reported.  If the separating axis is a face normal N[i0] of the aligned
+// box0 in dimension i0, then (i0,-1) is returned.  If the axis is a face
+// normal box1.Axis[i1], then (-1,i1) is returned.  If the axis is a cross
+// product of edges, Cross(N[i0],box1.Axis[i1]), then (i0,i1) is returned.
+
+namespace gte
+{
+
+template <typename Real>
+class TIQuery<Real, AlignedBox3<Real>, OrientedBox3<Real>>
+{
+public:
+    struct Result
+    {
+        // The 'epsilon' value must be nonnegative.
+        Result(Real inEpsilon = (Real)0);
+
+        bool intersect;
+        Real epsilon;
+        int separating[2];
+    };
+
+    Result operator()(AlignedBox3<Real> const& box0,
+        OrientedBox3<Real> const& box1);
+};
+
+
+template <typename Real>
+TIQuery<Real, AlignedBox3<Real>, OrientedBox3<Real>>::Result::
+Result(Real inEpsilon)
+{
+    if (inEpsilon >= (Real)0)
+    {
+        epsilon = inEpsilon;
+    }
+    else
+    {
+        epsilon = (Real)0;
+    }
+}
+
+template <typename Real>
+typename TIQuery<Real, AlignedBox3<Real>, OrientedBox3<Real>>::Result
+TIQuery<Real, AlignedBox3<Real>, OrientedBox3<Real>>::operator()(
+    AlignedBox3<Real> const& box0, OrientedBox3<Real> const& box1)
+{
+    Result result;
+
+    // Get the centered form of the aligned box.  The axes are implicitly
+    // A0[0] = (1,0,0), A0[1] = (0,1,0), and A0[2] = (0,0,1).
+    Vector3<Real> C0, E0;
+    box0.GetCenteredForm(C0, E0);
+
+    // Convenience variables.
+    Vector3<Real> const& C1 = box1.center;
+    Vector3<Real> const* A1 = &box1.axis[0];
+    Vector3<Real> const& E1 = box1.extent;
+
+    Real const cutoff = (Real)1 - result.epsilon;
+    bool existsParallelPair = false;
+
+    // Compute the difference of box centers.
+    Vector3<Real> D = C1 - C0;
+
+    Real dot01[3][3];       // dot01[i][j] = Dot(A0[i],A1[j]) = A1[j][i]
+    Real absDot01[3][3];    // |dot01[i][j]|
+    Real r0, r1, r;         // interval radii and distance between centers
+    Real r01;               // r0 + r1
+
+    // Test for separation on the axis C0 + t*A0[0].
+    for (int i = 0; i < 3; ++i)
+    {
+        dot01[0][i] = A1[i][0];
+        absDot01[0][i] = std::abs(A1[i][0]);
+        if (absDot01[0][i] >= cutoff)
+        {
+            existsParallelPair = true;
+        }
+    }
+    r = std::abs(D[0]);
+    r1 = E1[0] * absDot01[0][0] + E1[1] * absDot01[0][1] + E1[2] * absDot01[0][2];
+    r01 = E0[0] + r1;
+    if (r > r01)
+    {
+        result.intersect = false;
+        result.separating[0] = 0;
+        result.separating[1] = -1;
+        return result;
+    }
+
+    // Test for separation on the axis C0 + t*A0[1].
+    for (int i = 0; i < 3; ++i)
+    {
+        dot01[1][i] = A1[i][1];
+        absDot01[1][i] = std::abs(A1[i][1]);
+        if (absDot01[1][i] >= cutoff)
+        {
+            existsParallelPair = true;
+        }
+    }
+    r = std::abs(D[1]);
+    r1 = E1[0] * absDot01[1][0] + E1[1] * absDot01[1][1] + E1[2] * absDot01[1][2];
+    r01 = E0[1] + r1;
+    if (r > r01)
+    {
+        result.intersect = false;
+        result.separating[0] = 1;
+        result.separating[1] = -1;
+        return result;
+    }
+
+    // Test for separation on the axis C0 + t*A0[2].
+    for (int i = 0; i < 3; ++i)
+    {
+        dot01[2][i] = A1[i][2];
+        absDot01[2][i] = std::abs(A1[i][2]);
+        if (absDot01[2][i] >= cutoff)
+        {
+            existsParallelPair = true;
+        }
+    }
+    r = std::abs(D[2]);
+    r1 = E1[0] * absDot01[2][0] + E1[1] * absDot01[2][1] + E1[2] * absDot01[2][2];
+    r01 = E0[2] + r1;
+    if (r > r01)
+    {
+        result.intersect = false;
+        result.separating[0] = 2;
+        result.separating[1] = -1;
+        return result;
+    }
+
+    // Test for separation on the axis C0 + t*A1[0].
+    r = std::abs(Dot(D, A1[0]));
+    r0 = E0[0] * absDot01[0][0] + E0[1] * absDot01[1][0] + E0[2] * absDot01[2][0];
+    r01 = r0 + E1[0];
+    if (r > r01)
+    {
+        result.intersect = false;
+        result.separating[0] = -1;
+        result.separating[1] = 0;
+        return result;
+    }
+
+    // Test for separation on the axis C0 + t*A1[1].
+    r = std::abs(Dot(D, A1[1]));
+    r0 = E0[0] * absDot01[0][1] + E0[1] * absDot01[1][1] + E0[2] * absDot01[2][1];
+    r01 = r0 + E1[1];
+    if (r > r01)
+    {
+        result.intersect = false;
+        result.separating[0] = -1;
+        result.separating[1] = 1;
+        return result;
+    }
+
+    // Test for separation on the axis C0 + t*A1[2].
+    r = std::abs(Dot(D, A1[2]));
+    r0 = E0[0] * absDot01[0][2] + E0[1] * absDot01[1][2] + E0[2] * absDot01[2][2];
+    r01 = r0 + E1[2];
+    if (r > r01)
+    {
+        result.intersect = false;
+        result.separating[0] = -1;
+        result.separating[1] = 2;
+        return result;
+    }
+
+    // At least one pair of box axes was parallel, so the separation is
+    // effectively in 2D.  The edge-edge axes do not need to be tested.
+    if (existsParallelPair)
+    {
+        result.intersect = true;
+        return result;
+    }
+
+    // Test for separation on the axis C0 + t*A0[0]xA1[0].
+    r = std::abs(D[2] * dot01[1][0] - D[1] * dot01[2][0]);
+    r0 = E0[1] * absDot01[2][0] + E0[2] * absDot01[1][0];
+    r1 = E1[1] * absDot01[0][2] + E1[2] * absDot01[0][1];
+    r01 = r0 + r1;
+    if (r > r01)
+    {
+        result.intersect = false;
+        result.separating[0] = 0;
+        result.separating[1] = 0;
+        return result;
+    }
+
+    // Test for separation on the axis C0 + t*A0[0]xA1[1].
+    r = std::abs(D[2] * dot01[1][1] - D[1] * dot01[2][1]);
+    r0 = E0[1] * absDot01[2][1] + E0[2] * absDot01[1][1];
+    r1 = E1[0] * absDot01[0][2] + E1[2] * absDot01[0][0];
+    r01 = r0 + r1;
+    if (r > r01)
+    {
+        result.intersect = false;
+        result.separating[0] = 0;
+        result.separating[1] = 1;
+        return result;
+    }
+
+    // Test for separation on the axis C0 + t*A0[0]xA1[2].
+    r = std::abs(D[2] * dot01[1][2] - D[1] * dot01[2][2]);
+    r0 = E0[1] * absDot01[2][2] + E0[2] * absDot01[1][2];
+    r1 = E1[0] * absDot01[0][1] + E1[1] * absDot01[0][0];
+    r01 = r0 + r1;
+    if (r > r01)
+    {
+        result.intersect = false;
+        result.separating[0] = 0;
+        result.separating[1] = 2;
+        return result;
+    }
+
+    // Test for separation on the axis C0 + t*A0[1]xA1[0].
+    r = std::abs(D[0] * dot01[2][0] - D[2] * dot01[0][0]);
+    r0 = E0[0] * absDot01[2][0] + E0[2] * absDot01[0][0];
+    r1 = E1[1] * absDot01[1][2] + E1[2] * absDot01[1][1];
+    r01 = r0 + r1;
+    if (r > r01)
+    {
+        result.intersect = false;
+        result.separating[0] = 1;
+        result.separating[1] = 0;
+        return result;
+    }
+
+    // Test for separation on the axis C0 + t*A0[1]xA1[1].
+    r = std::abs(D[0] * dot01[2][1] - D[2] * dot01[0][1]);
+    r0 = E0[0] * absDot01[2][1] + E0[2] * absDot01[0][1];
+    r1 = E1[0] * absDot01[1][2] + E1[2] * absDot01[1][0];
+    r01 = r0 + r1;
+    if (r > r01)
+    {
+        result.intersect = false;
+        result.separating[0] = 1;
+        result.separating[1] = 1;
+        return result;
+    }
+
+    // Test for separation on the axis C0 + t*A0[1]xA1[2].
+    r = std::abs(D[0] * dot01[2][2] - D[2] * dot01[0][2]);
+    r0 = E0[0] * absDot01[2][2] + E0[2] * absDot01[0][2];
+    r1 = E1[0] * absDot01[1][1] + E1[1] * absDot01[1][0];
+    r01 = r0 + r1;
+    if (r > r01)
+    {
+        result.intersect = false;
+        result.separating[0] = 1;
+        result.separating[1] = 2;
+        return result;
+    }
+
+    // Test for separation on the axis C0 + t*A0[2]xA1[0].
+    r = std::abs(D[1] * dot01[0][0] - D[0] * dot01[1][0]);
+    r0 = E0[0] * absDot01[1][0] + E0[1] * absDot01[0][0];
+    r1 = E1[1] * absDot01[2][2] + E1[2] * absDot01[2][1];
+    r01 = r0 + r1;
+    if (r > r01)
+    {
+        result.intersect = false;
+        result.separating[0] = 2;
+        result.separating[1] = 0;
+        return result;
+    }
+
+    // Test for separation on the axis C0 + t*A0[2]xA1[1].
+    r = std::abs(D[1] * dot01[0][1] - D[0] * dot01[1][1]);
+    r0 = E0[0] * absDot01[1][1] + E0[1] * absDot01[0][1];
+    r1 = E1[0] * absDot01[2][2] + E1[2] * absDot01[2][0];
+    r01 = r0 + r1;
+    if (r > r01)
+    {
+        result.intersect = false;
+        result.separating[0] = 2;
+        result.separating[1] = 1;
+        return result;
+    }
+
+    // Test for separation on the axis C0 + t*A0[2]xA1[2].
+    r = std::abs(D[1] * dot01[0][2] - D[0] * dot01[1][2]);
+    r0 = E0[0] * absDot01[1][2] + E0[1] * absDot01[0][2];
+    r1 = E1[0] * absDot01[2][1] + E1[1] * absDot01[2][0];
+    r01 = r0 + r1;
+    if (r > r01)
+    {
+        result.intersect = false;
+        result.separating[0] = 2;
+        result.separating[1] = 2;
+        return result;
+    }
+
+    result.intersect = true;
+    return result;
+}
+
+
+}
diff --git a/Engine/Plugins/Experimental/GeometryProcessing/Source/GeometryAlgorithms/Private/GTEngine/Mathematics/GteIntrAlignedBox3Sphere3.h b/Engine/Plugins/Experimental/GeometryProcessing/Source/GeometryAlgorithms/Private/GTEngine/Mathematics/GteIntrAlignedBox3Sphere3.h
new file mode 100644
--- /dev/null
+++ b/Engine/Plugins/Experimental/GeometryProcessing/Source/GeometryAlgorithms/Private/GTEngine/Mathematics/GteIntrAlignedBox3Sphere3.h
@@ -0,0 +1,595 @@
+// David Eberly, Geometric Tools, Redmond WA 98052
+// Copyright (c) 1998-2019
+// Distributed under the Boost Software License, Version 1.0.
+// http://www.boost.org/LICENSE_1_0.txt
+// http://www.geometrictools.com/License/Boost/LICENSE_1_0.txt
+// File Version: 3.0.4 (2018/10/05)
+
+#pragma once
+
+#include <GTEngine/Mathematics/GteVector3.h>
+#include <GTEngine/Mathematics/GteHypersphere.h>
+#include <GTEngine/Mathematics/GteDistPointAlignedBox.h>
+#include <GTEngine/Mathematics/GteIntrRay3AlignedBox3.h>
+
+// The find-intersection query is based on the document
+// https://www.geometrictools.com/Documentation/IntersectionMovingSphereBox.pdf
+
+namespace gte
+{
+    template <typename Real>
+    class TIQuery<Real, AlignedBox3<Real>, Sphere3<Real>>
+    {
+    public:
+        // The intersection query considers the box and sphere to be solids;
+        // that is, the sphere object includes the region inside the spherical
+        // boundary and the box object includes the region inside the cuboid
+        // boundary.  If the sphere object and box object object overlap, the
+        // objects intersect.
+        struct Result
+        {
+            bool intersect;
+        };
+
+        Result operator()(AlignedBox3<Real> const& box, Sphere3<Real> const& sphere)
+        {
+            DCPQuery<Real, Vector3<Real>, AlignedBox3<Real>> pbQuery;
+            auto pbResult = pbQuery(sphere.center, box);
+            Result result;
+            result.intersect = (pbResult.sqrDistance <= sphere.radius * sphere.radius);
+            return result;
+        }
+    };
+
+    template <typename Real>
+    class FIQuery<Real, AlignedBox3<Real>, Sphere3<Real>>
+    {
+    public:
+        // Currently, only a dynamic query is supported.  A static query will
+        // need to compute the intersection set of (solid) box and sphere.
+        struct Result
+        {
+            // The cases are
+            // 1. Objects initially overlapping.  The contactPoint is only one
+            //    of infinitely many points in the overlap.
+            //      intersectionType = -1
+            //      contactTime = 0
+            //      contactPoint = sphere.center
+            // 2. Objects initially separated but do not intersect later.  The
+            //      contactTime and contactPoint are invalid.
+            //      intersectionType = 0
+            //      contactTime = 0
+            //      contactPoint = (0,0,0)
+            // 3. Objects initially separated but intersect later.
+            //      intersectionType = +1
+            //      contactTime = first time T > 0
+            //      contactPoint = corresponding first contact
+            int intersectionType;
+            Real contactTime;
+            Vector3<Real> contactPoint;
+
+            // TODO: To support arbitrary precision for the contactTime,
+            // return q0, q1 and q2 where contactTime = (q0 - sqrt(q1)) / q2.
+            // The caller can compute contactTime to desired number of digits
+            // of precision.  These are valid when intersectionType is +1 but
+            // are set to zero (invalid) in the other cases.  Do the same for
+            // the contactPoint.
+        };
+
+        Result operator()(AlignedBox3<Real> const& box, Vector3<Real> const& boxVelocity,
+            Sphere3<Real> const& sphere, Vector3<Real> const& sphereVelocity)
+        {
+            Result result = { 0, (Real)0, { (Real)0, (Real)0, (Real)0 } };
+
+            // Translate the sphere and box so that the box center becomes
+            // the origin.  Compute the velocity of the sphere relative to
+            // the box.
+            Vector3<Real> boxCenter = (box.max + box.min) * (Real)0.5;
+            Vector3<Real> extent = (box.max - box.min) * (Real)0.5;
+            Vector3<Real> C = sphere.center - boxCenter;
+            Vector3<Real> V = sphereVelocity - boxVelocity;
+
+            // Test for no-intersection that leads to an early exit.  The test
+            // is fast, using the method of separating axes.
+            AlignedBox3<Real> superBox;
+            for (int i = 0; i < 3; ++i)
+            {
+                superBox.max[i] = extent[i] + sphere.radius;
+                superBox.min[i] = -superBox.max[i];
+            }
+            TIQuery<Real, Ray3<Real>, AlignedBox3<Real>> rbQuery;
+            auto rbResult = rbQuery(Ray3<Real>(C, V), superBox);
+            if (!rbResult.intersect)
+            {
+                return result;
+            }
+
+            // Change signs on components, if necessary, to transform C to the
+            // first quadrant.  Adjust the velocity accordingly.
+            Real sign[3];
+            for (int i = 0; i < 3; ++i)
+            {
+                if (C[i] >= (Real)0)
+                {
+                    sign[i] = (Real)1;
+                }
+                else
+                {
+                    C[i] = -C[i];
+                    V[i] = -V[i];
+                    sign[i] = (Real)-1;
+                }
+            }
+
+            DoQuery(extent, C, sphere.radius, V, result);
+
+            if (result.intersectionType != 0)
+            {
+                // Translate back to the original coordinate system.
+                for (int i = 0; i < 3; ++i)
+                {
+                    if (sign[i] < (Real)0)
+                    {
+                        result.contactPoint[i] = -result.contactPoint[i];
+                    }
+                }
+
+                result.contactPoint += boxCenter;
+            }
+
+            return result;
+        }
+
+    protected:
+        // The query assumes the box is axis-aligned with center at the
+        // origin.  Callers need to convert the results back to the original
+        // coordinate system of the query.
+        void DoQuery(Vector3<Real> const& K, Vector3<Real> const& C,
+            Real radius, Vector3<Real> const& V, Result& result)
+        {
+            Vector3<Real> delta = C - K;
+            if (delta[2] <= radius)
+            {
+                if (delta[1] <= radius)
+                {
+                    if (delta[0] <= radius)
+                    {
+                        if (delta[2] <= (Real)0)
+                        {
+                            if (delta[1] <= (Real)0)
+                            {
+                                if (delta[0] <= (Real)0)
+                                {
+                                    InteriorOverlap(C, result);
+                                }
+                                else
+                                {
+                                    // x-face
+                                    FaceOverlap(0, 1, 2, K, C, radius, delta, result);
+                                }
+                            }
+                            else
+                            {
+                                if (delta[0] <= (Real)0)
+                                {
+                                    // y-face
+                                    FaceOverlap(1, 2, 0, K, C, radius, delta, result);
+                                }
+                                else
+                                {
+                                    // xy-edge
+                                    if (delta[0] * delta[0] + delta[1] * delta[1] <= radius * radius)
+                                    {
+                                        EdgeOverlap(0, 1, 2, K, C, radius, delta, result);
+                                    }
+                                    else
+                                    {
+                                        EdgeSeparated(0, 1, 2, K, C, radius, delta, V, result);
+                                    }
+                                }
+                            }
+                        }
+                        else
+                        {
+                            if (delta[1] <= (Real)0)
+                            {
+                                if (delta[0] <= (Real)0)
+                                {
+                                    // z-face
+                                    FaceOverlap(2, 0, 1, K, C, radius, delta, result);
+                                }
+                                else
+                                {
+                                    // xz-edge
+                                    if (delta[0] * delta[0] + delta[2] * delta[2] <= radius * radius)
+                                    {
+                                        EdgeOverlap(2, 0, 1, K, C, radius, delta, result);
+                                    }
+                                    else
+                                    {
+                                        EdgeSeparated(2, 0, 1, K, C, radius, delta, V, result);
+                                    }
+                                }
+                            }
+                            else
+                            {
+                                if (delta[0] <= (Real)0)
+                                {
+                                    // yz-edge
+                                    if (delta[1] * delta[1] + delta[2] * delta[2] <= radius * radius)
+                                    {
+                                        EdgeOverlap(1, 2, 0, K, C, radius, delta, result);
+                                    }
+                                    else
+                                    {
+                                        EdgeSeparated(1, 2, 0, K, C, radius, delta, V, result);
+                                    }
+                                }
+                                else
+                                {
+                                    // xyz-vertex
+                                    if (Dot(delta, delta) <= radius * radius)
+                                    {
+                                        VertexOverlap(K, radius, delta, result);
+                                    }
+                                    else
+                                    {
+                                        VertexSeparated(K, radius, delta, V, result);
+                                    }
+                                }
+                            }
+                        }
+                    }
+                    else
+                    {
+                        // x-face
+                        FaceUnbounded(0, 1, 2, K, C, radius, delta, V, result);
+                    }
+                }
+                else
+                {
+                    if (delta[0] <= radius)
+                    {
+                        // y-face
+                        FaceUnbounded(1, 2, 0, K, C, radius, delta, V, result);
+                    }
+                    else
+                    {
+                        // xy-edge
+                        EdgeUnbounded(0, 1, 2, K, C, radius, delta, V, result);
+                    }
+                }
+            }
+            else
+            {
+                if (delta[1] <= radius)
+                {
+                    if (delta[0] <= radius)
+                    {
+                        // z-face
+                        FaceUnbounded(2, 0, 1, K, C, radius, delta, V, result);
+                    }
+                    else
+                    {
+                        // xz-edge
+                        EdgeUnbounded(2, 0, 1, K, C, radius, delta, V, result);
+                    }
+                }
+                else
+                {
+                    if (delta[0] <= radius)
+                    {
+                        // yz-edge
+                        EdgeUnbounded(1, 2, 0, K, C, radius, delta, V, result);
+                    }
+                    else
+                    {
+                        // xyz-vertex
+                        VertexUnbounded(K, C, radius, delta, V, result);
+                    }
+                }
+            }
+        }
+
+    private:
+        void InteriorOverlap(Vector3<Real> const& C, Result& result)
+        {
+            result.intersectionType = -1;
+            result.contactTime = (Real)0;
+            result.contactPoint = C;
+        }
+
+        void VertexOverlap(Vector3<Real> const& K, Real radius,
+            Vector3<Real> const& delta, Result& result)
+        {
+            result.intersectionType = (Dot(delta, delta) < radius * radius ? -1 : 1);
+            result.contactTime = (Real)0;
+            result.contactPoint = K;
+        }
+
+        void EdgeOverlap(int i0, int i1, int i2, Vector3<Real> const& K,
+            Vector3<Real> const& C, Real radius, Vector3<Real> const& delta,
+            Result& result)
+        {
+            result.intersectionType = (delta[i0] * delta[i0] + delta[i1] * delta[i1] < radius * radius ? -1 : 1);
+            result.contactTime = (Real)0;
+            result.contactPoint[i0] = K[i0];
+            result.contactPoint[i1] = K[i1];
+            result.contactPoint[i2] = C[i2];
+        }
+
+        void FaceOverlap(int i0, int i1, int i2, Vector3<Real> const& K,
+            Vector3<Real> const& C, Real radius, Vector3<Real> const& delta,
+            Result& result)
+        {
+            result.intersectionType = (delta[i0] < radius ? -1 : 1);
+            result.contactTime = (Real)0;
+            result.contactPoint[i0] = K[i0];
+            result.contactPoint[i1] = C[i1];
+            result.contactPoint[i2] = C[i2];
+        }
+
+        void VertexSeparated(Vector3<Real> const& K, Real radius,
+            Vector3<Real> const& delta, Vector3<Real> const& V, Result& result)
+        {
+            if (V[0] < (Real)0 || V[1] < (Real)0 || V[2] < (Real)0)
+            {
+                DoQueryRayRoundedVertex(K, radius, delta, V, result);
+            }
+        }
+
+        void EdgeSeparated(int i0, int i1, int i2, Vector3<Real> const& K,
+            Vector3<Real> const& C, Real radius, Vector3<Real> const& delta,
+            Vector3<Real> const& V, Result& result)
+        {
+            if (V[i0] < (Real)0 || V[i1] < (Real)0)
+            {
+                DoQueryRayRoundedEdge(i0, i1, i2, K, C, radius, delta, V, result);
+            }
+        }
+
+        void VertexUnbounded(Vector3<Real> const& K, Vector3<Real> const& C, Real radius,
+            Vector3<Real> const& delta, Vector3<Real> const& V, Result& result)
+        {
+            if (V[0] < (Real)0 && V[1] < (Real)0 && V[2] < (Real)0)
+            {
+                // Determine the face of the rounded box that is intersected
+                // by the ray C+T*V.
+                Real T = (radius - delta[0]) / V[0];
+                int j0 = 0;
+                Real temp = (radius - delta[1]) / V[1];
+                if (temp > T)
+                {
+                    T = temp;
+                    j0 = 1;
+                }
+                temp = (radius - delta[2]) / V[2];
+                if (temp > T)
+                {
+                    T = temp;
+                    j0 = 2;
+                }
+
+                // The j0-rounded face is the candidate for intersection.
+                int j1 = (j0 + 1) % 3;
+                int j2 = (j1 + 1) % 3;
+                DoQueryRayRoundedFace(j0, j1, j2, K, C, radius, delta, V, result);
+            }
+        }
+
+        void EdgeUnbounded(int i0, int i1, int /* i2 */, Vector3<Real> const& K,
+            Vector3<Real> const& C, Real radius, Vector3<Real> const& delta,
+            Vector3<Real> const& V, Result& result)
+        {
+            if (V[i0] < (Real)0 && V[i1] < (Real)0)
+            {
+                // Determine the face of the rounded box that is intersected
+                // by the ray C+T*V.
+                Real T = (radius - delta[i0]) / V[i0];
+                int j0 = i0;
+                Real temp = (radius - delta[i1]) / V[i1];
+                if (temp > T)
+                {
+                    T = temp;
+                    j0 = i1;
+                }
+
+                // The j0-rounded face is the candidate for intersection.
+                int j1 = (j0 + 1) % 3;
+                int j2 = (j1 + 1) % 3;
+                DoQueryRayRoundedFace(j0, j1, j2, K, C, radius, delta, V, result);
+            }
+        }
+
+        void FaceUnbounded(int i0, int i1, int i2, Vector3<Real> const& K,
+            Vector3<Real> const& C, Real radius, Vector3<Real> const& delta,
+            Vector3<Real> const& V, Result& result)
+        {
+            if (V[i0] < (Real)0)
+            {
+                DoQueryRayRoundedFace(i0, i1, i2, K, C, radius, delta, V, result);
+            }
+        }
+
+        void DoQueryRayRoundedVertex(Vector3<Real> const& K, Real radius,
+            Vector3<Real> const& delta, Vector3<Real> const& V, Result& result)
+        {
+            Real a1 = Dot(V, delta);
+            if (a1 < (Real)0)
+            {
+                // The caller must ensure that a0 > 0 and a2 > 0.
+                Real a0 = Dot(delta, delta) - radius * radius;
+                Real a2 = Dot(V, V);
+                Real adiscr = a1 * a1 - a2 * a0;
+                if (adiscr >= (Real)0)
+                {
+                    // The ray intersects the rounded vertex, so the sphere-box
+                    // contact point is the vertex.
+                    result.intersectionType = 1;
+                    result.contactTime = -(a1 + std::sqrt(adiscr)) / a2;
+                    result.contactPoint = K;
+                }
+            }
+        }
+
+        void DoQueryRayRoundedEdge(int i0, int i1, int i2, Vector3<Real> const& K,
+            Vector3<Real> const& C, Real radius, Vector3<Real> const& delta,
+            Vector3<Real> const& V, Result& result)
+        {
+            Real b1 = V[i0] * delta[i0] + V[i1] * delta[i1];
+            if (b1 < (Real)0)
+            {
+                // The caller must ensure that b0 > 0 and b2 > 0.
+                Real b0 = delta[i0] * delta[i0] + delta[i1] * delta[i1] - radius * radius;
+                Real b2 = V[i0] * V[i0] + V[i1] * V[i1];
+                Real bdiscr = b1 * b1 - b2 * b0;
+                if (bdiscr >= (Real)0)
+                {
+                    Real T = -(b1 + std::sqrt(bdiscr)) / b2;
+                    Real p2 = C[i2] + T * V[i2];
+                    if (-K[i2] <= p2)
+                    {
+                        if (p2 <= K[i2])
+                        {
+                            // The ray intersects the finite cylinder of the
+                            // rounded edge, so the sphere-box contact point
+                            // is on the corresponding box edge.
+                            result.intersectionType = 1;
+                            result.contactTime = T;
+                            result.contactPoint[i0] = K[i0];
+                            result.contactPoint[i1] = K[i1];
+                            result.contactPoint[i2] = p2;
+                        }
+                        else
+                        {
+                            // The ray intersects the infinite cylinder but
+                            // not the finite cylinder of the rounded edge.
+                            // It is possible the ray intersects the rounded
+                            // vertex for K.
+                            DoQueryRayRoundedVertex(K, radius, delta, V, result);
+                        }
+                    }
+                    else
+                    {
+                        // The ray intersects the infinite cylinder but
+                        // not the finite cylinder of the rounded edge.
+                        // It is possible the ray intersects the rounded
+                        // vertex for otherK.
+                        Vector3<Real> otherK, otherDelta;
+                        otherK[i0] = K[i0];
+                        otherK[i1] = K[i1];
+                        otherK[i2] = -K[i2];
+                        otherDelta[i0] = C[i0] - otherK[i0];
+                        otherDelta[i1] = C[i1] - otherK[i1];
+                        otherDelta[i2] = C[i2] - otherK[i2];
+                        DoQueryRayRoundedVertex(otherK, radius, otherDelta, V, result);
+                    }
+                }
+            }
+        }
+
+        void DoQueryRayRoundedFace(int i0, int i1, int i2, Vector3<Real> const& K,
+            Vector3<Real> const& C, Real radius, Vector3<Real> const& delta,
+            Vector3<Real> const& V, Result& result)
+        {
+            Vector3<Real> otherK, otherDelta;
+
+            Real T = (radius - delta[i0]) / V[i0];
+            Real p1 = C[i1] + T * V[i1];
+            Real p2 = C[i2] + T * V[i2];
+
+            if (p1 < -K[i1])
+            {
+                // The ray potentially intersects the rounded (i0,i1)-edge
+                // whose top-most vertex is otherK.
+                otherK[i0] = K[i0];
+                otherK[i1] = -K[i1];
+                otherK[i2] = K[i2];
+                otherDelta[i0] = C[i0] - otherK[i0];
+                otherDelta[i1] = C[i1] - otherK[i1];
+                otherDelta[i2] = C[i2] - otherK[i2];
+                DoQueryRayRoundedEdge(i0, i1, i2, otherK, C, radius, otherDelta, V, result);
+                if (result.intersectionType == 0)
+                {
+                    if (p2 < -K[i2])
+                    {
+                        // The ray potentially intersects the rounded
+                        // (i2,i0)-edge whose right-most vertex is otherK.
+                        otherK[i0] = K[i0];
+                        otherK[i1] = K[i1];
+                        otherK[i2] = -K[i2];
+                        otherDelta[i0] = C[i0] - otherK[i0];
+                        otherDelta[i1] = C[i1] - otherK[i1];
+                        otherDelta[i2] = C[i2] - otherK[i2];
+                        DoQueryRayRoundedEdge(i2, i0, i1, otherK, C, radius, otherDelta, V, result);
+                    }
+                    else if (p2 > K[i2])
+                    {
+                        // The ray potentially intersects the rounded
+                        // (i2,i0)-edge whose right-most vertex is K.
+                        DoQueryRayRoundedEdge(i2, i0, i1, K, C, radius, delta, V, result);
+                    }
+                }
+            }
+            else if (p1 <= K[i1])
+            {
+                if (p2 < -K[i2])
+                {
+                    // The ray potentially intersects the rounded
+                    // (i2,i0)-edge whose right-most vertex is otherK.
+                    otherK[i0] = K[i0];
+                    otherK[i1] = K[i1];
+                    otherK[i2] = -K[i2];
+                    otherDelta[i0] = C[i0] - otherK[i0];
+                    otherDelta[i1] = C[i1] - otherK[i1];
+                    otherDelta[i2] = C[i2] - otherK[i2];
+                    DoQueryRayRoundedEdge(i2, i0, i1, otherK, C, radius, otherDelta, V, result);
+                }
+                else if (p2 <= K[i2])
+                {
+                    // The ray intersects the i0-face of the rounded box, so
+                    // the sphere-box contact point is on the corresponding
+                    // box face.
+                    result.intersectionType = 1;
+                    result.contactTime = T;
+                    result.contactPoint[i0] = K[i0];
+                    result.contactPoint[i1] = p1;
+                    result.contactPoint[i2] = p2;
+                }
+                else  // p2 > K[i2]
+                {
+                    // The ray potentially intersects the rounded
+                    // (i2,i0)-edge whose right-most vertex is K.
+                    DoQueryRayRoundedEdge(i2, i0, i1, K, C, radius, delta, V, result);
+                }
+            }
+            else // p1 > K[i1]
+            {
+                // The ray potentially intersects the rounded (i0,i1)-edge
+                // whose top-most vertex is K.
+                DoQueryRayRoundedEdge(i0, i1, i2, K, C, radius, delta, V, result);
+                if (result.intersectionType == 0)
+                {
+                    if (p2 < -K[i2])
+                    {
+                        // The ray potentially intersects the rounded
+                        // (i2,i0)-edge whose right-most vertex is otherK.
+                        otherK[i0] = K[i0];
+                        otherK[i1] = K[i1];
+                        otherK[i2] = -K[i2];
+                        otherDelta[i0] = C[i0] - otherK[i0];
+                        otherDelta[i1] = C[i1] - otherK[i1];
+                        otherDelta[i2] = C[i2] - otherK[i2];
+                        DoQueryRayRoundedEdge(i2, i0, i1, otherK, C, radius, otherDelta, V, result);
+                    }
+                    else if (p2 > K[i2])
+                    {
+                        // The ray potentially intersects the rounded
+                        // (i2,i0)-edge whose right-most vertex is K.
+                        DoQueryRayRoundedEdge(i2, i0, i1, K, C, radius, delta, V, result);
+                    }
+                }
+            }
+        }
+    };
+}
diff --git a/Engine/Plugins/Experimental/GeometryProcessing/Source/GeometryAlgorithms/Private/GTEngine/Mathematics/GteIntrArc2Arc2.h b/Engine/Plugins/Experimental/GeometryProcessing/Source/GeometryAlgorithms/Private/GTEngine/Mathematics/GteIntrArc2Arc2.h
new file mode 100644
--- /dev/null
+++ b/Engine/Plugins/Experimental/GeometryProcessing/Source/GeometryAlgorithms/Private/GTEngine/Mathematics/GteIntrArc2Arc2.h
@@ -0,0 +1,212 @@
+// David Eberly, Geometric Tools, Redmond WA 98052
+// Copyright (c) 1998-2019
+// Distributed under the Boost Software License, Version 1.0.
+// http://www.boost.org/LICENSE_1_0.txt
+// http://www.geometrictools.com/License/Boost/LICENSE_1_0.txt
+// File Version: 3.0.1 (2017/02/23)
+
+#pragma once
+
+#include <GTEngine/Mathematics/GteArc2.h>
+#include <GTEngine/Mathematics/GteIntrCircle2Circle2.h>
+
+namespace gte
+{
+
+template <typename Real>
+class FIQuery<Real, Arc2<Real>, Arc2<Real>>
+{
+public:
+    // The possible 'configuration' in Result are listed as an enumeration.
+    // The valid array elements are listed in the comments.
+    enum
+    {
+        NO_INTERSECTION,
+        NONCOCIRCULAR_ONE_POINT,        // point[0]
+        NONCOCIRCULAR_TWO_POINTS,       // point[0], point[1]
+        COCIRCULAR_ONE_POINT,           // point[0]
+        COCIRCULAR_TWO_POINTS,          // point[0], point[1]
+        COCIRCULAR_ONE_POINT_ONE_ARC,   // point[0], arc[0]
+        COCIRCULAR_ONE_ARC,             // arc[0]
+        COCIRCULAR_TWO_ARCS             // arc[0], arc[1]
+    };
+
+    struct Result
+    {
+        bool intersect;  // 'true' iff configuration != NO_INTERSECTION
+        int configuration;  // one of the enumerations listed previously
+        Vector2<Real> point[2];
+        Arc2<Real> arc[2];
+    };
+
+    Result operator()(Arc2<Real> const& arc0, Arc2<Real> const& arc1);
+};
+
+
+template <typename Real>
+typename FIQuery<Real, Arc2<Real>, Arc2<Real>>::Result
+FIQuery<Real, Arc2<Real>, Arc2<Real>>::operator()(Arc2<Real> const& arc0, Arc2<Real> const& arc1)
+{
+    // Assume initially there are no intersections.  If we find at least one
+    // intersection, we will set result.intersect to 'true'.
+    Result result;
+    result.intersect = false;
+
+    Circle2<Real> circle0(arc0.center, arc0.radius);
+    Circle2<Real> circle1(arc1.center, arc1.radius);
+    FIQuery<Real, Circle2<Real>, Circle2<Real>> ccQuery;
+    auto ccResult = ccQuery(circle0, circle1);
+    if (!ccResult.intersect)
+    {
+        // The arcs do not intersect.
+        result.configuration = NO_INTERSECTION;
+        return result;
+    }
+
+    if (ccResult.numIntersections == std::numeric_limits<int>::max())
+    {
+        // The arcs are cocircular.  Determine whether they overlap.  Let
+        // arc0 be <A0,A1> and arc1 be <B0,B1>.  The points are ordered
+        // counterclockwise around the circle of the arc.
+        if (arc1.Contains(arc0.end[0]))
+        {
+            result.intersect = true;
+            if (arc1.Contains(arc0.end[1]))
+            {
+                if (arc0.Contains(arc1.end[0]) && arc0.Contains(arc1.end[1]))
+                {
+                    if (arc0.end[0] == arc1.end[0] && arc0.end[1] == arc1.end[1])
+                    {
+                        // The arcs are the same.
+                        result.configuration = COCIRCULAR_ONE_ARC;
+                        result.arc[0] = arc0;
+                    }
+                    else
+                    {
+                        // arc0 and arc1 overlap in two disjoint subsets.
+                        if (arc0.end[0] != arc1.end[1])
+                        {
+                            if (arc1.end[0] != arc0.end[1])
+                            {
+                                // The arcs overlap in two disjoint subarcs, each
+                                // of positive subtended angle: <A0,B1>, <A1,B0>
+                                result.configuration = COCIRCULAR_TWO_ARCS;
+                                result.arc[0] = Arc2<Real>(arc0.center, arc0.radius, arc0.end[0], arc1.end[1]);
+                                result.arc[1] = Arc2<Real>(arc0.center, arc0.radius, arc1.end[0], arc0.end[1]);
+                            }
+                            else  // B0 = A1
+                            {
+                                // The intersection is a point {A1} and an
+                                // arc <A0,B1>.
+                                result.configuration = COCIRCULAR_ONE_POINT_ONE_ARC;
+                                result.point[0] = arc0.end[1];
+                                result.arc[0] = Arc2<Real>(arc0.center, arc0.radius, arc0.end[0], arc1.end[1]);
+                            }
+                        }
+                        else  // A0 = B1
+                        {
+                            if (arc1.end[0] != arc0.end[1])
+                            {
+                                // The intersection is a point {A0} and an
+                                // arc <A1,B0>.
+                                result.configuration = COCIRCULAR_ONE_POINT_ONE_ARC;
+                                result.point[0] = arc0.end[0];
+                                result.arc[0] = Arc2<Real>(arc0.center, arc0.radius, arc1.end[0], arc0.end[1]);
+                            }
+                            else
+                            {
+                                // The arcs shared endpoints, so the union is a circle.
+                                result.configuration = COCIRCULAR_TWO_POINTS;
+                                result.point[0] = arc0.end[0];
+                                result.point[1] = arc0.end[1];
+                            }
+                        }
+                    }
+                }
+                else
+                {
+                    // Arc0 inside arc1, <B0,A0,A1,B1>.
+                    result.configuration = COCIRCULAR_ONE_ARC;
+                    result.arc[0] = arc0;
+                }
+            }
+            else
+            {
+                if (arc0.end[0] != arc1.end[1])
+                {
+                    // Arc0 and arc1 overlap, <B0,A0,B1,A1>.
+                    result.configuration = COCIRCULAR_ONE_ARC;
+                    result.arc[0] = Arc2<Real>(arc0.center, arc0.radius, arc0.end[0], arc1.end[1]);
+                }
+                else
+                {
+                    // Arc0 and arc1 share endpoint, <B0,A0,B1,A1>, A0 = B1.
+                    result.configuration = COCIRCULAR_ONE_POINT;
+                    result.point[0] = arc0.end[0];
+                }
+            }
+            return result;
+        }
+
+        if (arc1.Contains(arc0.end[1]))
+        {
+            result.intersect = true;
+            if (arc0.end[1] != arc1.end[0])
+            {
+                // Arc0 and arc1 overlap in a single arc, <A0,B0,A1,B1>.
+                result.configuration = COCIRCULAR_ONE_ARC;
+                result.arc[0] = Arc2<Real>(arc0.center, arc0.radius, arc1.end[0], arc0.end[1]);
+            }
+            else
+            {
+                // Arc0 and arc1 share endpoint, <A0,B0,A1,B1>, B0 = A1.
+                result.configuration = COCIRCULAR_ONE_POINT;
+                result.point[0] = arc1.end[0];
+            }
+            return result;
+        }
+
+        if (arc0.Contains(arc1.end[0]))
+        {
+            // Arc1 inside arc0, <A0,B0,B1,A1>.
+            result.intersect = true;
+            result.configuration = COCIRCULAR_ONE_ARC;
+            result.arc[0] = arc1;
+        }
+        else
+        {
+            // Arcs do not overlap, <A0,A1,B0,B1>.
+            result.configuration = NO_INTERSECTION;
+        }
+        return result;
+    }
+
+    // Test whether circle-circle intersection points are on the arcs.
+    int numIntersections = 0;
+    for (int i = 0; i < ccResult.numIntersections; ++i)
+    {
+        if (arc0.Contains(ccResult.point[i]) && arc1.Contains(ccResult.point[i]))
+        {
+            result.point[numIntersections++] = ccResult.point[i];
+            result.intersect = true;
+        }
+    }
+
+    if (numIntersections == 2)
+    {
+        result.configuration = NONCOCIRCULAR_TWO_POINTS;
+    }
+    else if (numIntersections == 1)
+    {
+        result.configuration = NONCOCIRCULAR_ONE_POINT;
+    }
+    else
+    {
+        result.configuration = NO_INTERSECTION;
+    }
+
+    return result;
+}
+
+
+}
diff --git a/Engine/Plugins/Experimental/GeometryProcessing/Source/GeometryAlgorithms/Private/GTEngine/Mathematics/GteIntrCapsule3Capsule3.h b/Engine/Plugins/Experimental/GeometryProcessing/Source/GeometryAlgorithms/Private/GTEngine/Mathematics/GteIntrCapsule3Capsule3.h
new file mode 100644
--- /dev/null
+++ b/Engine/Plugins/Experimental/GeometryProcessing/Source/GeometryAlgorithms/Private/GTEngine/Mathematics/GteIntrCapsule3Capsule3.h
@@ -0,0 +1,45 @@
+// David Eberly, Geometric Tools, Redmond WA 98052
+// Copyright (c) 1998-2019
+// Distributed under the Boost Software License, Version 1.0.
+// http://www.boost.org/LICENSE_1_0.txt
+// http://www.geometrictools.com/License/Boost/LICENSE_1_0.txt
+// File Version: 3.0.0 (2016/06/19)
+
+#pragma once
+
+#include <GTEngine/Mathematics/GteDistSegmentSegment.h>
+#include <GTEngine/Mathematics/GteCapsule.h>
+#include <GTEngine/Mathematics/GteTIQuery.h>
+
+namespace gte
+{
+
+template <typename Real>
+class TIQuery<Real, Capsule3<Real>, Capsule3<Real>>
+{
+public:
+    struct Result
+    {
+        bool intersect;
+    };
+
+    Result operator()(Capsule3<Real> const& capsule0,
+        Capsule3<Real> const& capsule1);
+};
+
+
+template <typename Real>
+typename TIQuery<Real, Capsule3<Real>, Capsule3<Real>>::Result
+TIQuery<Real, Capsule3<Real>, Capsule3<Real>>::operator()(
+    Capsule3<Real> const& capsule0, Capsule3<Real> const& capsule1)
+{
+    Result result;
+    DCPQuery<Real, Segment3<Real>, Segment3<Real>> ssQuery;
+    auto ssResult = ssQuery(capsule0.segment, capsule1.segment);
+    Real rSum = capsule0.radius + capsule1.radius;
+    result.intersect = (ssResult.distance <= rSum);
+    return result;
+}
+
+
+}
diff --git a/Engine/Plugins/Experimental/GeometryProcessing/Source/GeometryAlgorithms/Private/GTEngine/Mathematics/GteIntrCircle2Arc2.h b/Engine/Plugins/Experimental/GeometryProcessing/Source/GeometryAlgorithms/Private/GTEngine/Mathematics/GteIntrCircle2Arc2.h
new file mode 100644
--- /dev/null
+++ b/Engine/Plugins/Experimental/GeometryProcessing/Source/GeometryAlgorithms/Private/GTEngine/Mathematics/GteIntrCircle2Arc2.h
@@ -0,0 +1,82 @@
+// David Eberly, Geometric Tools, Redmond WA 98052
+// Copyright (c) 1998-2019
+// Distributed under the Boost Software License, Version 1.0.
+// http://www.boost.org/LICENSE_1_0.txt
+// http://www.geometrictools.com/License/Boost/LICENSE_1_0.txt
+// File Version: 3.0.0 (2016/06/19)
+
+#pragma once
+
+#include <GTEngine/Mathematics/GteArc2.h>
+#include <GTEngine/Mathematics/GteIntrCircle2Circle2.h>
+
+namespace gte
+{
+
+template <typename Real>
+class FIQuery<Real, Circle2<Real>, Arc2<Real>>
+{
+public:
+    struct Result
+    {
+        bool intersect;
+
+        // The number of intersections is 0, 1, 2, or maxInt =
+        // std::numeric_limits<int>::max().  When 1, the arc and circle
+        // intersect in a single point.  When 2, the arc is not on the circle
+        // and they intersect in two points.  When maxInt, the arc is on the
+        // circle.
+        int numIntersections;
+
+        // Valid only when numIntersections = 1 or 2.
+        Vector2<Real> point[2];
+
+        // Valid only when numIntersections = maxInt.
+        Arc2<Real> arc;
+    };
+
+    Result operator()(Circle2<Real> const& circle, Arc2<Real> const& arc);
+};
+
+
+template <typename Real>
+typename FIQuery<Real, Circle2<Real>, Arc2<Real>>::Result
+FIQuery<Real, Circle2<Real>, Arc2<Real>>::operator()(
+    Circle2<Real> const& circle, Arc2<Real> const& arc)
+{
+    Result result;
+
+    Circle2<Real> circleOfArc(arc.center, arc.radius);
+    FIQuery<Real, Circle2<Real>, Circle2<Real>> ccQuery;
+    auto ccResult = ccQuery(circle, circleOfArc);
+    if (!ccResult.intersect)
+    {
+        result.intersect = false;
+        result.numIntersections = 0;
+        return result;
+    }
+
+    if (ccResult.numIntersections == std::numeric_limits<int>::max())
+    {
+        // The arc is on the circle.
+        result.intersect = true;
+        result.numIntersections = std::numeric_limits<int>::max();
+        result.arc = arc;
+        return result;
+    }
+
+    // Test whether circle-circle intersection points are on the arc.
+    for (int i = 0; i < ccResult.numIntersections; ++i)
+    {
+        result.numIntersections = 0;
+        if (arc.Contains(ccResult.point[i]))
+        {
+            result.point[result.numIntersections++] = ccResult.point[i];
+            result.intersect = true;
+        }
+    }
+    return result;
+}
+
+
+}
diff --git a/Engine/Plugins/Experimental/GeometryProcessing/Source/GeometryAlgorithms/Private/GTEngine/Mathematics/GteIntrCircle2Circle2.h b/Engine/Plugins/Experimental/GeometryProcessing/Source/GeometryAlgorithms/Private/GTEngine/Mathematics/GteIntrCircle2Circle2.h
new file mode 100644
--- /dev/null
+++ b/Engine/Plugins/Experimental/GeometryProcessing/Source/GeometryAlgorithms/Private/GTEngine/Mathematics/GteIntrCircle2Circle2.h
@@ -0,0 +1,172 @@
+// David Eberly, Geometric Tools, Redmond WA 98052
+// Copyright (c) 1998-2019
+// Distributed under the Boost Software License, Version 1.0.
+// http://www.boost.org/LICENSE_1_0.txt
+// http://www.geometrictools.com/License/Boost/LICENSE_1_0.txt
+// File Version: 3.0.1 (2018/10/05)
+
+#pragma once
+
+#include <GTEngine/Mathematics/GteVector2.h>
+#include <GTEngine/Mathematics/GteHypersphere.h>
+#include <GTEngine/Mathematics/GteFIQuery.h>
+#include <GTEngine/Mathematics/GteTIQuery.h>
+#include <limits>
+
+namespace gte
+{
+
+template <typename Real>
+class TIQuery<Real, Circle2<Real>, Circle2<Real>>
+{
+public:
+    struct Result
+    {
+        bool intersect;
+    };
+
+    Result operator()(Circle2<Real> const& circle0,
+        Circle2<Real> const& circle1);
+};
+
+template <typename Real>
+class FIQuery<Real, Circle2<Real>, Circle2<Real>>
+{
+public:
+    struct Result
+    {
+        bool intersect;
+
+        // The number of intersections is 0, 1, 2, or maxInt =
+        // std::numeric_limits<int>::max().  When 1, the circles are tangent
+        // and intersect in a single point.  When 2, circles have two
+        // transverse intersection points.  When maxInt, the circles are the
+        // same.
+        int numIntersections;
+
+        // Valid only when numIntersections = 1 or 2.
+        Vector2<Real> point[2];
+
+        // Valid only when numIntersections = maxInt.
+        Circle2<Real> circle;
+    };
+
+    Result operator()(Circle2<Real> const& circle0,
+        Circle2<Real> const& circle1);
+};
+
+
+template <typename Real>
+typename TIQuery<Real, Circle2<Real>, Circle2<Real>>::Result
+TIQuery<Real, Circle2<Real>, Circle2<Real>>::operator()(
+    Circle2<Real> const& circle0, Circle2<Real> const& circle1)
+{
+    Result result;
+    Vector2<Real> diff = circle0.center - circle1.center;
+    result.intersect = (Length(diff) <= circle0.radius + circle1.radius);
+    return result;
+}
+
+template <typename Real>
+typename FIQuery<Real, Circle2<Real>, Circle2<Real>>::Result
+FIQuery<Real, Circle2<Real>, Circle2<Real>>::operator()(
+    Circle2<Real> const& circle0, Circle2<Real> const& circle1)
+{
+    // The two circles are |X-C0| = R0 and |X-C1| = R1.  Define U = C1 - C0
+    // and V = Perp(U) where Perp(x,y) = (y,-x).  Note that Dot(U,V) = 0 and
+    // |V|^2 = |U|^2.  The intersection points X can be written in the form
+    // X = C0+s*U+t*V and X = C1+(s-1)*U+t*V.  Squaring the circle equations
+    // and substituting these formulas into them yields
+    //   R0^2 = (s^2 + t^2)*|U|^2
+    //   R1^2 = ((s-1)^2 + t^2)*|U|^2.
+    // Subtracting and solving for s yields
+    //   s = ((R0^2-R1^2)/|U|^2 + 1)/2
+    // Then replace in the first equation and solve for t^2
+    //   t^2 = (R0^2/|U|^2) - s^2.
+    // In order for there to be solutions, the right-hand side must be
+    // nonnegative.  Some algebra leads to the condition for existence of
+    // solutions,
+    //   (|U|^2 - (R0+R1)^2)*(|U|^2 - (R0-R1)^2) <= 0.
+    // This reduces to
+    //   |R0-R1| <= |U| <= |R0+R1|.
+    // If |U| = |R0-R1|, then the circles are side-by-side and just tangent.
+    // If |U| = |R0+R1|, then the circles are nested and just tangent.
+    // If |R0-R1| < |U| < |R0+R1|, then the two circles to intersect in two
+    // points.
+
+    Result result;
+
+    Vector2<Real> U = circle1.center - circle0.center;
+    Real USqrLen = Dot(U, U);
+    Real R0 = circle0.radius, R1 = circle1.radius;
+    Real R0mR1 = R0 - R1;
+    if (USqrLen == (Real)0 && R0mR1 == (Real)0)
+    {
+        // Circles are the same.
+        result.intersect = true;
+        result.numIntersections = std::numeric_limits<int>::max();
+        result.circle = circle0;
+        return result;
+    }
+
+    Real R0mR1Sqr = R0mR1*R0mR1;
+    if (USqrLen < R0mR1Sqr)
+    {
+        // The circles do not intersect.
+        result.intersect = false;
+        result.numIntersections = 0;
+        return result;
+    }
+
+    Real R0pR1 = R0 + R1;
+    Real R0pR1Sqr = R0pR1*R0pR1;
+    if (USqrLen > R0pR1Sqr)
+    {
+        // The circles do not intersect.
+        result.intersect = false;
+        result.numIntersections = 0;
+        return result;
+    }
+
+    if (USqrLen < R0pR1Sqr)
+    {
+        if (R0mR1Sqr < USqrLen)
+        {
+            Real invUSqrLen = ((Real)1) / USqrLen;
+            Real s = ((Real)0.5)*((R0*R0 - R1*R1)*invUSqrLen + (Real)1);
+            Vector2<Real> tmp = circle0.center + s*U;
+
+            // In theory, discr is nonnegative.  However, numerical round-off
+            // errors can make it slightly negative.  Clamp it to zero.
+            Real discr = R0*R0*invUSqrLen - s*s;
+            if (discr < (Real)0)
+            {
+                discr = (Real)0;
+            }
+            Real t = std::sqrt(discr);
+            Vector2<Real> V{ U[1], -U[0] };
+            result.point[0] = tmp - t*V;
+            result.point[1] = tmp + t*V;
+            result.numIntersections = (t >(Real)0 ? 2 : 1);
+        }
+        else
+        {
+            // |U| = |R0-R1|, circles are tangent.
+            result.numIntersections = 1;
+            result.point[0] = circle0.center + (R0 / R0mR1)*U;
+        }
+    }
+    else
+    {
+        // |U| = |R0+R1|, circles are tangent.
+        result.numIntersections = 1;
+        result.point[0] = circle0.center + (R0 / R0pR1)*U;
+    }
+
+    // The circles intersect in 1 or 2 points.
+    result.intersect = true;
+    return result;
+}
+
+
+}
