move GeometryProcessing plugin from NotForLicensees to Experimental

diff --git a/Engine/Plugins/Experimental/GeometryProcessing/Source/GeometryAlgorithms/Private/GTEngine/Mathematics/GteIntrConvexPolygonHyperplane.h b/Engine/Plugins/Experimental/GeometryProcessing/Source/GeometryAlgorithms/Private/GTEngine/Mathematics/GteIntrConvexPolygonHyperplane.h
new file mode 100644
--- /dev/null
+++ b/Engine/Plugins/Experimental/GeometryProcessing/Source/GeometryAlgorithms/Private/GTEngine/Mathematics/GteIntrConvexPolygonHyperplane.h
@@ -0,0 +1,402 @@
+// David Eberly, Geometric Tools, Redmond WA 98052
+// Copyright (c) 1998-2019
+// Distributed under the Boost Software License, Version 1.0.
+// http://www.boost.org/LICENSE_1_0.txt
+// http://www.geometrictools.com/License/Boost/LICENSE_1_0.txt
+// File Version: 3.0.2 (2019/02/13)
+
+#pragma once
+
+#include <GTEngine/Mathematics/GteTIQuery.h>
+#include <GTEngine/Mathematics/GteFIQuery.h>
+#include <GTEngine/Mathematics/GteHyperplane.h>
+#include <GTEngine/Mathematics/GteVector.h>
+#include <limits>
+#include <list>
+#include <vector>
+
+// The intersection queries are based on the document
+// https://www.geometrictools.com/Documentation/ClipConvexPolygonByHyperplane.pdf
+
+namespace gte
+{
+    template <int N, typename Real>
+    class TIQuery<Real, std::vector<Vector<N, Real>>, Hyperplane<N, Real>>
+    {
+    public:
+        enum class Configuration
+        {
+            SPLIT,
+            POSITIVE_SIDE_VERTEX,
+            POSITIVE_SIDE_EDGE,
+            POSITIVE_SIDE_STRICT,
+            NEGATIVE_SIDE_VERTEX,
+            NEGATIVE_SIDE_EDGE,
+            NEGATIVE_SIDE_STRICT,
+            CONTAINED,
+            INVALID_POLYGON
+        };
+
+        struct Result
+        {
+            bool intersect;
+            Configuration configuration;
+        };
+
+        Result operator()(std::vector<Vector<N, Real>> const& polygon, Hyperplane<N, Real> const& hyperplane)
+        {
+            Result result;
+
+            size_t const numVertices = polygon.size();
+            if (numVertices < 3)
+            {
+                // The convex polygon must have at least 3 vertices.
+                result.intersect = false;
+                result.configuration = Configuration::INVALID_POLYGON;
+                return result;
+            }
+
+            // Determine on which side of the hyperplane each vertex lies.
+            size_t numPositive = 0, numNegative = 0, numZero = 0;
+            for (size_t i = 0; i < numVertices; ++i)
+            {
+                Real h = Dot(hyperplane.normal, polygon[i]) - hyperplane.constant;
+                if (h > (Real)0)
+                {
+                    ++numPositive;
+                }
+                else if (h < (Real)0)
+                {
+                    ++numNegative;
+                }
+                else
+                {
+                    ++numZero;
+                }
+            }
+
+            if (numPositive > 0)
+            {
+                if (numNegative > 0)
+                {
+                    result.intersect = true;
+                    result.configuration = Configuration::SPLIT;
+                }
+                else if (numZero == 0)
+                {
+                    result.intersect = false;
+                    result.configuration = Configuration::POSITIVE_SIDE_STRICT;
+                }
+                else if (numZero == 1)
+                {
+                    result.intersect = true;
+                    result.configuration = Configuration::POSITIVE_SIDE_VERTEX;
+                }
+                else // numZero > 1
+                {
+                    result.intersect = true;
+                    result.configuration = Configuration::POSITIVE_SIDE_EDGE;
+                }
+            }
+            else if (numNegative > 0)
+            {
+                if (numZero == 0)
+                {
+                    result.intersect = false;
+                    result.configuration = Configuration::NEGATIVE_SIDE_STRICT;
+                }
+                else if (numZero == 1)
+                {
+                    // The polygon touches the plane in a vertex or an edge.
+                    result.intersect = true;
+                    result.configuration = Configuration::NEGATIVE_SIDE_VERTEX;
+                }
+                else // numZero > 1
+                {
+                    result.intersect = true;
+                    result.configuration = Configuration::NEGATIVE_SIDE_EDGE;
+                }
+            }
+            else  // numZero == numVertices
+            {
+                result.intersect = true;
+                result.configuration = Configuration::CONTAINED;
+            }
+
+            return result;
+        }
+    };
+
+    template <int N, typename Real>
+    class FIQuery<Real, std::vector<Vector<N, Real>>, Hyperplane<N, Real>>
+    {
+    public:
+        enum class Configuration
+        {
+            SPLIT,
+            POSITIVE_SIDE_VERTEX,
+            POSITIVE_SIDE_EDGE,
+            POSITIVE_SIDE_STRICT,
+            NEGATIVE_SIDE_VERTEX,
+            NEGATIVE_SIDE_EDGE,
+            NEGATIVE_SIDE_STRICT,
+            CONTAINED,
+            INVALID_POLYGON
+        };
+
+        struct Result
+        {
+            // The intersection is either empty, a single vertex, a single
+            // edge or the polygon is contained by the hyperplane.
+            Configuration configuration;
+            std::vector<Vector<N, Real>> intersection;
+
+            // If 'configuration' is POSITIVE_* or SPLIT, this polygon is the
+            // portion of the query input 'polygon' on the positive side of
+            // the hyperplane with possibly a vertex or edge on the hyperplane.
+            std::vector<Vector<N, Real>> positivePolygon;
+
+            // If 'configuration' is NEGATIVE_* or SPLIT, this polygon is the
+            // portion of the query input 'polygon' on the negative side of
+            // the hyperplane with possibly a vertex or edge on the hyperplane.
+            std::vector<Vector<N, Real>> negativePolygon;
+        };
+
+        Result operator()(std::vector<Vector<N, Real>> const& polygon, Hyperplane<N, Real> const& hyperplane)
+        {
+            Result result;
+
+            size_t const numVertices = polygon.size();
+            if (numVertices < 3)
+            {
+                // The convex polygon must have at least 3 vertices.
+                result.configuration = Configuration::INVALID_POLYGON;
+                return result;
+            }
+
+            // Determine on which side of the hyperplane the vertices live.
+            // The index maxPosIndex stores the index of the vertex on the
+            // positive side of the hyperplane that is farthest from the
+            // hyperplane.  The index maxNegIndex stores the index of the
+            // vertex on the negative side of the hyperplane that is farthest
+            // from the hyperplane.  If one or the other such vertex does not
+            // exist, the corresponding index will remain its initial value of
+            // max(size_t).
+            std::vector<Real> height(numVertices);
+            std::vector<size_t> zeroHeightIndices;
+            zeroHeightIndices.reserve(numVertices);
+            size_t numPositive = 0, numNegative = 0;
+            Real maxPosHeight = -std::numeric_limits<Real>::max();
+            Real maxNegHeight = std::numeric_limits<Real>::max();
+            size_t maxPosIndex = std::numeric_limits<size_t>::max();
+            size_t maxNegIndex = std::numeric_limits<size_t>::max();
+            for (size_t i = 0; i < numVertices; ++i)
+            {
+                height[i] = Dot(hyperplane.normal, polygon[i]) - hyperplane.constant;
+                if (height[i] > (Real)0)
+                {
+                    ++numPositive;
+                    if (height[i] > maxPosHeight)
+                    {
+                        maxPosHeight = height[i];
+                        maxPosIndex = i;
+                    }
+                }
+                else if (height[i] < (Real)0)
+                {
+                    ++numNegative;
+                    if (height[i] < maxNegHeight)
+                    {
+                        maxNegHeight = height[i];
+                        maxNegIndex = i;
+                    }
+                }
+                else
+                {
+                    zeroHeightIndices.push_back(i);
+                }
+            }
+
+            if (numPositive > 0)
+            {
+                if (numNegative > 0)
+                {
+                    result.configuration = Configuration::SPLIT;
+
+                    bool doSwap = (maxPosHeight < -maxNegHeight);
+                    if (doSwap)
+                    {
+                        for (auto& h : height)
+                        {
+                            h = -h;
+                        }
+                        std::swap(maxPosIndex, maxNegIndex);
+                    }
+
+                    SplitPolygon(polygon, height, maxPosIndex, result);
+
+                    if (doSwap)
+                    {
+                        std::swap(result.positivePolygon, result.negativePolygon);
+                    }
+                }
+                else
+                {
+                    size_t numZero = zeroHeightIndices.size();
+                    if (numZero == 0)
+                    {
+                        result.configuration = Configuration::POSITIVE_SIDE_STRICT;
+                    }
+                    else if (numZero == 1)
+                    {
+                        result.configuration = Configuration::POSITIVE_SIDE_VERTEX;
+                        result.intersection =
+                        {
+                            polygon[zeroHeightIndices[0]]
+                        };
+                    }
+                    else // numZero > 1
+                    {
+                        result.configuration = Configuration::POSITIVE_SIDE_EDGE;
+                        result.intersection =
+                        {
+                            polygon[zeroHeightIndices[0]],
+                            polygon[zeroHeightIndices[1]]
+                        };
+                    }
+                    result.positivePolygon = polygon;
+                }
+            }
+            else if (numNegative > 0)
+            {
+                size_t numZero = zeroHeightIndices.size();
+                if (numZero == 0)
+                {
+                    result.configuration = Configuration::NEGATIVE_SIDE_STRICT;
+                }
+                else if (numZero == 1)
+                {
+                    result.configuration = Configuration::NEGATIVE_SIDE_VERTEX;
+                    result.intersection =
+                    {
+                        polygon[zeroHeightIndices[0]]
+                    };
+                }
+                else  // numZero > 1
+                {
+                    result.configuration = Configuration::NEGATIVE_SIDE_EDGE;
+                    result.intersection =
+                    {
+                        polygon[zeroHeightIndices[0]],
+                        polygon[zeroHeightIndices[1]]
+                    };
+                }
+                result.negativePolygon = polygon;
+            }
+            else  // numZero == numVertices
+            {
+                result.configuration = Configuration::CONTAINED;
+                result.intersection = polygon;
+            }
+
+            return result;
+        }
+
+    protected:
+        void SplitPolygon(std::vector<Vector<N, Real>> const& polygon,
+            std::vector<Real> const& height, size_t maxPosIndex, Result& result)
+        {
+            // Find the largest contiguous subset of indices for which
+            // height[i] >= 0.
+            size_t const numVertices = polygon.size();
+            std::list<Vector<N, Real>> positiveList;
+            positiveList.push_back(polygon[maxPosIndex]);
+            size_t end0 = maxPosIndex;
+            size_t end0prev = std::numeric_limits<size_t>::max();
+            for (size_t i = 0; i < numVertices; ++i)
+            {
+                end0prev = (end0 + numVertices - 1) % numVertices;
+                if (height[end0prev] >= (Real)0)
+                {
+                    positiveList.push_front(polygon[end0prev]);
+                    end0 = end0prev;
+                }
+                else
+                {
+                    break;
+                }
+            }
+
+            size_t end1 = maxPosIndex;
+            size_t end1next = std::numeric_limits<size_t>::max();
+            for (size_t i = 0; i < numVertices; ++i)
+            {
+                end1next = (end1 + 1) % numVertices;
+                if (height[end1next] >= (Real)0)
+                {
+                    positiveList.push_back(polygon[end1next]);
+                    end1 = end1next;
+                }
+                else
+                {
+                    break;
+                }
+            }
+
+            size_t index = end1next;
+            std::list<Vector<N, Real>> negativeList;
+            for (size_t i = 0; i < numVertices; ++i)
+            {
+                negativeList.push_back(polygon[index]);
+                index = (index + 1) % numVertices;
+                if (index == end0)
+                {
+                    break;
+                }
+            }
+
+            // Clip the polygon.
+            if (height[end0] > (Real)0)
+            {
+                Real t = -height[end0prev] / (height[end0] - height[end0prev]);
+                Real omt = (Real)1 - t;
+                Vector<N, Real> V = omt * polygon[end0prev] + t * polygon[end0];
+                positiveList.push_front(V);
+                negativeList.push_back(V);
+                result.intersection.push_back(V);
+            }
+            else
+            {
+                negativeList.push_back(polygon[end0]);
+                result.intersection.push_back(polygon[end0]);
+            }
+
+            if (height[end1] > (Real)0)
+            {
+                Real t = -height[end1next] / (height[end1] - height[end1next]);
+                Real omt = (Real)1 - t;
+                Vector<N, Real> V = omt * polygon[end1next] + t * polygon[end1];
+                positiveList.push_back(V);
+                negativeList.push_front(V);
+                result.intersection.push_back(V);
+            }
+            else
+            {
+                negativeList.push_front(polygon[end1]);
+                result.intersection.push_back(polygon[end1]);
+            }
+
+            result.positivePolygon.reserve(positiveList.size());
+            for (auto const& p : positiveList)
+            {
+                result.positivePolygon.push_back(p);
+            }
+
+            result.negativePolygon.reserve(negativeList.size());
+            for (auto const& p : negativeList)
+            {
+                result.negativePolygon.push_back(p);
+            }
+        }
+    };
+}
diff --git a/Engine/Plugins/Experimental/GeometryProcessing/Source/GeometryAlgorithms/Private/GTEngine/Mathematics/GteIntrDisk2Sector2.h b/Engine/Plugins/Experimental/GeometryProcessing/Source/GeometryAlgorithms/Private/GTEngine/Mathematics/GteIntrDisk2Sector2.h
new file mode 100644
--- /dev/null
+++ b/Engine/Plugins/Experimental/GeometryProcessing/Source/GeometryAlgorithms/Private/GTEngine/Mathematics/GteIntrDisk2Sector2.h
@@ -0,0 +1,185 @@
+// David Eberly, Geometric Tools, Redmond WA 98052
+// Copyright (c) 1998-2019
+// Distributed under the Boost Software License, Version 1.0.
+// http://www.boost.org/LICENSE_1_0.txt
+// http://www.geometrictools.com/License/Boost/LICENSE_1_0.txt
+// File Version: 3.0.1 (2018/10/05)
+
+#pragma once
+
+#include <GTEngine/Mathematics/GteSector2.h>
+#include <GTEngine/Mathematics/GteHypersphere.h>
+#include <GTEngine/Mathematics/GteTIQuery.h>
+
+// The Circle2 object is considered to be a disk whose points X satisfy the
+// constraint |X-C| <= R, where C is the disk center and R is the disk
+// radius.  The Sector2 object is also considered to be a solid.  Also,
+// the Sector2 object is required to be convex, so the sector angle must
+// be in (0,pi/2], even though the Sector2 definition allows for angles
+// larger than pi/2 (leading to nonconvex sectors).  The sector vertex is
+// V, the radius is L, the axis direction is D, and the angle is A. Sector
+// points X satisfy |X-V| <= L and Dot(D,X-V) >= cos(A)|X-V| >= 0.
+//
+// A subproblem for the test-intersection query is to determine whether
+// the disk intersects the cone of the sector.  Although the query is in
+// 2D, it is analogous to the 3D problem of determining whether a sphere
+// and cone overlap.  That algorithm is described in
+//   http://www.geometrictools.com/Documentation/IntersectionSphereCone.pdf
+// The algorithm leads to coordinate-free pseudocod which applies to 2D
+// as well as 3D.  That function is the first SphereIntersectsCone on
+// page 4 of the PDF.
+//
+// If the disk is outside the cone, there is no intersection.  If the disk
+// overlaps the cone, we then need to test whether the disk overlaps the
+// disk of the sector.
+
+namespace gte
+{
+
+template <typename Real>
+class TIQuery<Real, Circle2<Real>, Sector2<Real>>
+{
+public:
+    struct Result
+    {
+        bool intersect;
+    };
+
+    Result operator()(Circle2<Real> const& disk, Sector2<Real> const& sector);
+};
+
+template <typename Real>
+typename TIQuery<Real, Circle2<Real>, Sector2<Real>>::Result
+TIQuery<Real, Circle2<Real>, Sector2<Real>>::operator()(
+    Circle2<Real> const& disk, Sector2<Real> const& sector)
+{
+    Result result;
+
+    // Test whether the disk and the disk of the sector overlap.
+    Vector2<Real> CmV = disk.center - sector.vertex;
+    Real sqrLengthCmV = Dot(CmV, CmV);
+    Real lengthCmV = std::sqrt(sqrLengthCmV);
+    if (lengthCmV > disk.radius + sector.radius)
+    {
+        // The disk is outside the disk of the sector.
+        result.intersect = false;
+        return result;
+    }
+
+    // Test whether the disk and cone of the sector overlap.  The comments
+    // about K, K', and K" refer to the PDF mentioned previously.
+    Vector2<Real> U = sector.vertex - (disk.radius / sector.sinAngle) * sector.direction;
+    Vector2<Real> CmU = disk.center - U;
+    Real lengthCmU = Length(CmU);
+    if (Dot(sector.direction, CmU) < lengthCmU * sector.cosAngle)
+    {
+        // The disk center is outside K" (in the white or gray regions).
+        result.intersect = false;
+        return result;
+    }
+    // The disk center is inside K" (in the red, orange, blue, or green regions).
+    Real dotDirCmV = Dot(sector.direction, CmV);
+    if (-dotDirCmV >= lengthCmV * sector.sinAngle)
+    {
+        // The disk center is inside K" and inside K' (in the blue or green regions).
+        if (lengthCmV <= disk.radius)
+        {
+            // The disk center is in the blue region, in which case the disk
+            // contains the sector's vertex.
+            result.intersect = true;
+        }
+        else
+        {
+            // The disk center is in the green region.
+            result.intersect = false;
+        }
+        return result;
+    }
+
+    // To reach here, we know that the disk overlaps the sector's disk and
+    // the sector's cone.  The disk center is in the orange region or in
+    // in the red region (not including the segments that separate the red
+    // and blue regions).
+
+    // Test whether the ray of the right boundary of the sector overlaps
+    // the disk.  The ray direction U0 is a clockwise rotation of the
+    // cone axis by the cone angle.
+    Vector2<Real> U0
+    {
+        +sector.cosAngle * sector.direction[0] + sector.sinAngle * sector.direction[1],
+        -sector.sinAngle * sector.direction[0] + sector.cosAngle * sector.direction[1]
+    };
+    Real dp0 = Dot(U0, CmV);
+    Real discr0 = disk.radius * disk.radius + dp0 * dp0 - sqrLengthCmV;
+    if (discr0 >= (Real)0)
+    {
+        // The ray intersects the disk.  Now test whether the sector boundary
+        // segment contained by the ray overlaps the disk.  The quadratic
+        // root tmin generates the ray-disk point of intersection closest to
+        // the sector vertex.
+        Real tmin = dp0 - std::sqrt(discr0);
+        if (sector.radius >= tmin)
+        {
+            // The segment overlaps the disk.
+            result.intersect = true;
+            return result;
+        }
+        else
+        {
+            // The segment does not overlap the disk.  We know the disks
+            // overlap, so if the disk center is outside the sector cone
+            // or on the right-boundary ray, the overlap occurs outside the
+            // cone, which implies the disk and sector do not intersect.
+            if (dotDirCmV <= lengthCmV * sector.cosAngle)
+            {
+                // The disk center is not inside the sector cone.
+                result.intersect = false;
+                return result;
+            }
+        }
+    }
+
+    // Test whether the ray of the left boundary of the sector overlaps
+    // the disk.  The ray direction U1 is a counterclockwise rotation of the
+    // cone axis by the cone angle.
+    Vector2<Real> U1
+    {
+        +sector.cosAngle * sector.direction[0] - sector.sinAngle * sector.direction[1],
+        +sector.sinAngle * sector.direction[0] + sector.cosAngle * sector.direction[1]
+    };
+    Real dp1 = Dot(U1, CmV);
+    Real discr1 = disk.radius * disk.radius + dp1 * dp1 - sqrLengthCmV;
+    if (discr1 >= (Real)0)
+    {
+        // The ray intersects the disk.  Now test whether the sector boundary
+        // segment contained by the ray overlaps the disk.  The quadratic
+        // root tmin generates the ray-disk point of intersection closest to
+        // the sector vertex.
+        Real tmin = dp1 - std::sqrt(discr1);
+        if (sector.radius >= tmin)
+        {
+            result.intersect = true;
+            return result;
+        }
+        else
+        {
+            // The segment does not overlap the disk.  We know the disks
+            // overlap, so if the disk center is outside the sector cone
+            // or on the right-boundary ray, the overlap occurs outside the
+            // cone, which implies the disk and sector do not intersect.
+            if (dotDirCmV <= lengthCmV * sector.cosAngle)
+            {
+                // The disk center is not inside the sector cone.
+                result.intersect = false;
+                return result;
+            }
+        }
+    }
+
+    // To reach here, a strict subset of the sector's arc boundary must
+    // intersect the disk.
+    result.intersect = true;
+    return result;
+}
+
+}
diff --git a/Engine/Plugins/Experimental/GeometryProcessing/Source/GeometryAlgorithms/Private/GTEngine/Mathematics/GteIntrEllipse2Ellipse2.h b/Engine/Plugins/Experimental/GeometryProcessing/Source/GeometryAlgorithms/Private/GTEngine/Mathematics/GteIntrEllipse2Ellipse2.h
new file mode 100644
--- /dev/null
+++ b/Engine/Plugins/Experimental/GeometryProcessing/Source/GeometryAlgorithms/Private/GTEngine/Mathematics/GteIntrEllipse2Ellipse2.h
@@ -0,0 +1,1369 @@
+// David Eberly, Geometric Tools, Redmond WA 98052
+// Copyright (c) 1998-2019
+// Distributed under the Boost Software License, Version 1.0.
+// http://www.boost.org/LICENSE_1_0.txt
+// http://www.geometrictools.com/License/Boost/LICENSE_1_0.txt
+// File Version: 3.0.4 (2018/10/05)
+
+#pragma once
+
+#include <GTEngine/LowLevel/GteLogger.h>
+#include <GTEngine/Mathematics/GteMatrix2x2.h>
+#include <GTEngine/Mathematics/GteHyperellipsoid.h>
+#include <GTEngine/Mathematics/GteFIQuery.h>
+#include <GTEngine/Mathematics/GteTIQuery.h>
+#include <GTEngine/Mathematics/GteRootsBisection.h>
+#include <GTEngine/Mathematics/GteRootsPolynomial.h>
+#include <GTEngine/Mathematics/GteSymmetricEigensolver2x2.h>
+
+// The test-intersection and find-intersection queries implemented here are
+// discussed in the document
+//   http://www.geometrictools.com/Documentation/IntersectionOfEllipses.pdf
+// The Real type should support exact rational arithmetic in order for the
+// polynomial root construction to be robust.  The classification of the
+// intersections depends on various sign tests of computed values.  If these
+// values are computed with floating-point arithmetic, the sign tests can
+// lead to misclassification.
+//
+// The area-of-intersection query is discussed in the document
+//   http://www.geometrictools.com/Documentation/AreaIntersectingEllipses.pdf
+
+namespace gte
+{
+
+template <typename Real>
+class TIQuery<Real, Ellipse2<Real>, Ellipse2<Real>>
+{
+public:
+    // The query tests the relationship between the ellipses as solid
+    // objects.
+    enum
+    {
+        ELLIPSES_SEPARATED,
+        ELLIPSES_OVERLAP,
+        ELLIPSE0_OUTSIDE_ELLIPSE1_BUT_TANGENT,
+        ELLIPSE0_STRICTLY_CONTAINS_ELLIPSE1,
+        ELLIPSE0_CONTAINS_ELLIPSE1_BUT_TANGENT,
+        ELLIPSE1_STRICTLY_CONTAINS_ELLIPSE0,
+        ELLIPSE1_CONTAINS_ELLIPSE0_BUT_TANGENT,
+        ELLIPSES_EQUAL
+    };
+
+    // The ellipse axes are already normalized, which most likely introduced
+    // rounding errors.
+    int operator()(Ellipse2<Real> const& ellipse0,
+        Ellipse2<Real> const& ellipse1);
+
+private:
+    void GetRoots(Real d0, Real c0, int& numRoots, Real* roots);
+    void GetRoots(Real d0, Real d1, Real c0, Real c1, int& numRoots,
+        Real* roots);
+
+    int Classify(Real minSqrDistance, Real maxSqrDistance, Real d0c0pd1c1);
+};
+
+template <typename Real>
+class FIQuery<Real, Ellipse2<Real>, Ellipse2<Real>>
+{
+public:
+    // The queries find the intersections (if any) of the ellipses treated as
+    // hollow objects.  The implementations use the same concepts.
+    FIQuery();
+
+    struct Result
+    {
+        // This value is true when the ellipses intersect in at least one
+        // point.
+        bool intersect;
+
+        // If the ellipses are not the same, numPoints is 0 through 4 and
+        // that number of elements of 'points' are valid.  If the ellipses
+        // are the same, numPoints is set to maxInt and 'points' is invalid.
+        int numPoints;
+        std::array<Vector2<Real>, 4> points;
+        std::array<bool, 4> isTransverse;
+    };
+
+    // The ellipse axes are already normalized, which most likely introduced
+    // rounding errors.
+    Result operator()(Ellipse2<Real> const& ellipse0,
+        Ellipse2<Real> const& ellipse1);
+
+    // The axis directions do not have to be unit length.  The quadratic
+    // equations are constructed according to the details of the PDF document
+    // about the intersection of ellipses.
+    Result operator()(Vector2<Real> const& center0,
+        Vector2<Real> const axis0[2], Vector2<Real> const& sqrExtent0,
+        Vector2<Real> const& center1, Vector2<Real> const axis1[2],
+        Vector2<Real> const& sqrExtent1);
+
+
+    // Compute the area of intersection of ellipses.
+    struct AreaResult
+    {
+        // The configuration of the two ellipses.
+        enum
+        {
+            ELLIPSES_ARE_EQUAL,
+            ELLIPSES_ARE_SEPARATED,
+            E0_CONTAINS_E1,
+            E1_CONTAINS_E0,
+            ONE_CHORD_REGION,
+            FOUR_CHORD_REGION
+        };
+
+        // One of the enumerates, determined in the call to AreaDispatch.
+        int configuration;
+
+        // Information about the ellipse-ellipse intersection points.
+        Result findResult;
+
+        // The area of intersection of the ellipses.
+        Real area;
+    };
+
+    // The ellipse axes are already normalized, which most likely introduced
+    // rounding errors.
+    AreaResult AreaOfIntersection(Ellipse2<Real> const& ellipse0,
+        Ellipse2<Real> const& ellipse1);
+
+    // The axis directions do not have to be unit length.  The positive
+    // definite matrices are constructed according to the details of the PDF
+    // documentabout the intersection of ellipses.
+    AreaResult AreaOfIntersection(Vector2<Real> const& center0,
+        Vector2<Real> const axis0[2], Vector2<Real> const& sqrExtent0,
+        Vector2<Real> const& center1, Vector2<Real> const axis1[2],
+        Vector2<Real> const& sqrExtent1);
+
+private:
+    // Compute the coefficients of the quadratic equation but with the
+    // y^2-coefficient of 1.
+    void ToCoefficients(Vector2<Real> const& center,
+        Vector2<Real> const axis[2], Vector2<Real> const& SqrExtent,
+        std::array<Real, 5>& coeff);
+
+    void DoRootFinding(Result& result);
+    void D4NotZeroEBarNotZero(Result& result);
+    void D4NotZeroEBarZero(Real const& xbar, Result& result);
+    void D4ZeroD2NotZeroE2NotZero(Result& result);
+    void D4ZeroD2NotZeroE2Zero(Result& result);
+    void D4ZeroD2Zero(Result& result);
+
+    // Helper function for D4ZeroD2Zero.
+    void SpecialIntersection(Real const& x, bool transverse, Result& result);
+
+
+    // Helper functions for AreaOfIntersection.
+    struct EllipseInfo
+    {
+        Vector2<Real> center;
+        std::array<Vector2<Real>, 2> axis;
+        Vector2<Real> extent, sqrExtent;
+        Matrix2x2<Real> M;
+        Real AB;        // extent[0] * extent[1]
+        Real halfAB;    // extent[0] * extent[1] / 2
+        Real BpA;       // extent[1] + extent[0]
+        Real BmA;       // extent[1] - extent[0]
+    };
+
+    // The axis, extent, and sqrExtent members of E must be set before this
+    // function is called.
+    void FinishEllipseInfo(EllipseInfo& E);
+
+    void AreaDispatch(EllipseInfo const& E0, EllipseInfo const& E1,
+        AreaResult& ar);
+
+    void AreaCS(EllipseInfo const& E0, EllipseInfo const& E1,
+        AreaResult& ar);
+
+    void Area2(EllipseInfo const& E0, EllipseInfo const& E1, int i0, int i1,
+        AreaResult& ar);
+
+    void Area4(EllipseInfo const& E0, EllipseInfo const& E1,
+        AreaResult& ar);
+
+    Real ComputeAreaChordRegion(EllipseInfo const& E,
+        Vector2<Real> const& P0mC, Vector2<Real> const& P1mC);
+
+    Real ComputeIntegral(EllipseInfo const& E, Real const& theta);
+
+
+    // Constants that are set up once (optimization for rational arithmetic).
+    Real mZero, mOne, mTwo, mPi, mTwoPi;
+
+    // Various polynomial coefficients.  The short names are meant to match
+    // the variable names used in the PDF documentation.
+    std::array<Real, 5> mA, mB, mD, mF;
+    std::array<Real, 3> mC, mE;
+    Real mA2Div2, mA4Div2;
+};
+
+// Template aliases for convenience.
+template <typename Real>
+using TIEllipses2 = TIQuery<Real, Ellipse2<Real>, Ellipse2<Real>>;
+
+template <typename Real>
+using FIEllipses2 = FIQuery<Real, Ellipse2<Real>, Ellipse2<Real>>;
+
+
+// TIQuery for intersections
+
+template <typename Real>
+int TIQuery<Real, Ellipse2<Real>, Ellipse2<Real>>::operator()(
+    Ellipse2<Real> const& ellipse0, Ellipse2<Real> const& ellipse1)
+{
+    Real const zero = 0, one = 1;
+
+    // Get the parameters of ellipe0.
+    Vector2<Real> K0 = ellipse0.center;
+    Matrix2x2<Real> R0;
+    R0.SetCol(0, ellipse0.axis[0]);
+    R0.SetCol(1, ellipse0.axis[1]);
+    Matrix2x2<Real> D0{
+        one / (ellipse0.extent[0] * ellipse0.extent[0]), zero,
+        zero, one / (ellipse0.extent[1] * ellipse0.extent[1]) };
+
+    // Get the parameters of ellipse1.
+    Vector2<Real> K1 = ellipse1.center;
+    Matrix2x2<Real> R1;
+    R1.SetCol(0, ellipse1.axis[0]);
+    R1.SetCol(1, ellipse1.axis[1]);
+    Matrix2x2<Real> D1{
+        one / (ellipse1.extent[0] * ellipse1.extent[0]), zero,
+        zero, one / (ellipse1.extent[1] * ellipse1.extent[1]) };
+
+    // Compute K2.
+    Matrix2x2<Real> D0NegHalf{
+        ellipse0.extent[0], zero,
+        zero, ellipse0.extent[1] };
+    Matrix2x2<Real> D0Half{
+        one / ellipse0.extent[0], zero,
+        zero, one / ellipse0.extent[1] };
+    Vector2<Real> K2 = D0Half*((K1 - K0)*R0);
+
+    // Compute M2.
+    Matrix2x2<Real> R1TR0D0NegHalf = MultiplyATB(R1, R0 * D0NegHalf);
+    Matrix2x2<Real> M2 = MultiplyATB(R1TR0D0NegHalf, D1) * R1TR0D0NegHalf;
+
+    // Factor M2 = R*D*R^T.
+    SymmetricEigensolver2x2<Real> es;
+    std::array<Real, 2> D;
+    std::array<std::array<Real, 2>, 2> evec;
+    es(M2(0, 0), M2(0, 1), M2(1, 1), +1, D, evec);
+    Matrix2x2<Real> R;
+    R.SetCol(0, evec[0]);
+    R.SetCol(1, evec[1]);
+
+    // Compute K = R^T*K2.
+    Vector2<Real> K = K2 * R;
+
+    // Transformed ellipse0 is Z^T*Z = 1 and transformed ellipse1 is
+    // (Z-K)^T*D*(Z-K) = 0.
+
+    // The minimum and maximum squared distances from the origin of points on
+    // transformed ellipse1 are used to determine whether the ellipses
+    // intersect, are separated, or one contains the other.
+    Real minSqrDistance = std::numeric_limits<Real>::max();
+    Real maxSqrDistance = zero;
+
+    if (K == Vector2<Real>::Zero())
+    {
+        // The special case of common centers must be handled separately.  It
+        // is not possible for the ellipses to be separated.
+        for (int i = 0; i < 2; ++i)
+        {
+            Real invD = one / D[i];
+            if (invD < minSqrDistance)
+            {
+                minSqrDistance = invD;
+            }
+            if (invD > maxSqrDistance)
+            {
+                maxSqrDistance = invD;
+            }
+        }
+        return Classify(minSqrDistance, maxSqrDistance, zero);
+    }
+
+    // The closest point P0 and farthest point P1 are solutions to
+    // s0*D*(P0 - K) = P0 and s1*D1*(P1 - K) = P1 for some scalars s0 and s1
+    // that are roots to the function
+    //   f(s) = d0*k0^2/(d0*s-1)^2 + d1*k1^2/(d1*s-1)^2 - 1
+    // where D = diagonal(d0,d1) and K = (k0,k1).
+    Real d0 = D[0], d1 = D[1];
+    Real c0 = K[0] * K[0], c1 = K[1] * K[1];
+
+    // Sort the values so that d0 >= d1.  This allows us to bound the roots of
+    // f(s), of which there are at most 4.
+    std::vector<std::pair<Real, Real>> param(2);
+    if (d0 >= d1)
+    {
+        param[0] = std::make_pair(d0, c0);
+        param[1] = std::make_pair(d1, c1);
+    }
+    else
+    {
+        param[0] = std::make_pair(d1, c1);
+        param[1] = std::make_pair(d0, c0);
+    }
+
+    std::vector<std::pair<Real, Real>> valid;
+    valid.reserve(2);
+    if (param[0].first > param[1].first)
+    {
+        // d0 > d1
+        for (int i = 0; i < 2; ++i)
+        {
+            if (param[i].second > zero)
+            {
+                valid.push_back(param[i]);
+            }
+        }
+    }
+    else
+    {
+        // d0 = d1
+        param[0].second += param[1].second;
+        if (param[0].second > zero)
+        {
+            valid.push_back(param[0]);
+        }
+    }
+
+    size_t numValid = valid.size();
+    int numRoots = 0;
+    Real roots[4];
+    if (numValid == 2)
+    {
+        GetRoots(valid[0].first, valid[1].first, valid[0].second,
+            valid[1].second, numRoots, roots);
+    }
+    else if (numValid == 1)
+    {
+        GetRoots(valid[0].first, valid[0].second, numRoots, roots);
+    }
+    // else: numValid cannot be zero because we already handled case K = 0
+
+    for (int i = 0; i < numRoots; ++i)
+    {
+        Real s = roots[i];
+        Real p0 = d0 * K[0] * s / (d0 * s - (Real)1);
+        Real p1 = d1 * K[1] * s / (d1 * s - (Real)1);
+        Real sqrDistance = p0 * p0 + p1 * p1;
+        if (sqrDistance < minSqrDistance)
+        {
+            minSqrDistance = sqrDistance;
+        }
+        if (sqrDistance > maxSqrDistance)
+        {
+            maxSqrDistance = sqrDistance;
+        }
+    }
+
+    return Classify(minSqrDistance, maxSqrDistance, d0 * c0 + d1 * c1);
+}
+
+template <typename Real>
+void TIQuery<Real, Ellipse2<Real>, Ellipse2<Real>>::GetRoots(Real d0, Real c0,
+    int& numRoots, Real* roots)
+{
+    // f(s) = d0*c0/(d0*s-1)^2 - 1
+    Real const one = (Real)1;
+    Real temp = std::sqrt(d0*c0);
+    Real inv = one / d0;
+    numRoots = 2;
+    roots[0] = (one - temp) * inv;
+    roots[1] = (one + temp) * inv;
+}
+
+template <typename Real>
+void TIQuery<Real, Ellipse2<Real>, Ellipse2<Real>>::GetRoots(Real d0, Real d1,
+    Real c0, Real c1, int& numRoots, Real* roots)
+{
+    // f(s) = d0*c0/(d0*s-1)^2 + d1*c1/(d1*s-1)^2 - 1
+    // with d0 > d1
+
+    Real const zero = 0, one = (Real)1;
+    Real d0c0 = d0 * c0;
+    Real d1c1 = d1 * c1;
+    Real sum = d0c0 + d1c1;
+    Real sqrtsum = std::sqrt(sum);
+
+    std::function<Real(Real)> F = [&one, d0, d1, d0c0, d1c1](Real s)
+    {
+        Real invN0 = one / (d0*s - one);
+        Real invN1 = one / (d1*s - one);
+        Real term0 = d0c0 * invN0 * invN0;
+        Real term1 = d1c1 * invN1 * invN1;
+        Real f = term0 + term1 - one;
+        return f;
+    };
+
+    std::function<Real(Real)> DF = [&one, d0, d1, d0c0, d1c1](Real s)
+    {
+        Real const two = 2;
+        Real invN0 = one / (d0*s - one);
+        Real invN1 = one / (d1*s - one);
+        Real term0 = d0 * d0c0 * invN0 * invN0 * invN0;
+        Real term1 = d1 * d1c1 * invN1 * invN1 * invN1;
+        Real df = -two * (term0 + term1);
+        return df;
+    };
+
+    unsigned int const maxIterations = (unsigned int)(
+        3 + std::numeric_limits<Real>::digits -
+        std::numeric_limits<Real>::min_exponent);
+    unsigned int iterations;
+    numRoots = 0;
+
+    Real invD0 = one / d0;
+    Real invD1 = one / d1;
+    Real smin, smax, s, fval;
+
+    // Compute root in (-infinity,1/d0).  Obtain a lower bound for the root
+    // better than -std::numeric_limits<Real>::max().
+    smax = invD0;
+    fval = sum - one;
+    if (fval > zero)
+    {
+        smin = (one - sqrtsum) * invD1;  // < 0
+        fval = F(smin);
+        LogAssert(fval <= zero, "Unexpected condition.");
+    }
+    else
+    {
+        smin = zero;
+    }
+    iterations = RootsBisection<Real>::Find(F, smin, smax, -one, one,
+        maxIterations, s);
+    fval = F(s);
+    LogAssert(iterations > 0, "Unexpected condition.");
+    roots[numRoots++] = s;
+
+    // Compute roots (if any) in (1/d0,1/d1).  It is the case that
+    //   F(1/d0) = +infinity, F'(1/d0) = -infinity
+    //   F(1/d1) = +infinity, F'(1/d1) = +infinity
+    //   F"(s) > 0 for all s in the domain of F
+    // Compute the unique root r of F'(s) on (1/d0,1/d1).  The bisector needs
+    // only the signs at the endpoints, so we pass -1 and +1 instead of the
+    // infinite values.  If F(r) < 0, F(s) has two roots in the interval.
+    // If F(r) = 0, F(s) has only one root in the interval.
+    Real const oneThird = (Real)(1.0 / 3.0);
+    Real rho = std::pow(d0 * d0c0 / (d1 * d1c1), oneThird);
+    Real smid = (one + rho) / (d0 + rho * d1);
+    Real fmid = F(smid);
+    if (fmid <= zero)
+    {
+        // Pass in signs rather than infinities, because the bisector cares
+        // only about the signs.
+        iterations = RootsBisection<Real>::Find(F, invD0, smid, one, -one,
+            maxIterations, s);
+        fval = F(s);
+        LogAssert(iterations > 0, "Unexpected condition.");
+        roots[numRoots++] = s;
+        iterations = RootsBisection<Real>::Find(F, smid, invD1, -one, one,
+            maxIterations, s);
+        fval = F(s);
+        LogAssert(iterations > 0, "Unexpected condition.");
+        roots[numRoots++] = s;
+    }
+    else if (fmid == zero)
+    {
+        roots[numRoots++] = smid;
+    }
+
+    // Compute root in (1/d1,+infinity).  Obtain an upper bound for the root
+    // better than std::numeric_limits<Real>::max().
+    smin = invD1;
+    smax = (one + sqrtsum) * invD1;  // > 1/d1
+    fval = F(smax);
+    LogAssert(fval <= zero, "Unexpected condition.");
+    iterations = RootsBisection<Real>::Find(F, smin, smax, one, -one,
+        maxIterations, s);
+    fval = F(s);
+    LogAssert(iterations > 0, "Unexpected condition.");
+    roots[numRoots++] = s;
+}
+
+template <typename Real>
+int TIQuery<Real, Ellipse2<Real>, Ellipse2<Real>>::Classify(
+    Real minSqrDistance, Real maxSqrDistance, Real d0c0pd1c1)
+{
+    Real const one = (Real)1;
+
+    if (maxSqrDistance < one)
+    {
+        return ELLIPSE0_STRICTLY_CONTAINS_ELLIPSE1;
+    }
+    else if (maxSqrDistance > one)
+    {
+        if (minSqrDistance < one)
+        {
+            return ELLIPSES_OVERLAP;
+        }
+        else if (minSqrDistance > one)
+        {
+            if (d0c0pd1c1 > one)
+            {
+                return ELLIPSES_SEPARATED;
+            }
+            else
+            {
+                return ELLIPSE1_STRICTLY_CONTAINS_ELLIPSE0;
+            }
+        }
+        else  // minSqrDistance = 1
+        {
+            if (d0c0pd1c1 > one)
+            {
+                return ELLIPSE0_OUTSIDE_ELLIPSE1_BUT_TANGENT;
+            }
+            else
+            {
+                return ELLIPSE1_CONTAINS_ELLIPSE0_BUT_TANGENT;
+            }
+        }
+    }
+    else  // maxSqrDistance = 1
+    {
+        if (minSqrDistance < one)
+        {
+            return ELLIPSE0_CONTAINS_ELLIPSE1_BUT_TANGENT;
+        }
+        else // minSqrDistance = 1
+        {
+            return ELLIPSES_EQUAL;
+        }
+    }
+}
+
+
+
+// FIQuery constructor for both intersection and area queries.
+
+template <typename Real>
+FIQuery<Real, Ellipse2<Real>, Ellipse2<Real>>::FIQuery()
+    :
+    mZero((Real)0),
+    mOne((Real)1),
+    mTwo((Real)2),
+    mPi((Real)GTE_C_PI),
+    mTwoPi((Real)GTE_C_TWO_PI)
+{
+}
+
+
+
+// FIQuery functions for ellipse-ellipse intersection points.
+
+template <typename Real>
+typename FIQuery<Real, Ellipse2<Real>, Ellipse2<Real>>::Result
+FIQuery<Real, Ellipse2<Real>, Ellipse2<Real>>::operator()(
+    Ellipse2<Real> const& ellipse0, Ellipse2<Real> const& ellipse1)
+{
+    Vector2<Real> rCenter, rAxis[2], rSqrExtent;
+
+    rCenter = { ellipse0.center[0], ellipse0.center[1] };
+    rAxis[0] = { ellipse0.axis[0][0], ellipse0.axis[0][1] };
+    rAxis[1] = { ellipse0.axis[1][0], ellipse0.axis[1][1] };
+    rSqrExtent =
+    {
+        ellipse0.extent[0] * ellipse0.extent[0],
+        ellipse0.extent[1] * ellipse0.extent[1]
+    };
+    ToCoefficients(rCenter, rAxis, rSqrExtent, mA);
+
+    rCenter = { ellipse1.center[0], ellipse1.center[1] };
+    rAxis[0] = { ellipse1.axis[0][0], ellipse1.axis[0][1] };
+    rAxis[1] = { ellipse1.axis[1][0], ellipse1.axis[1][1] };
+    rSqrExtent =
+    {
+        ellipse1.extent[0] * ellipse1.extent[0],
+        ellipse1.extent[1] * ellipse1.extent[1]
+    };
+    ToCoefficients(rCenter, rAxis, rSqrExtent, mB);
+
+    Result result;
+    DoRootFinding(result);
+    return result;
+}
+
+template <typename Real>
+typename FIQuery<Real, Ellipse2<Real>, Ellipse2<Real>>::Result
+FIQuery<Real, Ellipse2<Real>, Ellipse2<Real>>::operator()(
+    Vector2<Real> const& center0, Vector2<Real> const axis0[2],
+    Vector2<Real> const& sqrExtent0, Vector2<Real> const& center1,
+    Vector2<Real> const axis1[2], Vector2<Real> const& sqrExtent1)
+{
+    Vector2<Real> rCenter, rAxis[2], rSqrExtent;
+
+    rCenter = { center0[0], center0[1] };
+    rAxis[0] = { axis0[0][0], axis0[0][1] };
+    rAxis[1] = { axis0[1][0], axis0[1][1] };
+    rSqrExtent = { sqrExtent0[0], sqrExtent0[1] };
+    ToCoefficients(rCenter, rAxis, rSqrExtent, mA);
+
+    rCenter = { center1[0], center1[1] };
+    rAxis[0] = { axis1[0][0], axis1[0][1] };
+    rAxis[1] = { axis1[1][0], axis1[1][1] };
+    rSqrExtent = { sqrExtent1[0], sqrExtent1[1] };
+    ToCoefficients(rCenter, rAxis, rSqrExtent, mB);
+
+    Result result;
+    DoRootFinding(result);
+    return result;
+}
+
+template <typename Real>
+void FIQuery<Real, Ellipse2<Real>, Ellipse2<Real>>::ToCoefficients(
+    Vector2<Real> const& center, Vector2<Real> const axis[2],
+    Vector2<Real> const& sqrExtent, std::array<Real, 5>& coeff)
+{
+    Real denom0 = Dot(axis[0], axis[0]) * sqrExtent[0];
+    Real denom1 = Dot(axis[1], axis[1]) * sqrExtent[1];
+    Matrix2x2<Real> outer0 = OuterProduct(axis[0], axis[0]);
+    Matrix2x2<Real> outer1 = OuterProduct(axis[1], axis[1]);
+    Matrix2x2<Real> A = outer0 / denom0 + outer1 / denom1;
+    Vector2<Real> product = A * center;
+    Vector2<Real> B = -mTwo * product;
+    Real const& denom = A(1, 1);
+    coeff[0] = (Dot(center, product) - mOne) / denom;
+    coeff[1] = B[0] / denom;
+    coeff[2] = B[1] / denom;
+    coeff[3] = A(0, 0) / denom;
+    coeff[4] = mTwo * A(0, 1) / denom;
+    // coeff[5] = A(1, 1) / denom = 1;
+}
+
+template <typename Real>
+void FIQuery<Real, Ellipse2<Real>, Ellipse2<Real>>::DoRootFinding(
+    Result& result)
+{
+    bool allZero = true;
+    for (int i = 0; i < 5; ++i)
+    {
+        mD[i] = mA[i] - mB[i];
+        if (mD[i] != mZero)
+        {
+            allZero = false;
+        }
+    }
+    if (allZero)
+    {
+        result.intersect = false;
+        result.numPoints = std::numeric_limits<int>::max();
+        return;
+    }
+
+    result.numPoints = 0;
+
+    mA2Div2 = mA[2] / mTwo;
+    mA4Div2 = mA[4] / mTwo;
+    mC[0] = mA[0] - mA2Div2 * mA2Div2;
+    mC[1] = mA[1] - mA2Div2 * mA[4];
+    mC[2] = mA[3] - mA4Div2 * mA4Div2;  // c[2] > 0
+    mE[0] = mD[0] - mA2Div2 * mD[2];
+    mE[1] = mD[1] - mA2Div2 * mD[4] - mA4Div2 * mD[2];
+    mE[2] = mD[3] - mA4Div2 * mD[4];
+
+    if (mD[4] != mZero)
+    {
+        Real xbar = -mD[2] / mD[4];
+        Real ebar = mE[0] + xbar * (mE[1] + xbar * mE[2]);
+        if (ebar != mZero)
+        {
+            D4NotZeroEBarNotZero(result);
+        }
+        else
+        {
+            D4NotZeroEBarZero(xbar, result);
+        }
+    }
+    else if (mD[2] != mZero)  // d[4] = 0
+    {
+        if (mE[2] != mZero)
+        {
+            D4ZeroD2NotZeroE2NotZero(result);
+        }
+        else
+        {
+            D4ZeroD2NotZeroE2Zero(result);
+        }
+    }
+    else  // d[2] = d[4] = 0
+    {
+        D4ZeroD2Zero(result);
+    }
+
+    result.intersect = (result.numPoints > 0);
+}
+
+template <typename Real>
+void FIQuery<Real, Ellipse2<Real>, Ellipse2<Real>>::D4NotZeroEBarNotZero(
+    Result& result)
+{
+    // The graph of w = -e(x)/d(x) is a hyperbola.
+    Real d2d2 = mD[2] * mD[2], d2d4 = mD[2] * mD[4], d4d4 = mD[4] * mD[4];
+    Real e0e0 = mE[0] * mE[0], e0e1 = mE[0] * mE[1], e0e2 = mE[0] * mE[2];
+    Real e1e1 = mE[1] * mE[1], e1e2 = mE[1] * mE[2], e2e2 = mE[2] * mE[2];
+    std::array<Real, 5> f =
+    {
+        mC[0] * d2d2 + e0e0,
+        mC[1] * d2d2 + mTwo * (mC[0] * d2d4 + e0e1),
+        mC[2] * d2d2 + mC[0] * d4d4 + e1e1 + mTwo * (mC[1] * d2d4 + e0e2),
+        mC[1] * d4d4 + mTwo * (mC[2] * d2d4 + e1e2),
+        mC[2] * d4d4 + e2e2  // > 0
+    };
+
+    std::map<Real, int> rmMap;
+    RootsPolynomial<Real>::template SolveQuartic<Real>(f[0], f[1], f[2],
+        f[3], f[4], rmMap);
+
+    // xbar cannot be a root of f(x), so d(x) != 0 and we can solve
+    // directly for w = -e(x)/d(x).
+    for (auto const& rm : rmMap)
+    {
+        Real const& x = rm.first;
+        Real w = -(mE[0] + x * (mE[1] + x * mE[2])) / (mD[2] + mD[4] * x);
+        Real y = w - (mA2Div2 + x * mA4Div2);
+        result.points[result.numPoints] = { x, y };
+        result.isTransverse[result.numPoints++] = (rm.second == 1);
+    }
+}
+
+template <typename Real>
+void FIQuery<Real, Ellipse2<Real>, Ellipse2<Real>>::D4NotZeroEBarZero(
+    Real const& xbar, Result& result)
+{
+    // Factor e(x) = (d2 + d4*x)*(h0 + h1*x).  The w-equation has
+    // two solution components, x = xbar and w = -(h0 + h1*x).
+    Real translate, w, y;
+
+    // Compute intersection of x = xbar with ellipse.
+    Real ncbar = -(mC[0] + xbar * (mC[1] + xbar * mC[2]));
+    if (ncbar >= mZero)
+    {
+        translate = mA2Div2 + xbar * mA4Div2;
+        w = std::sqrt(ncbar);
+        y = w - translate;
+        result.points[result.numPoints] = { xbar, y };
+        if (w > mZero)
+        {
+            result.isTransverse[result.numPoints++] = true;
+            w = -w;
+            y = w - translate;
+            result.points[result.numPoints] = { xbar, y };
+            result.isTransverse[result.numPoints++] = true;
+        }
+        else
+        {
+            result.isTransverse[result.numPoints++] = false;
+        }
+    }
+
+    // Compute intersections of w = -(h0 + h1*x) with ellipse.
+    std::array<Real, 2> h;
+    h[1] = mE[2] / mD[4];
+    h[0] = (mE[1] - mD[2] * h[1]) / mD[4];
+    std::array<Real, 3> f =
+    {
+        mC[0] + h[0] * h[0],
+        mC[1] + mTwo * h[0] * h[1],
+        mC[2] + h[1] * h[1]  // > 0
+    };
+
+    std::map<Real, int> rmMap;
+    RootsPolynomial<Real>::template SolveQuadratic<Real>(f[0], f[1], f[2],
+        rmMap);
+    for (auto const& rm : rmMap)
+    {
+        Real const& x = rm.first;
+        translate = mA2Div2 + x * mA4Div2;
+        w = -(h[0] + x * h[1]);
+        y = w - translate;
+        result.points[result.numPoints] = { x, y };
+        result.isTransverse[result.numPoints++] = (rm.second == 1);
+    }
+}
+
+template <typename Real>
+void FIQuery<Real, Ellipse2<Real>, Ellipse2<Real>>::D4ZeroD2NotZeroE2NotZero(
+    Result& result)
+{
+    Real d2d2 = mD[2] * mD[2];
+    std::array<Real, 5> f =
+    {
+        mC[0] * d2d2 + mE[0] * mE[0],
+        mC[1] * d2d2 + mTwo * mE[0] * mE[1],
+        mC[2] * d2d2 + mE[1] * mE[1] + mTwo * mE[0] * mE[2],
+        mTwo * mE[1] * mE[2],
+        mE[2] * mE[2]  // > 0
+    };
+
+    std::map<Real, int> rmMap;
+    RootsPolynomial<Real>::template SolveQuartic<Real>(f[0], f[1], f[2], f[3],
+        f[4], rmMap);
+    for (auto const& rm : rmMap)
+    {
+        Real const& x = rm.first;
+        Real translate = mA2Div2 + x * mA4Div2;
+        Real w = -(mE[0] + x * (mE[1] + x * mE[2])) / mD[2];
+        Real y = w - translate;
+        result.points[result.numPoints] = { x, y };
+        result.isTransverse[result.numPoints++] = (rm.second == 1);
+    }
+}
+
+template <typename Real>
+void FIQuery<Real, Ellipse2<Real>, Ellipse2<Real>>::D4ZeroD2NotZeroE2Zero(
+    Result& result)
+{
+    Real d2d2 = mD[2] * mD[2];
+    std::array<Real, 3> f =
+    {
+        mC[0] * d2d2 + mE[0] * mE[0],
+        mC[1] * d2d2 + mTwo * mE[0] * mE[1],
+        mC[2] * d2d2 + mE[1] * mE[1]
+    };
+
+    std::map<Real, int> rmMap;
+    RootsPolynomial<Real>::template SolveQuadratic<Real>(f[0], f[1], f[2],
+        rmMap);
+    for (auto const& rm : rmMap)
+    {
+        Real const& x = rm.first;
+        Real translate = mA2Div2 + x * mA4Div2;
+        Real w = -(mE[0] + x * mE[1]) / mD[2];
+        Real y = w - translate;
+        result.points[result.numPoints] = { x, y };
+        result.isTransverse[result.numPoints++] = (rm.second == 1);
+    }
+}
+
+template <typename Real>
+void FIQuery<Real, Ellipse2<Real>, Ellipse2<Real>>::D4ZeroD2Zero(
+    Result& result)
+{
+    // e(x) cannot be identically zero, because if it were, then all
+    // d[i] = 0.  But we tested that case previously and exited.
+
+    if (mE[2] != mZero)
+    {
+        // Make e(x) monic, f(x) = e(x)/e2 = x^2 + (e1/e2)*x + (e0/e2)
+        // = x^2 + f1*x + f0.
+        std::array<Real, 2> f = { mE[0] / mE[2], mE[1] / mE[2] };
+
+        Real mid = -f[1] / mTwo;
+        Real discr = mid * mid - f[0];
+        if (discr > mZero)
+        {
+            // The theoretical roots of e(x) are x = -f1/2 + s*sqrt(discr)
+            // where s in {-1,+1}.  For each root, determine exactly the
+            // sign of c(x).  We need c(x) <= 0 in order to solve for
+            // w^2 = -c(x).  At the root,
+            //   c(x) = c0 + c1*x + c2*x^2 = c0 + c1*x - c2*(f0 + f1*x)
+            //        = (c0 - c2*f0) + (c1 - c2*f1)*x
+            //        = g0 + g1*x
+            // We need g0 + g1*x <= 0.
+            Real sqrtDiscr = std::sqrt(discr);
+            std::array<Real, 2> g =
+            {
+                mC[0] - mC[2] * f[0],
+                mC[1] - mC[2] * f[1]
+            };
+
+            if (g[1] > mZero)
+            {
+                // We need s*sqrt(discr) <= -g[0]/g[1] + f1/2.
+                Real r = -g[0] / g[1] - mid;
+
+                // s = +1:
+                if (r >= mZero)
+                {
+                    Real rsqr = r * r;
+                    if (discr < rsqr)
+                    {
+                        SpecialIntersection(mid + sqrtDiscr, true, result);
+                    }
+                    else if (discr == rsqr)
+                    {
+                        SpecialIntersection(mid + sqrtDiscr, false, result);
+                    }
+                }
+
+                // s = -1:
+                if (r > mZero)
+                {
+                    SpecialIntersection(mid - sqrtDiscr, true, result);
+                }
+                else
+                {
+                    Real rsqr = r * r;
+                    if (discr > rsqr)
+                    {
+                        SpecialIntersection(mid - sqrtDiscr, true, result);
+                    }
+                    else if (discr == rsqr)
+                    {
+                        SpecialIntersection(mid - sqrtDiscr, false, result);
+                    }
+                }
+            }
+            else if (g[1] < mZero)
+            {
+                // We need s*sqrt(discr) >= -g[0]/g[1] + f1/2.
+                Real r = -g[0] / g[1] - mid;
+
+                // s = -1:
+                if (r <= mZero)
+                {
+                    Real rsqr = r * r;
+                    if (discr < rsqr)
+                    {
+                        SpecialIntersection(mid - sqrtDiscr, true, result);
+                    }
+                    else
+                    {
+                        SpecialIntersection(mid - sqrtDiscr, false, result);
+                    }
+                }
+
+                // s = +1:
+                if (r < mZero)
+                {
+                    SpecialIntersection(mid + sqrtDiscr, true, result);
+                }
+                else
+                {
+                    Real rsqr = r * r;
+                    if (discr > rsqr)
+                    {
+                        SpecialIntersection(mid + sqrtDiscr, true, result);
+                    }
+                    else if (discr == rsqr)
+                    {
+                        SpecialIntersection(mid + sqrtDiscr, false, result);
+                    }
+                }
+            }
+            else  // g[1] = 0
+            {
+                // The graphs of c(x) and f(x) are parabolas of the same
+                // shape.  One is a vertical translation of the other.
+                if (g[0] < mZero)
+                {
+                    // The graph of f(x) is above that of c(x).
+                    SpecialIntersection(mid - sqrtDiscr, true, result);
+                    SpecialIntersection(mid + sqrtDiscr, true, result);
+                }
+                else if (g[0] == mZero)
+                {
+                    // The graphs of c(x) and f(x) are the same parabola.
+                    SpecialIntersection(mid - sqrtDiscr, false, result);
+                    SpecialIntersection(mid + sqrtDiscr, false, result);
+                }
+            }
+        }
+        else if (discr == mZero)
+        {
+            // The theoretical root of f(x) is x = -f1/2.
+            Real nchat = -(mC[0] + mid * (mC[1] + mid * mC[2]));
+            if (nchat > mZero)
+            {
+                SpecialIntersection(mid, true, result);
+            }
+            else if (nchat == mZero)
+            {
+                SpecialIntersection(mid, false, result);
+            }
+        }
+    }
+    else if (mE[1] != mZero)
+    {
+        Real xhat = -mE[0] / mE[1];
+        Real nchat = -(mC[0] + xhat * (mC[1] + xhat * mC[2]));
+        if (nchat > mZero)
+        {
+            SpecialIntersection(xhat, true, result);
+        }
+        else if (nchat == mZero)
+        {
+            SpecialIntersection(xhat, false, result);
+        }
+    }
+}
+
+template <typename Real>
+void FIQuery<Real, Ellipse2<Real>, Ellipse2<Real>>::SpecialIntersection(
+    Real const& x, bool transverse, Result& result)
+{
+    if (transverse)
+    {
+        Real translate = mA2Div2 + x * mA4Div2;
+        Real nc = -(mC[0] + x * (mC[1] + x * mC[2]));
+        if (nc < mZero)
+        {
+            // Clamp to eliminate the rounding error, but duplicate the point
+            // because we know that it is a transverse intersection.
+            nc = mZero;
+        }
+
+        Real w = std::sqrt(nc);
+        Real y = w - translate;
+        result.points[result.numPoints] = { x, y };
+        result.isTransverse[result.numPoints++] = true;
+        w = -w;
+        y = w - translate;
+        result.points[result.numPoints] = { x, y };
+        result.isTransverse[result.numPoints++] = true;
+    }
+    else
+    {
+        // The vertical line at the root is tangent to the ellipse.
+        Real y = -(mA2Div2 + x * mA4Div2);  // w = 0
+        result.points[result.numPoints] = { x, y };
+        result.isTransverse[result.numPoints++] = false;
+    }
+}
+
+
+
+// FIQuery functions for area of intersection of ellipses.
+
+template <typename Real>
+typename FIQuery<Real, Ellipse2<Real>, Ellipse2<Real>>::AreaResult
+FIQuery<Real, Ellipse2<Real>, Ellipse2<Real>>::AreaOfIntersection(
+    Ellipse2<Real> const& ellipse0, Ellipse2<Real> const& ellipse1)
+{
+    EllipseInfo E0;
+    E0.center = ellipse0.center;
+    E0.axis = ellipse0.axis;
+    E0.extent = ellipse0.extent;
+    E0.sqrExtent = ellipse0.extent * ellipse0.extent;
+    FinishEllipseInfo(E0);
+
+    EllipseInfo E1;
+    E1.center = ellipse1.center;
+    E1.axis = ellipse1.axis;
+    E1.extent = ellipse1.extent;
+    E1.sqrExtent = ellipse1.extent * ellipse0.extent;
+    FinishEllipseInfo(E1);
+
+    AreaResult ar;
+    ar.configuration = 0;
+    ar.findResult = operator()(ellipse0, ellipse1);
+    ar.area = mZero;
+    AreaDispatch(E0, E1, ar);
+    return ar;
+}
+
+template <typename Real>
+typename FIQuery<Real, Ellipse2<Real>, Ellipse2<Real>>::AreaResult
+FIQuery<Real, Ellipse2<Real>, Ellipse2<Real>>::AreaOfIntersection(
+    Vector2<Real> const& center0, Vector2<Real> const axis0[2],
+    Vector2<Real> const& sqrExtent0, Vector2<Real> const& center1,
+    Vector2<Real> const axis1[2], Vector2<Real> const& sqrExtent1)
+{
+    EllipseInfo E0;
+    E0.center = center0;
+    E0.axis = { axis0[0], axis0[1] };
+    E0.extent =
+    {
+        std::sqrt(sqrExtent0[0]),
+        std::sqrt(sqrExtent0[1])
+    };
+    E0.sqrExtent = sqrExtent0;
+    FinishEllipseInfo(E0);
+
+    EllipseInfo E1;
+    E1.center = center1;
+    E1.axis = { axis1[0], axis1[1] };
+    E1.extent =
+    {
+        std::sqrt(sqrExtent1[0]),
+        std::sqrt(sqrExtent1[1])
+    };
+    E1.sqrExtent = sqrExtent1;
+    FinishEllipseInfo(E1);
+
+    AreaResult ar;
+    ar.configuration = 0;
+    ar.findResult = operator()(center0, axis0, sqrExtent0, center1, axis1,
+        sqrExtent1);
+    ar.area = mZero;
+    AreaDispatch(E0, E1, ar);
+    return ar;
+}
+
+template <typename Real>
+void FIQuery<Real, Ellipse2<Real>, Ellipse2<Real>>::FinishEllipseInfo(
+    EllipseInfo& E)
+{
+    E.M = OuterProduct(E.axis[0], E.axis[0]) /
+        (E.sqrExtent[0] * Dot(E.axis[0], E.axis[0]));
+    E.M += OuterProduct(E.axis[1], E.axis[1]) /
+        (E.sqrExtent[1] * Dot(E.axis[1], E.axis[1]));
+    E.AB = E.extent[0] * E.extent[1];
+    E.halfAB = E.AB / mTwo;
+    E.BpA = E.extent[1] + E.extent[0];
+    E.BmA = E.extent[1] - E.extent[0];
+}
+
+template <typename Real>
+void FIQuery<Real, Ellipse2<Real>, Ellipse2<Real>>::AreaDispatch(
+    EllipseInfo const& E0, EllipseInfo const& E1, AreaResult& ar)
+{
+    if (ar.findResult.intersect)
+    {
+        if (ar.findResult.numPoints == 1)
+        {
+            // Containment or separation.
+            AreaCS(E0, E1, ar);
+        }
+        else if (ar.findResult.numPoints == 2)
+        {
+            if (ar.findResult.isTransverse[0])
+            {
+                // Both intersection points are transverse.
+                Area2(E0, E1, 0, 1, ar);
+            }
+            else
+            {
+                // Both intersection points are tangential, so one ellipse
+                // is contained in the other.
+                AreaCS(E0, E1, ar);
+            }
+        }
+        else if (ar.findResult.numPoints == 3)
+        {
+            // The tangential intersection is irrelevant in the area
+            // computation.
+            if (!ar.findResult.isTransverse[0])
+            {
+                Area2(E0, E1, 1, 2, ar);
+            }
+            else if (!ar.findResult.isTransverse[1])
+            {
+                Area2(E0, E1, 2, 0, ar);
+            }
+            else  // ar.findResult.isTransverse[2] == false
+            {
+                Area2(E0, E1, 0, 1, ar);
+            }
+        }
+        else  // ar.findResult.numPoints == 4
+        {
+            Area4(E0, E1, ar);
+        }
+    }
+    else
+    {
+        // Containment, separation, or same ellipse.
+        AreaCS(E0, E1, ar);
+    }
+}
+
+template <typename Real>
+void FIQuery<Real, Ellipse2<Real>, Ellipse2<Real>>::AreaCS(
+    EllipseInfo const& E0, EllipseInfo const& E1, AreaResult& ar)
+{
+    if (ar.findResult.numPoints <= 1)
+    {
+        Vector2<Real> diff = E0.center - E1.center;
+        Real qform0 = Dot(diff, E0.M * diff);
+        Real qform1 = Dot(diff, E1.M * diff);
+        if (qform0 > mOne && qform1 > mOne)
+        {
+            // Each ellipse center is outside the other ellipse, so the
+            // ellipses are separated (numPoints == 0) or outside each
+            // other and just touching (numPoints == 1).
+            ar.configuration = AreaResult::ELLIPSES_ARE_SEPARATED;
+            ar.area = mZero;
+        }
+        else
+        {
+            // One ellipse is inside the other.  Determine this indirectly by
+            // comparing areas.
+            if (E0.AB < E1.AB)
+            {
+                ar.configuration = AreaResult::E1_CONTAINS_E0;
+                ar.area = mPi * E0.AB;
+            }
+            else
+            {
+                ar.configuration = AreaResult::E0_CONTAINS_E1;
+                ar.area = mPi * E1.AB;
+            }
+        }
+    }
+    else
+    {
+        ar.configuration = AreaResult::ELLIPSES_ARE_EQUAL;
+        ar.area = mPi * E0.AB;
+    }
+}
+
+template <typename Real>
+void FIQuery<Real, Ellipse2<Real>, Ellipse2<Real>>::Area2(
+    EllipseInfo const& E0, EllipseInfo const& E1, int i0, int i1,
+    AreaResult& ar)
+{
+    ar.configuration = AreaResult::ONE_CHORD_REGION;
+
+    // The endpoints of the chord.
+    Vector2<Real> const& P0 = ar.findResult.points[i0];
+    Vector2<Real> const& P1 = ar.findResult.points[i1];
+
+    // Compute locations relative to the ellipses.
+    Vector2<Real> P0mC0 = P0 - E0.center, P0mC1 = P0 - E1.center;
+    Vector2<Real> P1mC0 = P1 - E0.center, P1mC1 = P1 - E1.center;
+
+    // Compute the ellipse normal vectors at endpoint P0.  This is sufficient
+    // information to determine chord endpoint order.
+    Vector2<Real> N0 = E0.M * P0mC0, N1 = E1.M * P0mC1;
+    Real dotperp = DotPerp(N1, N0);
+
+    // Choose the endpoint order for the chord region associated with E0.
+    if (dotperp > mZero)
+    {
+        // The chord order for E0 is <P0,P1> and for E1 is <P1,P0>.
+        ar.area =
+            ComputeAreaChordRegion(E0, P0mC0, P1mC0) +
+            ComputeAreaChordRegion(E1, P1mC1, P0mC1);
+    }
+    else
+    {
+        // The chord order for E0 is <P1,P0> and for E1 is <P0,P1>.
+        ar.area =
+            ComputeAreaChordRegion(E0, P1mC0, P0mC0) +
+            ComputeAreaChordRegion(E1, P0mC1, P1mC1);
+    }
+}
+
+template <typename Real>
+void FIQuery<Real, Ellipse2<Real>, Ellipse2<Real>>::Area4(
+    EllipseInfo const& E0, EllipseInfo const& E1, AreaResult& ar)
+{
+    ar.configuration = AreaResult::FOUR_CHORD_REGION;
+
+    // Select a counterclockwise ordering of the points of intersection.  Use
+    // the polar coordinates for E0 to do this.  The multimap is used in the
+    // event that computing the intersections involved numerical rounding
+    // errors that lead to a duplicate intersection, even though the
+    // intersections are all labeled as transverse.  [See the comment in the
+    // function SpecialIntersection.]
+    std::multimap<Real, int> ordering;
+    int i;
+    for (i = 0; i < 4; ++i)
+    {
+        Vector2<Real> PmC = ar.findResult.points[i] - E0.center;
+        Real x = Dot(E0.axis[0], PmC);
+        Real y = Dot(E0.axis[1], PmC);
+        Real theta = std::atan2(y, x);
+        ordering.insert(std::make_pair(theta, i));
+    }
+
+    std::array<int, 4> permute;
+    i = 0;
+    for (auto const& element : ordering)
+    {
+        permute[i++] = element.second;
+    }
+
+    // Start with the area of the convex quadrilateral.
+    Vector2<Real> diag20 =
+        ar.findResult.points[permute[2]] - ar.findResult.points[permute[0]];
+    Vector2<Real> diag31 =
+        ar.findResult.points[permute[3]] - ar.findResult.points[permute[1]];
+    ar.area = std::abs(DotPerp(diag20, diag31)) / mTwo;
+
+    // Visit each pair of consecutive points.  The selection of ellipse for
+    // the chord-region area calculation uses the "most counterclockwise"
+    // tangent vector.
+    for (int i0 = 3, i1 = 0; i1 < 4; i0 = i1++)
+    {
+        // Get a pair of consecutive points.
+        Vector2<Real> const& P0 = ar.findResult.points[permute[i0]];
+        Vector2<Real> const& P1 = ar.findResult.points[permute[i1]];
+
+        // Compute locations relative to the ellipses.
+        Vector2<Real> P0mC0 = P0 - E0.center, P0mC1 = P0 - E1.center;
+        Vector2<Real> P1mC0 = P1 - E0.center, P1mC1 = P1 - E1.center;
+
+        // Compute the ellipse normal vectors at endpoint P0.
+        Vector2<Real> N0 = E0.M * P0mC0, N1 = E1.M * P0mC1;
+        Real dotperp = DotPerp(N1, N0);
+        if (dotperp > mZero)
+        {
+            // The chord goes with ellipse E0.
+            ar.area += ComputeAreaChordRegion(E0, P0mC0, P1mC0);
+        }
+        else
+        {
+            // The chord goes with ellipse E1.
+            ar.area += ComputeAreaChordRegion(E1, P0mC1, P1mC1);
+        }
+    }
+}
+
+template <typename Real>
+Real FIQuery<Real, Ellipse2<Real>, Ellipse2<Real>>::ComputeAreaChordRegion(
+    EllipseInfo const& E, Vector2<Real> const& P0mC,
+    Vector2<Real> const& P1mC)
+{
+    // Compute polar coordinates for P0 and P1 on the ellipse.
+    Real x0 = Dot(E.axis[0], P0mC);
+    Real y0 = Dot(E.axis[1], P0mC);
+    Real theta0 = std::atan2(y0, x0);
+    Real x1 = Dot(E.axis[0], P1mC);
+    Real y1 = Dot(E.axis[1], P1mC);
+    Real theta1 = std::atan2(y1, x1);
+
+    // The arc straddles the atan2 discontinuity on the negative x-axis.  Wrap
+    // the second angle to be larger than the first angle.
+    if (theta1 < theta0)
+    {
+        theta1 += mTwoPi;
+    }
+
+    // Compute the area portion of the sector due to the triangle.
+    Real triArea = std::abs(DotPerp(P0mC, P1mC)) / mTwo;
+
+    // Compute the chord region area.
+    Real dtheta = theta1 - theta0;
+    Real F0, F1, sectorArea;
+    if (dtheta <= mPi)
+    {
+        // Use the area formula directly.
+        //   area(theta0,theta1) = F(theta1) - F(theta0) - area(triangle)
+        F0 = ComputeIntegral(E, theta0);
+        F1 = ComputeIntegral(E, theta1);
+        sectorArea = F1 - F0;
+        return sectorArea - triArea;
+    }
+    else
+    {
+        // The angle of the elliptical sector is larger than pi radians.
+        // Use the area formula
+        //   area(theta0,theta1) = pi*a*b - area(theta1,theta0)
+        theta0 += mTwoPi;  // ensure theta0 > theta1
+        F0 = ComputeIntegral(E, theta0);
+        F1 = ComputeIntegral(E, theta1);
+        sectorArea = F0 - F1;
+        return mPi * E.AB - (sectorArea - triArea);
+    }
+}
+
+template <typename Real>
+Real FIQuery<Real, Ellipse2<Real>, Ellipse2<Real>>::ComputeIntegral(
+    EllipseInfo const& E, Real const& theta)
+{
+    Real twoTheta = mTwo * theta;
+    Real sn = std::sin(twoTheta);
+    Real cs = std::cos(twoTheta);
+    Real arg = E.BmA * sn / (E.BpA + E.BmA * cs);
+    return E.halfAB * (theta - std::atan(arg));
+}
+
+
+}
diff --git a/Engine/Plugins/Experimental/GeometryProcessing/Source/GeometryAlgorithms/Private/GTEngine/Mathematics/GteIntrEllipsoid3Ellipsoid3.h b/Engine/Plugins/Experimental/GeometryProcessing/Source/GeometryAlgorithms/Private/GTEngine/Mathematics/GteIntrEllipsoid3Ellipsoid3.h
new file mode 100644
--- /dev/null
+++ b/Engine/Plugins/Experimental/GeometryProcessing/Source/GeometryAlgorithms/Private/GTEngine/Mathematics/GteIntrEllipsoid3Ellipsoid3.h
@@ -0,0 +1,544 @@
+// David Eberly, Geometric Tools, Redmond WA 98052
+// Copyright (c) 1998-2019
+// Distributed under the Boost Software License, Version 1.0.
+// http://www.boost.org/LICENSE_1_0.txt
+// http://www.geometrictools.com/License/Boost/LICENSE_1_0.txt
+// File Version: 3.0.2 (2018/10/05)
+
+#pragma once
+
+#include <GTEngine/LowLevel/GteLogger.h>
+#include <GTEngine/Mathematics/GteHyperellipsoid.h>
+#include <GTEngine/Mathematics/GteMatrix3x3.h>
+#include <GTEngine/Mathematics/GteFIQuery.h>
+#include <GTEngine/Mathematics/GteTIQuery.h>
+#include <GTEngine/Mathematics/GteRootsBisection.h>
+#include <GTEngine/Mathematics/GteRootsPolynomial.h>
+#include <GTEngine/Mathematics/GteSymmetricEigensolver3x3.h>
+
+namespace gte
+{
+
+template <typename Real>
+class TIQuery<Real, Ellipsoid3<Real>, Ellipsoid3<Real>>
+{
+public:
+    enum
+    {
+        ELLIPSOIDS_SEPARATED,
+        ELLIPSOIDS_INTERSECTING,
+        ELLIPSOID0_CONTAINS_ELLIPSOID1,
+        ELLIPSOID1_CONTAINS_ELLIPSOID0
+    };
+
+    struct Result
+    {
+        // As solids, the ellipsoids intersect as long as they are not
+        // separated.
+        bool intersect;
+
+        // This is one of the four enumerations listed above.
+        int relationship;
+    };
+
+    Result operator()(Ellipsoid3<Real> const& ellipsoid0,
+        Ellipsoid3<Real> const& ellipsoid1);
+
+private:
+    void GetRoots(Real d0, Real c0, int& numRoots, Real* roots);
+    void GetRoots(Real d0, Real d1, Real c0, Real c1, int& numRoots,
+        Real* roots);
+    void GetRoots(Real d0, Real d1, Real d2, Real c0, Real c1, Real c2,
+        int& numRoots, Real* roots);
+};
+
+
+template <typename Real>
+typename TIQuery<Real, Ellipsoid3<Real>, Ellipsoid3<Real>>::Result
+TIQuery<Real, Ellipsoid3<Real>, Ellipsoid3<Real>>::operator()(
+    Ellipsoid3<Real> const& ellipsoid0, Ellipsoid3<Real> const& ellipsoid1)
+{
+    Result result;
+
+    Real const zero = (Real)0;
+    Real const one = (Real)1;
+
+    // Get the parameters of ellipsoid0.
+    Vector3<Real> K0 = ellipsoid0.center;
+    Matrix3x3<Real> R0;
+    R0.SetCol(0, ellipsoid0.axis[0]);
+    R0.SetCol(1, ellipsoid0.axis[1]);
+    R0.SetCol(2, ellipsoid0.axis[2]);
+    Matrix3x3<Real> D0{
+        one / (ellipsoid0.extent[0] * ellipsoid0.extent[0]), zero, zero,
+        zero, one / (ellipsoid0.extent[1] * ellipsoid0.extent[1]), zero,
+        zero, zero, one / (ellipsoid0.extent[2] * ellipsoid0.extent[2]) };
+
+    // Get the parameters of ellipsoid1.
+    Vector3<Real> K1 = ellipsoid1.center;
+    Matrix3x3<Real> R1;
+    R1.SetCol(0, ellipsoid1.axis[0]);
+    R1.SetCol(1, ellipsoid1.axis[1]);
+    R1.SetCol(2, ellipsoid1.axis[2]);
+    Matrix3x3<Real> D1{
+        one / (ellipsoid1.extent[0] * ellipsoid1.extent[0]), zero, zero,
+        zero, one / (ellipsoid1.extent[1] * ellipsoid1.extent[1]), zero,
+        zero, zero, one / (ellipsoid1.extent[2] * ellipsoid1.extent[2]) };
+
+    // Compute K2.
+    Matrix3x3<Real> D0NegHalf{
+        ellipsoid0.extent[0], zero, zero,
+        zero, ellipsoid0.extent[1], zero,
+        zero, zero, ellipsoid0.extent[2] };
+    Matrix3x3<Real> D0Half{
+        one / ellipsoid0.extent[0], zero, zero,
+        zero, one / ellipsoid0.extent[1], zero,
+        zero, zero, one / ellipsoid0.extent[2] };
+    Vector3<Real> K2 = D0Half*((K1 - K0)*R0);
+
+    // Compute M2.
+    Matrix3x3<Real> R1TR0D0NegHalf = MultiplyATB(R1, R0*D0NegHalf);
+    Matrix3x3<Real> M2 = MultiplyATB(R1TR0D0NegHalf, D1)*R1TR0D0NegHalf;
+
+    // Factor M2 = R*D*R^T.
+    SymmetricEigensolver3x3<Real> es;
+    std::array<Real, 3> D;
+    std::array<std::array<Real, 3>, 3> evec;
+    es(M2(0, 0), M2(0, 1), M2(0, 2), M2(1, 1), M2(1, 2), M2(2, 2), false, +1,
+        D, evec);
+    Matrix3x3<Real> R;
+    R.SetCol(0, evec[0]);
+    R.SetCol(1, evec[1]);
+    R.SetCol(2, evec[2]);
+
+    // Compute K = R^T*K2.
+    Vector3<Real> K = K2*R;
+
+    // Transformed ellipsoid0 is Z^T*Z = 1 and transformed ellipsoid1 is
+    // (Z-K)^T*D*(Z-K) = 0.
+
+    // The minimum and maximum squared distances from the origin of points on
+    // transformed ellipsoid1 are used to determine whether the ellipsoids
+    // intersect, are separated, or one contains the other.
+    Real minSqrDistance = std::numeric_limits<Real>::max();
+    Real maxSqrDistance = zero;
+    int i;
+
+    if (K == Vector3<Real>::Zero())
+    {
+        // The special case of common centers must be handled separately.  It
+        // is not possible for the ellipsoids to be separated.
+        for (i = 0; i < 3; ++i)
+        {
+            Real invD = one / D[i];
+            if (invD < minSqrDistance)
+            {
+                minSqrDistance = invD;
+            }
+            if (invD > maxSqrDistance)
+            {
+                maxSqrDistance = invD;
+            }
+        }
+
+        if (maxSqrDistance < one)
+        {
+            result.relationship = ELLIPSOID0_CONTAINS_ELLIPSOID1;
+        }
+        else if (minSqrDistance >(Real)1)
+        {
+            result.relationship = ELLIPSOID1_CONTAINS_ELLIPSOID0;
+        }
+        else
+        {
+            result.relationship = ELLIPSOIDS_INTERSECTING;
+        }
+        result.intersect = true;
+        return result;
+    }
+
+    // The closest point P0 and farthest point P1 are solutions to
+    // s0*D*(P0 - K) = P0 and s1*D*(P1 - K) = P1 for some scalars s0 and s1
+    // that are roots to the function
+    //   f(s) = d0*k0^2/(d0*s-1)^2+d1*k1^2/(d1*s-1)^2+d2*k2^2/(d2*s-1)^2-1
+    // where D = diagonal(d0,d1,d2) and K = (k0,k1,k2).
+    Real d0 = D[0], d1 = D[1], d2 = D[2];
+    Real c0 = K[0] * K[0], c1 = K[1] * K[1], c2 = K[2] * K[2];
+
+    // Sort the values so that d0 >= d1 >= d2.  This allows us to bound the
+    // roots of f(s), of which there are at most 6.
+    std::vector<std::pair<Real, Real>> param(3);
+    param[0] = std::make_pair(d0, c0);
+    param[1] = std::make_pair(d1, c1);
+    param[2] = std::make_pair(d2, c2);
+    std::sort(param.begin(), param.end(),
+        std::greater<std::pair<Real, Real>>());
+
+    std::vector<std::pair<Real, Real>> valid;
+    valid.reserve(3);
+    if (param[0].first > param[1].first)
+    {
+        if (param[1].first > param[2].first)
+        {
+            // d0 > d1 > d2
+            for (i = 0; i < 3; ++i)
+            {
+                if (param[i].second >(Real)0)
+                {
+                    valid.push_back(param[i]);
+                }
+            }
+        }
+        else
+        {
+            // d0 > d1 = d2
+            if (param[0].second > (Real)0)
+            {
+                valid.push_back(param[0]);
+            }
+            param[1].second += param[0].second;
+            if (param[1].second > (Real)0)
+            {
+                valid.push_back(param[1]);
+            }
+        }
+    }
+    else
+    {
+        if (param[1].first > param[2].first)
+        {
+            // d0 = d1 > d2
+            param[0].second += param[1].second;
+            if (param[0].second > (Real)0)
+            {
+                valid.push_back(param[0]);
+            }
+            if (param[2].second > (Real)0)
+            {
+                valid.push_back(param[2]);
+            }
+        }
+        else
+        {
+            // d0 = d1 = d2
+            param[0].second += param[1].second + param[2].second;
+            if (param[0].second > (Real)0)
+            {
+                valid.push_back(param[0]);
+            }
+        }
+    }
+
+    size_t numValid = valid.size();
+    int numRoots;
+    Real roots[6];
+    if (numValid == 3)
+    {
+        GetRoots(valid[0].first, valid[1].first, valid[2].first,
+            valid[0].second, valid[1].second, valid[2].second, numRoots,
+            roots);
+    }
+    else if (numValid == 2)
+    {
+        GetRoots(valid[0].first, valid[1].first, valid[0].second,
+            valid[1].second, numRoots, roots);
+    }
+    else if (numValid == 1)
+    {
+        GetRoots(valid[0].first, valid[0].second, numRoots, roots);
+    }
+    else
+    {
+        // numValid cannot be zero because we already handled case K = 0
+        LogError("Unexpected condition.");
+        result.intersect = true;
+        result.relationship = ELLIPSOIDS_INTERSECTING;
+        return result;
+    }
+
+    for (i = 0; i < numRoots; ++i)
+    {
+        Real s = roots[i];
+        Real p0 = d0*K[0] * s / (d0 * s - (Real)1);
+        Real p1 = d1*K[1] * s / (d1 * s - (Real)1);
+        Real p2 = d2*K[2] * s / (d2 * s - (Real)1);
+        Real sqrDistance = p0 * p0 + p1 * p1 + p2 * p2;
+        if (sqrDistance < minSqrDistance)
+        {
+            minSqrDistance = sqrDistance;
+        }
+        if (sqrDistance > maxSqrDistance)
+        {
+            maxSqrDistance = sqrDistance;
+        }
+    }
+
+    if (maxSqrDistance < one)
+    {
+        result.intersect = true;
+        result.relationship = ELLIPSOID0_CONTAINS_ELLIPSOID1;
+    }
+    else if (minSqrDistance >(Real)1)
+    {
+        if (d0 * c0 + d1 * c1 + d2 * c2 > one)
+        {
+            result.intersect = false;
+            result.relationship = ELLIPSOIDS_SEPARATED;
+        }
+        else
+        {
+            result.intersect = true;
+            result.relationship = ELLIPSOID1_CONTAINS_ELLIPSOID0;
+        }
+    }
+    else
+    {
+        result.intersect = true;
+        result.relationship = ELLIPSOIDS_INTERSECTING;
+    }
+
+    return result;
+}
+
+template <typename Real>
+void TIQuery<Real, Ellipsoid3<Real>, Ellipsoid3<Real>>::GetRoots(Real d0,
+    Real c0, int& numRoots, Real* roots)
+{
+    // f(s) = d0*c0/(d0*s-1)^2 - 1
+    Real const one = (Real)1;
+    Real temp = std::sqrt(d0*c0);
+    Real inv = one / d0;
+    numRoots = 2;
+    roots[0] = (one - temp) * inv;
+    roots[1] = (one + temp) * inv;
+}
+
+template <typename Real>
+void TIQuery<Real, Ellipsoid3<Real>, Ellipsoid3<Real>>::GetRoots(Real d0,
+    Real d1, Real c0, Real c1, int& numRoots, Real* roots)
+{
+    // f(s) = d0*c0/(d0*s-1)^2 + d1*c1/(d1*s-1)^2 - 1
+    // with d0 > d1
+
+    Real const zero = (Real)0;
+    Real const one = (Real)1;
+    Real const two = (Real)2;
+    Real d0c0 = d0 * c0;
+    Real d1c1 = d1 * c1;
+
+    std::function<Real(Real)> F = [&one, d0, d1, d0c0, d1c1](Real s)
+    {
+        Real invN0 = one / (d0*s - one);
+        Real invN1 = one / (d1*s - one);
+        Real term0 = d0c0 * invN0 * invN0;
+        Real term1 = d1c1 * invN1 * invN1;
+        Real f = term0 + term1 - one;
+        return f;
+    };
+
+    std::function<Real(Real)> DF = [&one, &two, d0, d1, d0c0, d1c1](Real s)
+    {
+        Real invN0 = one / (d0*s - one);
+        Real invN1 = one / (d1*s - one);
+        Real term0 = d0 * d0c0 * invN0 * invN0 * invN0;
+        Real term1 = d1 * d1c1 * invN1 * invN1 * invN1;
+        Real df = -two * (term0 + term1);
+        return df;
+    };
+
+    unsigned int const maxIterations = 1024;
+    unsigned int iterations;
+    numRoots = 0;
+
+    // TODO: What role does epsilon play?
+    Real const epsilon = (Real)0.001;
+    Real multiplier0 = std::sqrt(two / (one - epsilon));
+    Real multiplier1 = std::sqrt(one / (one + epsilon));
+    Real sqrtd0c0 = std::sqrt(d0c0);
+    Real sqrtd1c1 = std::sqrt(d1c1);
+    Real invD0 = one / d0;
+    Real invD1 = one / d1;
+    Real temp0, temp1, smin, smax, s;
+
+    // Compute root in (-infinity,1/d0).
+    temp0 = (one - multiplier0 * sqrtd0c0) * invD0;
+    temp1 = (one - multiplier0 * sqrtd1c1) * invD1;
+    smin = std::min(temp0, temp1);
+    LogAssert(F(smin) < zero, "Unexpected condition.");
+    smax = (one - multiplier1*sqrtd0c0)*invD0;
+    LogAssert(F(smax) > zero, "Unexpected condition.");
+    iterations = RootsBisection<Real>::Find(F, smin, smax, maxIterations, s);
+    LogAssert(iterations > 0, "Unexpected condition.");
+    roots[numRoots++] = s;
+
+    // Compute roots (if any) in (1/d0,1/d1).  It is the case that
+    //   F(1/d0) = +infinity, F'(1/d0) = -infinity
+    //   F(1/d1) = +infinity, F'(1/d1) = +infinity
+    //   F"(s) > 0 for all s in the domain of F
+    // Compute the unique root r of F'(s) on (1/d0,1/d1).  The bisector needs
+    // only the signs at the endpoints, so we pass -1 and +1 instead of the
+    // infinite values.  If F(r) < 0, F(s) has two roots in the interval.
+    // If F(r) = 0, F(s) has only one root in the interval.
+    Real smid;
+    iterations = RootsBisection<Real>::Find(DF, invD0, invD1, -one, one,
+        maxIterations, smid);
+    LogAssert(iterations > 0, "Unexpected condition.");
+    if (F(smid) < zero)
+    {
+        // Pass in signs rather than infinities, because the bisector cares
+        // only about the signs.
+        iterations = RootsBisection<Real>::Find(F, invD0, smid, one, -one,
+            maxIterations, s);
+        LogAssert(iterations > 0, "Unexpected condition.");
+        roots[numRoots++] = s;
+        iterations = RootsBisection<Real>::Find(F, smid, invD1, -one, one,
+            maxIterations, s);
+        LogAssert(iterations > 0, "Unexpected condition.");
+        roots[numRoots++] = s;
+    }
+
+    // Compute root in (1/d1,+infinity).
+    temp0 = (one + multiplier0 * sqrtd0c0) * invD0;
+    temp1 = (one + multiplier0 * sqrtd1c1) * invD1;
+    smax = std::max(temp0, temp1);
+    LogAssert(F(smax) < zero, "Unexpected condition.");
+    smin = (one + multiplier1 * sqrtd1c1) * invD1;
+    LogAssert(F(smin) > zero, "Unexpected condition.");
+    iterations = RootsBisection<Real>::Find(F, smin, smax, maxIterations, s);
+    LogAssert(iterations > 0, "Unexpected condition.");
+    roots[numRoots++] = s;
+}
+
+template <typename Real>
+void TIQuery<Real, Ellipsoid3<Real>, Ellipsoid3<Real>>::GetRoots(Real d0,
+    Real d1, Real d2, Real c0, Real c1, Real c2, int& numRoots, Real* roots)
+{
+    // f(s) = d0*c0/(d0*s-1)^2 + d1*c1/(d1*s-1)^2 + d2*c2/(d2*s-1)^2 - 1
+    // with d0 > d1 > d2
+
+    Real const zero = (Real)0;
+    Real const one = (Real)1;
+    Real const three = (Real)3;
+    Real d0c0 = d0 * c0;
+    Real d1c1 = d1 * c1;
+    Real d2c2 = d2 * c2;
+
+    std::function<Real(Real)> F = [&one, d0, d1, d2, d0c0, d1c1, d2c2](Real s)
+    {
+        Real invN0 = one / (d0*s - one);
+        Real invN1 = one / (d1*s - one);
+        Real invN2 = one / (d2*s - one);
+        Real term0 = d0c0 * invN0 * invN0;
+        Real term1 = d1c1 * invN1 * invN1;
+        Real term2 = d2c2 * invN2 * invN2;
+        Real f = term0 + term1 + term2 - one;
+        return f;
+    };
+
+    std::function<Real(Real)> DF = [&one, d0, d1, d2, d0c0, d1c1, d2c2](Real s)
+    {
+        Real const two = (Real)2;
+        Real invN0 = one / (d0*s - one);
+        Real invN1 = one / (d1*s - one);
+        Real invN2 = one / (d2*s - one);
+        Real term0 = d0 * d0c0 * invN0 * invN0 * invN0;
+        Real term1 = d1 * d1c1 * invN1 * invN1 * invN1;
+        Real term2 = d2 * d2c2 * invN2 * invN2 * invN2;
+        Real df = -two * (term0 + term1 + term2);
+        return df;
+    };
+
+    unsigned int const maxIterations = 1024;
+    unsigned int iterations;
+    numRoots = 0;
+
+    // TODO: What role does epsilon play?
+    Real epsilon = (Real)0.001;
+    Real multiplier0 = std::sqrt(three / (one - epsilon));
+    Real multiplier1 = std::sqrt(one / (one + epsilon));
+    Real sqrtd0c0 = std::sqrt(d0c0);
+    Real sqrtd1c1 = std::sqrt(d1c1);
+    Real sqrtd2c2 = std::sqrt(d2c2);
+    Real invD0 = one / d0;
+    Real invD1 = one / d1;
+    Real invD2 = one / d2;
+    Real temp0, temp1, temp2, smin, smax, s;
+
+    // Compute root in (-infinity,1/d0).
+    temp0 = (one - multiplier0*sqrtd0c0)*invD0;
+    temp1 = (one - multiplier0*sqrtd1c1)*invD1;
+    temp2 = (one - multiplier0*sqrtd2c2)*invD2;
+    smin = std::min(std::min(temp0, temp1), temp2);
+    LogAssert(F(smin) < zero, "Unexpected condition.");
+    smax = (one - multiplier1*sqrtd0c0)*invD0;
+    LogAssert(F(smax) > zero, "Unexpected condition.");
+    iterations = RootsBisection<Real>::Find(F, smin, smax, maxIterations, s);
+    LogAssert(iterations > 0, "Unexpected condition.");
+    roots[numRoots++] = s;
+
+    // Compute roots (if any) in (1/d0,1/d1).  It is the case that
+    //   F(1/d0) = +infinity, F'(1/d0) = -infinity
+    //   F(1/d1) = +infinity, F'(1/d1) = +infinity
+    //   F"(s) > 0 for all s in the domain of F
+    // Compute the unique root r of F'(s) on (1/d0,1/d1).  The bisector needs
+    // only the signs at the endpoints, so we pass -1 and +1 instead of the
+    // infinite values.  If F(r) < 0, F(s) has two roots in the interval.
+    // If F(r) = 0, F(s) has only one root in the interval.
+    Real smid;
+    iterations = RootsBisection<Real>::Find(DF, invD0, invD1, -one, one,
+        maxIterations, smid);
+    LogAssert(iterations > 0, "Unexpected condition.");
+    if (F(smid) < zero)
+    {
+        // Pass in signs rather than infinities, because the bisector cares
+        // only about the signs.
+        iterations = RootsBisection<Real>::Find(F, invD0, smid, one, -one,
+            maxIterations, s);
+        LogAssert(iterations > 0, "Unexpected condition.");
+        roots[numRoots++] = s;
+        iterations = RootsBisection<Real>::Find(F, smid, invD1, -one, one,
+            maxIterations, s);
+        LogAssert(iterations > 0, "Unexpected condition.");
+        roots[numRoots++] = s;
+    }
+
+    // Compute roots (if any) in (1/d1,1/d2).  It is the case that
+    //   F(1/d1) = +infinity, F'(1/d1) = -infinity
+    //   F(1/d2) = +infinity, F'(1/d2) = +infinity
+    //   F"(s) > 0 for all s in the domain of F
+    // Compute the unique root r of F'(s) on (1/d1,1/d2).  The bisector needs
+    // only the signs at the endpoints, so we pass -1 and +1 instead of the
+    // infinite values.  If F(r) < 0, F(s) has two roots in the interval.
+    // If F(r) = 0, F(s) has only one root in the interval.
+    iterations = RootsBisection<Real>::Find(DF, invD1, invD2, -one, one,
+        maxIterations, smid);
+    LogAssert(iterations > 0, "Unexpected condition.");
+    if (F(smid) < zero)
+    {
+        // Pass in signs rather than infinities, because the bisector cares
+        // only about the signs.
+        iterations = RootsBisection<Real>::Find(F, invD1, smid, one, -one,
+            maxIterations, s);
+        LogAssert(iterations > 0, "Unexpected condition.");
+        roots[numRoots++] = s;
+        iterations = RootsBisection<Real>::Find(F, smid, invD2, -one, one,
+            maxIterations, s);
+        LogAssert(iterations > 0, "Unexpected condition.");
+        roots[numRoots++] = s;
+    }
+
+    // Compute root in (1/d2,+infinity).
+    temp0 = (one + multiplier0*sqrtd0c0)*invD0;
+    temp1 = (one + multiplier0*sqrtd1c1)*invD1;
+    temp2 = (one + multiplier0*sqrtd2c2)*invD2;
+    smax = std::max(std::max(temp0, temp1), temp2);
+    LogAssert(F(smax) < zero, "Unexpected condition.");
+    smin = (one + multiplier1*sqrtd2c2)*invD2;
+    LogAssert(F(smin) > zero, "Unexpected condition.");
+    iterations = RootsBisection<Real>::Find(F, smin, smax, maxIterations, s);
+    LogAssert(iterations > 0, "Unexpected condition.");
+    roots[numRoots++] = s;
+}
+
+
+}
diff --git a/Engine/Plugins/Experimental/GeometryProcessing/Source/GeometryAlgorithms/Private/GTEngine/Mathematics/GteIntrHalfspace2Polygon2.h b/Engine/Plugins/Experimental/GeometryProcessing/Source/GeometryAlgorithms/Private/GTEngine/Mathematics/GteIntrHalfspace2Polygon2.h
new file mode 100644
--- /dev/null
+++ b/Engine/Plugins/Experimental/GeometryProcessing/Source/GeometryAlgorithms/Private/GTEngine/Mathematics/GteIntrHalfspace2Polygon2.h
@@ -0,0 +1,167 @@
+// David Eberly, Geometric Tools, Redmond WA 98052
+// Copyright (c) 1998-2019
+// Distributed under the Boost Software License, Version 1.0.
+// http://www.boost.org/LICENSE_1_0.txt
+// http://www.geometrictools.com/License/Boost/LICENSE_1_0.txt
+// File Version: 3.0.0 (2016/06/19)
+
+#pragma once
+
+#include <GTEngine/Mathematics/GteVector2.h>
+#include <GTEngine/Mathematics/GteHalfspace.h>
+#include <GTEngine/Mathematics/GteOrientedBox.h>
+#include <GTEngine/Mathematics/GteFIQuery.h>
+#include <vector>
+
+// The queries consider the box to be a solid and the polygon to be a
+// convex solid.
+
+namespace gte
+{
+
+template <typename Real>
+class FIQuery<Real, Halfspace<2, Real>, std::vector<Vector2<Real>>>
+{
+public:
+    struct Result
+    {
+        bool intersect;
+
+        // If 'intersect' is true, the halfspace and polygon intersect in a
+        // convex polygon.
+        std::vector<Vector2<Real>> polygon;
+    };
+
+    Result operator()(Halfspace<2, Real> const& halfspace,
+        std::vector<Vector2<Real>> const& polygon);
+};
+
+
+template <typename Real>
+typename FIQuery<Real, Halfspace<2, Real>, std::vector<Vector2<Real>>>::Result
+FIQuery<Real, Halfspace<2, Real>, std::vector<Vector2<Real>>>::operator()(
+   Halfspace<2, Real> const& halfspace, std::vector<Vector2<Real>> const& polygon)
+{
+    Result result;
+
+    // Determine whether the polygon vertices are outside the halfspace,
+    // inside the halfspace, or on the halfspace boundary.
+    int const numVertices = static_cast<int>(polygon.size());
+    std::vector<Real> distance(numVertices);
+    int positive = 0, negative = 0, positiveIndex = -1;
+    for (int i = 0; i < numVertices; ++i)
+    {
+        distance[i] = Dot(halfspace.normal, polygon[i]) - halfspace.constant;
+        if (distance[i] > (Real)0)
+        {
+            ++positive;
+            if (positiveIndex == -1)
+            {
+                positiveIndex = i;
+            }
+        }
+        else if (distance[i] < (Real)0)
+        {
+            ++negative;
+        }
+    }
+
+    if (positive == 0)
+    {
+        // The polygon is strictly outside the halfspace.
+        result.intersect = false;
+        return result;
+    }
+
+    if (negative == 0)
+    {
+        // The polygon is contained in the closed halfspace, so it is
+        // fully visible and no clipping is necessary.
+        result.intersect = true;
+        return result;
+    }
+
+    // The line transversely intersects the polygon.  Clip the polygon.
+    Vector2<Real> vertex;
+    int curr, prev;
+    Real t;
+
+    if (positiveIndex > 0)
+    {
+        // Compute the first clip vertex on the line.
+        curr = positiveIndex;
+        prev = curr - 1;
+        t = distance[curr] / (distance[curr] - distance[prev]);
+        vertex = polygon[curr] + t * (polygon[prev] - polygon[curr]);
+        result.polygon.push_back(vertex);
+
+        // Include the vertices on the positive side of line.
+        while (curr < numVertices && distance[curr] >(Real)0)
+        {
+            result.polygon.push_back(polygon[curr++]);
+        }
+
+        // Compute the kast clip vertex on the line.
+        if (curr < numVertices)
+        {
+            prev = curr - 1;
+        }
+        else
+        {
+            curr = 0;
+            prev = numVertices - 1;
+        }
+        t = distance[curr] / (distance[curr] - distance[prev]);
+        vertex = polygon[curr] + t * (polygon[prev] - polygon[curr]);
+        result.polygon.push_back(vertex);
+    }
+    else  // positiveIndex is 0
+    {
+        // Include the vertices on the positive side of line.
+        curr = 0;
+        while (curr < numVertices && distance[curr] >(Real)0)
+        {
+            result.polygon.push_back(polygon[curr++]);
+        }
+
+        // Compute the last clip vertex on the line.
+        prev = curr - 1;
+        t = distance[curr] / (distance[curr] - distance[prev]);
+        vertex = polygon[curr] + t * (polygon[prev] - polygon[curr]);
+        result.polygon.push_back(vertex);
+
+        // Skip the vertices on the negative side of the line.
+        while (curr < numVertices && distance[curr] <= (Real)0)
+        {
+            curr++;
+        }
+
+        // Compute the first clip vertex on the line.
+        if (curr < numVertices)
+        {
+            prev = curr - 1;
+            t = distance[curr] / (distance[curr] - distance[prev]);
+            vertex = polygon[curr] + t * (polygon[prev] - polygon[curr]);
+            result.polygon.push_back(vertex);
+
+            // Keep the vertices on the positive side of the line.
+            while (curr < numVertices && distance[curr] >(Real)0)
+            {
+                result.polygon.push_back(polygon[curr++]);
+            }
+        }
+        else
+        {
+            curr = 0;
+            prev = numVertices - 1;
+            t = distance[curr] / (distance[curr] - distance[prev]);
+            vertex = polygon[curr] + t * (polygon[prev] - polygon[curr]);
+            result.polygon.push_back(vertex);
+        }
+    }
+
+    result.intersect = true;
+    return result;
+}
+
+}
diff --git a/Engine/Plugins/Experimental/GeometryProcessing/Source/GeometryAlgorithms/Private/GTEngine/Mathematics/GteIntrHalfspace3Capsule3.h b/Engine/Plugins/Experimental/GeometryProcessing/Source/GeometryAlgorithms/Private/GTEngine/Mathematics/GteIntrHalfspace3Capsule3.h
new file mode 100644
--- /dev/null
+++ b/Engine/Plugins/Experimental/GeometryProcessing/Source/GeometryAlgorithms/Private/GTEngine/Mathematics/GteIntrHalfspace3Capsule3.h
@@ -0,0 +1,56 @@
+// David Eberly, Geometric Tools, Redmond WA 98052
+// Copyright (c) 1998-2019
+// Distributed under the Boost Software License, Version 1.0.
+// http://www.boost.org/LICENSE_1_0.txt
+// http://www.geometrictools.com/License/Boost/LICENSE_1_0.txt
+// File Version: 3.0.0 (2016/06/19)
+
+#pragma once
+
+#include <GTEngine/Mathematics/GteCapsule.h>
+#include <GTEngine/Mathematics/GteHalfspace.h>
+#include <GTEngine/Mathematics/GteFIQuery.h>
+#include <GTEngine/Mathematics/GteTIQuery.h>
+
+// Queries for intersection of objects with halfspaces.  These are useful for
+// containment testing, object culling, and clipping.
+
+namespace gte
+{
+
+template <typename Real>
+class TIQuery<Real, Halfspace3<Real>, Capsule3<Real>>
+{
+public:
+    struct Result
+    {
+        bool intersect;
+    };
+
+    Result operator()(Halfspace3<Real> const& halfspace,
+        Capsule3<Real> const& capsule);
+};
+
+
+template <typename Real>
+typename TIQuery<Real, Halfspace3<Real>, Capsule3<Real>>::Result
+TIQuery<Real, Halfspace3<Real>, Capsule3<Real>>::operator()(
+    Halfspace3<Real> const& halfspace, Capsule3<Real> const& capsule)
+{
+    Result result;
+
+    // Project the capsule onto the normal line.  The plane of the halfspace
+    // occurs at the origin (zero) of the normal line.
+    Real e0 =
+        Dot(halfspace.normal, capsule.segment.p[0]) - halfspace.constant;
+    Real e1 =
+        Dot(halfspace.normal, capsule.segment.p[1]) - halfspace.constant;
+
+    // The capsule and halfspace intersect when the projection interval
+    // maximum is nonnegative.
+    result.intersect = (std::max(e0, e1) + capsule.radius >= (Real)0);
+    return result;
+}
+
+
+}
diff --git a/Engine/Plugins/Experimental/GeometryProcessing/Source/GeometryAlgorithms/Private/GTEngine/Mathematics/GteIntrHalfspace3Cylinder3.h b/Engine/Plugins/Experimental/GeometryProcessing/Source/GeometryAlgorithms/Private/GTEngine/Mathematics/GteIntrHalfspace3Cylinder3.h
new file mode 100644
--- /dev/null
+++ b/Engine/Plugins/Experimental/GeometryProcessing/Source/GeometryAlgorithms/Private/GTEngine/Mathematics/GteIntrHalfspace3Cylinder3.h
@@ -0,0 +1,60 @@
+// David Eberly, Geometric Tools, Redmond WA 98052
+// Copyright (c) 1998-2019
+// Distributed under the Boost Software License, Version 1.0.
+// http://www.boost.org/LICENSE_1_0.txt
+// http://www.geometrictools.com/License/Boost/LICENSE_1_0.txt
+// File Version: 3.0.1 (2018/10/05)
+
+#pragma once
+
+#include <GTEngine/Mathematics/GteCylinder3.h>
+#include <GTEngine/Mathematics/GteHalfspace.h>
+#include <GTEngine/Mathematics/GteFIQuery.h>
+#include <GTEngine/Mathematics/GteTIQuery.h>
+
+// Queries for intersection of objects with halfspaces.  These are useful for
+// containment testing, object culling, and clipping.
+
+namespace gte
+{
+
+template <typename Real>
+class TIQuery<Real, Halfspace3<Real>, Cylinder3<Real>>
+{
+public:
+    struct Result
+    {
+        bool intersect;
+    };
+
+    Result operator()(Halfspace3<Real> const& halfspace,
+        Cylinder3<Real> const& cylinder);
+};
+
+
+template <typename Real>
+typename TIQuery<Real, Halfspace3<Real>, Cylinder3<Real>>::Result
+TIQuery<Real, Halfspace3<Real>, Cylinder3<Real>>::operator()(
+    Halfspace3<Real> const& halfspace, Cylinder3<Real> const& cylinder)
+{
+    Result result;
+
+    // Compute extremes of signed distance Dot(N,X)-d for points on the
+    // cylinder.  These are
+    //   min = (Dot(N,C)-d) - r*sqrt(1-Dot(N,W)^2) - (h/2)*|Dot(N,W)|
+    //   max = (Dot(N,C)-d) + r*sqrt(1-Dot(N,W)^2) + (h/2)*|Dot(N,W)|
+    Real center = Dot(halfspace.normal, cylinder.axis.origin) -
+        halfspace.constant;
+    Real absNdW = std::abs(Dot(halfspace.normal, cylinder.axis.direction));
+    Real root = std::sqrt(std::max((Real)1, (Real)1 - absNdW * absNdW));
+    Real tmax = center + cylinder.radius*root +
+        ((Real)0.5)*cylinder.height*absNdW;
+
+    // The cylinder and halfspace intersect when the projection interval
+    // maximum is nonnegative.
+    result.intersect = (tmax >= (Real)0);
+    return result;
+}
+
+
+}
