move GeometryProcessing plugin from NotForLicensees to Experimental

diff --git a/Engine/Plugins/Experimental/GeometryProcessing/Source/GeometryAlgorithms/Private/GTEngine/Mathematics/GteIntrHalfspace3Ellipsoid3.h b/Engine/Plugins/Experimental/GeometryProcessing/Source/GeometryAlgorithms/Private/GTEngine/Mathematics/GteIntrHalfspace3Ellipsoid3.h
new file mode 100644
--- /dev/null
+++ b/Engine/Plugins/Experimental/GeometryProcessing/Source/GeometryAlgorithms/Private/GTEngine/Mathematics/GteIntrHalfspace3Ellipsoid3.h
@@ -0,0 +1,58 @@
+// David Eberly, Geometric Tools, Redmond WA 98052
+// Copyright (c) 1998-2019
+// Distributed under the Boost Software License, Version 1.0.
+// http://www.boost.org/LICENSE_1_0.txt
+// http://www.geometrictools.com/License/Boost/LICENSE_1_0.txt
+// File Version: 3.0.1 (2018/10/05)
+
+#pragma once
+
+#include <GTEngine/Mathematics/GteMatrix3x3.h>
+#include <GTEngine/Mathematics/GteHalfspace.h>
+#include <GTEngine/Mathematics/GteHyperellipsoid.h>
+#include <GTEngine/Mathematics/GteFIQuery.h>
+#include <GTEngine/Mathematics/GteTIQuery.h>
+
+// Queries for intersection of objects with halfspaces.  These are useful for
+// containment testing, object culling, and clipping.
+
+namespace gte
+{
+
+template <typename Real>
+class TIQuery<Real, Halfspace3<Real>, Ellipsoid3<Real>>
+{
+public:
+    struct Result
+    {
+        bool intersect;
+    };
+
+    Result operator()(Halfspace3<Real> const& halfspace,
+        Ellipsoid3<Real> const& ellipsoid);
+};
+
+
+template <typename Real>
+typename TIQuery<Real, Halfspace3<Real>, Ellipsoid3<Real>>::Result
+TIQuery<Real, Halfspace3<Real>, Ellipsoid3<Real>>::operator()(
+    Halfspace3<Real> const& halfspace, Ellipsoid3<Real> const& ellipsoid)
+{
+    // Project the ellipsoid onto the normal line.  The plane of the
+    // halfspace occurs at the origin (zero) of the normal line.
+    Result result;
+    Matrix3x3<Real> MInverse;
+    ellipsoid.GetMInverse(MInverse);
+    Real discr = Dot(halfspace.normal, MInverse * halfspace.normal);
+    Real extent = std::sqrt(std::max(discr, (Real)0));
+    Real center = Dot(halfspace.normal, ellipsoid.center) - halfspace.constant;
+    Real tmax = center + extent;
+
+    // The ellipsoid and halfspace intersect when the projection interval
+    // maximum is nonnegative.
+    result.intersect = (tmax >= (Real)0);
+    return result;
+}
+
+
+}
diff --git a/Engine/Plugins/Experimental/GeometryProcessing/Source/GeometryAlgorithms/Private/GTEngine/Mathematics/GteIntrHalfspace3OrientedBox3.h b/Engine/Plugins/Experimental/GeometryProcessing/Source/GeometryAlgorithms/Private/GTEngine/Mathematics/GteIntrHalfspace3OrientedBox3.h
new file mode 100644
--- /dev/null
+++ b/Engine/Plugins/Experimental/GeometryProcessing/Source/GeometryAlgorithms/Private/GTEngine/Mathematics/GteIntrHalfspace3OrientedBox3.h
@@ -0,0 +1,59 @@
+// David Eberly, Geometric Tools, Redmond WA 98052
+// Copyright (c) 1998-2019
+// Distributed under the Boost Software License, Version 1.0.
+// http://www.boost.org/LICENSE_1_0.txt
+// http://www.geometrictools.com/License/Boost/LICENSE_1_0.txt
+// File Version: 3.0.0 (2016/06/19)
+
+#pragma once
+
+#include <GTEngine/Mathematics/GteHalfspace.h>
+#include <GTEngine/Mathematics/GteOrientedBox.h>
+#include <GTEngine/Mathematics/GteFIQuery.h>
+#include <GTEngine/Mathematics/GteTIQuery.h>
+
+// Queries for intersection of objects with halfspaces.  These are useful for
+// containment testing, object culling, and clipping.
+
+namespace gte
+{
+
+template <typename Real>
+class TIQuery<Real, Halfspace3<Real>, OrientedBox3<Real>>
+{
+public:
+    struct Result
+    {
+        bool intersect;
+    };
+
+    Result operator()(Halfspace3<Real> const& halfspace,
+        OrientedBox3<Real> const& box);
+};
+
+
+template <typename Real>
+typename TIQuery<Real, Halfspace3<Real>, OrientedBox3<Real>>::Result
+TIQuery<Real, Halfspace3<Real>, OrientedBox3<Real>>::operator()(
+    Halfspace3<Real> const& halfspace, OrientedBox3<Real> const& box)
+{
+    Result result;
+
+    // Project the box center onto the normal line.  The plane of the
+    // halfspace occurs at the origin (zero) of the normal line.
+    Real center = Dot(halfspace.normal, box.center) - halfspace.constant;
+
+    // Compute the radius of the interval of projection.
+    Real radius =
+        std::abs(box.extent[0] * Dot(halfspace.normal, box.axis[0])) +
+        std::abs(box.extent[1] * Dot(halfspace.normal, box.axis[1])) +
+        std::abs(box.extent[2] * Dot(halfspace.normal, box.axis[2]));
+
+    // The box and halfspace intersect when the projection interval maximum
+    // is nonnegative.
+    result.intersect = (center + radius >= (Real)0);
+    return result;
+}
+
+
+}
diff --git a/Engine/Plugins/Experimental/GeometryProcessing/Source/GeometryAlgorithms/Private/GTEngine/Mathematics/GteIntrHalfspace3Segment3.h b/Engine/Plugins/Experimental/GeometryProcessing/Source/GeometryAlgorithms/Private/GTEngine/Mathematics/GteIntrHalfspace3Segment3.h
new file mode 100644
--- /dev/null
+++ b/Engine/Plugins/Experimental/GeometryProcessing/Source/GeometryAlgorithms/Private/GTEngine/Mathematics/GteIntrHalfspace3Segment3.h
@@ -0,0 +1,157 @@
+// David Eberly, Geometric Tools, Redmond WA 98052
+// Copyright (c) 1998-2019
+// Distributed under the Boost Software License, Version 1.0.
+// http://www.boost.org/LICENSE_1_0.txt
+// http://www.geometrictools.com/License/Boost/LICENSE_1_0.txt
+// File Version: 3.0.0 (2016/06/19)
+
+#pragma once
+
+#include <GTEngine/Mathematics/GteVector3.h>
+#include <GTEngine/Mathematics/GteHalfspace.h>
+#include <GTEngine/Mathematics/GteSegment.h>
+#include <GTEngine/Mathematics/GteFIQuery.h>
+#include <GTEngine/Mathematics/GteTIQuery.h>
+
+// Queries for intersection of objects with halfspaces.  These are useful for
+// containment testing, object culling, and clipping.
+
+namespace gte
+{
+
+template <typename Real>
+class TIQuery<Real, Halfspace3<Real>, Segment3<Real>>
+{
+public:
+    struct Result
+    {
+        bool intersect;
+    };
+
+    Result operator()(Halfspace3<Real> const& halfspace,
+        Segment3<Real> const& segment);
+};
+
+template <typename Real>
+class FIQuery<Real, Halfspace3<Real>, Segment3<Real>>
+{
+public:
+    struct Result
+    {
+        bool intersect;
+
+        // The segment is clipped against the plane defining the halfspace.
+        // The 'numPoints' is either 0 (no intersection), 1 (point), or 2
+        // (segment).
+        int numPoints;
+        Vector3<Real> point[2];
+    };
+
+    Result operator()(Halfspace3<Real> const& halfspace,
+        Segment3<Real> const& segment);
+};
+
+
+template <typename Real>
+typename TIQuery<Real, Halfspace3<Real>, Segment3<Real>>::Result
+TIQuery<Real, Halfspace3<Real>, Segment3<Real>>::operator()(
+    Halfspace3<Real> const& halfspace, Segment3<Real> const& segment)
+{
+    Result result;
+
+    // Project the segment endpoints onto the normal line.    The plane of the
+    // halfspace occurs at the origin (zero) of the normal line.
+    Real s[2];
+    for (int i = 0; i < 2; ++i)
+    {
+        s[i] = Dot(halfspace.normal, segment.p[i]) - halfspace.constant;
+    }
+
+    // The segment and halfspace intersect when the projection interval
+    // maximum is nonnegative.
+    result.intersect = (std::max(s[0], s[1]) >= (Real)0);
+    return result;
+}
+
+template <typename Real>
+typename FIQuery<Real, Halfspace3<Real>, Segment3<Real>>::Result
+FIQuery<Real, Halfspace3<Real>, Segment3<Real>>::operator()(
+    Halfspace3<Real> const& halfspace, Segment3<Real> const& segment)
+{
+    // Determine on which side of the plane the endpoints lie.  The table of
+    // possibilities is listed next with n = numNegative, p = numPositive, and
+    // z = numZero.
+    //
+    //   n p z  intersection
+    //   -------------------------
+    //   0 2 0  segment (original)
+    //   0 1 1  segment (original)
+    //   0 0 2  segment (original)
+    //   1 1 0  segment (clipped)
+    //   1 0 1  point (endpoint)
+    //   2 0 0  none
+
+    Real s[3];
+    int numPositive = 0, numNegative = 0, numZero = 0;
+    for (int i = 0; i < 2; ++i)
+    {
+        s[i] = Dot(halfspace.normal, segment.p[i]) - halfspace.constant;
+        if (s[i] >(Real)0)
+        {
+            ++numPositive;
+        }
+        else if (s[i] < (Real)0)
+        {
+            ++numNegative;
+        }
+        else
+        {
+            ++numZero;
+        }
+    }
+
+    Result result;
+
+    if (numNegative == 0)
+    {
+        // The segment is in the halfspace.
+        result.intersect = true;
+        result.numPoints = 2;
+        result.point[0] = segment.p[0];
+        result.point[1] = segment.p[1];
+    }
+    else if (numNegative == 1)
+    {
+        result.intersect = true;
+        result.numPoints = 1;
+        if (numPositive == 1)
+        {
+            // The segment is intersected at an interior point.
+            result.point[0] = segment.p[0] +
+                (s[0] / (s[0] - s[1]))*(segment.p[1] - segment.p[0]);
+        }
+        else  // numZero = 1
+        {
+            // One segment endpoint is on the plane.
+            if (s[0] == (Real)0)
+            {
+                result.point[0] = segment.p[0];
+            }
+            else
+            {
+                result.point[0] = segment.p[1];
+            }
+        }
+    }
+    else  // numNegative == 2
+    {
+        // The segment is outside the halfspace.  (numNegative == 2)
+        result.intersect = false;
+        result.numPoints = 0;
+    }
+
+    return result;
+}
+
+
+}
diff --git a/Engine/Plugins/Experimental/GeometryProcessing/Source/GeometryAlgorithms/Private/GTEngine/Mathematics/GteIntrHalfspace3Sphere3.h b/Engine/Plugins/Experimental/GeometryProcessing/Source/GeometryAlgorithms/Private/GTEngine/Mathematics/GteIntrHalfspace3Sphere3.h
new file mode 100644
--- /dev/null
+++ b/Engine/Plugins/Experimental/GeometryProcessing/Source/GeometryAlgorithms/Private/GTEngine/Mathematics/GteIntrHalfspace3Sphere3.h
@@ -0,0 +1,53 @@
+// David Eberly, Geometric Tools, Redmond WA 98052
+// Copyright (c) 1998-2019
+// Distributed under the Boost Software License, Version 1.0.
+// http://www.boost.org/LICENSE_1_0.txt
+// http://www.geometrictools.com/License/Boost/LICENSE_1_0.txt
+// File Version: 3.0.0 (2016/06/19)
+
+#pragma once
+
+#include <GTEngine/Mathematics/GteHalfspace.h>
+#include <GTEngine/Mathematics/GteHypersphere.h>
+#include <GTEngine/Mathematics/GteFIQuery.h>
+#include <GTEngine/Mathematics/GteTIQuery.h>
+
+// Queries for intersection of objects with halfspaces.  These are useful for
+// containment testing, object culling, and clipping.
+
+namespace gte
+{
+
+template <typename Real>
+class TIQuery<Real, Halfspace3<Real>, Sphere3<Real>>
+{
+public:
+    struct Result
+    {
+        bool intersect;
+    };
+
+    Result operator()(Halfspace3<Real> const& halfspace,
+        Sphere3<Real> const& sphere);
+};
+
+
+template <typename Real>
+typename TIQuery<Real, Halfspace3<Real>, Sphere3<Real>>::Result
+TIQuery<Real, Halfspace3<Real>, Sphere3<Real>>::operator()(
+    Halfspace3<Real> const& halfspace, Sphere3<Real> const& sphere)
+{
+    Result result;
+
+    // Project the sphere center onto the normal line.  The plane of the
+    // halfspace occurs at the origin (zero) of the normal line.
+    Real center = Dot(halfspace.normal, sphere.center) - halfspace.constant;
+
+    // The sphere and halfspace intersect when the projection interval
+    // maximum is nonnegative.
+    result.intersect = (center + sphere.radius >= (Real)0);
+    return result;
+}
+
+
+}
diff --git a/Engine/Plugins/Experimental/GeometryProcessing/Source/GeometryAlgorithms/Private/GTEngine/Mathematics/GteIntrHalfspace3Triangle3.h b/Engine/Plugins/Experimental/GeometryProcessing/Source/GeometryAlgorithms/Private/GTEngine/Mathematics/GteIntrHalfspace3Triangle3.h
new file mode 100644
--- /dev/null
+++ b/Engine/Plugins/Experimental/GeometryProcessing/Source/GeometryAlgorithms/Private/GTEngine/Mathematics/GteIntrHalfspace3Triangle3.h
@@ -0,0 +1,242 @@
+// David Eberly, Geometric Tools, Redmond WA 98052
+// Copyright (c) 1998-2019
+// Distributed under the Boost Software License, Version 1.0.
+// http://www.boost.org/LICENSE_1_0.txt
+// http://www.geometrictools.com/License/Boost/LICENSE_1_0.txt
+// File Version: 3.0.0 (2016/06/19)
+
+#pragma once
+
+#include <GTEngine/Mathematics/GteVector3.h>
+#include <GTEngine/Mathematics/GteHalfspace.h>
+#include <GTEngine/Mathematics/GteTriangle.h>
+#include <GTEngine/Mathematics/GteFIQuery.h>
+#include <GTEngine/Mathematics/GteTIQuery.h>
+
+// Queries for intersection of objects with halfspaces.  These are useful for
+// containment testing, object culling, and clipping.
+
+namespace gte
+{
+
+template <typename Real>
+class TIQuery<Real, Halfspace3<Real>, Triangle3<Real>>
+{
+public:
+    struct Result
+    {
+        bool intersect;
+    };
+
+    Result operator()(Halfspace3<Real> const& halfspace,
+        Triangle3<Real> const& triangle);
+};
+
+template <typename Real>
+class FIQuery<Real, Halfspace3<Real>, Triangle3<Real>>
+{
+public:
+    struct Result
+    {
+        bool intersect;
+
+        // The triangle is clipped against the plane defining the halfspace.
+        // The 'numPoints' is either 0 (no intersection), 1 (point), 2
+        // (segment), 3 (triangle), or 4 (quadrilateral).
+        int numPoints;
+        Vector3<Real> point[4];
+    };
+
+    Result operator()(Halfspace3<Real> const& halfspace,
+        Triangle3<Real> const& triangle);
+};
+
+
+template <typename Real>
+typename TIQuery<Real, Halfspace3<Real>, Triangle3<Real>>::Result
+TIQuery<Real, Halfspace3<Real>, Triangle3<Real>>::operator()(
+    Halfspace3<Real> const& halfspace, Triangle3<Real> const& triangle)
+{
+    Result result;
+
+    // Project the triangle vertices onto the normal line.  The plane of the
+    // halfspace occurs at the origin (zero) of the normal line.
+    Real s[3];
+    for (int i = 0; i < 3; ++i)
+    {
+        s[i] = Dot(halfspace.normal, triangle.v[i]) - halfspace.constant;
+    }
+
+    // The triangle and halfspace intersect when the projection interval
+    // maximum is nonnegative.
+    result.intersect = (std::max(std::max(s[0], s[1]), s[2]) >= (Real)0);
+    return result;
+}
+
+template <typename Real>
+typename FIQuery<Real, Halfspace3<Real>, Triangle3<Real>>::Result
+FIQuery<Real, Halfspace3<Real>, Triangle3<Real>>::operator()(
+    Halfspace3<Real> const& halfspace, Triangle3<Real> const& triangle)
+{
+    Result result;
+
+    // Determine on which side of the plane the vertices lie.  The table of
+    // possibilities is listed next with n = numNegative, p = numPositive, and
+    // z = numZero.
+    //
+    //   n p z  intersection
+    //   ---------------------------------
+    //   0 3 0  triangle (original)
+    //   0 2 1  triangle (original)
+    //   0 1 2  triangle (original)
+    //   0 0 3  triangle (original)
+    //   1 2 0  quad (2 edges clipped)
+    //   1 1 1  triangle (1 edge clipped)
+    //   1 0 2  edge
+    //   2 1 0  triangle (2 edges clipped)
+    //   2 0 1  vertex
+    //   3 0 0  none
+
+    Real s[3];
+    int numPositive = 0, numNegative = 0, numZero = 0;
+    for (int i = 0; i < 3; ++i)
+    {
+        s[i] = Dot(halfspace.normal, triangle.v[i]) - halfspace.constant;
+        if (s[i] >(Real)0)
+        {
+            ++numPositive;
+        }
+        else if (s[i] < (Real)0)
+        {
+            ++numNegative;
+        }
+        else
+        {
+            ++numZero;
+        }
+    }
+
+    if (numNegative == 0)
+    {
+        // The triangle is in the halfspace.
+        result.intersect = true;
+        result.numPoints = 3;
+        result.point[0] = triangle.v[0];
+        result.point[1] = triangle.v[1];
+        result.point[2] = triangle.v[2];
+    }
+    else if (numNegative == 1)
+    {
+        result.intersect = true;
+        if (numPositive == 2)
+        {
+            // The portion of the triangle in the halfspace is a
+            // quadrilateral.
+            result.numPoints = 4;
+            for (int i0 = 0; i0 < 3; ++i0)
+            {
+                if (s[i0] < (Real)0)
+                {
+                    int i1 = (i0 + 1) % 3, i2 = (i0 + 2) % 3;
+                    result.point[0] = triangle.v[i1];
+                    result.point[1] = triangle.v[i2];
+                    Real t2 = s[i2] / (s[i2] - s[i0]);
+                    Real t0 = s[i0] / (s[i0] - s[i1]);
+                    result.point[2] = triangle.v[i2] + t2 *
+                        (triangle.v[i0] - triangle.v[i2]);
+                    result.point[3] = triangle.v[i0] + t0 *
+                        (triangle.v[i1] - triangle.v[i0]);
+                    break;
+                }
+            }
+        }
+        else if (numPositive == 1)
+        {
+            // The portion of the triangle in the halfspace is a triangle
+            // with one vertex on the plane.
+            result.numPoints = 3;
+            for (int i0 = 0; i0 < 3; ++i0)
+            {
+                if (s[i0] == (Real)0)
+                {
+                    int i1 = (i0 + 1) % 3, i2 = (i0 + 2) % 3;
+                    result.point[0] = triangle.v[i0];
+                    Real t1 = s[i1] / (s[i1] - s[i2]);
+                    Vector3<Real> p = triangle.v[i1] + t1 *
+                        (triangle.v[i2] - triangle.v[i1]);
+                    if (s[i1] >(Real)0)
+                    {
+                        result.point[1] = triangle.v[i1];
+                        result.point[2] = p;
+                    }
+                    else
+                    {
+                        result.point[1] = p;
+                        result.point[2] = triangle.v[i2];
+                    }
+                    break;
+                }
+            }
+        }
+        else
+        {
+            // Only an edge of the triangle is in the halfspace.
+            result.numPoints = 0;
+            for (int i = 0; i < 3; ++i)
+            {
+                if (s[i] == (Real)0)
+                {
+                    result.point[result.numPoints++] = triangle.v[i];
+                }
+            }
+        }
+    }
+    else if (numNegative == 2)
+    {
+        result.intersect = true;
+        if (numPositive == 1)
+        {
+            // The portion of the triangle in the halfspace is a triangle.
+            result.numPoints = 3;
+            for (int i0 = 0; i0 < 3; ++i0)
+            {
+                if (s[i0] >(Real)0)
+                {
+                    int i1 = (i0 + 1) % 3, i2 = (i0 + 2) % 3;
+                    result.point[0] = triangle.v[i0];
+                    Real t0 = s[i0] / (s[i0] - s[i1]);
+                    Real t2 = s[i2] / (s[i2] - s[i0]);
+                    result.point[1] = triangle.v[i0] + t0 *
+                        (triangle.v[i1] - triangle.v[i0]);
+                    result.point[2] = triangle.v[i2] + t2 *
+                        (triangle.v[i0] - triangle.v[i2]);
+                    break;
+                }
+            }
+        }
+        else
+        {
+            // Only a vertex of the triangle is in the halfspace.
+            result.numPoints = 1;
+            for (int i = 0; i < 3; ++i)
+            {
+                if (s[i] == (Real)0)
+                {
+                    result.point[0] = triangle.v[i];
+                    break;
+                }
+            }
+        }
+    }
+    else  // numNegative == 3
+    {
+        // The triangle is outside the halfspace.  (numNegative == 3)
+        result.intersect = false;
+        result.numPoints = 0;
+    }
+
+    return result;
+}
+
+
+}
diff --git a/Engine/Plugins/Experimental/GeometryProcessing/Source/GeometryAlgorithms/Private/GTEngine/Mathematics/GteIntrIntervals.h b/Engine/Plugins/Experimental/GeometryProcessing/Source/GeometryAlgorithms/Private/GTEngine/Mathematics/GteIntrIntervals.h
new file mode 100644
--- /dev/null
+++ b/Engine/Plugins/Experimental/GeometryProcessing/Source/GeometryAlgorithms/Private/GTEngine/Mathematics/GteIntrIntervals.h
@@ -0,0 +1,331 @@
+// David Eberly, Geometric Tools, Redmond WA 98052
+// Copyright (c) 1998-2019
+// Distributed under the Boost Software License, Version 1.0.
+// http://www.boost.org/LICENSE_1_0.txt
+// http://www.geometrictools.com/License/Boost/LICENSE_1_0.txt
+// File Version: 3.0.0 (2016/06/19)
+
+#pragma once
+
+#include <GTEngine/Mathematics/GteFIQuery.h>
+#include <GTEngine/Mathematics/GteTIQuery.h>
+#include <array>
+#include <limits>
+
+namespace gte
+{
+
+// The intervals are [u0,u1] and [v0,v1], where u0 <= u1 and v0 <= v1, and
+// where the endpoints are any finite floating-point numbers.  Degenerate
+// intervals are allowed (u0 = u1 or v0 = v1).  The queries do not perform
+// validation on the input intervals.  In the comments, maxReal refers to
+// std::numeric_limits<Real>::max().
+
+template <typename Real>
+class TIQuery<Real, std::array<Real,2>, std::array<Real,2>>
+{
+public:
+    // The query tests overlap, whether a single point or an entire interval.
+    struct Result
+    {
+        bool intersect;
+
+        // Dynamic queries (intervals moving with constant speeds).  If
+        // 'intersect' is true, the contact times are valid and
+        //     0 <= firstTime <= lastTime,  firstTime <= maxTime
+        // If 'intersect' is false, there are two cases reported.  If the
+        // intervals will intersect at firstTime > maxTime, the contact times
+        // are reported just as when 'intersect' is true.  However, if the
+        // intervals will not intersect, then firstTime = maxReal and
+        // lastTime = -maxReal.
+        Real firstTime, lastTime;
+    };
+
+    // Static query.
+    Result operator()(std::array<Real, 2> const& interval0,
+        std::array<Real, 2> const& interval1);
+
+    // Dynamic query.  Current time is 0, maxTime > 0 is required.
+    Result operator()(Real maxTime, std::array<Real, 2> const& interval0,
+        Real speed0, std::array<Real, 2> const& interval1, Real speed1);
+};
+
+template <typename Real>
+class FIQuery<Real, std::array<Real, 2>, std::array<Real, 2>>
+{
+public:
+    // The query finds overlap, whether a single point or an entire interval.
+    struct Result
+    {
+        bool intersect;
+
+        // Static queries (no motion of intervals over time).  The number of
+        // number of intersections is 0 (no overlap), 1 (intervals are just
+        // touching), or 2 (intervals overlap in an interval).  If 'intersect'
+        // is false, numIntersections is 0 and 'overlap' is set to
+        // [maxReal,-maxReal].  If 'intersect' is true, numIntersections is
+        // 1 or 2.  When 1, 'overlap' is set to [x,x], which is degenerate and
+        // represents the single intersection point x.  When 2, 'overlap' is
+        // the interval of intersection.
+        int numIntersections;
+        std::array<Real, 2> overlap;
+
+        // Dynamic queries (intervals moving with constant speeds).  If
+        // 'intersect' is true, the contact times are valid and
+        //     0 <= firstTime <= lastTime,  firstTime <= maxTime
+        // If 'intersect' is false, there are two cases reported.  If the
+        // intervals will intersect at firstTime > maxTime, the contact times
+        // are reported just as when 'intersect' is true.  However, if the
+        // intervals will not intersect, then firstTime = maxReal and
+        // lastTime = -maxReal.
+        Real firstTime, lastTime;
+    };
+
+    // Static query.
+    Result operator()(std::array<Real, 2> const& interval0,
+        std::array<Real, 2> const& interval1);
+
+    // Dynamic query.  Current time is 0, maxTime > 0 is required.
+    Result operator()(Real maxTime, std::array<Real, 2> const& interval0,
+        Real speed0, std::array<Real, 2> const& interval1, Real speed1);
+};
+
+// Template aliases for convenience.
+template <typename Real>
+using TIIntervalInterval =
+TIQuery<Real, std::array<Real, 2>, std::array<Real, 2>>;
+
+template <typename Real>
+using FIIntervalInterval =
+FIQuery<Real, std::array<Real, 2>, std::array<Real, 2>>;
+
+
+template <typename Real>
+typename TIQuery<Real, std::array<Real, 2>, std::array<Real, 2>>::Result
+TIQuery<Real, std::array<Real, 2>, std::array<Real, 2>>::operator()(
+    std::array<Real, 2> const& interval0,
+    std::array<Real, 2> const& interval1)
+{
+    Result result;
+    result.intersect =
+        interval0[0] <= interval1[1] && interval0[1] >= interval1[0];
+    return result;
+}
+
+template <typename Real>
+typename TIQuery<Real, std::array<Real, 2>, std::array<Real, 2>>::Result
+TIQuery<Real, std::array<Real, 2>, std::array<Real, 2>>::operator()(
+    Real maxTime, std::array<Real, 2> const& interval0, Real speed0,
+    std::array<Real, 2> const& interval1, Real speed1)
+{
+    Result result;
+
+    if (interval0[1] < interval1[0])
+    {
+        // interval0 initially to the left of interval1.
+        Real diffSpeed = speed0 - speed1;
+        if (diffSpeed > (Real)0)
+        {
+            // The intervals must move towards each other.  'intersect' is
+            // true when the intervals will intersect by maxTime.
+            Real diffPos = interval1[0] - interval0[1];
+            Real invDiffSpeed = ((Real)1) / diffSpeed;
+            result.intersect = (diffPos <= maxTime*diffSpeed);
+            result.firstTime = diffPos*invDiffSpeed;
+            result.lastTime = (interval1[1] - interval0[0])*invDiffSpeed;
+            return result;
+        }
+    }
+    else if (interval0[0] > interval1[1])
+    {
+        // interval0 initially to the right of interval1.
+        Real diffSpeed = speed1 - speed0;
+        if (diffSpeed > (Real)0)
+        {
+            // The intervals must move towards each other.  'intersect' is
+            // true when the intervals will intersect by maxTime.
+            Real diffPos = interval0[0] - interval1[1];
+            Real invDiffSpeed = ((Real)1) / diffSpeed;
+            result.intersect = (diffPos <= maxTime*diffSpeed);
+            result.firstTime = diffPos*invDiffSpeed;
+            result.lastTime = (interval0[1] - interval1[0])*invDiffSpeed;
+            return result;
+        }
+    }
+    else
+    {
+        // The intervals are initially intersecting.
+        result.intersect = true;
+        result.firstTime = (Real)0;
+        if (speed1 > speed0)
+        {
+            result.lastTime = (interval0[1] - interval1[0])/(speed1 - speed0);
+        }
+        else if (speed1 < speed0)
+        {
+            result.lastTime = (interval1[1] - interval0[0])/(speed0 - speed1);
+        }
+        else
+        {
+            result.lastTime = std::numeric_limits<Real>::max();
+        }
+        return result;
+    }
+
+    result.intersect = false;
+    result.firstTime = std::numeric_limits<Real>::max();
+    result.lastTime = -std::numeric_limits<Real>::max();
+    return result;
+}
+
+
+
+template <typename Real>
+typename FIQuery<Real, std::array<Real, 2>, std::array<Real, 2>>::Result
+FIQuery<Real, std::array<Real, 2>, std::array<Real, 2>>::operator()(
+    std::array<Real, 2> const& interval0,
+    std::array<Real, 2> const& interval1)
+{
+    Result result;
+    result.firstTime = std::numeric_limits<Real>::max();
+    result.lastTime = -std::numeric_limits<Real>::max();
+
+    if (interval0[1] < interval1[0] || interval0[0] > interval1[1])
+    {
+        result.numIntersections = 0;
+        result.overlap[0] = std::numeric_limits<Real>::max();
+        result.overlap[1] = -std::numeric_limits<Real>::max();
+    }
+    else if (interval0[1] > interval1[0])
+    {
+        if (interval0[0] < interval1[1])
+        {
+            result.numIntersections = 2;
+            result.overlap[0] =
+                (interval0[0] < interval1[0] ? interval1[0] : interval0[0]);
+            result.overlap[1] =
+                (interval0[1] > interval1[1] ? interval1[1] : interval0[1]);
+            if (result.overlap[0] == result.overlap[1])
+            {
+                result.numIntersections = 1;
+            }
+        }
+        else  // interval0[0] == interval1[1]
+        {
+            result.numIntersections = 1;
+            result.overlap[0] = interval0[0];
+            result.overlap[1] = result.overlap[0];
+        }
+    }
+    else  // interval0[1] == interval1[0]
+    {
+        result.numIntersections = 1;
+        result.overlap[0] = interval0[1];
+        result.overlap[1] = result.overlap[0];
+    }
+
+    result.intersect = (result.numIntersections > 0);
+    return result;
+}
+
+template <typename Real>
+typename FIQuery<Real, std::array<Real, 2>, std::array<Real, 2>>::Result
+FIQuery<Real, std::array<Real, 2>, std::array<Real, 2>>::operator()(
+    Real maxTime, std::array<Real, 2> const& interval0, Real speed0,
+    std::array<Real, 2> const& interval1, Real speed1)
+{
+    Result result;
+
+    if (interval0[1] < interval1[0])
+    {
+        // interval0 initially to the left of interval1.
+        Real diffSpeed = speed0 - speed1;
+        if (diffSpeed > (Real)0)
+        {
+            // The intervals must move towards each other.  'intersect' is
+            // true when the intervals will intersect by maxTime.
+            Real diffPos = interval1[0] - interval0[1];
+            Real invDiffSpeed = ((Real)1) / diffSpeed;
+            result.intersect = (diffPos <= maxTime*diffSpeed);
+            result.numIntersections = 1;
+            result.firstTime = diffPos*invDiffSpeed;
+            result.lastTime = (interval1[1] - interval0[0])*invDiffSpeed;
+            result.overlap[0] = interval0[0] + result.firstTime*speed0;
+            result.overlap[1] = result.overlap[0];
+            return result;
+        }
+    }
+    else if (interval0[0] > interval1[1])
+    {
+        // interval0 initially to the right of interval1.
+        Real diffSpeed = speed1 - speed0;
+        if (diffSpeed > (Real)0)
+        {
+            // The intervals must move towards each other.  'intersect' is
+            // true when the intervals will intersect by maxTime.
+            Real diffPos = interval0[0] - interval1[1];
+            Real invDiffSpeed = ((Real)1) / diffSpeed;
+            result.intersect = (diffPos <= maxTime*diffSpeed);
+            result.numIntersections = 1;
+            result.firstTime = diffPos*invDiffSpeed;
+            result.lastTime = (interval0[1] - interval1[0])*invDiffSpeed;
+            result.overlap[0] = interval1[1] + result.firstTime*speed1;
+            result.overlap[1] = result.overlap[0];
+            return result;
+        }
+    }
+    else
+    {
+        // The intervals are initially intersecting.
+        result.intersect = true;
+        result.firstTime = (Real)0;
+        if (speed1 > speed0)
+        {
+            result.lastTime = (interval0[1] - interval1[0]) / (speed1 - speed0);
+        }
+        else if (speed1 < speed0)
+        {
+            result.lastTime = (interval1[1] - interval0[0]) / (speed0 - speed1);
+        }
+        else
+        {
+            result.lastTime = std::numeric_limits<Real>::max();
+        }
+
+        if (interval0[1] > interval1[0])
+        {
+            if (interval0[0] < interval1[1])
+            {
+                result.numIntersections = 2;
+                result.overlap[0] = (interval0[0] < interval1[0] ?
+                    interval1[0] : interval0[0]);
+                result.overlap[1] = (interval0[1] > interval1[1] ?
+                    interval1[1] : interval0[1]);
+            }
+            else  // interval0[0] == interval1[1]
+            {
+                result.numIntersections = 1;
+                result.overlap[0] = interval0[0];
+                result.overlap[1] = result.overlap[0];
+            }
+        }
+        else  // interval0[1] == interval1[0]
+        {
+            result.numIntersections = 1;
+            result.overlap[0] = interval0[1];
+            result.overlap[1] = result.overlap[0];
+        }
+        return result;
+    }
+
+    result.intersect = false;
+    result.numIntersections = 0;
+    result.overlap[0] = std::numeric_limits<Real>::max();
+    result.overlap[1] = -std::numeric_limits<Real>::max();
+    result.firstTime = std::numeric_limits<Real>::max();
+    result.lastTime = -std::numeric_limits<Real>::max();
+    return result;
+}
+
+
+}
diff --git a/Engine/Plugins/Experimental/GeometryProcessing/Source/GeometryAlgorithms/Private/GTEngine/Mathematics/GteIntrLine2AlignedBox2.h b/Engine/Plugins/Experimental/GeometryProcessing/Source/GeometryAlgorithms/Private/GTEngine/Mathematics/GteIntrLine2AlignedBox2.h
new file mode 100644
--- /dev/null
+++ b/Engine/Plugins/Experimental/GeometryProcessing/Source/GeometryAlgorithms/Private/GTEngine/Mathematics/GteIntrLine2AlignedBox2.h
@@ -0,0 +1,196 @@
+// David Eberly, Geometric Tools, Redmond WA 98052
+// Copyright (c) 1998-2019
+// Distributed under the Boost Software License, Version 1.0.
+// http://www.boost.org/LICENSE_1_0.txt
+// http://www.geometrictools.com/License/Boost/LICENSE_1_0.txt
+// File Version: 3.0.0 (2016/06/19)
+
+#pragma once
+
+#include <GTEngine/Mathematics/GteVector2.h>
+#include <GTEngine/Mathematics/GteLine.h>
+#include <GTEngine/Mathematics/GteAlignedBox.h>
+#include <GTEngine/Mathematics/GteFIQuery.h>
+#include <GTEngine/Mathematics/GteTIQuery.h>
+#include <limits>
+
+// The queries consider the box to be a solid.
+//
+// The test-intersection queries use the method of separating axes.  The
+// find-intersection queries use parametric clipping against the four edges of
+// the box.
+
+namespace gte
+{
+
+template <typename Real>
+class TIQuery<Real, Line2<Real>, AlignedBox2<Real>>
+{
+public:
+    struct Result
+    {
+        bool intersect;
+    };
+
+    Result operator()(Line2<Real> const& line, AlignedBox2<Real> const& box);
+
+protected:
+    void DoQuery(Vector2<Real> const& lineOrigin,
+        Vector2<Real> const& lineDirection, Vector2<Real> const& boxExtent,
+        Result& result);
+};
+
+template <typename Real>
+class FIQuery<Real, Line2<Real>, AlignedBox2<Real>>
+{
+public:
+    struct Result
+    {
+        bool intersect;
+        int numIntersections;
+        std::array<Real, 2> parameter;
+        std::array<Vector2<Real>, 2> point;
+    };
+
+    Result operator()(Line2<Real> const& line, AlignedBox2<Real> const& box);
+
+protected:
+    void DoQuery(Vector2<Real> const& lineOrigin,
+        Vector2<Real> const& lineDirection, Vector2<Real> const& boxExtent,
+        Result& result);
+
+private:
+    // Test whether the current clipped segment intersects the current test
+    // plane.  If the return value is 'true', the segment does intersect the
+    // plane and is clipped; otherwise, the segment is culled (no intersection
+    // with box).
+    static bool Clip(Real denom, Real numer, Real& t0, Real& t1);
+};
+
+
+template <typename Real>
+typename TIQuery<Real, Line2<Real>, AlignedBox2<Real>>::Result
+TIQuery<Real, Line2<Real>, AlignedBox2<Real>>::operator()(
+    Line2<Real> const& line, AlignedBox2<Real> const& box)
+{
+    // Get the centered form of the aligned box.  The axes are implicitly
+    // Axis[d] = Vector2<Real>::Unit(d).
+    Vector2<Real> boxCenter, boxExtent;
+    box.GetCenteredForm(boxCenter, boxExtent);
+
+    // Transform the line to the aligned-box coordinate system.
+    Vector2<Real> lineOrigin = line.origin - boxCenter;
+
+    Result result;
+    DoQuery(lineOrigin, line.direction, boxExtent, result);
+    return result;
+}
+
+template <typename Real>
+void TIQuery<Real, Line2<Real>, AlignedBox2<Real>>::DoQuery(
+    Vector2<Real> const& lineOrigin, Vector2<Real> const& lineDirection,
+    Vector2<Real> const& boxExtent, Result& result)
+{
+    Real LHS = std::abs(DotPerp(lineDirection, lineOrigin));
+    Real RHS =
+        boxExtent[0] * std::abs(lineDirection[1]) +
+        boxExtent[1] * std::abs(lineDirection[0]);
+    result.intersect = (LHS <= RHS);
+}
+
+template <typename Real>
+typename FIQuery<Real, Line2<Real>, AlignedBox2<Real>>::Result
+FIQuery<Real, Line2<Real>, AlignedBox2<Real>>::operator()(
+    Line2<Real> const& line, AlignedBox2<Real> const& box)
+{
+    // Get the centered form of the aligned box.  The axes are implicitly
+    // Axis[d] = Vector2<Real>::Unit(d).
+    Vector2<Real> boxCenter, boxExtent;
+    box.GetCenteredForm(boxCenter, boxExtent);
+
+    // Transform the line to the aligned-box coordinate system.
+    Vector2<Real> lineOrigin = line.origin - boxCenter;
+
+    Result result;
+    DoQuery(lineOrigin, line.direction, boxExtent, result);
+    for (int i = 0; i < result.numIntersections; ++i)
+    {
+        result.point[i] = line.origin + result.parameter[i] * line.direction;
+    }
+    return result;
+}
+
+template <typename Real>
+void FIQuery<Real, Line2<Real>, AlignedBox2<Real>>::DoQuery(
+    Vector2<Real> const& lineOrigin, Vector2<Real> const& lineDirection,
+    Vector2<Real> const& boxExtent, Result& result)
+{
+    // The line t-values are in the interval (-infinity,+infinity).  Clip the
+    // line against all four planes of an aligned box in centered form.  The
+    // result.numPoints is
+    //   0, no intersection
+    //   1, intersect in a single point (t0 is line parameter of point)
+    //   2, intersect in a segment (line parameter interval is [t0,t1])
+    Real t0 = -std::numeric_limits<Real>::max();
+    Real t1 = std::numeric_limits<Real>::max();
+    if (Clip(+lineDirection[0], -lineOrigin[0] - boxExtent[0], t0, t1) &&
+        Clip(-lineDirection[0], +lineOrigin[0] - boxExtent[0], t0, t1) &&
+        Clip(+lineDirection[1], -lineOrigin[1] - boxExtent[1], t0, t1) &&
+        Clip(-lineDirection[1], +lineOrigin[1] - boxExtent[1], t0, t1))
+    {
+        result.intersect = true;
+        if (t1 > t0)
+        {
+            result.numIntersections = 2;
+            result.parameter[0] = t0;
+            result.parameter[1] = t1;
+        }
+        else
+        {
+            result.numIntersections = 1;
+            result.parameter[0] = t0;
+            result.parameter[1] = t0;  // Used by derived classes.
+        }
+        return;
+    }
+
+    result.intersect = false;
+    result.numIntersections = 0;
+}
+
+template <typename Real>
+bool FIQuery<Real, Line2<Real>, AlignedBox2<Real>>::Clip(Real denom,
+    Real numer, Real& t0, Real& t1)
+{
+    if (denom > (Real)0)
+    {
+        if (numer > denom*t1)
+        {
+            return false;
+        }
+        if (numer > denom*t0)
+        {
+            t0 = numer / denom;
+        }
+        return true;
+    }
+    else if (denom < (Real)0)
+    {
+        if (numer > denom*t0)
+        {
+            return false;
+        }
+        if (numer > denom*t1)
+        {
+            t1 = numer / denom;
+        }
+        return true;
+    }
+    else
+    {
+        return numer <= (Real)0;
+    }
+}
+
+
+}
diff --git a/Engine/Plugins/Experimental/GeometryProcessing/Source/GeometryAlgorithms/Private/GTEngine/Mathematics/GteIntrLine2Arc2.h b/Engine/Plugins/Experimental/GeometryProcessing/Source/GeometryAlgorithms/Private/GTEngine/Mathematics/GteIntrLine2Arc2.h
new file mode 100644
--- /dev/null
+++ b/Engine/Plugins/Experimental/GeometryProcessing/Source/GeometryAlgorithms/Private/GTEngine/Mathematics/GteIntrLine2Arc2.h
@@ -0,0 +1,94 @@
+// David Eberly, Geometric Tools, Redmond WA 98052
+// Copyright (c) 1998-2019
+// Distributed under the Boost Software License, Version 1.0.
+// http://www.boost.org/LICENSE_1_0.txt
+// http://www.geometrictools.com/License/Boost/LICENSE_1_0.txt
+// File Version: 3.0.0 (2016/06/19)
+
+#pragma once
+
+#include <GTEngine/Mathematics/GteArc2.h>
+#include <GTEngine/Mathematics/GteIntrLine2Circle2.h>
+
+// The queries consider the arc to be a 1-dimensional object.
+
+namespace gte
+{
+
+template <typename Real>
+class TIQuery<Real, Line2<Real>, Arc2<Real>>
+{
+public:
+    struct Result
+    {
+        bool intersect;
+    };
+
+    Result operator()(Line2<Real> const& line, Arc2<Real> const& arc);
+};
+
+template <typename Real>
+class FIQuery<Real, Line2<Real>, Arc2<Real>>
+{
+public:
+    struct Result
+    {
+        bool intersect;
+        int numIntersections;
+        std::array<Real, 2> parameter;
+        std::array<Vector2<Real>, 2> point;
+    };
+
+    Result operator()(Line2<Real> const& line, Arc2<Real> const& arc);
+};
+
+
+template <typename Real>
+typename TIQuery<Real, Line2<Real>, Arc2<Real>>::Result
+TIQuery<Real, Line2<Real>, Arc2<Real>>::operator()(
+    Line2<Real> const& line, Arc2<Real> const& arc)
+{
+    Result result;
+    FIQuery<Real, Line2<Real>, Arc2<Real>> laQuery;
+    auto laResult = laQuery(line, arc);
+    result.intersect = laResult.intersect;
+    return result;
+}
+
+template <typename Real>
+typename FIQuery<Real, Line2<Real>, Arc2<Real>>::Result
+FIQuery<Real, Line2<Real>, Arc2<Real>>::operator()(
+    Line2<Real> const& line, Arc2<Real> const& arc)
+{
+    Result result;
+
+    FIQuery<Real, Line2<Real>, Circle2<Real>> lcQuery;
+    Circle2<Real> circle(arc.center, arc.radius);
+    auto lcResult = lcQuery(line, circle);
+    if (lcResult.intersect)
+    {
+        // Test whether line-circle intersections are on the arc.
+        result.numIntersections = 0;
+        for (int i = 0; i < lcResult.numIntersections; ++i)
+        {
+            if (arc.Contains(lcResult.point[i]))
+            {
+                result.intersect = true;
+                result.parameter[result.numIntersections]
+                    = lcResult.parameter[i];
+                result.point[result.numIntersections++]
+                    = lcResult.point[i];
+            }
+        }
+    }
+    else
+    {
+        result.intersect = false;
+        result.numIntersections = 0;
+    }
+
+    return result;
+}
+
+
+}
diff --git a/Engine/Plugins/Experimental/GeometryProcessing/Source/GeometryAlgorithms/Private/GTEngine/Mathematics/GteIntrLine2Circle2.h b/Engine/Plugins/Experimental/GeometryProcessing/Source/GeometryAlgorithms/Private/GTEngine/Mathematics/GteIntrLine2Circle2.h
new file mode 100644
--- /dev/null
+++ b/Engine/Plugins/Experimental/GeometryProcessing/Source/GeometryAlgorithms/Private/GTEngine/Mathematics/GteIntrLine2Circle2.h
@@ -0,0 +1,118 @@
+// David Eberly, Geometric Tools, Redmond WA 98052
+// Copyright (c) 1998-2019
+// Distributed under the Boost Software License, Version 1.0.
+// http://www.boost.org/LICENSE_1_0.txt
+// http://www.geometrictools.com/License/Boost/LICENSE_1_0.txt
+// File Version: 3.0.1 (2018/10/05)
+
+#pragma once
+
+#include <GTEngine/Mathematics/GteVector2.h>
+#include <GTEngine/Mathematics/GteDistPointLine.h>
+#include <GTEngine/Mathematics/GteHypersphere.h>
+#include <GTEngine/Mathematics/GteFIQuery.h>
+#include <GTEngine/Mathematics/GteTIQuery.h>
+
+// The queries consider the circle to be a solid (disk).
+
+namespace gte
+{
+
+template <typename Real>
+class TIQuery<Real, Line2<Real>, Circle2<Real>>
+{
+public:
+    struct Result
+    {
+        bool intersect;
+    };
+
+    Result operator()(Line2<Real> const& line, Circle2<Real> const& circle);
+};
+
+template <typename Real>
+class FIQuery<Real, Line2<Real>, Circle2<Real>>
+{
+public:
+    struct Result
+    {
+        bool intersect;
+        int numIntersections;
+        std::array<Real,2> parameter;
+        std::array<Vector2<Real>, 2> point;
+    };
+
+    Result operator()(Line2<Real> const& line, Circle2<Real> const& circle);
+
+protected:
+    void DoQuery(Vector2<Real> const& lineOrigin,
+        Vector2<Real> const& lineDirection, Circle2<Real> const& circle,
+        Result& result);
+};
+
+
+template <typename Real>
+typename TIQuery<Real, Line2<Real>, Circle2<Real>>::Result
+TIQuery<Real, Line2<Real>, Circle2<Real>>::operator()(
+    Line2<Real> const& line, Circle2<Real> const& circle)
+{
+    Result result;
+    DCPQuery<Real, Vector2<Real>, Line2<Real>> plQuery;
+    auto plResult = plQuery(circle.center, line);
+    result.intersect = (plResult.distance <= circle.radius);
+    return result;
+}
+
+template <typename Real>
+typename FIQuery<Real, Line2<Real>, Circle2<Real>>::Result
+FIQuery<Real, Line2<Real>, Circle2<Real>>::operator()(
+    Line2<Real> const& line, Circle2<Real> const& circle)
+{
+    Result result;
+    DoQuery(line.origin, line.direction, circle, result);
+    for (int i = 0; i < result.numIntersections; ++i)
+    {
+        result.point[i] = line.origin + result.parameter[i] * line.direction;
+    }
+    return result;
+}
+
+template <typename Real>
+void FIQuery<Real, Line2<Real>, Circle2<Real>>::DoQuery(
+    Vector2<Real> const& lineOrigin, Vector2<Real> const& lineDirection,
+    Circle2<Real> const& circle, Result& result)
+{
+    // Intersection of a the line P+t*D and the circle |X-C| = R.  The line
+    // direction is unit length. The t-value is a real-valued root to the
+    // quadratic equation
+    //   0 = |t*D+P-C|^2 - R^2
+    //     = t^2 + 2*Dot(D,P-C)*t + |P-C|^2-R^2
+    //     = t^2 + 2*a1*t + a0
+    // If there are two distinct roots, the order is t0 < t1.
+    Vector2<Real> diff = lineOrigin - circle.center;
+    Real a0 = Dot(diff, diff) - circle.radius * circle.radius;
+    Real a1 = Dot(lineDirection, diff);
+    Real discr = a1 * a1 - a0;
+    if (discr > (Real)0)
+    {
+        Real root = std::sqrt(discr);
+        result.intersect = true;
+        result.numIntersections = 2;
+        result.parameter[0] = -a1 - root;
+        result.parameter[1] = -a1 + root;
+    }
+    else if (discr < (Real)0)
+    {
+        result.intersect = false;
+        result.numIntersections = 0;
+    }
+    else  // discr == 0
+    {
+        result.intersect = true;
+        result.numIntersections = 1;
+        result.parameter[0] = -a1;
+    }
+}
+
+
+}
diff --git a/Engine/Plugins/Experimental/GeometryProcessing/Source/GeometryAlgorithms/Private/GTEngine/Mathematics/GteIntrLine2Line2.h b/Engine/Plugins/Experimental/GeometryProcessing/Source/GeometryAlgorithms/Private/GTEngine/Mathematics/GteIntrLine2Line2.h
new file mode 100644
--- /dev/null
+++ b/Engine/Plugins/Experimental/GeometryProcessing/Source/GeometryAlgorithms/Private/GTEngine/Mathematics/GteIntrLine2Line2.h
@@ -0,0 +1,166 @@
+// David Eberly, Geometric Tools, Redmond WA 98052
+// Copyright (c) 1998-2019
+// Distributed under the Boost Software License, Version 1.0.
+// http://www.boost.org/LICENSE_1_0.txt
+// http://www.geometrictools.com/License/Boost/LICENSE_1_0.txt
+// File Version: 3.0.0 (2016/06/19)
+
+#pragma once
+
+#include <GTEngine/Mathematics/GteVector2.h>
+#include <GTEngine/Mathematics/GteLine.h>
+#include <GTEngine/Mathematics/GteFIQuery.h>
+#include <GTEngine/Mathematics/GteTIQuery.h>
+#include <limits>
+
+namespace gte
+{
+
+template <typename Real>
+class TIQuery<Real, Line2<Real>, Line2<Real>>
+{
+public:
+    struct Result
+    {
+        bool intersect;
+
+        // The number is 0 (no intersection), 1 (lines intersect in a single
+        // point) or std::numeric_limits<int>::max() (lines are the same).
+        int numIntersections;
+    };
+
+    Result operator()(Line2<Real> const& line0, Line2<Real> const& line1);
+};
+
+template <typename Real>
+class FIQuery<Real, Line2<Real>, Line2<Real>>
+{
+public:
+    struct Result
+    {
+        bool intersect;
+
+        // The number is 0 (no intersection), 1 (lines intersect in a single
+        // point) or std::numeric_limits<int>::max() (lines are the same).
+        int numIntersections;
+
+        // If numIntersections is 1, the intersection is
+        //   point = line0.origin + line0parameter[0] * line0.direction
+        //         = line1.origin + line1parameter[0] * line1.direction
+        // If numIntersections is maxInt, point is not valid but the
+        // intervals are
+        //   line0Parameter[] = { -maxReal, +maxReal }
+        //   line1Parameter[] = { -maxReal, +maxReal }
+        Real line0Parameter[2], line1Parameter[2];
+        Vector2<Real> point;
+    };
+
+    Result operator()(Line2<Real> const& line0, Line2<Real> const& line1);
+};
+
+
+template <typename Real>
+typename TIQuery<Real, Line2<Real>, Line2<Real>>::Result
+TIQuery<Real, Line2<Real>, Line2<Real>>::operator()(
+    Line2<Real> const& line0, Line2<Real> const& line1)
+{
+    Result result;
+
+    // The intersection of two lines is a solution to P0 + s0*D0 = P1 + s1*D1.
+    // Rewrite this as s0*D0 - s1*D1 = P1 - P0 = Q.  If DotPerp(D0, D1)) = 0,
+    // the lines are parallel.  Additionally, if DotPerp(Q, D1)) = 0, the
+    // lines are the same.  If Dotperp(D0, D1)) is not zero, then
+    //   s0 = DotPerp(Q, D1))/DotPerp(D0, D1))
+    // produces the point of intersection.  Also,
+    //   s1 = DotPerp(Q, D0))/DotPerp(D0, D1))
+
+    Vector2<Real> diff = line1.origin - line0.origin;
+    Real D0DotPerpD1 = DotPerp(line0.direction, line1.direction);
+    if (D0DotPerpD1 != (Real)0)
+    {
+        // The lines are not parallel.
+        result.intersect = true;
+        result.numIntersections = 1;
+    }
+    else
+    {
+        // The lines are parallel.
+        Normalize(diff);
+        Real diffNDotPerpD1 = DotPerp(diff, line1.direction);
+        if (diffNDotPerpD1 != (Real)0)
+        {
+            // The lines are parallel but distinct.
+            result.intersect = false;
+            result.numIntersections = 0;
+        }
+        else
+        {
+            // The lines are the same.
+            result.intersect = true;
+            result.numIntersections = std::numeric_limits<int>::max();
+        }
+    }
+
+    return result;
+}
+
+template <typename Real>
+typename FIQuery<Real, Line2<Real>, Line2<Real>>::Result
+FIQuery<Real, Line2<Real>, Line2<Real>>::operator()(
+    Line2<Real> const& line0, Line2<Real> const& line1)
+{
+    Result result;
+
+    // The intersection of two lines is a solution to P0 + s0*D0 = P1 + s1*D1.
+    // Rewrite this as s0*D0 - s1*D1 = P1 - P0 = Q.  If DotPerp(D0, D1)) = 0,
+    // the lines are parallel.  Additionally, if DotPerp(Q, D1)) = 0, the
+    // lines are the same.  If Dotperp(D0, D1)) is not zero, then
+    //   s0 = DotPerp(Q, D1))/DotPerp(D0, D1))
+    // produces the point of intersection.  Also,
+    //   s1 = DotPerp(Q, D0))/DotPerp(D0, D1))
+
+    Vector2<Real> diff = line1.origin - line0.origin;
+    Real D0DotPerpD1 = DotPerp(line0.direction, line1.direction);
+    if (D0DotPerpD1 != (Real)0)
+    {
+        // The lines are not parallel.
+        result.intersect = true;
+        result.numIntersections = 1;
+        Real invD0DotPerpD1 = ((Real)1) / D0DotPerpD1;
+        Real diffDotPerpD0 = DotPerp(diff, line0.direction);
+        Real diffDotPerpD1 = DotPerp(diff, line1.direction);
+        Real s0 = diffDotPerpD1*invD0DotPerpD1;
+        Real s1 = diffDotPerpD0*invD0DotPerpD1;
+        result.line0Parameter[0] = s0;
+        result.line1Parameter[0] = s1;
+        result.point = line0.origin + s0 * line0.direction;
+    }
+    else
+    {
+        // The lines are parallel.
+        Normalize(diff);
+        Real diffNDotPerpD1 = DotPerp(diff, line1.direction);
+        if (std::abs(diffNDotPerpD1) != (Real)0)
+        {
+            // The lines are parallel but distinct.
+            result.intersect = false;
+            result.numIntersections = 0;
+        }
+        else
+        {
+            // The lines are the same.
+            result.intersect = true;
+            result.numIntersections = std::numeric_limits<int>::max();
+            Real maxReal = std::numeric_limits<Real>::max();
+            result.line0Parameter[0] = -maxReal;
+            result.line0Parameter[1] = +maxReal;
+            result.line1Parameter[0] = -maxReal;
+            result.line1Parameter[1] = +maxReal;
+        }
+    }
+
+    return result;
+}
+
+
+}
diff --git a/Engine/Plugins/Experimental/GeometryProcessing/Source/GeometryAlgorithms/Private/GTEngine/Mathematics/GteIntrLine2OrientedBox2.h b/Engine/Plugins/Experimental/GeometryProcessing/Source/GeometryAlgorithms/Private/GTEngine/Mathematics/GteIntrLine2OrientedBox2.h
new file mode 100644
--- /dev/null
+++ b/Engine/Plugins/Experimental/GeometryProcessing/Source/GeometryAlgorithms/Private/GTEngine/Mathematics/GteIntrLine2OrientedBox2.h
@@ -0,0 +1,106 @@
+// David Eberly, Geometric Tools, Redmond WA 98052
+// Copyright (c) 1998-2019
+// Distributed under the Boost Software License, Version 1.0.
+// http://www.boost.org/LICENSE_1_0.txt
+// http://www.geometrictools.com/License/Boost/LICENSE_1_0.txt
+// File Version: 3.0.0 (2016/06/19)
+
+#pragma once
+
+#include <GTEngine/Mathematics/GteOrientedBox.h>
+#include <GTEngine/Mathematics/GteIntrLine2AlignedBox2.h>
+
+// The queries consider the box to be a solid.
+//
+// The test-intersection queries use the method of separating axes.  The
+// find-intersection queries use parametric clipping against the four edges of
+// the box.
+
+namespace gte
+{
+
+template <typename Real>
+class TIQuery<Real, Line2<Real>, OrientedBox2<Real>>
+    :
+    public TIQuery<Real, Line2<Real>, AlignedBox2<Real>>
+{
+public:
+    struct Result
+        :
+        public TIQuery<Real, Line2<Real>, AlignedBox2<Real>>::Result
+    {
+        // No additional relevant information to compute.
+    };
+
+    Result operator()(Line2<Real> const& line, OrientedBox2<Real> const& box);
+};
+
+template <typename Real>
+class FIQuery<Real, Line2<Real>, OrientedBox2<Real>>
+    :
+    public FIQuery<Real, Line2<Real>, AlignedBox2<Real>>
+{
+public:
+    struct Result
+        :
+        public FIQuery<Real, Line2<Real>, AlignedBox2<Real>>::Result
+    {
+        // No additional relevant information to compute.
+    };
+
+    Result operator()(Line2<Real> const& line, OrientedBox2<Real> const& box);
+};
+
+
+template <typename Real>
+typename TIQuery<Real, Line2<Real>, OrientedBox2<Real>>::Result
+TIQuery<Real, Line2<Real>, OrientedBox2<Real>>::operator()(
+    Line2<Real> const& line, OrientedBox2<Real> const& box)
+{
+    // Transform the line to the oriented-box coordinate system.
+    Vector2<Real> diff = line.origin - box.center;
+    Vector2<Real> lineOrigin
+    {
+        Dot(diff, box.axis[0]),
+        Dot(diff, box.axis[1])
+    };
+    Vector2<Real> lineDirection
+    {
+        Dot(line.direction, box.axis[0]),
+        Dot(line.direction, box.axis[1])
+    };
+
+    Result result;
+    this->DoQuery(lineOrigin, lineDirection, box.extent, result);
+    return result;
+}
+
+template <typename Real>
+typename FIQuery<Real, Line2<Real>, OrientedBox2<Real>>::Result
+FIQuery<Real, Line2<Real>, OrientedBox2<Real>>::operator()(
+    Line2<Real> const& line, OrientedBox2<Real> const& box)
+{
+    // Transform the line to the oriented-box coordinate system.
+    Vector2<Real> diff = line.origin - box.center;
+    Vector2<Real> lineOrigin
+    {
+        Dot(diff, box.axis[0]),
+        Dot(diff, box.axis[1])
+    };
+    Vector2<Real> lineDirection
+    {
+        Dot(line.direction, box.axis[0]),
+        Dot(line.direction, box.axis[1])
+    };
+
+    Result result;
+    this->DoQuery(lineOrigin, lineDirection, box.extent, result);
+    for (int i = 0; i < result.numIntersections; ++i)
+    {
+        result.point[i] = line.origin + result.parameter[i] * line.direction;
+    }
+    return result;
+}
+
+
+}
diff --git a/Engine/Plugins/Experimental/GeometryProcessing/Source/GeometryAlgorithms/Private/GTEngine/Mathematics/GteIntrLine2Ray2.h b/Engine/Plugins/Experimental/GeometryProcessing/Source/GeometryAlgorithms/Private/GTEngine/Mathematics/GteIntrLine2Ray2.h
new file mode 100644
--- /dev/null
+++ b/Engine/Plugins/Experimental/GeometryProcessing/Source/GeometryAlgorithms/Private/GTEngine/Mathematics/GteIntrLine2Ray2.h
@@ -0,0 +1,135 @@
+// David Eberly, Geometric Tools, Redmond WA 98052
+// Copyright (c) 1998-2019
+// Distributed under the Boost Software License, Version 1.0.
+// http://www.boost.org/LICENSE_1_0.txt
+// http://www.geometrictools.com/License/Boost/LICENSE_1_0.txt
+// File Version: 3.0.0 (2016/06/19)
+
+#pragma once
+
+#include <GTEngine/Mathematics/GteRay.h>
+#include <GTEngine/Mathematics/GteIntrLine2Line2.h>
+
+namespace gte
+{
+
+template <typename Real>
+class TIQuery<Real, Line2<Real>, Ray2<Real>>
+{
+public:
+    struct Result
+    {
+        bool intersect;
+
+        // The number is 0 (no intersection), 1 (line and ray intersect in a
+        // single point) or std::numeric_limits<int>::max() (line and ray
+        // are collinear).
+        int numIntersections;
+    };
+
+    Result operator()(Line2<Real> const& line, Ray2<Real> const& ray);
+};
+
+template <typename Real>
+class FIQuery<Real, Line2<Real>, Ray2<Real>>
+{
+public:
+    struct Result
+    {
+        bool intersect;
+
+        // The number is 0 (no intersection), 1 (line and ray intersect in a
+        // single point) or std::numeric_limits<int>::max() (line and ray
+        // are collinear).
+        int numIntersections;
+
+        // If numIntersections is 1, the intersection is
+        //   point = line.origin + lineParameter[0] * line.direction
+        //         = ray.origin + rayParameter[0] * ray.direction
+        // If numIntersections is maxInt, point is not valid but the
+        // intervals are
+        //   lineParameter[] = { -maxReal, +maxReal }
+        //   rayParameter[] = { 0, +maxReal }
+        Real lineParameter[2], rayParameter[2];
+        Vector2<Real> point;
+    };
+
+    Result operator()(Line2<Real> const& line, Ray2<Real> const& ray);
+};
+
+
+template <typename Real>
+typename TIQuery<Real, Line2<Real>, Ray2<Real>>::Result
+TIQuery<Real, Line2<Real>, Ray2<Real>>::operator()(
+    Line2<Real> const& line, Ray2<Real> const& ray)
+{
+    Result result;
+    FIQuery<Real, Line2<Real>, Line2<Real>> llQuery;
+    auto llResult = llQuery(line, Line2<Real>(ray.origin, ray.direction));
+    if (llResult.numIntersections == 1)
+    {
+        // Test whether the line-line intersection is on the ray.
+        if (llResult.line1Parameter[0] >= (Real)0)
+        {
+            result.intersect = true;
+            result.numIntersections = 1;
+        }
+        else
+        {
+            result.intersect = false;
+            result.numIntersections = 0;
+        }
+    }
+    else
+    {
+        result.intersect = llResult.intersect;
+        result.numIntersections = llResult.numIntersections;
+    }
+    return result;
+}
+
+template <typename Real>
+typename FIQuery<Real, Line2<Real>, Ray2<Real>>::Result
+FIQuery<Real, Line2<Real>, Ray2<Real>>::operator()(
+    Line2<Real> const& line, Ray2<Real> const& ray)
+{
+    Result result;
+    FIQuery<Real, Line2<Real>, Line2<Real>> llQuery;
+    auto llResult = llQuery(line, Line2<Real>(ray.origin, ray.direction));
+    if (llResult.numIntersections == 1)
+    {
+        // Test whether the line-line intersection is on the ray.
+        if (llResult.line1Parameter[0] >= (Real)0)
+        {
+            result.intersect = true;
+            result.numIntersections = 1;
+            result.lineParameter[0] = llResult.line0Parameter[0];
+            result.rayParameter[0] = llResult.line1Parameter[0];
+            result.point = llResult.point;
+        }
+        else
+        {
+            result.intersect = false;
+            result.numIntersections = 0;
+        }
+    }
+    else if (llResult.numIntersections == std::numeric_limits<int>::max())
+    {
+        result.intersect = true;
+        result.numIntersections = std::numeric_limits<int>::max();
+        Real maxReal = std::numeric_limits<Real>::max();
+        result.lineParameter[0] = -maxReal;
+        result.lineParameter[1] = +maxReal;
+        result.rayParameter[0] = (Real)0;
+        result.rayParameter[1] = +maxReal;
+    }
+    else
+    {
+        result.intersect = false;
+        result.numIntersections = 0;
+    }
+    return result;
+}
+
+
+}
diff --git a/Engine/Plugins/Experimental/GeometryProcessing/Source/GeometryAlgorithms/Private/GTEngine/Mathematics/GteIntrLine2Segment2.h b/Engine/Plugins/Experimental/GeometryProcessing/Source/GeometryAlgorithms/Private/GTEngine/Mathematics/GteIntrLine2Segment2.h
new file mode 100644
--- /dev/null
+++ b/Engine/Plugins/Experimental/GeometryProcessing/Source/GeometryAlgorithms/Private/GTEngine/Mathematics/GteIntrLine2Segment2.h
@@ -0,0 +1,143 @@
+// David Eberly, Geometric Tools, Redmond WA 98052
+// Copyright (c) 1998-2019
+// Distributed under the Boost Software License, Version 1.0.
+// http://www.boost.org/LICENSE_1_0.txt
+// http://www.geometrictools.com/License/Boost/LICENSE_1_0.txt
+// File Version: 3.0.0 (2016/06/19)
+
+#pragma once
+
+#include <GTEngine/Mathematics/GteSegment.h>
+#include <GTEngine/Mathematics/GteIntrLine2Line2.h>
+
+namespace gte
+{
+
+template <typename Real>
+class TIQuery<Real, Line2<Real>, Segment2<Real>>
+{
+public:
+    struct Result
+    {
+        bool intersect;
+
+        // The number is 0 (no intersection), 1 (line and segment intersect in
+        // a single point) or std::numeric_limits<int>::max() (line and
+        // segment are collinear).
+        int numIntersections;
+    };
+
+    Result operator()(Line2<Real> const& line, Segment2<Real> const& segment);
+};
+
+template <typename Real>
+class FIQuery<Real, Line2<Real>, Segment2<Real>>
+{
+public:
+    struct Result
+    {
+        bool intersect;
+
+        // The number is 0 (no intersection), 1 (line and segment intersect in
+        // a single point) or std::numeric_limits<int>::max() (line and
+        // segment are collinear).
+        int numIntersections;
+
+        // If numIntersections is 1, the intersection is
+        //   point = line.origin + lineParameter[0] * line.direction
+        //         = segment.origin + segmentParameter[0] * segment.direction
+        // If numIntersections is maxInt, point is not valid but the
+        // intervals are
+        //   lineParameter[] = { -maxReal, +maxReal }
+        //   segmentParameter[] = { -segmentExtent, segmentExtent }
+        Real lineParameter[2], segmentParameter[2];
+        Vector2<Real> point;
+    };
+
+    Result operator()(Line2<Real> const& line, Segment2<Real> const& segment);
+};
+
+
+template <typename Real>
+typename TIQuery<Real, Line2<Real>, Segment2<Real>>::Result
+TIQuery<Real, Line2<Real>, Segment2<Real>>::operator()(
+    Line2<Real> const& line, Segment2<Real> const& segment)
+{
+    Result result;
+    Vector2<Real> segOrigin, segDirection;
+    Real segExtent;
+    segment.GetCenteredForm(segOrigin, segDirection, segExtent);
+
+    FIQuery<Real, Line2<Real>, Line2<Real>> llQuery;
+    auto llResult = llQuery(line, Line2<Real>(segOrigin, segDirection));
+    if (llResult.numIntersections == 1)
+    {
+        // Test whether the line-line intersection is on the segment.
+        if (std::abs(llResult.line1Parameter[0]) <= segExtent)
+        {
+            result.intersect = true;
+            result.numIntersections = 1;
+        }
+        else
+        {
+            result.intersect = false;
+            result.numIntersections = 0;
+        }
+    }
+    else
+    {
+        result.intersect = llResult.intersect;
+        result.numIntersections = llResult.numIntersections;
+    }
+    return result;
+}
+
+template <typename Real>
+typename FIQuery<Real, Line2<Real>, Segment2<Real>>::Result
+FIQuery<Real, Line2<Real>, Segment2<Real>>::operator()(
+    Line2<Real> const& line, Segment2<Real> const& segment)
+{
+    Result result;
+    Vector2<Real> segOrigin, segDirection;
+    Real segExtent;
+    segment.GetCenteredForm(segOrigin, segDirection, segExtent);
+
+    FIQuery<Real, Line2<Real>, Line2<Real>> llQuery;
+    auto llResult = llQuery(line, Line2<Real>(segOrigin, segDirection));
+    if (llResult.numIntersections == 1)
+    {
+        // Test whether the line-line intersection is on the ray.
+        if (std::abs(llResult.line1Parameter[0]) <= segExtent)
+        {
+            result.intersect = true;
+            result.numIntersections = 1;
+            result.lineParameter[0] = llResult.line0Parameter[0];
+            result.segmentParameter[0] = llResult.line1Parameter[0];
+            result.point = llResult.point;
+        }
+        else
+        {
+            result.intersect = false;
+            result.numIntersections = 0;
+        }
+    }
+    else if (llResult.numIntersections == std::numeric_limits<int>::max())
+    {
+        result.intersect = true;
+        result.numIntersections = std::numeric_limits<int>::max();
+        Real maxReal = std::numeric_limits<Real>::max();
+        result.lineParameter[0] = -maxReal;
+        result.lineParameter[1] = +maxReal;
+        result.segmentParameter[0] = -segExtent;
+        result.segmentParameter[1] = +segExtent;
+    }
+    else
+    {
+        result.intersect = false;
+        result.numIntersections = 0;
+    }
+    return result;
+}
+
+
+}
diff --git a/Engine/Plugins/Experimental/GeometryProcessing/Source/GeometryAlgorithms/Private/GTEngine/Mathematics/GteIntrLine2Triangle2.h b/Engine/Plugins/Experimental/GeometryProcessing/Source/GeometryAlgorithms/Private/GTEngine/Mathematics/GteIntrLine2Triangle2.h
new file mode 100644
--- /dev/null
+++ b/Engine/Plugins/Experimental/GeometryProcessing/Source/GeometryAlgorithms/Private/GTEngine/Mathematics/GteIntrLine2Triangle2.h
@@ -0,0 +1,243 @@
+// David Eberly, Geometric Tools, Redmond WA 98052
+// Copyright (c) 1998-2019
+// Distributed under the Boost Software License, Version 1.0.
+// http://www.boost.org/LICENSE_1_0.txt
+// http://www.geometrictools.com/License/Boost/LICENSE_1_0.txt
+// File Version: 3.0.0 (2016/06/19)
+
+#pragma once
+
+#include <GTEngine/Mathematics/GteVector2.h>
+#include <GTEngine/Mathematics/GteLine.h>
+#include <GTEngine/Mathematics/GteTriangle.h>
+#include <GTEngine/Mathematics/GteFIQuery.h>
+#include <GTEngine/Mathematics/GteTIQuery.h>
+
+// The queries consider the triangle to be a solid.
+
+namespace gte
+{
+
+template <typename Real>
+class TIQuery<Real, Line2<Real>, Triangle2<Real>>
+{
+public:
+    struct Result
+    {
+        bool intersect;
+    };
+
+    Result operator()(Line2<Real> const& line,
+        Triangle2<Real> const& triangle);
+};
+
+template <typename Real>
+class FIQuery<Real, Line2<Real>, Triangle2<Real>>
+{
+public:
+    struct Result
+    {
+        bool intersect;
+        int numIntersections;
+        std::array<Real, 2> parameter;
+        std::array<Vector2<Real>, 2> point;
+    };
+
+    Result operator()(Line2<Real> const& line,
+        Triangle2<Real> const& triangle);
+
+protected:
+    void DoQuery(Vector2<Real> const& lineOrigin,
+        Vector2<Real> const& lineDirection, Triangle2<Real> const& triangle,
+        Result& result);
+};
+
+
+template <typename Real>
+typename TIQuery<Real, Line2<Real>, Triangle2<Real>>::Result
+TIQuery<Real, Line2<Real>, Triangle2<Real>>::operator()(
+    Line2<Real> const& line, Triangle2<Real> const& triangle)
+{
+    Result result;
+
+    // Determine on which side of the line the vertices lie.  The table of
+    // possibilities is listed next with n = numNegative, p = numPositive, and
+    // z = numZero.
+    //
+    //   n p z  intersection
+    //   ------------------------------------
+    //   0 3 0  none
+    //   0 2 1  vertex
+    //   0 1 2  edge
+    //   0 0 3  none (degenerate triangle)
+    //   1 2 0  segment (2 edges clipped)
+    //   1 1 1  segment (1 edge clipped)
+    //   1 0 2  edge
+    //   2 1 0  segment (2 edges clipped)
+    //   2 0 1  vertex
+    //   3 0 0  none
+
+    Real s[3];
+    int numPositive = 0, numNegative = 0, numZero = 0;
+    for (int i = 0; i < 3; ++i)
+    {
+        Vector2<Real> diff = triangle.v[i] - line.origin;
+        s[i] = DotPerp(line.direction, diff);
+        if (s[i] >(Real)0)
+        {
+            ++numPositive;
+        }
+        else if (s[i] < (Real)0)
+        {
+            ++numNegative;
+        }
+        else
+        {
+            ++numZero;
+        }
+    }
+
+    result.intersect =
+        (numZero == 0 && (numPositive == 0 || numNegative == 0)) ||
+        (numZero == 3);
+
+    return result;
+}
+
+template <typename Real>
+typename FIQuery<Real, Line2<Real>, Triangle2<Real>>::Result
+FIQuery<Real, Line2<Real>, Triangle2<Real>>::operator()(
+    Line2<Real> const& line, Triangle2<Real> const& triangle)
+{
+    Result result;
+    DoQuery(line.origin, line.direction, triangle, result);
+    for (int i = 0; i < result.numIntersections; ++i)
+    {
+        result.point[i] = line.origin + result.parameter[i] * line.direction;
+    }
+    return result;
+}
+
+template <typename Real>
+void FIQuery<Real, Line2<Real>, Triangle2<Real>>::DoQuery(
+    Vector2<Real> const& lineOrigin, Vector2<Real> const& lineDirection,
+    Triangle2<Real> const& triangle, Result& result)
+{
+    // Determine on which side of the line the vertices lie.  The table of
+    // possibilities is listed next with n = numNegative, p = numPositive, and
+    // z = numZero.
+    //
+    //   n p z  intersection
+    //   ------------------------------------
+    //   0 3 0  none
+    //   0 2 1  vertex
+    //   0 1 2  edge
+    //   0 0 3  none (degenerate triangle)
+    //   1 2 0  segment (2 edges clipped)
+    //   1 1 1  segment (1 edge clipped)
+    //   1 0 2  edge
+    //   2 1 0  segment (2 edges clipped)
+    //   2 0 1  vertex
+    //   3 0 0  none
+
+    Real s[3];
+    int numPositive = 0, numNegative = 0, numZero = 0;
+    for (int i = 0; i < 3; ++i)
+    {
+        Vector2<Real> diff = triangle.v[i] - lineOrigin;
+        s[i] = DotPerp(lineDirection, diff);
+        if (s[i] >(Real)0)
+        {
+            ++numPositive;
+        }
+        else if (s[i] < (Real)0)
+        {
+            ++numNegative;
+        }
+        else
+        {
+            ++numZero;
+        }
+    }
+
+    if (numZero == 0 && numPositive > 0 && numNegative > 0)
+    {
+        result.intersect = true;
+        result.numIntersections = 2;
+        Real sign = (Real)3 - numPositive * (Real)2;
+        for (int i0 = 0; i0 < 3; ++i0)
+        {
+            if (sign * s[i0] >(Real)0)
+            {
+                int i1 = (i0 + 1) % 3, i2 = (i0 + 2) % 3;
+                Real s1 = s[i1] / (s[i1] - s[i0]);
+                Vector2<Real> p1 = (triangle.v[i1] - lineOrigin) +
+                    s1 * (triangle.v[i0] - triangle.v[i1]);
+                result.parameter[0] = Dot(lineDirection, p1);
+                Real s2 = s[i2] / (s[i2] - s[i0]);
+                Vector2<Real> p2 = (triangle.v[i2] - lineOrigin) +
+                    s2 * (triangle.v[i0] - triangle.v[i2]);
+                result.parameter[1] = Dot(lineDirection, p2);
+                break;
+            }
+        }
+        return;
+    }
+
+    if (numZero == 1)
+    {
+        result.intersect = true;
+        for (int i0 = 0; i0 < 3; ++i0)
+        {
+            if (s[i0] == (Real)0)
+            {
+                int i1 = (i0 + 1) % 3, i2 = (i0 + 2) % 3;
+                result.parameter[0] =
+                    Dot(lineDirection, triangle.v[i0] - lineOrigin);
+                if (numPositive == 2 || numNegative == 2)
+                {
+                    result.numIntersections = 1;
+
+                    // Used by derived classes.
+                    result.parameter[1] = result.parameter[0];
+                }
+                else
+                {
+                    result.numIntersections = 2;
+                    Real s1 = s[i1] / (s[i1] - s[i2]);
+                    Vector2<Real> p1 = (triangle.v[i1] - lineOrigin) +
+                        s1 * (triangle.v[i2] - triangle.v[i1]);
+                    result.parameter[1] = Dot(lineDirection, p1);
+                }
+                break;
+            }
+        }
+        return;
+    }
+
+    if (numZero == 2)
+    {
+        result.intersect = true;
+        result.numIntersections = 2;
+        for (int i0 = 0; i0 < 3; ++i0)
+        {
+            if (s[i0] != (Real)0)
+            {
+                int i1 = (i0 + 1) % 3, i2 = (i0 + 2) % 3;
+                result.parameter[0] =
+                    Dot(lineDirection, triangle.v[i1] - lineOrigin);
+                result.parameter[1] =
+                    Dot(lineDirection, triangle.v[i2] - lineOrigin);
+                break;
+            }
+        }
+        return;
+    }
+
+    // (n,p,z) one of (3,0,0), (0,3,0), (0,0,3)
+    result.intersect = false;
+    result.numIntersections = 0;
+}
+
+
+}
diff --git a/Engine/Plugins/Experimental/GeometryProcessing/Source/GeometryAlgorithms/Private/GTEngine/Mathematics/GteIntrLine3AlignedBox3.h b/Engine/Plugins/Experimental/GeometryProcessing/Source/GeometryAlgorithms/Private/GTEngine/Mathematics/GteIntrLine3AlignedBox3.h
new file mode 100644
--- /dev/null
+++ b/Engine/Plugins/Experimental/GeometryProcessing/Source/GeometryAlgorithms/Private/GTEngine/Mathematics/GteIntrLine3AlignedBox3.h
@@ -0,0 +1,190 @@
+// David Eberly, Geometric Tools, Redmond WA 98052
+// Copyright (c) 1998-2019
+// Distributed under the Boost Software License, Version 1.0.
+// http://www.boost.org/LICENSE_1_0.txt
+// http://www.geometrictools.com/License/Boost/LICENSE_1_0.txt
+// File Version: 3.0.1 (2018/11/28)
+
+#pragma once
+
+#include <GTEngine/Mathematics/GteVector3.h>
+#include <GTEngine/Mathematics/GteAlignedBox.h>
+#include <GTEngine/Mathematics/GteLine.h>
+#include <GTEngine/Mathematics/GteFIQuery.h>
+#include <GTEngine/Mathematics/GteTIQuery.h>
+#include <limits>
+
+// The test-intersection queries use the method of separating axes.  The
+// find-intersection queries use parametric clipping against the six faces of
+// the box.  The find-intersection queries use Liang-Barsky clipping.  The
+// queries consider the box to be a solid.  The algorithms are described in
+//   https://www.geometrictools.com/Documentation/IntersectionLineBox.pdf
+
+namespace gte
+{
+    template <typename Real>
+    class TIQuery<Real, Line3<Real>, AlignedBox3<Real>>
+    {
+    public:
+        struct Result
+        {
+            bool intersect;
+        };
+
+        Result operator()(Line3<Real> const& line, AlignedBox3<Real> const& box)
+        {
+            // Get the centered form of the aligned box.  The axes are implicitly
+            // Axis[d] = Vector3<Real>::Unit(d).
+            Vector3<Real> boxCenter, boxExtent;
+            box.GetCenteredForm(boxCenter, boxExtent);
+
+            // Transform the line to the aligned-box coordinate system.
+            Vector3<Real> lineOrigin = line.origin - boxCenter;
+
+            Result result;
+            DoQuery(lineOrigin, line.direction, boxExtent, result);
+            return result;
+        }
+
+    protected:
+        void DoQuery(Vector3<Real> const& lineOrigin, Vector3<Real> const& lineDirection,
+            Vector3<Real> const& boxExtent, Result& result)
+        {
+            Vector3<Real> WxD = Cross(lineDirection, lineOrigin);
+            Real absWdU[3] =
+            {
+                std::abs(lineDirection[0]),
+                std::abs(lineDirection[1]),
+                std::abs(lineDirection[2])
+            };
+
+            if (std::abs(WxD[0]) > boxExtent[1] * absWdU[2] + boxExtent[2] * absWdU[1])
+            {
+                result.intersect = false;
+                return;
+            }
+
+            if (std::abs(WxD[1]) > boxExtent[0] * absWdU[2] + boxExtent[2] * absWdU[0])
+            {
+                result.intersect = false;
+                return;
+            }
+
+            if (std::abs(WxD[2]) > boxExtent[0] * absWdU[1] + boxExtent[1] * absWdU[0])
+            {
+                result.intersect = false;
+                return;
+            }
+
+            result.intersect = true;
+        }
+    };
+
+    template <typename Real>
+    class FIQuery<Real, Line3<Real>, AlignedBox3<Real>>
+    {
+    public:
+        struct Result
+        {
+            bool intersect;
+            int numPoints;
+            Real lineParameter[2];
+            Vector3<Real> point[2];
+        };
+
+        Result operator()(Line3<Real> const& line, AlignedBox3<Real> const& box)
+        {
+            // Get the centered form of the aligned box.  The axes are
+            // implicitly Axis[d] = Vector3<Real>::Unit(d).
+            Vector3<Real> boxCenter, boxExtent;
+            box.GetCenteredForm(boxCenter, boxExtent);
+
+            // Transform the line to the aligned-box coordinate system.
+            Vector3<Real> lineOrigin = line.origin - boxCenter;
+
+            Result result;
+            DoQuery(lineOrigin, line.direction, boxExtent, result);
+            for (int i = 0; i < result.numPoints; ++i)
+            {
+                result.point[i] = line.origin + result.lineParameter[i] * line.direction;
+            }
+            return result;
+        }
+
+    protected:
+        void DoQuery(Vector3<Real> const& lineOrigin, Vector3<Real> const& lineDirection,
+            Vector3<Real> const& boxExtent, Result& result)
+        {
+            // The line t-values are in the interval (-infinity,+infinity).
+            // Clip the line against all six planes of an aligned box in
+            // centered form.  The result.numPoints is
+            //   0, no intersection
+            //   1, intersect in a single point (t0 is line parameter of point)
+            //   2, intersect in a segment (line parameter interval is [t0,t1])
+            Real t0 = -std::numeric_limits<Real>::max();
+            Real t1 = std::numeric_limits<Real>::max();
+            if (Clip(+lineDirection[0], -lineOrigin[0] - boxExtent[0], t0, t1) &&
+                Clip(-lineDirection[0], +lineOrigin[0] - boxExtent[0], t0, t1) &&
+                Clip(+lineDirection[1], -lineOrigin[1] - boxExtent[1], t0, t1) &&
+                Clip(-lineDirection[1], +lineOrigin[1] - boxExtent[1], t0, t1) &&
+                Clip(+lineDirection[2], -lineOrigin[2] - boxExtent[2], t0, t1) &&
+                Clip(-lineDirection[2], +lineOrigin[2] - boxExtent[2], t0, t1))
+            {
+                result.intersect = true;
+                if (t1 > t0)
+                {
+                    result.numPoints = 2;
+                    result.lineParameter[0] = t0;
+                    result.lineParameter[1] = t1;
+                }
+                else
+                {
+                    result.numPoints = 1;
+                    result.lineParameter[0] = t0;
+                    result.lineParameter[1] = t0;  // Used by derived classes.
+                }
+                return;
+            }
+
+            result.intersect = false;
+            result.numPoints = 0;
+        }
+
+    private:
+        // Test whether the current clipped segment intersects the current
+        // test plane.  If the return value is 'true', the segment does
+        // intersect the plane and is clipped; otherwise, the segment is
+        // culled (no intersection with box).
+        static bool Clip(Real denom, Real numer, Real& t0, Real& t1)
+        {
+            if (denom > (Real)0)
+            {
+                if (numer > denom * t1)
+                {
+                    return false;
+                }
+                if (numer > denom * t0)
+                {
+                    t0 = numer / denom;
+                }
+                return true;
+            }
+            else if (denom < (Real)0)
+            {
+                if (numer > denom * t0)
+                {
+                    return false;
+                }
+                if (numer > denom * t1)
+                {
+                    t1 = numer / denom;
+                }
+                return true;
+            }
+            else
+            {
+                return numer <= (Real)0;
+            }
+        }
+    };
+}
diff --git a/Engine/Plugins/Experimental/GeometryProcessing/Source/GeometryAlgorithms/Private/GTEngine/Mathematics/GteIntrLine3Capsule3.h b/Engine/Plugins/Experimental/GeometryProcessing/Source/GeometryAlgorithms/Private/GTEngine/Mathematics/GteIntrLine3Capsule3.h
new file mode 100644
--- /dev/null
+++ b/Engine/Plugins/Experimental/GeometryProcessing/Source/GeometryAlgorithms/Private/GTEngine/Mathematics/GteIntrLine3Capsule3.h
@@ -0,0 +1,341 @@
+// David Eberly, Geometric Tools, Redmond WA 98052
+// Copyright (c) 1998-2019
+// Distributed under the Boost Software License, Version 1.0.
+// http://www.boost.org/LICENSE_1_0.txt
+// http://www.geometrictools.com/License/Boost/LICENSE_1_0.txt
+// File Version: 3.0.1 (2018/10/05)
+
+#pragma once
+
+#include <GTEngine/Mathematics/GteVector3.h>
+#include <GTEngine/Mathematics/GteDistLineSegment.h>
+#include <GTEngine/Mathematics/GteCapsule.h>
+#include <GTEngine/Mathematics/GteFIQuery.h>
+#include <GTEngine/Mathematics/GteTIQuery.h>
+
+// The queries consider the capsule to be a solid.
+//
+// The test-intersection queries are based on distance computations.
+
+namespace gte
+{
+
+template <typename Real>
+class TIQuery<Real, Line3<Real>, Capsule3<Real>>
+{
+public:
+    struct Result
+    {
+        bool intersect;
+    };
+
+    Result operator()(Line3<Real> const& line, Capsule3<Real> const& capsule);
+};
+
+template <typename Real>
+class FIQuery<Real, Line3<Real>, Capsule3<Real>>
+{
+public:
+    struct Result
+    {
+        bool intersect;
+        int numIntersections;
+        std::array<Real, 2> parameter;
+        std::array<Vector3<Real>, 2> point;
+    };
+
+    Result operator()(Line3<Real> const& line, Capsule3<Real> const& capsule);
+
+protected:
+    void DoQuery(Vector3<Real> const& lineOrigin,
+        Vector3<Real> const& lineDirection, Capsule3<Real> const& capsule,
+        Result& result);
+};
+
+
+template <typename Real>
+typename TIQuery<Real, Line3<Real>, Capsule3<Real>>::Result
+TIQuery<Real, Line3<Real>, Capsule3<Real>>::operator()(
+    Line3<Real> const& line, Capsule3<Real> const& capsule)
+{
+    Result result;
+    DCPQuery<Real, Line3<Real>, Segment3<Real>> lsQuery;
+    auto lsResult = lsQuery(line, capsule.segment);
+    result.intersect = (lsResult.distance <= capsule.radius);
+    return result;
+}
+
+template <typename Real>
+typename FIQuery<Real, Line3<Real>, Capsule3<Real>>::Result
+FIQuery<Real, Line3<Real>, Capsule3<Real>>::operator()(
+    Line3<Real> const& line, Capsule3<Real> const& capsule)
+{
+    Result result;
+    DoQuery(line.origin, line.direction, capsule, result);
+    for (int i = 0; i < result.numIntersections; ++i)
+    {
+        result.point[i] = line.origin + result.parameter[i] * line.direction;
+    }
+    return result;
+}
+
+template <typename Real>
+void FIQuery<Real, Line3<Real>, Capsule3<Real>>::DoQuery(
+    Vector3<Real> const& lineOrigin, Vector3<Real> const& lineDirection,
+    Capsule3<Real> const& capsule, Result& result)
+{
+    // Initialize the result as if there is no intersection.  If we discover
+    // an intersection, these values will be modified accordingly.
+    result.intersect = false;
+    result.numIntersections = 0;
+
+    // Create a coordinate system for the capsule.  In this system, the
+    // capsule segment center C is the origin and the capsule axis direction
+    // W is the z-axis.  U and V are the other coordinate axis directions.
+    // If P = x*U+y*V+z*W, the cylinder containing the capsule wall is
+    // x^2 + y^2 = r^2, where r is the capsule radius.  The finite cylinder
+    // that makes up the capsule minus its hemispherical end caps has z-values
+    // |z| <= e, where e is the extent of the capsule segment.  The top
+    // hemisphere cap is x^2+y^2+(z-e)^2 = r^2 for z >= e, and the bottom
+    // hemisphere cap is x^2+y^2+(z+e)^2 = r^2 for z <= -e.
+
+    Vector3<Real> segOrigin, segDirection;
+    Real segExtent;
+    capsule.segment.GetCenteredForm(segOrigin, segDirection, segExtent);
+    Vector3<Real> basis[3];  // {W, U, V}
+    basis[0] = segDirection;
+    ComputeOrthogonalComplement(1, basis);
+    Real rSqr = capsule.radius * capsule.radius;
+
+    // Convert incoming line origin to capsule coordinates.
+    Vector3<Real> diff = lineOrigin - segOrigin;
+    Vector3<Real> P{ Dot(basis[1], diff), Dot(basis[2], diff),
+        Dot(basis[0], diff) };
+
+    // Get the z-value, in capsule coordinates, of the incoming line's
+    // unit-length direction.
+    Real dz = Dot(basis[0], lineDirection);
+    if (std::abs(dz) == (Real)1)
+    {
+        // The line is parallel to the capsule axis.  Determine whether the
+        // line intersects the capsule hemispheres.
+        Real radialSqrDist = rSqr - P[0] * P[0] - P[1] * P[1];
+        if (radialSqrDist >= (Real)0)
+        {
+            // The line intersects the hemispherical caps.
+            result.intersect = true;
+            result.numIntersections = 2;
+            Real zOffset = std::sqrt(radialSqrDist) + segExtent;
+            if (dz > (Real)0)
+            {
+                result.parameter[0] = -P[2] - zOffset;
+                result.parameter[1] = -P[2] + zOffset;
+            }
+            else
+            {
+                result.parameter[0] = P[2] - zOffset;
+                result.parameter[1] = P[2] + zOffset;
+            }
+        }
+        // else: The line outside the capsule's cylinder, no intersection.
+        return;
+    }
+
+    // Convert the incoming line unit-length direction to capsule coordinates.
+    Vector3<Real> D{ Dot(basis[1], lineDirection),
+        Dot(basis[2], lineDirection), dz };
+
+    // Test intersection of line P+t*D with infinite cylinder x^2+y^2 = r^2.
+    // This reduces to computing the roots of a quadratic equation.  If
+    // P = (px,py,pz) and D = (dx,dy,dz), then the quadratic equation is
+    //   (dx^2+dy^2)*t^2 + 2*(px*dx+py*dy)*t + (px^2+py^2-r^2) = 0
+    Real a0 = P[0] * P[0] + P[1] * P[1] - rSqr;
+    Real a1 = P[0] * D[0] + P[1] * D[1];
+    Real a2 = D[0] * D[0] + D[1] * D[1];
+    Real discr = a1 * a1 - a0 * a2;
+    if (discr < (Real)0)
+    {
+        // The line does not intersect the infinite cylinder, so it
+        // cannot intersect the capsule.
+        return;
+    }
+
+    Real root, inv, tValue, zValue;
+    if (discr >(Real)0)
+    {
+        // The line intersects the infinite cylinder in two places.
+        root = std::sqrt(discr);
+        inv = ((Real)1) / a2;
+        tValue = (-a1 - root) * inv;
+        zValue = P[2] + tValue * D[2];
+        if (std::abs(zValue) <= segExtent)
+        {
+            result.intersect = true;
+            result.parameter[result.numIntersections++] = tValue;
+        }
+
+        tValue = (-a1 + root) * inv;
+        zValue = P[2] + tValue * D[2];
+        if (std::abs(zValue) <= segExtent)
+        {
+            result.intersect = true;
+            result.parameter[result.numIntersections++] = tValue;
+        }
+
+        if (result.numIntersections == 2)
+        {
+            // The line intersects the capsule wall in two places.
+            return;
+        }
+    }
+    else
+    {
+        // The line is tangent to the infinite cylinder but intersects the
+        // cylinder in a single point.
+        tValue = -a1 / a2;
+        zValue = P[2] + tValue * D[2];
+        if (std::abs(zValue) <= segExtent)
+        {
+            result.intersect = true;
+            result.numIntersections = 1;
+            result.parameter[0] = tValue;
+            // Used by derived classes.
+            result.parameter[1] = result.parameter[0];
+            return;
+        }
+    }
+
+    // Test intersection with bottom hemisphere.  The quadratic equation is
+    //   t^2 + 2*(px*dx+py*dy+(pz+e)*dz)*t + (px^2+py^2+(pz+e)^2-r^2) = 0
+    // Use the fact that currently a1 = px*dx+py*dy and a0 = px^2+py^2-r^2.
+    // The leading coefficient is a2 = 1, so no need to include in the
+    // construction.
+    Real PZpE = P[2] + segExtent;
+    a1 += PZpE * D[2];
+    a0 += PZpE * PZpE;
+    discr = a1 * a1 - a0;
+    if (discr > (Real)0)
+    {
+        root = std::sqrt(discr);
+        tValue = -a1 - root;
+        zValue = P[2] + tValue * D[2];
+        if (zValue <= -segExtent)
+        {
+            result.parameter[result.numIntersections++] = tValue;
+            if (result.numIntersections == 2)
+            {
+                result.intersect = true;
+                if (result.parameter[0] > result.parameter[1])
+                {
+                    std::swap(result.parameter[0], result.parameter[1]);
+                }
+                return;
+            }
+        }
+
+        tValue = -a1 + root;
+        zValue = P[2] + tValue * D[2];
+        if (zValue <= -segExtent)
+        {
+            result.parameter[result.numIntersections++] = tValue;
+            if (result.numIntersections == 2)
+            {
+                result.intersect = true;
+                if (result.parameter[0] > result.parameter[1])
+                {
+                    std::swap(result.parameter[0], result.parameter[1]);
+                }
+                return;
+            }
+        }
+    }
+    else if (discr == (Real)0)
+    {
+        tValue = -a1;
+        zValue = P[2] + tValue * D[2];
+        if (zValue <= -segExtent)
+        {
+            result.parameter[result.numIntersections++] = tValue;
+            if (result.numIntersections == 2)
+            {
+                result.intersect = true;
+                if (result.parameter[0] > result.parameter[1])
+                {
+                    std::swap(result.parameter[0], result.parameter[1]);
+                }
+                return;
+            }
+        }
+    }
+
+    // Test intersection with top hemisphere.  The quadratic equation is
+    //   t^2 + 2*(px*dx+py*dy+(pz-e)*dz)*t + (px^2+py^2+(pz-e)^2-r^2) = 0
+    // Use the fact that currently a1 = px*dx+py*dy+(pz+e)*dz and
+    // a0 = px^2+py^2+(pz+e)^2-r^2.  The leading coefficient is a2 = 1, so
+    // no need to include in the construction.
+    a1 -= ((Real)2) * segExtent * D[2];
+    a0 -= ((Real)4) * segExtent * P[2];
+    discr = a1 * a1 - a0;
+    if (discr > (Real)0)
+    {
+        root = std::sqrt(discr);
+        tValue = -a1 - root;
+        zValue = P[2] + tValue * D[2];
+        if (zValue >= segExtent)
+        {
+            result.parameter[result.numIntersections++] = tValue;
+            if (result.numIntersections == 2)
+            {
+                result.intersect = true;
+                if (result.parameter[0] > result.parameter[1])
+                {
+                    std::swap(result.parameter[0], result.parameter[1]);
+                }
+                return;
+            }
+        }
+
+        tValue = -a1 + root;
+        zValue = P[2] + tValue * D[2];
+        if (zValue >= segExtent)
+        {
+            result.parameter[result.numIntersections++] = tValue;
+            if (result.numIntersections == 2)
+            {
+                result.intersect = true;
+                if (result.parameter[0] > result.parameter[1])
+                {
+                    std::swap(result.parameter[0], result.parameter[1]);
+                }
+                return;
+            }
+        }
+    }
+    else if (discr == (Real)0)
+    {
+        tValue = -a1;
+        zValue = P[2] + tValue * D[2];
+        if (zValue >= segExtent)
+        {
+            result.parameter[result.numIntersections++] = tValue;
+            if (result.numIntersections == 2)
+            {
+                result.intersect = true;
+                if (result.parameter[0] > result.parameter[1])
+                {
+                    std::swap(result.parameter[0], result.parameter[1]);
+                }
+                return;
+            }
+        }
+    }
+
+    if (result.numIntersections == 1)
+    {
+        // Used by derived classes.
+        result.parameter[1] = result.parameter[0];
+    }
+}
+
+
+}
