move GeometryProcessing plugin from NotForLicensees to Experimental

diff --git a/Engine/Plugins/Experimental/GeometryProcessing/Source/GeometryAlgorithms/Private/GTEngine/Mathematics/GteIntrLine3Cone3.h b/Engine/Plugins/Experimental/GeometryProcessing/Source/GeometryAlgorithms/Private/GTEngine/Mathematics/GteIntrLine3Cone3.h
new file mode 100644
--- /dev/null
+++ b/Engine/Plugins/Experimental/GeometryProcessing/Source/GeometryAlgorithms/Private/GTEngine/Mathematics/GteIntrLine3Cone3.h
@@ -0,0 +1,320 @@
+// David Eberly, Geometric Tools, Redmond WA 98052
+// Copyright (c) 1998-2019
+// Distributed under the Boost Software License, Version 1.0.
+// http://www.boost.org/LICENSE_1_0.txt
+// http://www.geometrictools.com/License/Boost/LICENSE_1_0.txt
+// File Version: 3.0.3 (2018/11/29)
+
+#pragma once
+
+#include <GTEngine/Mathematics/GteVector3.h>
+#include <GTEngine/Mathematics/GteCone.h>
+#include <GTEngine/Mathematics/GteLine.h>
+#include <GTEngine/Mathematics/GteIntrIntervals.h>
+
+// The queries consider the cone to be single sided and solid.
+
+namespace gte
+{
+
+template <typename Real>
+class FIQuery<Real, Line3<Real>, Cone3<Real>>
+{
+public:
+    struct Result
+    {
+        // Default construction.  Set 'intersect' to false, 'type' to 0,
+        // 'parameter' to (0,0), and 'point' to (veczero,veczero).
+        Result();
+
+        bool intersect;
+
+        // Because the intersection of line and cone with infinite height
+        // h > 0 can be a ray or a line, we use a 'type' value that allows
+        // you to decide how to interpret the parameter[] and point[] values.
+        //   type  intersect  valid data
+        //   0     none       none
+        //   1     point      parameter[0] = parameter[1], finite
+        //                    point[0] = point[1]
+        //   2     segment    parameter[0] < parameter[1], finite
+        //                    point[0,1] valid
+        //   3     ray        parameter[0] finite, parameter[1] maxReal
+        //                    point[0] = rayOrigin, point[1] = lineDirection
+        //   4     ray        parameter[0] -maxReal, parameter[1] finite
+        //                    point[0] = rayOrigin, point[1] = -lineDirection
+        //   5     line       parameter[0] -maxReal, parameter[1] maxReal,
+        //                    point[0] = lineOrigin, point[1] = lineDirection
+        // If the cone height h is finite, only types 0, 1, or 2 can occur.
+        int type;
+        std::array<Real, 2> parameter;  // Relative to incoming line.
+        std::array<Vector3<Real>, 2> point;
+    };
+
+    Result operator()(Line3<Real> const& line, Cone3<Real> const& cone);
+
+protected:
+    void DoQuery(Vector3<Real> const& lineOrigin,
+        Vector3<Real> const& lineDirection, Cone3<Real> const& cone,
+        Result& result);
+};
+
+
+template <typename Real>
+FIQuery<Real, Line3<Real>, Cone3<Real>>::Result::Result()
+    :
+    intersect(false),
+    type(0)
+{
+    parameter.fill((Real)0);
+    point.fill({ (Real)0, (Real)0, (Real)0 });
+}
+
+template <typename Real>
+typename FIQuery<Real, Line3<Real>, Cone3<Real>>::Result
+FIQuery<Real, Line3<Real>, Cone3<Real>>::operator()(Line3<Real> const& line,
+    Cone3<Real> const& cone)
+{
+    Result result;
+    DoQuery(line.origin, line.direction, cone, result);
+    switch (result.type)
+    {
+    case 1:  // point
+        result.point[0] = line.origin + result.parameter[0] * line.direction;
+        result.point[1] = result.point[0];
+        break;
+    case 2:  // segment
+        result.point[0] = line.origin + result.parameter[0] * line.direction;
+        result.point[1] = line.origin + result.parameter[1] * line.direction;
+        break;
+    case 3:  // ray
+        result.point[0] = line.origin + result.parameter[0] * line.direction;
+        result.point[1] = line.direction;
+        break;
+    case 4:  // ray
+        result.point[0] = line.origin + result.parameter[1] * line.direction;
+        result.point[1] = -line.direction;
+        break;
+    case 5:  // line
+        result.point[0] = line.origin;
+        result.point[1] = line.direction;
+        break;
+    default:  // no intersection
+        break;
+    }
+    return result;
+}
+
+template <typename Real>
+void FIQuery<Real, Line3<Real>, Cone3<Real>>::DoQuery(
+    Vector3<Real> const& lineOrigin, Vector3<Real> const& lineDirection,
+    Cone3<Real> const& cone, Result& result)
+{
+    // The cone has vertex V, unit-length axis direction D, angle theta in
+    // (0,pi/2), and height h in (0,+infinity).  The line is P + t*U, where U
+    // is a unit-length direction vector.  Define g = cos(theta).  The cone
+    // is represented by
+    //   (X-V)^T * (D*D^T - g^2*I) * (X-V) = 0,  0 <= Dot(D,X-V) <= h
+    // The first equation defines a double-sided cone.  The first inequality
+    // in the second equation limits this to a single-sided cone containing
+    // the ray V + s*D with s >= 0.  We will call this the 'positive cone'.
+    // The single-sided cone containing ray V + s * t with s <= 0 is called
+    // the 'negative cone'.  The double-sided cone is the union of the
+    // positive cone and negative cone.  The second inequality in the second
+    // equation limits the single-sided cone to the region bounded by the
+    // height.  Setting X(t) = P + t*U, the equations are
+    //   c2*t^2 + 2*c1*t + c0 = 0,  0 <= Dot(D,U)*t + Dot(D,P-V) <= h
+    // where
+    //   c2 = Dot(D,U)^2 - g^2
+    //   c1 = Dot(D,U)*Dot(D,P-V) - g^2*Dot(U,P-V)
+    //   c0 = Dot(D,P-V)^2 - g^2*Dot(P-V,P-V)
+    // The following code computes the t-interval that satisfies the quadratic
+    // equation subject to the linear inequality constraints.
+
+    Vector3<Real> PmV = lineOrigin - cone.ray.origin;
+    Real DdU = Dot(cone.ray.direction, lineDirection);
+    Real DdPmV = Dot(cone.ray.direction, PmV);
+    Real UdPmV = Dot(lineDirection, PmV);
+    Real PmVdPmV = Dot(PmV, PmV);
+    Real cosAngleSqr = cone.cosAngle * cone.cosAngle;
+    Real c2 = DdU * DdU - cosAngleSqr;
+    Real c1 = DdU * DdPmV - cosAngleSqr * UdPmV;
+    Real c0 = DdPmV * DdPmV - cosAngleSqr * PmVdPmV;
+    Real t;
+
+    if (c2 != (Real)0)
+    {
+        Real discr = c1 * c1 - c0 * c2;
+        if (discr < (Real)0)
+        {
+            // The quadratic has no real-valued roots.  The line does not
+            // intersect the double-sided cone.
+            result.intersect = false;
+            result.type = 0;
+            return;
+        }
+        else if (discr > (Real)0)
+        {
+            // The quadratic has two distinct real-valued roots.  However, one
+            // or both of them might intersect the negative cone.  We are
+            // interested only in those intersections with the positive cone.
+            Real root = std::sqrt(discr);
+            Real invC2 = ((Real)1) / c2;
+            int numParameters = 0;
+
+            t = (-c1 - root) * invC2;
+            if (DdU * t + DdPmV >= (Real)0)
+            {
+                result.parameter[numParameters++] = t;
+            }
+
+            t = (-c1 + root) * invC2;
+            if (DdU * t + DdPmV >= (Real)0)
+            {
+                result.parameter[numParameters++] = t;
+            }
+
+            if (numParameters == 2)
+            {
+                // The line intersects the positive cone in two distinct
+                // points.
+                result.intersect = true;
+                result.type = 2;
+                if (result.parameter[0] > result.parameter[1])
+                {
+                    std::swap(result.parameter[0], result.parameter[1]);
+                }
+            }
+            else if (numParameters == 1)
+            {
+                // The line intersects the positive cone in a single point and
+                // the negative cone in a single point.  We report only the
+                // intersection with the positive cone.
+                result.intersect = true;
+                if (DdU > (Real)0)
+                {
+                    result.type = 3;
+                    result.parameter[1] = std::numeric_limits<Real>::max();
+                }
+                else
+                {
+                    result.type = 4;
+                    result.parameter[1] = result.parameter[0];
+                    result.parameter[0] = -std::numeric_limits<Real>::max();
+
+                }
+            }
+            else
+            {
+                // The line intersects the negative cone in two distinct
+                // points, but we are interested only in the intersections
+                // with the positive cone.
+                result.intersect = false;
+                result.type = 0;
+                return;
+            }
+        }
+        else  // discr == 0
+        {
+            // One repeated real root; the line is tangent to the double-sided
+            // cone at a single point.  Report only the point if it is on the
+            // positive cone.
+            t = -c1 / c2;
+            if (DdU * t + DdPmV >= (Real)0)
+            {
+                result.intersect = true;
+                result.type = 1;
+                result.parameter[0] = t;
+                result.parameter[1] = t;
+            }
+            else
+            {
+                result.intersect = false;
+                result.type = 0;
+                return;
+            }
+        }
+    }
+    else if (c1 != (Real)0)
+    {
+        // c2 = 0, c1 != 0; U is a direction vector on the cone boundary
+        t = -((Real)0.5)*c0 / c1;
+        if (DdU * t + DdPmV >= (Real)0)
+        {
+            // The line intersects the positive cone and the ray of
+            // intersection is interior to the positive cone.
+            result.intersect = true;
+            if (DdU > (Real)0)
+            {
+                result.type = 3;
+                result.parameter[0] = t;
+                result.parameter[1] = std::numeric_limits<Real>::max();
+            }
+            else
+            {
+                result.type = 4;
+                result.parameter[0] = -std::numeric_limits<Real>::max();
+                result.parameter[1] = t;
+            }
+        }
+        else
+        {
+            // The line intersects the negative cone and the ray of
+            // intersection is interior to the positive cone.
+            result.intersect = false;
+            result.type = 0;
+            return;
+        }
+    }
+    else if (c0 != (Real)0)
+    {
+        // c2 = c1 = 0, c0 != 0.  Cross(D,U) is perpendicular to Cross(P-V,U)
+        result.intersect = false;
+        result.type = 0;
+        return;
+    }
+    else
+    {
+        // c2 = c1 = c0 = 0; the line is on the cone boundary.
+        result.intersect = true;
+        result.type = 5;
+        result.parameter[0] = -std::numeric_limits<Real>::max();
+        result.parameter[1] = +std::numeric_limits<Real>::max();
+    }
+
+    if (cone.maxHeight < std::numeric_limits<Real>::max())
+    {
+        if (DdU != (Real)0)
+        {
+            // Clamp the intersection to the height of the cone.
+            Real invDdU = ((Real)1) / DdU;
+            std::array<Real, 2> hInterval;
+            if (DdU >(Real)0)
+            {
+                hInterval[0] = -DdPmV * invDdU;
+                hInterval[1] = (cone.maxHeight - DdPmV) * invDdU;
+            }
+            else // (DdU < (Real)0)
+            {
+                hInterval[0] = (cone.maxHeight - DdPmV) * invDdU;
+                hInterval[1] = -DdPmV * invDdU;
+            }
+
+            FIIntervalInterval<Real> iiQuery;
+            auto iiResult = iiQuery(result.parameter, hInterval);
+            result.intersect = (iiResult.numIntersections > 0);
+            result.type = iiResult.numIntersections;
+            result.parameter = iiResult.overlap;
+        }
+        else if (result.intersect)
+        {
+            if (DdPmV > cone.maxHeight)
+            {
+                result.intersect = false;
+                result.type = 0;
+            }
+        }
+    }
+}
+
+
+}
diff --git a/Engine/Plugins/Experimental/GeometryProcessing/Source/GeometryAlgorithms/Private/GTEngine/Mathematics/GteIntrLine3Cylinder3.h b/Engine/Plugins/Experimental/GeometryProcessing/Source/GeometryAlgorithms/Private/GTEngine/Mathematics/GteIntrLine3Cylinder3.h
new file mode 100644
--- /dev/null
+++ b/Engine/Plugins/Experimental/GeometryProcessing/Source/GeometryAlgorithms/Private/GTEngine/Mathematics/GteIntrLine3Cylinder3.h
@@ -0,0 +1,264 @@
+// David Eberly, Geometric Tools, Redmond WA 98052
+// Copyright (c) 1998-2019
+// Distributed under the Boost Software License, Version 1.0.
+// http://www.boost.org/LICENSE_1_0.txt
+// http://www.geometrictools.com/License/Boost/LICENSE_1_0.txt
+// File Version: 3.0.1 (2018/10/05)
+
+#pragma once
+
+#include <GTEngine/Mathematics/GteVector3.h>
+#include <GTEngine/Mathematics/GteCylinder3.h>
+#include <GTEngine/Mathematics/GteLine.h>
+#include <GTEngine/Mathematics/GteFIQuery.h>
+
+// The queries consider the cylinder to be a solid.
+
+namespace gte
+{
+
+template <typename Real>
+class FIQuery<Real, Line3<Real>, Cylinder3<Real>>
+{
+public:
+    struct Result
+    {
+        bool intersect;
+        int numIntersections;
+        std::array<Real, 2> parameter;
+        std::array<Vector3<Real>, 2> point;
+    };
+
+    Result operator()(Line3<Real> const& line,
+        Cylinder3<Real> const& cylinder);
+
+protected:
+    void DoQuery(Vector3<Real> const& lineOrigin,
+        Vector3<Real> const& lineDirection, Cylinder3<Real> const& cylinder,
+        Result& result);
+};
+
+
+template <typename Real>
+typename FIQuery<Real, Line3<Real>, Cylinder3<Real>>::Result
+FIQuery<Real, Line3<Real>, Cylinder3<Real>>::operator()(
+    Line3<Real> const& line, Cylinder3<Real> const& cylinder)
+{
+    Result result;
+    DoQuery(line.origin, line.direction, cylinder, result);
+    for (int i = 0; i < result.numIntersections; ++i)
+    {
+        result.point[i] = line.origin + result.parameter[i] * line.direction;
+    }
+    return result;
+}
+
+template <typename Real>
+void FIQuery<Real, Line3<Real>, Cylinder3<Real>>::DoQuery(
+    Vector3<Real> const& lineOrigin, Vector3<Real> const& lineDirection,
+    Cylinder3<Real> const& cylinder, Result& result)
+{
+    // Initialize the result as if there is no intersection.  If we discover
+    // an intersection, these values will be modified accordingly.
+    result.intersect = false;
+    result.numIntersections = 0;
+
+    // Create a coordinate system for the cylinder.  In this system, the
+    // cylinder segment center C is the origin and the cylinder axis direction
+    // W is the z-axis.  U and V are the other coordinate axis directions.
+    // If P = x*U+y*V+z*W, the cylinder is x^2 + y^2 = r^2, where r is the
+    // cylinder radius.  The end caps are |z| = h/2, where h is the cylinder
+    // height.
+    Vector3<Real> basis[3];  // {W, U, V}
+    basis[0] = cylinder.axis.direction;
+    ComputeOrthogonalComplement(1, basis);
+    Real halfHeight = ((Real)0.5) * cylinder.height;
+    Real rSqr = cylinder.radius * cylinder.radius;
+
+    // Convert incoming line origin to capsule coordinates.
+    Vector3<Real> diff = lineOrigin - cylinder.axis.origin;
+    Vector3<Real> P{ Dot(basis[1], diff), Dot(basis[2], diff),
+        Dot(basis[0], diff) };
+
+    // Get the z-value, in cylinder coordinates, of the incoming line's
+    // unit-length direction.
+    Real dz = Dot(basis[0], lineDirection);
+    if (std::abs(dz) == (Real)1)
+    {
+        // The line is parallel to the cylinder axis.  Determine whether the
+        // line intersects the cylinder end disks.
+        Real radialSqrDist = rSqr - P[0] * P[0] - P[1] * P[1];
+        if (radialSqrDist >= (Real)0)
+        {
+            // The line intersects the cylinder end disks.
+            result.intersect = true;
+            result.numIntersections = 2;
+            if (dz > (Real)0)
+            {
+                result.parameter[0] = -P[2] - halfHeight;
+                result.parameter[1] = -P[2] + halfHeight;
+            }
+            else
+            {
+                result.parameter[0] = P[2] - halfHeight;
+                result.parameter[1] = P[2] + halfHeight;
+            }
+        }
+        // else:  The line is outside the cylinder, no intersection.
+        return;
+    }
+
+    // Convert the incoming line unit-length direction to cylinder
+    // coordinates.
+    Vector3<Real> D{ Dot(basis[1], lineDirection),
+        Dot(basis[2], lineDirection), dz };
+
+    Real a0, a1, a2, discr, root, inv, tValue;
+
+    if (D[2] == (Real)0)
+    {
+        // The line is perpendicular to the cylinder axis.
+        if (std::abs(P[2]) <= halfHeight)
+        {
+            // Test intersection of line P+t*D with infinite cylinder
+            // x^2+y^2 = r^2.  This reduces to computing the roots of a
+            // quadratic equation.  If P = (px,py,pz) and D = (dx,dy,dz),
+            // then the quadratic equation is
+            //   (dx^2+dy^2)*t^2 + 2*(px*dx+py*dy)*t + (px^2+py^2-r^2) = 0
+            a0 = P[0] * P[0] + P[1] * P[1] - rSqr;
+            a1 = P[0] * D[0] + P[1] * D[1];
+            a2 = D[0] * D[0] + D[1] * D[1];
+            discr = a1 * a1 - a0 * a2;
+            if (discr > (Real)0)
+            {
+                // The line intersects the cylinder in two places.
+                result.intersect = true;
+                result.numIntersections = 2;
+                root = std::sqrt(discr);
+                inv = ((Real)1) / a2;
+                result.parameter[0] = (-a1 - root) * inv;
+                result.parameter[1] = (-a1 + root) * inv;
+            }
+            else if (discr == (Real)0)
+            {
+                // The line is tangent to the cylinder.
+                result.intersect = true;
+                result.numIntersections = 1;
+                result.parameter[0] = -a1 / a2;
+                // Used by derived classes.
+                result.parameter[1] = result.parameter[0];
+            }
+            // else: The line does not intersect the cylinder.
+        }
+        // else: The line is outside the planes of the cylinder end disks.
+        return;
+    }
+
+    // Test for intersections with the planes of the end disks.
+    inv = ((Real)1) / D[2];
+
+    Real t0 = (-halfHeight - P[2]) * inv;
+    Real xTmp = P[0] + t0 * D[0];
+    Real yTmp = P[1] + t0 * D[1];
+    if (xTmp * xTmp + yTmp * yTmp <= rSqr)
+    {
+        // Plane intersection inside the top cylinder end disk.
+        result.parameter[result.numIntersections++] = t0;
+    }
+
+    Real t1 = (+halfHeight - P[2]) * inv;
+    xTmp = P[0] + t1 * D[0];
+    yTmp = P[1] + t1 * D[1];
+    if (xTmp * xTmp + yTmp * yTmp <= rSqr)
+    {
+        // Plane intersection inside the bottom cylinder end disk.
+        result.parameter[result.numIntersections++] = t1;
+    }
+
+    if (result.numIntersections < 2)
+    {
+        // Test for intersection with the cylinder wall.
+        a0 = P[0] * P[0] + P[1] * P[1] - rSqr;
+        a1 = P[0] * D[0] + P[1] * D[1];
+        a2 = D[0] * D[0] + D[1] * D[1];
+        discr = a1 * a1 - a0 * a2;
+        if (discr >(Real)0)
+        {
+            root = std::sqrt(discr);
+            inv = ((Real)1) / a2;
+            tValue = (-a1 - root) * inv;
+            if (t0 <= t1)
+            {
+                if (t0 <= tValue && tValue <= t1)
+                {
+                    result.parameter[result.numIntersections++] = tValue;
+                }
+            }
+            else
+            {
+                if (t1 <= tValue && tValue <= t0)
+                {
+                    result.parameter[result.numIntersections++] = tValue;
+                }
+            }
+
+            if (result.numIntersections < 2)
+            {
+                tValue = (-a1 + root) * inv;
+                if (t0 <= t1)
+                {
+                    if (t0 <= tValue && tValue <= t1)
+                    {
+                        result.parameter[result.numIntersections++] = tValue;
+                    }
+                }
+                else
+                {
+                    if (t1 <= tValue && tValue <= t0)
+                    {
+                        result.parameter[result.numIntersections++] = tValue;
+                    }
+                }
+            }
+            // else: Line intersects end disk and cylinder wall.
+        }
+        else if (discr == (Real)0)
+        {
+            tValue = -a1 / a2;
+            if (t0 <= t1)
+            {
+                if (t0 <= tValue && tValue <= t1)
+                {
+                    result.parameter[result.numIntersections++] = tValue;
+                }
+            }
+            else
+            {
+                if (t1 <= tValue && tValue <= t0)
+                {
+                    result.parameter[result.numIntersections++] = tValue;
+                }
+            }
+        }
+        // else: Line does not intersect cylinder wall.
+    }
+    // else: Line intersects both top and bottom cylinder end disks.
+
+    if (result.numIntersections == 2)
+    {
+        result.intersect = true;
+        if (result.parameter[0] > result.parameter[1])
+        {
+            std::swap(result.parameter[0], result.parameter[1]);
+        }
+    }
+    else if (result.numIntersections == 1)
+    {
+        result.intersect = true;
+        // Used by derived classes.
+        result.parameter[1] = result.parameter[0];
+    }
+}
+
+
+}
diff --git a/Engine/Plugins/Experimental/GeometryProcessing/Source/GeometryAlgorithms/Private/GTEngine/Mathematics/GteIntrLine3Ellipsoid3.h b/Engine/Plugins/Experimental/GeometryProcessing/Source/GeometryAlgorithms/Private/GTEngine/Mathematics/GteIntrLine3Ellipsoid3.h
new file mode 100644
--- /dev/null
+++ b/Engine/Plugins/Experimental/GeometryProcessing/Source/GeometryAlgorithms/Private/GTEngine/Mathematics/GteIntrLine3Ellipsoid3.h
@@ -0,0 +1,141 @@
+// David Eberly, Geometric Tools, Redmond WA 98052
+// Copyright (c) 1998-2019
+// Distributed under the Boost Software License, Version 1.0.
+// http://www.boost.org/LICENSE_1_0.txt
+// http://www.geometrictools.com/License/Boost/LICENSE_1_0.txt
+// File Version: 3.0.1 (2018/10/05)
+
+#pragma once
+
+#include <GTEngine/Mathematics/GteMatrix3x3.h>
+#include <GTEngine/Mathematics/GteHyperellipsoid.h>
+#include <GTEngine/Mathematics/GteLine.h>
+#include <GTEngine/Mathematics/GteFIQuery.h>
+#include <GTEngine/Mathematics/GteTIQuery.h>
+
+// The queries consider the ellipsoid to be a solid.
+
+namespace gte
+{
+
+template <typename Real>
+class TIQuery<Real, Line3<Real>, Ellipsoid3<Real>>
+{
+public:
+    struct Result
+    {
+        bool intersect;
+    };
+
+    Result operator()(Line3<Real> const& line,
+        Ellipsoid3<Real> const& ellipsoid);
+};
+
+template <typename Real>
+class FIQuery<Real, Line3<Real>, Ellipsoid3<Real>>
+{
+public:
+    struct Result
+    {
+        bool intersect;
+        int numIntersections;
+        std::array<Real, 2> parameter;
+        std::array<Vector3<Real>, 2> point;
+    };
+
+    Result operator()(Line3<Real> const& line,
+        Ellipsoid3<Real> const& ellipsoid);
+
+protected:
+    void DoQuery(Vector3<Real> const& lineOrigin,
+        Vector3<Real> const& lineDirection, Ellipsoid3<Real> const& ellipsoid,
+        Result& result);
+};
+
+
+template <typename Real>
+typename TIQuery<Real, Line3<Real>, Ellipsoid3<Real>>::Result
+TIQuery<Real, Line3<Real>, Ellipsoid3<Real>>::operator()(
+    Line3<Real> const& line, Ellipsoid3<Real> const& ellipsoid)
+{
+    // The ellipsoid is (X-K)^T*M*(X-K)-1 = 0 and the line is X = P+t*D.
+    // Substitute the line equation into the ellipsoid equation to obtain
+    // a quadratic equation Q(t) = a2*t^2 + 2*a1*t + a0 = 0, where
+    // a2 = D^T*M*D, a1 = D^T*M*(P-K), and a0 = (P-K)^T*M*(P-K)-1.
+    Result result;
+
+    Matrix3x3<Real> M;
+    ellipsoid.GetM(M);
+
+    Vector3<Real> diff = line.origin - ellipsoid.center;
+    Vector3<Real> matDir = M*line.direction;
+    Vector3<Real> matDiff = M*diff;
+    Real a2 = Dot(line.direction, matDir);
+    Real a1 = Dot(line.direction, matDiff);
+    Real a0 = Dot(diff, matDiff) - (Real)1;
+
+    // Intersection occurs when Q(t) has real roots.
+    Real discr = a1*a1 - a0*a2;
+    result.intersect = (discr >= (Real)0);
+    return result;
+}
+
+template <typename Real>
+typename FIQuery<Real, Line3<Real>, Ellipsoid3<Real>>::Result
+FIQuery<Real, Line3<Real>, Ellipsoid3<Real>>::operator()(
+    Line3<Real> const& line, Ellipsoid3<Real> const& ellipsoid)
+{
+    Result result;
+    DoQuery(line.origin, line.direction, ellipsoid, result);
+    for (int i = 0; i < result.numIntersections; ++i)
+    {
+        result.point[i] = line.origin + result.parameter[i] * line.direction;
+    }
+    return result;
+}
+
+template <typename Real>
+void FIQuery<Real, Line3<Real>, Ellipsoid3<Real>>::DoQuery(
+    Vector3<Real> const& lineOrigin, Vector3<Real> const& lineDirection,
+    Ellipsoid3<Real> const& ellipsoid, Result& result)
+{
+    // The ellipsoid is (X-K)^T*M*(X-K)-1 = 0 and the line is X = P+t*D.
+    // Substitute the line equation into the ellipsoid equation to obtain
+    // a quadratic equation Q(t) = a2*t^2 + 2*a1*t + a0 = 0, where
+    // a2 = D^T*M*D, a1 = D^T*M*(P-K), and a0 = (P-K)^T*M*(P-K)-1.
+    Matrix3x3<Real> M;
+    ellipsoid.GetM(M);
+
+    Vector3<Real> diff = lineOrigin - ellipsoid.center;
+    Vector3<Real> matDir = M*lineDirection;
+    Vector3<Real> matDiff = M*diff;
+    Real a2 = Dot(lineDirection, matDir);
+    Real a1 = Dot(lineDirection, matDiff);
+    Real a0 = Dot(diff, matDiff) - (Real)1;
+
+    // Intersection occurs when Q(t) has real roots.
+    Real discr = a1*a1 - a0*a2;
+    if (discr > (Real)0)
+    {
+        result.intersect = true;
+        result.numIntersections = 2;
+        Real root = std::sqrt(discr);
+        Real inv = ((Real)1) / a2;
+        result.parameter[0] = (-a1 - root)*inv;
+        result.parameter[1] = (-a1 + root)*inv;
+    }
+    else if (discr < (Real)0)
+    {
+        result.intersect = false;
+        result.numIntersections = 0;
+    }
+    else
+    {
+        result.intersect = true;
+        result.numIntersections = 1;
+        result.parameter[0] = -a1 / a2;
+    }
+}
+
+
+}
diff --git a/Engine/Plugins/Experimental/GeometryProcessing/Source/GeometryAlgorithms/Private/GTEngine/Mathematics/GteIntrLine3OrientedBox3.h b/Engine/Plugins/Experimental/GeometryProcessing/Source/GeometryAlgorithms/Private/GTEngine/Mathematics/GteIntrLine3OrientedBox3.h
new file mode 100644
--- /dev/null
+++ b/Engine/Plugins/Experimental/GeometryProcessing/Source/GeometryAlgorithms/Private/GTEngine/Mathematics/GteIntrLine3OrientedBox3.h
@@ -0,0 +1,97 @@
+// David Eberly, Geometric Tools, Redmond WA 98052
+// Copyright (c) 1998-2019
+// Distributed under the Boost Software License, Version 1.0.
+// http://www.boost.org/LICENSE_1_0.txt
+// http://www.geometrictools.com/License/Boost/LICENSE_1_0.txt
+// File Version: 3.0.1 (2016/11/28)
+
+#pragma once
+
+#include <GTEngine/Mathematics/GteOrientedBox.h>
+#include <GTEngine/Mathematics/GteIntrLine3AlignedBox3.h>
+
+// The test-intersection queries use the method of separating axes.  The
+// find-intersection queries use parametric clipping against the six faces of
+// the box.  The find-intersection queries use Liang-Barsky clipping.  The
+// queries consider the box to be a solid.  The algorithms are described in
+//   https://www.geometrictools.com/Documentation/IntersectionLineBox.pdf
+
+namespace gte
+{
+    template <typename Real>
+    class TIQuery<Real, Line3<Real>, OrientedBox3<Real>>
+        :
+        public TIQuery<Real, Line3<Real>, AlignedBox3<Real>>
+    {
+    public:
+        struct Result
+            :
+            public TIQuery<Real, Line3<Real>, AlignedBox3<Real>>::Result
+        {
+            // No additional relevant information to compute.
+        };
+
+        Result operator()(Line3<Real> const& line, OrientedBox3<Real> const& box)
+        {
+            // Transform the line to the oriented-box coordinate system.
+            Vector3<Real> diff = line.origin - box.center;
+            Vector3<Real> lineOrigin
+            {
+                Dot(diff, box.axis[0]),
+                Dot(diff, box.axis[1]),
+                Dot(diff, box.axis[2])
+            };
+            Vector3<Real> lineDirection
+            {
+                Dot(line.direction, box.axis[0]),
+                Dot(line.direction, box.axis[1]),
+                Dot(line.direction, box.axis[2])
+            };
+
+            Result result;
+            this->DoQuery(lineOrigin, lineDirection, box.extent, result);
+            return result;
+        }
+    };
+
+    template <typename Real>
+    class FIQuery<Real, Line3<Real>, OrientedBox3<Real>>
+        :
+        public FIQuery<Real, Line3<Real>, AlignedBox3<Real>>
+    {
+    public:
+        struct Result
+            :
+            public FIQuery<Real, Line3<Real>, AlignedBox3<Real>>::Result
+        {
+            // No additional relevant information to compute.
+        };
+
+        Result operator()(Line3<Real> const& line, OrientedBox3<Real> const& box)
+        {
+            // Transform the line to the oriented-box coordinate system.
+            Vector3<Real> diff = line.origin - box.center;
+            Vector3<Real> lineOrigin
+            {
+                Dot(diff, box.axis[0]),
+                Dot(diff, box.axis[1]),
+                Dot(diff, box.axis[2])
+            };
+            Vector3<Real> lineDirection
+            {
+                Dot(line.direction, box.axis[0]),
+                Dot(line.direction, box.axis[1]),
+                Dot(line.direction, box.axis[2])
+            };
+
+            Result result;
+            this->DoQuery(lineOrigin, lineDirection, box.extent, result);
+            for (int i = 0; i < result.numPoints; ++i)
+            {
+                result.point[i] =
+                    line.origin + result.lineParameter[i] * line.direction;
+            }
+            return result;
+        }
+    };
+}
diff --git a/Engine/Plugins/Experimental/GeometryProcessing/Source/GeometryAlgorithms/Private/GTEngine/Mathematics/GteIntrLine3Plane3.h b/Engine/Plugins/Experimental/GeometryProcessing/Source/GeometryAlgorithms/Private/GTEngine/Mathematics/GteIntrLine3Plane3.h
new file mode 100644
--- /dev/null
+++ b/Engine/Plugins/Experimental/GeometryProcessing/Source/GeometryAlgorithms/Private/GTEngine/Mathematics/GteIntrLine3Plane3.h
@@ -0,0 +1,133 @@
+// David Eberly, Geometric Tools, Redmond WA 98052
+// Copyright (c) 1998-2019
+// Distributed under the Boost Software License, Version 1.0.
+// http://www.boost.org/LICENSE_1_0.txt
+// http://www.geometrictools.com/License/Boost/LICENSE_1_0.txt
+// File Version: 3.0.0 (2016/06/19)
+
+#pragma once
+
+#include <GTEngine/Mathematics/GteDistPoint3Plane3.h>
+#include <GTEngine/Mathematics/GteLine.h>
+#include <GTEngine/Mathematics/GteFIQuery.h>
+#include <GTEngine/Mathematics/GteTIQuery.h>
+#include <limits>
+
+namespace gte
+{
+
+template <typename Real>
+class TIQuery<Real, Line3<Real>, Plane3<Real>>
+{
+public:
+    struct Result
+    {
+        bool intersect;
+    };
+
+    Result operator()(Line3<Real> const& line, Plane3<Real> const& plane);
+};
+
+template <typename Real>
+class FIQuery<Real, Line3<Real>, Plane3<Real>>
+{
+public:
+    struct Result
+    {
+        bool intersect;
+
+        // The number of intersections is 0 (no intersection), 1 (linear
+        // component and plane intersect in a point), or
+        // std::numeric_limits<int>::max() (linear component is on the plane).
+        // If the linear component is on the plane, 'point' component's
+        // origin and 'parameter' is zero.
+        int numIntersections;
+        Real parameter;
+        Vector3<Real> point;
+    };
+
+    Result operator()(Line3<Real> const& line, Plane3<Real> const& plane);
+
+protected:
+    void DoQuery(Vector3<Real> const& lineOrigin,
+        Vector3<Real> const& lineDirection, Plane3<Real> const& plane,
+        Result& result);
+};
+
+
+template <typename Real>
+typename TIQuery<Real, Line3<Real>, Plane3<Real>>::Result
+TIQuery<Real, Line3<Real>, Plane3<Real>>::operator()(
+    Line3<Real> const& line, Plane3<Real> const& plane)
+{
+    Result result;
+
+    Real DdN = Dot(line.direction, plane.normal);
+    if (DdN != (Real)0)
+    {
+        // The line is not parallel to the plane, so they must intersect.
+        result.intersect = true;
+    }
+    else
+    {
+        // The line and plane are parallel.
+        DCPQuery<Real, Vector3<Real>, Plane3<Real>> vpQuery;
+        result.intersect = (vpQuery(line.origin, plane).distance == (Real)0);
+    }
+
+    return result;
+}
+
+template <typename Real>
+typename FIQuery<Real, Line3<Real>, Plane3<Real>>::Result
+FIQuery<Real, Line3<Real>, Plane3<Real>>::operator()(
+    Line3<Real> const& line, Plane3<Real> const& plane)
+{
+    Result result;
+    DoQuery(line.origin, line.direction, plane, result);
+    if (result.intersect)
+    {
+        result.point = line.origin + result.parameter * line.direction;
+    }
+    return result;
+}
+
+template <typename Real>
+void FIQuery<Real, Line3<Real>, Plane3<Real>>::DoQuery(
+    Vector3<Real> const& lineOrigin, Vector3<Real> const& lineDirection,
+    Plane3<Real> const& plane, Result& result)
+{
+    Real DdN = Dot(lineDirection, plane.normal);
+    DCPQuery<Real, Vector3<Real>, Plane3<Real>> vpQuery;
+    auto vpResult = vpQuery(lineOrigin, plane);
+
+    if (DdN != (Real)0)
+    {
+        // The line is not parallel to the plane, so they must intersect.
+        result.intersect = true;
+        result.numIntersections = 1;
+        result.parameter = -vpResult.signedDistance / DdN;
+    }
+    else
+    {
+        // The line and plane are parallel.  Determine whether the line is on
+        // the plane.
+        if (vpResult.distance == (Real)0)
+        {
+            // The line is coincident with the plane, so choose t = 0 for the
+            // parameter.
+            result.intersect = true;
+            result.numIntersections = std::numeric_limits<int>::max();
+            result.parameter = (Real)0;
+        }
+        else
+        {
+            // The line is not on the plane.
+            result.intersect = false;
+            result.numIntersections = 0;
+        }
+    }
+}
+
+
+}
diff --git a/Engine/Plugins/Experimental/GeometryProcessing/Source/GeometryAlgorithms/Private/GTEngine/Mathematics/GteIntrLine3Sphere3.h b/Engine/Plugins/Experimental/GeometryProcessing/Source/GeometryAlgorithms/Private/GTEngine/Mathematics/GteIntrLine3Sphere3.h
new file mode 100644
--- /dev/null
+++ b/Engine/Plugins/Experimental/GeometryProcessing/Source/GeometryAlgorithms/Private/GTEngine/Mathematics/GteIntrLine3Sphere3.h
@@ -0,0 +1,124 @@
+// David Eberly, Geometric Tools, Redmond WA 98052
+// Copyright (c) 1998-2019
+// Distributed under the Boost Software License, Version 1.0.
+// http://www.boost.org/LICENSE_1_0.txt
+// http://www.geometrictools.com/License/Boost/LICENSE_1_0.txt
+// File Version: 3.0.1 (2018/10/05)
+
+#pragma once
+
+#include <GTEngine/Mathematics/GteVector3.h>
+#include <GTEngine/Mathematics/GteHypersphere.h>
+#include <GTEngine/Mathematics/GteLine.h>
+#include <GTEngine/Mathematics/GteFIQuery.h>
+#include <GTEngine/Mathematics/GteTIQuery.h>
+
+namespace gte
+{
+
+template <typename Real>
+class TIQuery<Real, Line3<Real>, Sphere3<Real>>
+{
+public:
+    struct Result
+    {
+        bool intersect;
+    };
+
+    Result operator()(Line3<Real> const& line, Sphere3<Real> const& sphere);
+};
+
+template <typename Real>
+class FIQuery<Real, Line3<Real>, Sphere3<Real>>
+{
+public:
+    struct Result
+    {
+        bool intersect;
+        int numIntersections;
+        std::array<Real, 2> parameter;
+        std::array<Vector3<Real>, 2> point;
+    };
+
+    Result operator()(Line3<Real> const& line, Sphere3<Real> const& sphere);
+
+protected:
+    void DoQuery(Vector3<Real> const& lineOrigin,
+        Vector3<Real> const& lineDirection, Sphere3<Real> const& sphere,
+        Result& result);
+};
+
+
+template <typename Real>
+typename TIQuery<Real, Line3<Real>, Sphere3<Real>>::Result
+TIQuery<Real, Line3<Real>, Sphere3<Real>>::operator()(
+    Line3<Real> const& line, Sphere3<Real> const& sphere)
+{
+    // The sphere is (X-C)^T*(X-C)-1 = 0 and the line is X = P+t*D.
+    // Substitute the line equation into the sphere equation to obtain a
+    // quadratic equation Q(t) = t^2 + 2*a1*t + a0 = 0, where a1 = D^T*(P-C),
+    // and a0 = (P-C)^T*(P-C)-1.
+    Result result;
+
+    Vector3<Real> diff = line.origin - sphere.center;
+    Real a0 = Dot(diff, diff) - sphere.radius * sphere.radius;
+    Real a1 = Dot(line.direction, diff);
+
+    // Intersection occurs when Q(t) has real roots.
+    Real discr = a1*a1 - a0;
+    result.intersect = (discr >= (Real)0);
+    return result;
+}
+
+template <typename Real>
+typename FIQuery<Real, Line3<Real>, Sphere3<Real>>::Result
+FIQuery<Real, Line3<Real>, Sphere3<Real>>::operator()(
+    Line3<Real> const& line, Sphere3<Real> const& sphere)
+{
+    Result result;
+    DoQuery(line.origin, line.direction, sphere, result);
+    for (int i = 0; i < result.numIntersections; ++i)
+    {
+        result.point[i] = line.origin + result.parameter[i] * line.direction;
+    }
+    return result;
+}
+
+template <typename Real>
+void FIQuery<Real, Line3<Real>, Sphere3<Real>>::DoQuery(
+    Vector3<Real> const& lineOrigin, Vector3<Real> const& lineDirection,
+    Sphere3<Real> const& sphere, Result& result)
+{
+    // The sphere is (X-C)^T*(X-C)-1 = 0 and the line is X = P+t*D.
+    // Substitute the line equation into the sphere equation to obtain a
+    // quadratic equation Q(t) = t^2 + 2*a1*t + a0 = 0, where a1 = D^T*(P-C),
+    // and a0 = (P-C)^T*(P-C)-1.
+    Vector3<Real> diff = lineOrigin - sphere.center;
+    Real a0 = Dot(diff, diff) - sphere.radius * sphere.radius;
+    Real a1 = Dot(lineDirection, diff);
+
+    // Intersection occurs when Q(t) has real roots.
+    Real discr = a1*a1 - a0;
+    if (discr > (Real)0)
+    {
+        result.intersect = true;
+        result.numIntersections = 2;
+        Real root = std::sqrt(discr);
+        result.parameter[0] = -a1 - root;
+        result.parameter[1] = -a1 + root;
+    }
+    else if (discr < (Real)0)
+    {
+        result.intersect = false;
+        result.numIntersections = 0;
+    }
+    else
+    {
+        result.intersect = true;
+        result.numIntersections = 1;
+        result.parameter[0] = -a1;
+    }
+}
+
+
+}
diff --git a/Engine/Plugins/Experimental/GeometryProcessing/Source/GeometryAlgorithms/Private/GTEngine/Mathematics/GteIntrLine3Triangle3.h b/Engine/Plugins/Experimental/GeometryProcessing/Source/GeometryAlgorithms/Private/GTEngine/Mathematics/GteIntrLine3Triangle3.h
new file mode 100644
--- /dev/null
+++ b/Engine/Plugins/Experimental/GeometryProcessing/Source/GeometryAlgorithms/Private/GTEngine/Mathematics/GteIntrLine3Triangle3.h
@@ -0,0 +1,190 @@
+// David Eberly, Geometric Tools, Redmond WA 98052
+// Copyright (c) 1998-2019
+// Distributed under the Boost Software License, Version 1.0.
+// http://www.boost.org/LICENSE_1_0.txt
+// http://www.geometrictools.com/License/Boost/LICENSE_1_0.txt
+// File Version: 3.0.0 (2016/06/19)
+
+#pragma once
+
+#include <GTEngine/Mathematics/GteVector3.h>
+#include <GTEngine/Mathematics/GteLine.h>
+#include <GTEngine/Mathematics/GteTriangle.h>
+#include <GTEngine/Mathematics/GteFIQuery.h>
+#include <GTEngine/Mathematics/GteTIQuery.h>
+
+namespace gte
+{
+
+template <typename Real>
+class TIQuery<Real,Line3<Real>,Triangle3<Real>>
+{
+public:
+    struct Result
+    {
+        bool intersect;
+    };
+
+    Result operator()(Line3<Real> const& line,
+        Triangle3<Real> const& triangle);
+};
+
+template <typename Real>
+class FIQuery<Real, Line3<Real>, Triangle3<Real>>
+{
+public:
+    struct Result
+    {
+        Result();
+
+        bool intersect;
+        Real parameter;
+        Real triangleBary[3];
+        Vector3<Real> point;
+    };
+
+    Result operator()(Line3<Real> const& line,
+        Triangle3<Real> const& triangle);
+};
+
+
+template <typename Real>
+typename TIQuery<Real, Line3<Real>, Triangle3<Real>>::Result
+TIQuery<Real, Line3<Real>, Triangle3<Real>>::operator()(
+    Line3<Real> const& line, Triangle3<Real> const& triangle)
+{
+    Result result;
+
+    // Compute the offset origin, edges, and normal.
+    Vector3<Real> diff = line.origin - triangle.v[0];
+    Vector3<Real> edge1 = triangle.v[1] - triangle.v[0];
+    Vector3<Real> edge2 = triangle.v[2] - triangle.v[0];
+    Vector3<Real> normal = Cross(edge1, edge2);
+
+    // Solve Q + t*D = b1*E1 + b2*E2 (Q = diff, D = line direction,
+    // E1 = edge1, E2 = edge2, N = Cross(E1,E2)) by
+    //   |Dot(D,N)|*b1 = sign(Dot(D,N))*Dot(D,Cross(Q,E2))
+    //   |Dot(D,N)|*b2 = sign(Dot(D,N))*Dot(D,Cross(E1,Q))
+    //   |Dot(D,N)|*t = -sign(Dot(D,N))*Dot(Q,N)
+    Real DdN = Dot(line.direction, normal);
+    Real sign;
+    if (DdN > (Real)0)
+    {
+        sign = (Real)1;
+    }
+    else if (DdN < (Real)0)
+    {
+        sign = (Real)-1;
+        DdN = -DdN;
+    }
+    else
+    {
+        // Line and triangle are parallel, call it a "no intersection"
+        // even if the line and triangle are coplanar and intersecting.
+        result.intersect = false;
+        return result;
+    }
+
+    Real DdQxE2 = sign*DotCross(line.direction, diff, edge2);
+    if (DdQxE2 >= (Real)0)
+    {
+        Real DdE1xQ = sign*DotCross(line.direction, edge1, diff);
+        if (DdE1xQ >= (Real)0)
+        {
+            if (DdQxE2 + DdE1xQ <= DdN)
+            {
+                // Line intersects triangle.
+                result.intersect = true;
+                return result;
+            }
+            // else: b1+b2 > 1, no intersection
+        }
+        // else: b2 < 0, no intersection
+    }
+    // else: b1 < 0, no intersection
+
+    result.intersect = false;
+    return result;
+}
+
+template <typename Real>
+FIQuery<Real, Line3<Real>, Triangle3<Real>>::Result::Result()
+    :
+    parameter((Real)0)
+{
+    triangleBary[0] = (Real)0;
+    triangleBary[1] = (Real)0;
+    triangleBary[2] = (Real)0;
+}
+
+template <typename Real>
+typename FIQuery<Real, Line3<Real>, Triangle3<Real>>::Result
+FIQuery<Real, Line3<Real>, Triangle3<Real>>::operator()(
+    Line3<Real> const& line, Triangle3<Real> const& triangle)
+{
+    Result result;
+
+    // Compute the offset origin, edges, and normal.
+    Vector3<Real> diff = line.origin - triangle.v[0];
+    Vector3<Real> edge1 = triangle.v[1] - triangle.v[0];
+    Vector3<Real> edge2 = triangle.v[2] - triangle.v[0];
+    Vector3<Real> normal = Cross(edge1, edge2);
+
+    // Solve Q + t*D = b1*E1 + b2*E2 (Q = diff, D = line direction,
+    // E1 = edge1, E2 = edge2, N = Cross(E1,E2)) by
+    //   |Dot(D,N)|*b1 = sign(Dot(D,N))*Dot(D,Cross(Q,E2))
+    //   |Dot(D,N)|*b2 = sign(Dot(D,N))*Dot(D,Cross(E1,Q))
+    //   |Dot(D,N)|*t = -sign(Dot(D,N))*Dot(Q,N)
+    Real DdN = Dot(line.direction, normal);
+    Real sign;
+    if (DdN > (Real)0)
+    {
+        sign = (Real)1;
+    }
+    else if (DdN < (Real)0)
+    {
+        sign = (Real)-1;
+        DdN = -DdN;
+    }
+    else
+    {
+        // Line and triangle are parallel, call it a "no intersection"
+        // even if the line and triangle are coplanar and intersecting.
+        result.intersect = false;
+        return result;
+    }
+
+    Real DdQxE2 = sign*DotCross(line.direction, diff, edge2);
+    if (DdQxE2 >= (Real)0)
+    {
+        Real DdE1xQ = sign*DotCross(line.direction, edge1, diff);
+        if (DdE1xQ >= (Real)0)
+        {
+            if (DdQxE2 + DdE1xQ <= DdN)
+            {
+                // Line intersects triangle.
+                Real QdN = -sign*Dot(diff, normal);
+                Real inv = ((Real)1) / DdN;
+
+                result.intersect = true;
+                result.parameter = QdN*inv;
+                result.triangleBary[1] = DdQxE2*inv;
+                result.triangleBary[2] = DdE1xQ*inv;
+                result.triangleBary[0] = (Real)1 - result.triangleBary[1]
+                    - result.triangleBary[2];
+                result.point = line.origin +
+                    result.parameter * line.direction;
+                return result;
+            }
+            // else: b1+b2 > 1, no intersection
+        }
+        // else: b2 < 0, no intersection
+    }
+    // else: b1 < 0, no intersection
+
+    result.intersect = false;
+    return result;
+}
+
+
+}
diff --git a/Engine/Plugins/Experimental/GeometryProcessing/Source/GeometryAlgorithms/Private/GTEngine/Mathematics/GteIntrOrientedBox2Circle2.h b/Engine/Plugins/Experimental/GeometryProcessing/Source/GeometryAlgorithms/Private/GTEngine/Mathematics/GteIntrOrientedBox2Circle2.h
new file mode 100644
--- /dev/null
+++ b/Engine/Plugins/Experimental/GeometryProcessing/Source/GeometryAlgorithms/Private/GTEngine/Mathematics/GteIntrOrientedBox2Circle2.h
@@ -0,0 +1,96 @@
+// David Eberly, Geometric Tools, Redmond WA 98052
+// Copyright (c) 1998-2019
+// Distributed under the Boost Software License, Version 1.0.
+// http://www.boost.org/LICENSE_1_0.txt
+// http://www.geometrictools.com/License/Boost/LICENSE_1_0.txt
+// File Version: 3.0.2 (2018/08/01)
+
+#pragma once
+
+#include <GTEngine/Mathematics/GteIntrAlignedBox2Circle2.h>
+#include <GTEngine/Mathematics/GteDistPointOrientedBox.h>
+
+// The find-intersection query is based on the document
+// https://www.geometrictools.com/Documentation/IntersectionMovingCircleRectangle.pdf
+
+namespace gte
+{
+    template <typename Real>
+    class TIQuery<Real, OrientedBox2<Real>, Circle2<Real>>
+    {
+    public:
+        // The intersection query considers the box and circle to be solids;
+        // that is, the circle object includes the region inside the circular
+        // boundary and the box object includes the region inside the
+        // rectangular boundary.  If the circle object and rectangle object
+        // overlap, the objects intersect.
+        struct Result
+        {
+            bool intersect;
+        };
+
+        Result operator()(OrientedBox2<Real> const& box, Circle2<Real> const& circle)
+        {
+            DCPQuery<Real, Vector2<Real>, OrientedBox2<Real>> pbQuery;
+            auto pbResult = pbQuery(circle.center, box);
+            Result result;
+            result.intersect = (pbResult.sqrDistance <= circle.radius * circle.radius);
+            return result;
+        }
+    };
+
+    template <typename Real>
+    class FIQuery<Real, OrientedBox2<Real>, Circle2<Real>>
+        :
+        public FIQuery<Real, AlignedBox2<Real>, Circle2<Real>>
+    {
+    public:
+        // See the base class for the definition of 'struct Result'.
+        typename FIQuery<Real, AlignedBox2<Real>, Circle2<Real>>::Result
+        operator()(OrientedBox2<Real> const& box, Vector2<Real> const& boxVelocity,
+            Circle2<Real> const& circle, Vector2<Real> const& circleVelocity)
+        {
+            // Transform the oriented box to an axis-aligned box centered at
+            // the origin and transform the circle accordingly.  Compute the
+            // velocity of the circle relative to the box.
+            Real const zero(0), one(1), minusOne(-1);
+            Vector2<Real> cdiff = circle.center - box.center;
+            Vector2<Real> vdiff = circleVelocity - boxVelocity;
+            Vector2<Real> C, V;
+            for (int i = 0; i < 2; ++i)
+            {
+                C[i] = Dot(cdiff, box.axis[i]);
+                V[i] = Dot(vdiff, box.axis[i]);
+            }
+
+            // Change signs on components, if necessary, to transform C to the
+            // first quadrant.  Adjust the velocity accordingly.
+            Real sign[2];
+            for (int i = 0; i < 2; ++i)
+            {
+                if (C[i] >= zero)
+                {
+                    sign[i] = one;
+                }
+                else
+                {
+                    C[i] = -C[i];
+                    V[i] = -V[i];
+                    sign[i] = minusOne;
+                }
+            }
+
+            typename FIQuery<Real, AlignedBox2<Real>, Circle2<Real>>::Result result = { 0, zero, { zero, zero } };
+            this->DoQuery(box.extent, C, circle.radius, V, result);
+
+            if (result.intersectionType != 0)
+            {
+                // Transform back to the original coordinate system.
+                result.contactPoint = box.center
+                    + (sign[0] * result.contactPoint[0]) * box.axis[0]
+                    + (sign[1] * result.contactPoint[1]) * box.axis[1];
+            }
+            return result;
+        }
+    };
+}
diff --git a/Engine/Plugins/Experimental/GeometryProcessing/Source/GeometryAlgorithms/Private/GTEngine/Mathematics/GteIntrOrientedBox2Cone2.h b/Engine/Plugins/Experimental/GeometryProcessing/Source/GeometryAlgorithms/Private/GTEngine/Mathematics/GteIntrOrientedBox2Cone2.h
new file mode 100644
--- /dev/null
+++ b/Engine/Plugins/Experimental/GeometryProcessing/Source/GeometryAlgorithms/Private/GTEngine/Mathematics/GteIntrOrientedBox2Cone2.h
@@ -0,0 +1,113 @@
+// David Eberly, Geometric Tools, Redmond WA 98052
+// Copyright (c) 1998-2019
+// Distributed under the Boost Software License, Version 1.0.
+// http://www.boost.org/LICENSE_1_0.txt
+// http://www.geometrictools.com/License/Boost/LICENSE_1_0.txt
+// File Version: 3.0.0 (2016/06/19)
+
+#pragma once
+
+#include <GTEngine/Mathematics/GteCone.h>
+#include <GTEngine/Mathematics/GteIntrRay2OrientedBox2.h>
+
+// The queries consider the box and cone to be solids.
+//
+// Define V = cone.ray.origin, D = cone.ray.direction, and cs = cone.cosAngle.
+// Define C = box.center, U0 = box.axis[0], U1 = box.axis[1],
+// e0 = box.extent[0], and e1 = box.extent[1].  A box point is
+// P = C + x*U0 + y*U1 where |x| <= e0 and |y| <= e1.  Define the function
+//   F(P) = Dot(D, (P-V)/Length(P-V)) = F(x,y)
+//     = Dot(D, (x*U0 + y*U1 + (C-V))/|x*U0 + y*U1 + (C-V)|
+//     = (a0*x + a1*y + a2)/(x^2 + y^2 + 2*b0*x + 2*b1*y + b2)^{1/2}
+// The function has an essential singularity when P = V.  The box intersects
+// the cone (with positive-area overlap) when at least one of the four box
+// corners is strictly inside the cone.  It is necessary that the numerator
+// of F(P) be positive at such a corner.  The (interior of the) solid cone
+// is defined by the quadratic inequality
+//   (Dot(D,P-V))^2 > |P-V|^2*(cone.cosAngle)^2
+// This inequality is inexpensive to compute.  In summary, overlap occurs
+// when there is a box corner P for which
+//   F(P) > 0 and (Dot(D,P-V))^2 > |P-V|^2*(cone.cosAngle)^2
+
+namespace gte
+{
+
+template <typename Real>
+class TIQuery<Real, OrientedBox<2, Real>, Cone<2, Real>>
+{
+public:
+    struct Result
+    {
+        // The value of 'intersect' is true when there is a box point that
+        // is strictly inside the cone.  If the box just touches the cone
+        // from the outside, an intersection is not reported, which supports
+        // the common operation of culling objects outside a cone.
+        bool intersect;
+    };
+
+    Result operator()(OrientedBox<2, Real> const& box, Cone<2, Real>& cone);
+};
+
+
+template <typename Real>
+typename TIQuery<Real, OrientedBox<2, Real>, Cone<2, Real>>::Result
+TIQuery<Real, OrientedBox<2, Real>, Cone<2, Real>>::operator()(
+    OrientedBox<2, Real> const& box, Cone<2, Real>& cone)
+{
+    Result result;
+
+    TIQuery<Real, Ray<2, Real>, OrientedBox<2, Real>> rbQuery;
+    auto rbResult = rbQuery(cone.ray, box);
+    if (rbResult.intersect)
+    {
+        // The cone intersects the box.
+        result.intersect = true;
+        return result;
+    }
+
+    // Define V = cone.ray.origin, D = cone.ray.direction, and
+    // cs = cone.cosAngle.  Define C = box.center, U0 = box.axis[0],
+    // U1 = box.axis[1], e0 = box.extent[0], and e1 = box.extent[1].
+    // A box point is P = C + x*U0 + y*U1 where |x| <= e0 and |y| <= e1.
+    // Define the function
+    //   F(x,y) = Dot(D, (P-V)/Length(P-V))
+    //     = Dot(D, (x*U0 + y*U1 + (C-V))/|x*U0 + y*U1 + (C-V)|
+    //     = (a0*x + a1*y + a2)/(x^2 + y^2 + 2*b0*x + 2*b1*y + b2)^{1/2}
+    // The function has an essential singularity when P = V.
+    Vector<2, Real> diff = box.center - cone.ray.origin;
+    Real a0 = Dot(cone.ray.direction, box.axis[0]);
+    Real a1 = Dot(cone.ray.direction, box.axis[1]);
+    Real a2 = Dot(cone.ray.direction, diff);
+    Real b0 = Dot(box.axis[0], diff);
+    Real b1 = Dot(box.axis[1], diff);
+    Real b2 = Dot(diff, diff);
+    Real csSqr = cone.cosAngle * cone.cosAngle;
+
+    for (int i1 = 0; i1 < 2; ++i1)
+    {
+        Real sign1 = i1 * (Real)2 - (Real)1;
+        Real y = sign1 * box.extent[1];
+        for (int i0 = 0; i0 < 2; ++i0)
+        {
+            Real sign0 = i0 * (Real)2 - (Real)1;
+            Real x = sign0 * box.extent[0];
+            Real fNumerator = a0 * x + a1 * y + a2;
+            if (fNumerator >(Real)0)
+            {
+                Real dSqr = x*x + y*y + (b0*x + b1*y)*(Real)2 + b2;
+                Real nSqr = fNumerator*fNumerator;
+                if (nSqr > dSqr * csSqr)
+                {
+                    result.intersect = true;
+                    return result;
+                }
+            }
+        }
+    }
+
+    result.intersect = false;
+    return result;
+}
+
+
+}
diff --git a/Engine/Plugins/Experimental/GeometryProcessing/Source/GeometryAlgorithms/Private/GTEngine/Mathematics/GteIntrOrientedBox2OrientedBox2.h b/Engine/Plugins/Experimental/GeometryProcessing/Source/GeometryAlgorithms/Private/GTEngine/Mathematics/GteIntrOrientedBox2OrientedBox2.h
new file mode 100644
--- /dev/null
+++ b/Engine/Plugins/Experimental/GeometryProcessing/Source/GeometryAlgorithms/Private/GTEngine/Mathematics/GteIntrOrientedBox2OrientedBox2.h
@@ -0,0 +1,294 @@
+// David Eberly, Geometric Tools, Redmond WA 98052
+// Copyright (c) 1998-2019
+// Distributed under the Boost Software License, Version 1.0.
+// http://www.boost.org/LICENSE_1_0.txt
+// http://www.geometrictools.com/License/Boost/LICENSE_1_0.txt
+// File Version: 3.0.0 (2016/06/19)
+
+#pragma once
+
+#include <GTEngine/Mathematics/GteVector2.h>
+#include <GTEngine/Mathematics/GteOrientedBox.h>
+#include <GTEngine/Mathematics/GteFIQuery.h>
+#include <GTEngine/Mathematics/GteTIQuery.h>
+#include <vector>
+
+// The queries consider the box to be a solid.
+//
+// The test-intersection query uses the method of separating axes.  The set of
+// potential separating directions includes the 2 edge normals of box0 and the
+// 2 edge normals of box1.  The integer 'separating' identifies the axis that
+// reported separation; there may be more than one but only one is reported.
+// The value is 0 when box0.axis[0] separates, 1 when box0.axis[1] separates,
+// 2 when box1.axis[0] separates, or 3 when box1.axis[1] separates.
+
+namespace gte
+{
+
+template <typename Real>
+class TIQuery<Real, OrientedBox2<Real>, OrientedBox2<Real>>
+{
+public:
+    struct Result
+    {
+        bool intersect;
+        int separating;
+    };
+
+    Result operator()(OrientedBox2<Real> const& box0,
+        OrientedBox2<Real> const& box1);
+};
+
+template <typename Real>
+class FIQuery<Real, OrientedBox2<Real>, OrientedBox2<Real>>
+{
+public:
+    struct Result
+    {
+        bool intersect;
+
+        // If 'intersect' is true, the boxes intersect in a convex 'polygon'.
+        std::vector<Vector2<Real>> polygon;
+    };
+
+    Result operator()(OrientedBox2<Real> const& box0,
+        OrientedBox2<Real> const& box1);
+
+private:
+    // The line normals are inner pointing.  The function returns true
+    // when the incoming polygon is outside the line, in which case the
+    // boxes do not intersect.  If the function returns false, the outgoing
+    // polygon is the incoming polygon intersected with the closed halfspace
+    // defined by the line.
+    bool Outside(Vector2<Real> const& origin, Vector2<Real> const& normal,
+        std::vector<Vector2<Real>>& polygon);
+};
+
+
+template <typename Real>
+typename TIQuery<Real, OrientedBox2<Real>, OrientedBox2<Real>>::Result
+TIQuery<Real, OrientedBox2<Real>, OrientedBox2<Real>>::operator()(
+    OrientedBox2<Real> const& box0, OrientedBox2<Real> const& box1)
+{
+    Result result;
+
+    // Convenience variables.
+    Vector2<Real> const* A0 = &box0.axis[0];
+    Vector2<Real> const* A1 = &box1.axis[0];
+    Vector2<Real> const& E0 = box0.extent;
+    Vector2<Real> const& E1 = box1.extent;
+
+    // Compute difference of box centers, D = C1-C0.
+    Vector2<Real> D = box1.center - box0.center;
+
+    Real absA0dA1[2][2], rSum;
+
+    // Test box0.axis[0].
+    absA0dA1[0][0] = std::abs(Dot(A0[0], A1[0]));
+    absA0dA1[0][1] = std::abs(Dot(A0[0], A1[1]));
+    rSum = E0[0] + E1[0] * absA0dA1[0][0] + E1[1] * absA0dA1[0][1];
+    if (std::abs(Dot(A0[0], D)) > rSum)
+    {
+        result.intersect = false;
+        result.separating = 0;
+        return result;
+    }
+
+    // Test axis box0.axis[1].
+    absA0dA1[1][0] = std::abs(Dot(A0[1], A1[0]));
+    absA0dA1[1][1] = std::abs(Dot(A0[1], A1[1]));
+    rSum = E0[1] + E1[0] * absA0dA1[1][0] + E1[1] * absA0dA1[1][1];
+    if (std::abs(Dot(A0[1], D)) > rSum)
+    {
+        result.intersect = false;
+        result.separating = 1;
+        return result;
+    }
+
+    // Test axis box1.axis[0].
+    rSum = E1[0] + E0[0] * absA0dA1[0][0] + E0[1] * absA0dA1[1][0];
+    if (std::abs(Dot(A1[0], D)) > rSum)
+    {
+        result.intersect = false;
+        result.separating = 2;
+        return result;
+    }
+
+    // Test axis box1.axis[1].
+    rSum = E1[1] + E0[0] * absA0dA1[0][1] + E0[1] * absA0dA1[1][1];
+    if (std::abs(Dot(A1[1], D)) > rSum)
+    {
+        result.intersect = false;
+        result.separating = 3;
+        return result;
+    }
+
+    result.intersect = true;
+    return result;
+}
+
+template <typename Real>
+typename FIQuery<Real, OrientedBox2<Real>, OrientedBox2<Real>>::Result
+FIQuery<Real, OrientedBox2<Real>, OrientedBox2<Real>>::operator()(
+    OrientedBox2<Real> const& box0, OrientedBox2<Real> const& box1)
+{
+    Result result;
+    result.intersect = true;
+
+    // Initialize the intersection polygon to box0, listing the vertices
+    // in counterclockwise order.
+    std::array<Vector2<Real>, 4> vertex;
+    box0.GetVertices(vertex);
+    result.polygon.push_back(vertex[0]);  // C - e0 * U0 - e1 * U1
+    result.polygon.push_back(vertex[1]);  // C + e0 * U0 - e1 * U1
+    result.polygon.push_back(vertex[3]);  // C + e0 * U0 + e1 * U1
+    result.polygon.push_back(vertex[2]);  // C - e0 * U0 + e1 * U1
+
+    // Clip the polygon using the lines defining edges of box1.  The
+    // line normal points inside box1.  The line origin is the first
+    // vertex of the edge when traversing box1 counterclockwise.
+    box1.GetVertices(vertex);
+    std::array<Vector2<Real>, 4> normal =
+    {
+        box1.axis[1], -box1.axis[0], box1.axis[0], -box1.axis[1]
+    };
+
+    for (int i = 0; i < 4; ++i)
+    {
+        if (Outside(vertex[i], normal[i], result.polygon))
+        {
+            // The boxes are separated.
+            result.intersect = false;
+            result.polygon.clear();
+            break;
+        }
+    }
+
+    return result;
+}
+
+template <typename Real>
+bool FIQuery<Real, OrientedBox2<Real>, OrientedBox2<Real>>::Outside(
+    Vector2<Real> const& origin, Vector2<Real> const& normal,
+    std::vector<Vector2<Real>>& polygon)
+{
+    // Determine whether the polygon vertices are outside the polygon, inside
+    // the polygon, or on the polygon boundary.
+    int const numVertices = static_cast<int>(polygon.size());
+    std::vector<Real> distance(numVertices);
+    int positive = 0, negative = 0, positiveIndex = -1;
+    for (int i = 0; i < numVertices; ++i)
+    {
+        distance[i] = Dot(normal, polygon[i] - origin);
+        if (distance[i] > (Real)0)
+        {
+            ++positive;
+            if (positiveIndex == -1)
+            {
+                positiveIndex = i;
+            }
+        }
+        else if (distance[i] < (Real)0)
+        {
+            ++negative;
+        }
+    }
+
+    if (positive == 0)
+    {
+        // The polygon is strictly outside the line.
+        return true;
+    }
+
+    if (negative == 0)
+    {
+        // The polygon is contained in the closed halfspace whose boundary
+        // is the line.  It is fully visible and no clipping is necessary.
+        return false;
+    }
+
+    // The line transversely intersects the polygon.  Clip the polygon.
+    std::vector<Vector2<Real>> clipPolygon;
+    Vector2<Real> vertex;
+    int curr, prev;
+    Real t;
+
+    if (positiveIndex > 0)
+    {
+        // Compute the first clip vertex on the line.
+        curr = positiveIndex;
+        prev = curr - 1;
+        t = distance[curr] / (distance[curr] - distance[prev]);
+        vertex = polygon[curr] + t * (polygon[prev] - polygon[curr]);
+        clipPolygon.push_back(vertex);
+
+        // Include the vertices on the positive side of line.
+        while (curr < numVertices && distance[curr] > (Real)0)
+        {
+            clipPolygon.push_back(polygon[curr++]);
+        }
+
+        // Compute the kast clip vertex on the line.
+        if (curr < numVertices)
+        {
+            prev = curr - 1;
+        }
+        else
+        {
+            curr = 0;
+            prev = numVertices - 1;
+        }
+        t = distance[curr] / (distance[curr] - distance[prev]);
+        vertex = polygon[curr] + t * (polygon[prev] - polygon[curr]);
+        clipPolygon.push_back(vertex);
+    }
+    else  // positiveIndex is 0
+    {
+        // Include the vertices on the positive side of line.
+        curr = 0;
+        while (curr < numVertices && distance[curr] > (Real)0)
+        {
+            clipPolygon.push_back(polygon[curr++]);
+        }
+
+        // Compute the last clip vertex on the line.
+        prev = curr - 1;
+        t = distance[curr] / (distance[curr] - distance[prev]);
+        vertex = polygon[curr] + t * (polygon[prev] - polygon[curr]);
+        clipPolygon.push_back(vertex);
+
+        // Skip the vertices on the negative side of the line.
+        while (curr < numVertices && distance[curr] <= (Real)0)
+        {
+            curr++;
+        }
+
+        // Compute the first clip vertex on the line.
+        if (curr < numVertices)
+        {
+            prev = curr - 1;
+            t = distance[curr] / (distance[curr] - distance[prev]);
+            vertex = polygon[curr] + t * (polygon[prev] - polygon[curr]);
+            clipPolygon.push_back(vertex);
+
+            // Keep the vertices on the positive side of the line.
+            while (curr < numVertices && distance[curr] > (Real)0)
+            {
+                clipPolygon.push_back(polygon[curr++]);
+            }
+        }
+        else
+        {
+            curr = 0;
+            prev = numVertices - 1;
+            t = distance[curr] / (distance[curr] - distance[prev]);
+            vertex = polygon[curr] + t * (polygon[prev] - polygon[curr]);
+            clipPolygon.push_back(vertex);
+        }
+    }
+
+    polygon = clipPolygon;
+    return false;
+}
+
+}
diff --git a/Engine/Plugins/Experimental/GeometryProcessing/Source/GeometryAlgorithms/Private/GTEngine/Mathematics/GteIntrOrientedBox2Sector2.h b/Engine/Plugins/Experimental/GeometryProcessing/Source/GeometryAlgorithms/Private/GTEngine/Mathematics/GteIntrOrientedBox2Sector2.h
new file mode 100644
--- /dev/null
+++ b/Engine/Plugins/Experimental/GeometryProcessing/Source/GeometryAlgorithms/Private/GTEngine/Mathematics/GteIntrOrientedBox2Sector2.h
@@ -0,0 +1,136 @@
+// David Eberly, Geometric Tools, Redmond WA 98052
+// Copyright (c) 1998-2019
+// Distributed under the Boost Software License, Version 1.0.
+// http://www.boost.org/LICENSE_1_0.txt
+// http://www.geometrictools.com/License/Boost/LICENSE_1_0.txt
+// File Version: 3.0.1 (2018/10/04)
+
+#pragma once
+
+#include <GTEngine/Mathematics/GteSector2.h>
+#include <GTEngine/Mathematics/GteOrientedBox.h>
+#include <GTEngine/Mathematics/GteDistPointSegment.h>
+#include <GTEngine/Mathematics/GteIntrHalfspace2Polygon2.h>
+#include <GTEngine/Mathematics/GteTIQuery.h>
+
+// The OrientedBox2 object is considered to be a solid.
+
+namespace gte
+{
+
+template <typename Real>
+class TIQuery<Real, OrientedBox2<Real>, Sector2<Real>>
+{
+public:
+    struct Result
+    {
+        bool intersect;
+    };
+
+    Result operator()(OrientedBox2<Real> const& box, Sector2<Real> const& sector);
+};
+
+template <typename Real>
+typename TIQuery<Real, OrientedBox2<Real>, Sector2<Real>>::Result
+TIQuery<Real, OrientedBox2<Real>, Sector2<Real>>::operator()(
+    OrientedBox2<Real> const& box, Sector2<Real> const& sector)
+{
+    Result result;
+
+    // Determine whether the vertex is inside the box.
+    Vector2<Real> CmV = box.center - sector.vertex;
+    Vector2<Real> P{ Dot(box.axis[0], CmV), Dot(box.axis[1], CmV) };
+    if (std::abs(P[0]) <= box.extent[0] && std::abs(P[1]) <= box.extent[1])
+    {
+        // The vertex is inside the box.
+        result.intersect = true;
+        return result;
+    }
+
+    // Test whether the box is outside the right ray boundary of the sector.
+    Vector2<Real> U0
+    {
+        +sector.cosAngle * sector.direction[0] + sector.sinAngle * sector.direction[1],
+        -sector.sinAngle * sector.direction[0] + sector.cosAngle * sector.direction[1]
+    };
+    Vector2<Real> N0 = Perp(U0);
+    Real prjcen0 = Dot(N0, CmV);
+    Real radius0 = box.extent[0] * std::abs(Dot(N0, box.axis[0]))
+        + box.extent[1] * std::abs(Dot(N0, box.axis[1]));
+    if (prjcen0 > radius0)
+    {
+        result.intersect = false;
+        return result;
+    }
+
+    // Test whether the box is outside the ray of the left boundary of the
+    // sector.
+    Vector2<Real> U1
+    {
+        +sector.cosAngle * sector.direction[0] - sector.sinAngle * sector.direction[1],
+        +sector.sinAngle * sector.direction[0] + sector.cosAngle * sector.direction[1]
+    };
+    Vector2<Real> N1 = -Perp(U1);
+    Real prjcen1 = Dot(N1, CmV);
+    Real radius1 = box.extent[0] * std::abs(Dot(N1, box.axis[0]))
+        + box.extent[1] * std::abs(Dot(N1, box.axis[1]));
+    if (prjcen1 > radius1)
+    {
+        result.intersect = false;
+        return result;
+    }
+
+    // Initialize the polygon of intersection to be the box.
+    Vector2<Real> e0U0 = box.extent[0] * box.axis[0];
+    Vector2<Real> e1U1 = box.extent[1] * box.axis[1];
+    std::vector<Vector2<Real>> polygon;
+    polygon.reserve(8);
+    polygon.push_back(box.center - e0U0 - e1U1);
+    polygon.push_back(box.center + e0U0 - e1U1);
+    polygon.push_back(box.center + e0U0 + e1U1);
+    polygon.push_back(box.center - e0U0 + e1U1);
+
+    FIQuery<Real, Halfspace<2, Real>, std::vector<Vector2<Real>>> hpQuery;
+    typename FIQuery<Real, Halfspace<2, Real>, std::vector<Vector2<Real>>>::Result hpResult;
+    Halfspace<2, Real> halfspace;
+
+    // Clip the box against the right-ray sector boundary.
+    if (prjcen0 >= -radius0)
+    {
+        halfspace.normal = -N0;
+        halfspace.constant = Dot(halfspace.normal, sector.vertex);
+        hpResult = hpQuery(halfspace, polygon);
+        polygon = std::move(hpResult.polygon);
+    }
+
+    // Clip the box against the left-ray sector boundary.
+    if (prjcen1 >= -radius1)
+    {
+        halfspace.normal = -N1;
+        halfspace.constant = Dot(halfspace.normal, sector.vertex);
+        hpResult = hpQuery(halfspace, polygon);
+        polygon = std::move(hpResult.polygon);
+    }
+
+    DCPQuery<Real, Vector2<Real>, Segment2<Real>> psQuery;
+    typename DCPQuery<Real, Vector2<Real>, Segment2<Real>>::Result psResult;
+    int const numVertices = static_cast<int>(polygon.size());
+    if (numVertices >= 2)
+    {
+        for (int i0 = numVertices - 1, i1 = 0; i1 < numVertices; i0 = i1++)
+        {
+            Segment2<Real> segment(polygon[i0], polygon[i1]);
+            psResult = psQuery(sector.vertex, segment);
+            if (psResult.distance <= sector.radius)
+            {
+                result.intersect = true;
+                return result;
+            }
+        }
+    }
+
+    result.intersect = false;
+    return result;
+}
+
+}
diff --git a/Engine/Plugins/Experimental/GeometryProcessing/Source/GeometryAlgorithms/Private/GTEngine/Mathematics/GteIntrOrientedBox3Cone3.h b/Engine/Plugins/Experimental/GeometryProcessing/Source/GeometryAlgorithms/Private/GTEngine/Mathematics/GteIntrOrientedBox3Cone3.h
new file mode 100644
--- /dev/null
+++ b/Engine/Plugins/Experimental/GeometryProcessing/Source/GeometryAlgorithms/Private/GTEngine/Mathematics/GteIntrOrientedBox3Cone3.h
@@ -0,0 +1,79 @@
+// David Eberly, Geometric Tools, Redmond WA 98052
+// Copyright (c) 1998-2019
+// Distributed under the Boost Software License, Version 1.0.
+// http://www.boost.org/LICENSE_1_0.txt
+// http://www.geometrictools.com/License/Boost/LICENSE_1_0.txt
+// File Version: 3.0.2 (2018/11/29)
+
+#pragma once
+
+#include <GTEngine/Mathematics/GteCone.h>
+#include <GTEngine/Mathematics/GteOrientedBox.h>
+#include <GTEngine/Mathematics/GteIntrAlignedBox3Cone3.h>
+
+// Test for intersection of a box and a cone.  The cone can be infinite
+//   0 <= minHeight < maxHeight = std::numeric_limits<Real>::max()
+// or finite (cone frustum)
+//   0 <= minHeight < maxHeight < std::numeric_limits<Real>::max().
+// The algorithm is described in
+//   http://www.geometrictools.com/Documentation/IntersectionBoxCone.pdf
+// and reports an intersection only when th intersection set has positive
+// volume.  For example, let the box be outside the cone.  If the box is
+// below the minHeight plane at the cone vertex and just touches the cone
+// vertex, no intersection is reported.  If the box is above the maxHeight
+// plane and just touches the disk capping the cone, either at a single
+// point, a line segment of points or a polygon of points, no intersection
+// is reported.  However, if the box straddles the minHeight plane (part of
+// the box strictly above the plane and part of the box strictly below the
+// plane) and just touches the cone vertex, an intersection is reported.
+
+namespace gte
+{
+    template <typename Real>
+    class TIQuery<Real, OrientedBox<3, Real>, Cone<3, Real>>
+        :
+        public TIQuery<Real, AlignedBox<3, Real>, Cone<3, Real>>
+    {
+    public:
+        struct Result
+            :
+            public TIQuery<Real, AlignedBox<3, Real>, Cone<3, Real>>::Result
+        {
+            // No additional information to compute.
+        };
+
+        Result operator()(OrientedBox<3, Real> const& box, Cone<3, Real> const& cone)
+        {
+            // Transform the cone and box so that the cone vertex is at the
+            // origin and the box is axis aligned.  This allows us to call the
+            // base class operator()(...).
+            Vector<3, Real> diff = box.center - cone.ray.origin;
+            Vector<3, Real> xfrmBoxCenter
+            {
+                Dot(box.axis[0], diff),
+                Dot(box.axis[1], diff),
+                Dot(box.axis[2], diff)
+            };
+            AlignedBox<3, Real> xfrmBox;
+            xfrmBox.min = xfrmBoxCenter - box.extent;
+            xfrmBox.max = xfrmBoxCenter + box.extent;
+
+            Cone<3, Real> xfrmCone = cone;
+            for (int i = 0; i < 3; ++i)
+            {
+                xfrmCone.ray.origin[i] = (Real)0;
+                xfrmCone.ray.direction[i] = Dot(box.axis[i], cone.ray.direction);
+            }
+
+            // Test for intersection between the aligned box and the cone.
+            auto bcResult = TIQuery<Real, AlignedBox<3, Real>, Cone<3, Real>>::operator()(xfrmBox, xfrmCone);
+            Result result;
+            result.intersect = bcResult.intersect;
+            return result;
+        }
+    };
+
+    // Template alias for convenience.
+    template <typename Real>
+    using TIOrientedBox3Cone3 = TIQuery<Real, OrientedBox<3, Real>, Cone<3, Real>>;
+}
diff --git a/Engine/Plugins/Experimental/GeometryProcessing/Source/GeometryAlgorithms/Private/GTEngine/Mathematics/GteIntrOrientedBox3Cylinder3.h b/Engine/Plugins/Experimental/GeometryProcessing/Source/GeometryAlgorithms/Private/GTEngine/Mathematics/GteIntrOrientedBox3Cylinder3.h
new file mode 100644
--- /dev/null
+++ b/Engine/Plugins/Experimental/GeometryProcessing/Source/GeometryAlgorithms/Private/GTEngine/Mathematics/GteIntrOrientedBox3Cylinder3.h
@@ -0,0 +1,55 @@
+// David Eberly, Geometric Tools, Redmond WA 98052
+// Copyright (c) 1998-2019
+// Distributed under the Boost Software License, Version 1.0.
+// http://www.boost.org/LICENSE_1_0.txt
+// http://www.geometrictools.com/License/Boost/LICENSE_1_0.txt
+// File Version: 3.4.0 (2016/11/08)
+
+#pragma once
+
+#include <GTEngine/Mathematics/GteIntrAlignedBox3Cylinder3.h>
+#include <GTEngine/Mathematics/GteOrientedBox.h>
+
+// The query considers the cylinder and box to be solids.
+
+namespace gte
+{
+
+template <typename Real>
+class TIQuery<Real, OrientedBox3<Real>, Cylinder3<Real>>
+{
+public:
+    struct Result
+    {
+        bool intersect;
+    };
+
+    Result operator()(OrientedBox3<Real> const& box, Cylinder3<Real> const& cylinder);
+};
+
+
+template <typename Real>
+typename TIQuery<Real, OrientedBox3<Real>, Cylinder3<Real>>::Result
+TIQuery<Real, OrientedBox3<Real>, Cylinder3<Real>>::operator()(
+    OrientedBox3<Real> const& box, Cylinder3<Real> const& cylinder)
+{
+    // Transform the box and cylinder so that the box is axis-aligned.
+    AlignedBox3<Real> aabb(-box.extent, box.extent);
+    Vector3<Real> diff = cylinder.axis.origin - box.center;
+    Cylinder3<Real> transformedCylinder;
+    transformedCylinder.radius = cylinder.radius;
+    transformedCylinder.height = cylinder.height;
+    for (int i = 0; i < 3; ++i)
+    {
+        transformedCylinder.axis.origin[i] = Dot(box.axis[i], diff);
+        transformedCylinder.axis.direction[i] = Dot(box.axis[i], cylinder.axis.direction);
+    }
+
+    TIQuery<Real, AlignedBox3<Real>, Cylinder3<Real>> aabbCylinderQuery;
+    auto aabbCylinderResult = aabbCylinderQuery(aabb, transformedCylinder);
+    Result result;
+    result.intersect = aabbCylinderResult.intersect;
+    return result;
+}
+
+}
diff --git a/Engine/Plugins/Experimental/GeometryProcessing/Source/GeometryAlgorithms/Private/GTEngine/Mathematics/GteIntrOrientedBox3Frustum3.h b/Engine/Plugins/Experimental/GeometryProcessing/Source/GeometryAlgorithms/Private/GTEngine/Mathematics/GteIntrOrientedBox3Frustum3.h
new file mode 100644
--- /dev/null
+++ b/Engine/Plugins/Experimental/GeometryProcessing/Source/GeometryAlgorithms/Private/GTEngine/Mathematics/GteIntrOrientedBox3Frustum3.h
@@ -0,0 +1,373 @@
+// David Eberly, Geometric Tools, Redmond WA 98052
+// Copyright (c) 1998-2019
+// Distributed under the Boost Software License, Version 1.0.
+// http://www.boost.org/LICENSE_1_0.txt
+// http://www.geometrictools.com/License/Boost/LICENSE_1_0.txt
+// File Version: 3.0.0 (2016/06/19)
+
+#pragma once
+
+#include <GTEngine/Mathematics/GteFrustum3.h>
+#include <GTEngine/Mathematics/GteOrientedBox.h>
+#include <GTEngine/Mathematics/GteTIQuery.h>
+
+namespace gte
+{
+
+// The method of separating axes is used.  The potential separating axes
+// include the 3 box face normals, the 5 distinct frustum normals (near and
+// far plane have the same normal), and cross products of normals, one from
+// the box and one from the frustum.
+
+template <typename Real>
+class TIQuery<Real, OrientedBox3<Real>, Frustum3<Real>>
+{
+public:
+    struct Result
+    {
+        bool intersect;
+    };
+
+    Result operator()(OrientedBox3<Real> const& box,
+        Frustum3<Real> const& frustum);
+};
+
+
+template <typename Real>
+typename TIQuery<Real, OrientedBox3<Real>, Frustum3<Real>>::Result
+TIQuery<Real, OrientedBox3<Real>, Frustum3<Real>>::operator()(
+    OrientedBox3<Real> const& box, Frustum3<Real> const& frustum)
+{
+    Result result;
+
+    // Convenience variables.
+    Vector3<Real> const* axis = &box.axis[0];
+    Vector3<Real> const& extent = box.extent;
+
+    Vector3<Real> diff = box.center - frustum.origin;  // C-E
+
+    Real A[3];      // Dot(R,A[i])
+    Real B[3];      // Dot(U,A[i])
+    Real C[3];      // Dot(D,A[i])
+    Real D[3];      // (Dot(R,C-E),Dot(U,C-E),Dot(D,C-E))
+    Real NA[3];     // dmin*Dot(R,A[i])
+    Real NB[3];     // dmin*Dot(U,A[i])
+    Real NC[3];     // dmin*Dot(D,A[i])
+    Real ND[3];     // dmin*(Dot(R,C-E),Dot(U,C-E),?)
+    Real RC[3];     // rmax*Dot(D,A[i])
+    Real RD[3];     // rmax*(?,?,Dot(D,C-E))
+    Real UC[3];     // umax*Dot(D,A[i])
+    Real UD[3];     // umax*(?,?,Dot(D,C-E))
+    Real NApRC[3];  // dmin*Dot(R,A[i]) + rmax*Dot(D,A[i])
+    Real NAmRC[3];  // dmin*Dot(R,A[i]) - rmax*Dot(D,A[i])
+    Real NBpUC[3];  // dmin*Dot(U,A[i]) + umax*Dot(D,A[i])
+    Real NBmUC[3];  // dmin*Dot(U,A[i]) - umax*Dot(D,A[i])
+    Real RBpUA[3];  // rmax*Dot(U,A[i]) + umax*Dot(R,A[i])
+    Real RBmUA[3];  // rmax*Dot(U,A[i]) - umax*Dot(R,A[i])
+    Real DdD, radius, p, fmin, fmax, MTwoUF, MTwoRF, tmp;
+    int i, j;
+
+    // M = D
+    D[2] = Dot(diff, frustum.dVector);
+    for (i = 0; i < 3; ++i)
+    {
+        C[i] = Dot(axis[i], frustum.dVector);
+    }
+    radius =
+        extent[0] * std::abs(C[0]) +
+        extent[1] * std::abs(C[1]) +
+        extent[2] * std::abs(C[2]);
+    if (D[2] + radius < frustum.dMin
+        || D[2] - radius > frustum.dMax)
+    {
+        result.intersect = false;
+        return result;
+    }
+
+    // M = n*R - r*D
+    for (i = 0; i < 3; ++i)
+    {
+        A[i] = Dot(axis[i], frustum.rVector);
+        RC[i] = frustum.rBound*C[i];
+        NA[i] = frustum.dMin*A[i];
+        NAmRC[i] = NA[i] - RC[i];
+    }
+    D[0] = Dot(diff, frustum.rVector);
+    radius =
+        extent[0] * std::abs(NAmRC[0]) +
+        extent[1] * std::abs(NAmRC[1]) +
+        extent[2] * std::abs(NAmRC[2]);
+    ND[0] = frustum.dMin*D[0];
+    RD[2] = frustum.rBound*D[2];
+    DdD = ND[0] - RD[2];
+    MTwoRF = frustum.GetMTwoRF();
+    if (DdD + radius < MTwoRF || DdD > radius)
+    {
+        result.intersect = false;
+        return result;
+    }
+
+    // M = -n*R - r*D
+    for (i = 0; i < 3; ++i)
+    {
+        NApRC[i] = NA[i] + RC[i];
+    }
+    radius =
+        extent[0] * std::abs(NApRC[0]) +
+        extent[1] * std::abs(NApRC[1]) +
+        extent[2] * std::abs(NApRC[2]);
+    DdD = -(ND[0] + RD[2]);
+    if (DdD + radius < MTwoRF || DdD > radius)
+    {
+        result.intersect = false;
+        return result;
+    }
+
+    // M = n*U - u*D
+    for (i = 0; i < 3; ++i)
+    {
+        B[i] = Dot(axis[i], frustum.uVector);
+        UC[i] = frustum.uBound*C[i];
+        NB[i] = frustum.dMin*B[i];
+        NBmUC[i] = NB[i] - UC[i];
+    }
+    D[1] = Dot(diff, frustum.uVector);
+    radius =
+        extent[0] * std::abs(NBmUC[0]) +
+        extent[1] * std::abs(NBmUC[1]) +
+        extent[2] * std::abs(NBmUC[2]);
+    ND[1] = frustum.dMin*D[1];
+    UD[2] = frustum.uBound*D[2];
+    DdD = ND[1] - UD[2];
+    MTwoUF = frustum.GetMTwoUF();
+    if (DdD + radius < MTwoUF || DdD > radius)
+    {
+        result.intersect = false;
+        return result;
+    }
+
+    // M = -n*U - u*D
+    for (i = 0; i < 3; ++i)
+    {
+        NBpUC[i] = NB[i] + UC[i];
+    }
+    radius =
+        extent[0] * std::abs(NBpUC[0]) +
+        extent[1] * std::abs(NBpUC[1]) +
+        extent[2] * std::abs(NBpUC[2]);
+    DdD = -(ND[1] + UD[2]);
+    if (DdD + radius < MTwoUF || DdD > radius)
+    {
+        result.intersect = false;
+        return result;
+    }
+
+    // M = A[i]
+    for (i = 0; i < 3; ++i)
+    {
+        p = frustum.rBound*std::abs(A[i]) +
+            frustum.uBound*std::abs(B[i]);
+        NC[i] = frustum.dMin*C[i];
+        fmin = NC[i] - p;
+        if (fmin < (Real)0)
+        {
+            fmin *= frustum.GetDRatio();
+        }
+        fmax = NC[i] + p;
+        if (fmax >(Real)0)
+        {
+            fmax *= frustum.GetDRatio();
+        }
+        DdD = A[i] * D[0] + B[i] * D[1] + C[i] * D[2];
+        if (DdD + extent[i] < fmin || DdD - extent[i] > fmax)
+        {
+            result.intersect = false;
+            return result;
+        }
+    }
+
+    // M = Cross(R,A[i])
+    for (i = 0; i < 3; ++i)
+    {
+        p = frustum.uBound*std::abs(C[i]);
+        fmin = -NB[i] - p;
+        if (fmin < (Real)0)
+        {
+            fmin *= frustum.GetDRatio();
+        }
+        fmax = -NB[i] + p;
+        if (fmax >(Real)0)
+        {
+            fmax *= frustum.GetDRatio();
+        }
+        DdD = C[i] * D[1] - B[i] * D[2];
+        radius =
+            extent[0] * std::abs(B[i] * C[0] - B[0] * C[i]) +
+            extent[1] * std::abs(B[i] * C[1] - B[1] * C[i]) +
+            extent[2] * std::abs(B[i] * C[2] - B[2] * C[i]);
+        if (DdD + radius < fmin || DdD - radius > fmax)
+        {
+            result.intersect = false;
+            return result;
+        }
+    }
+
+    // M = Cross(U,A[i])
+    for (i = 0; i < 3; ++i)
+    {
+        p = frustum.rBound*std::abs(C[i]);
+        fmin = NA[i] - p;
+        if (fmin < (Real)0)
+        {
+            fmin *= frustum.GetDRatio();
+        }
+        fmax = NA[i] + p;
+        if (fmax >(Real)0)
+        {
+            fmax *= frustum.GetDRatio();
+        }
+        DdD = -C[i] * D[0] + A[i] * D[2];
+        radius =
+            extent[0] * std::abs(A[i] * C[0] - A[0] * C[i]) +
+            extent[1] * std::abs(A[i] * C[1] - A[1] * C[i]) +
+            extent[2] * std::abs(A[i] * C[2] - A[2] * C[i]);
+        if (DdD + radius < fmin || DdD - radius > fmax)
+        {
+            result.intersect = false;
+            return result;
+        }
+    }
+
+    // M = Cross(n*D+r*R+u*U,A[i])
+    for (i = 0; i < 3; ++i)
+    {
+        Real fRB = frustum.rBound*B[i];
+        Real fUA = frustum.uBound*A[i];
+        RBpUA[i] = fRB + fUA;
+        RBmUA[i] = fRB - fUA;
+    }
+    for (i = 0; i < 3; ++i)
+    {
+        p = frustum.rBound*std::abs(NBmUC[i]) +
+            frustum.uBound*std::abs(NAmRC[i]);
+        tmp = -frustum.dMin*RBmUA[i];
+        fmin = tmp - p;
+        if (fmin < (Real)0)
+        {
+            fmin *= frustum.GetDRatio();
+        }
+        fmax = tmp + p;
+        if (fmax >(Real)0)
+        {
+            fmax *= frustum.GetDRatio();
+        }
+        DdD = D[0] * NBmUC[i] - D[1] * NAmRC[i] - D[2] * RBmUA[i];
+        radius = (Real)0;
+        for (j = 0; j < 3; j++)
+        {
+            radius += extent[j] * std::abs(A[j] * NBmUC[i] -
+                B[j] * NAmRC[i] - C[j] * RBmUA[i]);
+        }
+        if (DdD + radius < fmin || DdD - radius > fmax)
+        {
+            result.intersect = false;
+            return result;
+        }
+    }
+
+    // M = Cross(n*D+r*R-u*U,A[i])
+    for (i = 0; i < 3; ++i)
+    {
+        p = frustum.rBound*std::abs(NBpUC[i]) +
+            frustum.uBound*std::abs(NAmRC[i]);
+        tmp = -frustum.dMin*RBpUA[i];
+        fmin = tmp - p;
+        if (fmin < (Real)0)
+        {
+            fmin *= frustum.GetDRatio();
+        }
+        fmax = tmp + p;
+        if (fmax >(Real)0)
+        {
+            fmax *= frustum.GetDRatio();
+        }
+        DdD = D[0] * NBpUC[i] - D[1] * NAmRC[i] - D[2] * RBpUA[i];
+        radius = (Real)0;
+        for (j = 0; j < 3; ++j)
+        {
+            radius += extent[j] * std::abs(A[j] * NBpUC[i] -
+                B[j] * NAmRC[i] - C[j] * RBpUA[i]);
+        }
+        if (DdD + radius < fmin || DdD - radius > fmax)
+        {
+            result.intersect = false;
+            return result;
+        }
+    }
+
+    // M = Cross(n*D-r*R+u*U,A[i])
+    for (i = 0; i < 3; ++i)
+    {
+        p = frustum.rBound*std::abs(NBmUC[i]) +
+            frustum.uBound*std::abs(NApRC[i]);
+        tmp = frustum.dMin*RBpUA[i];
+        fmin = tmp - p;
+        if (fmin < (Real)0)
+        {
+            fmin *= frustum.GetDRatio();
+        }
+        fmax = tmp + p;
+        if (fmax >(Real)0)
+        {
+            fmax *= frustum.GetDRatio();
+        }
+        DdD = D[0] * NBmUC[i] - D[1] * NApRC[i] + D[2] * RBpUA[i];
+        radius = (Real)0;
+        for (j = 0; j < 3; ++j)
+        {
+            radius += extent[j] * std::abs(A[j] * NBmUC[i] -
+                B[j] * NApRC[i] + C[j] * RBpUA[i]);
+        }
+        if (DdD + radius < fmin || DdD - radius > fmax)
+        {
+            result.intersect = false;
+            return result;
+        }
+    }
+
+    // M = Cross(n*D-r*R-u*U,A[i])
+    for (i = 0; i < 3; ++i)
+    {
+        p = frustum.rBound*std::abs(NBpUC[i]) +
+            frustum.uBound*std::abs(NApRC[i]);
+        tmp = frustum.dMin*RBmUA[i];
+        fmin = tmp - p;
+        if (fmin < (Real)0)
+        {
+            fmin *= frustum.GetDRatio();
+        }
+        fmax = tmp + p;
+        if (fmax >(Real)0)
+        {
+            fmax *= frustum.GetDRatio();
+        }
+        DdD = D[0] * NBpUC[i] - D[1] * NApRC[i] + D[2] * RBmUA[i];
+        radius = (Real)0;
+        for (j = 0; j < 3; ++j)
+        {
+            radius += extent[j] * std::abs(A[j] * NBpUC[i] -
+                B[j] * NApRC[i] + C[j] * RBmUA[i]);
+        }
+        if (DdD + radius < fmin || DdD - radius > fmax)
+        {
+            result.intersect = false;
+            return result;
+        }
+    }
+
+    result.intersect = true;
+    return result;
+}
+
+
+}
diff --git a/Engine/Plugins/Experimental/GeometryProcessing/Source/GeometryAlgorithms/Private/GTEngine/Mathematics/GteIntrOrientedBox3OrientedBox3.h b/Engine/Plugins/Experimental/GeometryProcessing/Source/GeometryAlgorithms/Private/GTEngine/Mathematics/GteIntrOrientedBox3OrientedBox3.h
new file mode 100644
--- /dev/null
+++ b/Engine/Plugins/Experimental/GeometryProcessing/Source/GeometryAlgorithms/Private/GTEngine/Mathematics/GteIntrOrientedBox3OrientedBox3.h
@@ -0,0 +1,336 @@
+// David Eberly, Geometric Tools, Redmond WA 98052
+// Copyright (c) 1998-2019
+// Distributed under the Boost Software License, Version 1.0.
+// http://www.boost.org/LICENSE_1_0.txt
+// http://www.geometrictools.com/License/Boost/LICENSE_1_0.txt
+// File Version: 3.0.0 (2016/06/19)
+
+#pragma once
+
+#include <GTEngine/Mathematics/GteVector3.h>
+#include <GTEngine/Mathematics/GteOrientedBox.h>
+#include <GTEngine/Mathematics/GteFIQuery.h>
+#include <GTEngine/Mathematics/GteTIQuery.h>
+
+// The queries consider the box to be a solid.
+//
+// The test-intersection query uses the method of separating axes.  The set of
+// potential separating directions includes the 3 face normals of box0, the 3
+// face normals of box1, and 9 directions, each of which is the cross product
+// of an edge of box0 and and an edge of box1.
+//
+// The separating axes involving cross products of edges has numerical
+// robustness problems when the two edges are nearly parallel.  The cross
+// product of the edges is nearly the zero vector, so normalization of the
+// cross product may produce unit-length directions that are not close to the
+// true direction.  Such a pair of edges occurs when a box0 face normal N0 and
+// a box1 face normal N1 are nearly parallel.  In this case, you may skip the
+// edge-edge directions, which is equivalent to projecting the boxes onto the
+// plane with normal N0 and applying a 2D separating axis test.  The ability
+// to do so involves choosing a small nonnegative epsilon .  It is used to
+// determine whether two face normals, one from each box, are nearly parallel:
+// |Dot(N0,N1)| >= 1 - epsilon.  If the input is negative, it is clamped to
+// zero.
+//
+// The pair of integers 'separating', say, (i0,i1), identify the axis that
+// reported separation; there may be more than one but only one is
+// reported.  If the separating axis is a face normal N[i0] of the aligned
+// box0 in dimension i0, then (i0,-1) is returned.  If the axis is a face
+// normal box1.Axis[i1], then (-1,i1) is returned.  If the axis is a cross
+// product of edges, Cross(N[i0],box1.Axis[i1]), then (i0,i1) is returned.
+
+namespace gte
+{
+
+template <typename Real>
+class TIQuery<Real, OrientedBox3<Real>, OrientedBox3<Real>>
+{
+public:
+    struct Result
+    {
+        // The 'epsilon' value must be nonnegative.
+        Result(Real inEpsilon = (Real)0);
+
+        bool intersect;
+        Real epsilon;
+        int separating[2];
+    };
+
+    Result operator()(OrientedBox3<Real> const& box0,
+        OrientedBox3<Real> const& box1);
+};
+
+
+template <typename Real>
+TIQuery<Real, OrientedBox3<Real>, OrientedBox3<Real>>::Result::
+Result(Real inEpsilon)
+{
+    if (inEpsilon >= (Real)0)
+    {
+        epsilon = inEpsilon;
+    }
+    else
+    {
+        epsilon = (Real)0;
+    }
+}
+
+template <typename Real>
+typename TIQuery<Real, OrientedBox3<Real>, OrientedBox3<Real>>::Result
+TIQuery<Real, OrientedBox3<Real>, OrientedBox3<Real>>::operator()(
+    OrientedBox3<Real> const& box0, OrientedBox3<Real> const& box1)
+{
+    Result result;
+
+    // Convenience variables.
+    Vector3<Real> const& C0 = box0.center;
+    Vector3<Real> const* A0 = &box0.axis[0];
+    Vector3<Real> const& E0 = box0.extent;
+    Vector3<Real> const& C1 = box1.center;
+    Vector3<Real> const* A1 = &box1.axis[0];
+    Vector3<Real> const& E1 = box1.extent;
+
+    const Real cutoff = (Real)1 - result.epsilon;
+    bool existsParallelPair = false;
+
+    // Compute difference of box centers.
+    Vector3<Real> D = C1 - C0;
+
+    Real dot01[3][3];       // dot01[i][j] = Dot(A0[i],A1[j]) = A1[j][i]
+    Real absDot01[3][3];    // |dot01[i][j]|
+    Real dotDA0[3];         // Dot(D, A0[i])
+    Real r0, r1, r;         // interval radii and distance between centers
+    Real r01;               // r0 + r1
+
+    // Test for separation on the axis C0 + t*A0[0].
+    for (int i = 0; i < 3; ++i)
+    {
+        dot01[0][i] = Dot(A0[0], A1[i]);
+        absDot01[0][i] = std::abs(dot01[0][i]);
+        if (absDot01[0][i] > cutoff)
+        {
+            existsParallelPair = true;
+        }
+    }
+    dotDA0[0] = Dot(D, A0[0]);
+    r = std::abs(dotDA0[0]);
+    r1 = E1[0] * absDot01[0][0] + E1[1] * absDot01[0][1] + E1[2] * absDot01[0][2];
+    r01 = E0[0] + r1;
+    if (r > r01)
+    {
+        result.intersect = false;
+        result.separating[0] = 0;
+        result.separating[1] = -1;
+        return result;
+    }
+
+    // Test for separation on the axis C0 + t*A0[1].
+    for (int i = 0; i < 3; ++i)
+    {
+        dot01[1][i] = Dot(A0[1], A1[i]);
+        absDot01[1][i] = std::abs(dot01[1][i]);
+        if (absDot01[1][i] > cutoff)
+        {
+            existsParallelPair = true;
+        }
+    }
+    dotDA0[1] = Dot(D, A0[1]);
+    r = std::abs(dotDA0[1]);
+    r1 = E1[0] * absDot01[1][0] + E1[1] * absDot01[1][1] + E1[2] * absDot01[1][2];
+    r01 = E0[1] + r1;
+    if (r > r01)
+    {
+        result.intersect = false;
+        result.separating[0] = 1;
+        result.separating[1] = -1;
+        return result;
+    }
+
+    // Test for separation on the axis C0 + t*A0[2].
+    for (int i = 0; i < 3; ++i)
+    {
+        dot01[2][i] = Dot(A0[2], A1[i]);
+        absDot01[2][i] = std::abs(dot01[2][i]);
+        if (absDot01[2][i] > cutoff)
+        {
+            existsParallelPair = true;
+        }
+    }
+    dotDA0[2] = Dot(D, A0[2]);
+    r = std::abs(dotDA0[2]);
+    r1 = E1[0] * absDot01[2][0] + E1[1] * absDot01[2][1] + E1[2] * absDot01[2][2];
+    r01 = E0[2] + r1;
+    if (r > r01)
+    {
+        result.intersect = false;
+        result.separating[0] = 2;
+        result.separating[1] = -1;
+        return result;
+    }
+
+    // Test for separation on the axis C0 + t*A1[0].
+    r = std::abs(Dot(D, A1[0]));
+    r0 = E0[0] * absDot01[0][0] + E0[1] * absDot01[1][0] + E0[2] * absDot01[2][0];
+    r01 = r0 + E1[0];
+    if (r > r01)
+    {
+        result.intersect = false;
+        result.separating[0] = -1;
+        result.separating[1] = 0;
+        return result;
+    }
+
+    // Test for separation on the axis C0 + t*A1[1].
+    r = std::abs(Dot(D, A1[1]));
+    r0 = E0[0] * absDot01[0][1] + E0[1] * absDot01[1][1] + E0[2] * absDot01[2][1];
+    r01 = r0 + E1[1];
+    if (r > r01)
+    {
+        result.intersect = false;
+        result.separating[0] = -1;
+        result.separating[1] = 1;
+        return result;
+    }
+
+    // Test for separation on the axis C0 + t*A1[2].
+    r = std::abs(Dot(D, A1[2]));
+    r0 = E0[0] * absDot01[0][2] + E0[1] * absDot01[1][2] + E0[2] * absDot01[2][2];
+    r01 = r0 + E1[2];
+    if (r > r01)
+    {
+        result.intersect = false;
+        result.separating[0] = -1;
+        result.separating[1] = 2;
+        return result;
+    }
+
+    // At least one pair of box axes was parallel, so the separation is
+    // effectively in 2D.  The edge-edge axes do not need to be tested.
+    if (existsParallelPair)
+    {
+        return true;
+    }
+
+    // Test for separation on the axis C0 + t*A0[0]xA1[0].
+    r = std::abs(dotDA0[2] * dot01[1][0] - dotDA0[1] * dot01[2][0]);
+    r0 = E0[1] * absDot01[2][0] + E0[2] * absDot01[1][0];
+    r1 = E1[1] * absDot01[0][2] + E1[2] * absDot01[0][1];
+    r01 = r0 + r1;
+    if (r > r01)
+    {
+        result.intersect = false;
+        result.separating[0] = 0;
+        result.separating[1] = 0;
+        return result;
+    }
+
+    // Test for separation on the axis C0 + t*A0[0]xA1[1].
+    r = std::abs(dotDA0[2] * dot01[1][1] - dotDA0[1] * dot01[2][1]);
+    r0 = E0[1] * absDot01[2][1] + E0[2] * absDot01[1][1];
+    r1 = E1[0] * absDot01[0][2] + E1[2] * absDot01[0][0];
+    r01 = r0 + r1;
+    if (r > r01)
+    {
+        result.intersect = false;
+        result.separating[0] = 0;
+        result.separating[1] = 1;
+        return result;
+    }
+
+    // Test for separation on the axis C0 + t*A0[0]xA1[2].
+    r = std::abs(dotDA0[2] * dot01[1][2] - dotDA0[1] * dot01[2][2]);
+    r0 = E0[1] * absDot01[2][2] + E0[2] * absDot01[1][2];
+    r1 = E1[0] * absDot01[0][1] + E1[1] * absDot01[0][0];
+    r01 = r0 + r1;
+    if (r > r01)
+    {
+        result.intersect = false;
+        result.separating[0] = 0;
+        result.separating[1] = 2;
+        return result;
+    }
+
+    // Test for separation on the axis C0 + t*A0[1]xA1[0].
+    r = std::abs(dotDA0[0] * dot01[2][0] - dotDA0[2] * dot01[0][0]);
+    r0 = E0[0] * absDot01[2][0] + E0[2] * absDot01[0][0];
+    r1 = E1[1] * absDot01[1][2] + E1[2] * absDot01[1][1];
+    r01 = r0 + r1;
+    if (r > r01)
+    {
+        result.intersect = false;
+        result.separating[0] = 1;
+        result.separating[1] = 0;
+        return result;
+    }
+
+    // Test for separation on the axis C0 + t*A0[1]xA1[1].
+    r = std::abs(dotDA0[0] * dot01[2][1] - dotDA0[2] * dot01[0][1]);
+    r0 = E0[0] * absDot01[2][1] + E0[2] * absDot01[0][1];
+    r1 = E1[0] * absDot01[1][2] + E1[2] * absDot01[1][0];
+    r01 = r0 + r1;
+    if (r > r01)
+    {
+        result.intersect = false;
+        result.separating[0] = 1;
+        result.separating[1] = 1;
+        return result;
+    }
+
+    // Test for separation on the axis C0 + t*A0[1]xA1[2].
+    r = std::abs(dotDA0[0] * dot01[2][2] - dotDA0[2] * dot01[0][2]);
+    r0 = E0[0] * absDot01[2][2] + E0[2] * absDot01[0][2];
+    r1 = E1[0] * absDot01[1][1] + E1[1] * absDot01[1][0];
+    r01 = r0 + r1;
+    if (r > r01)
+    {
+        result.intersect = false;
+        result.separating[0] = 1;
+        result.separating[1] = 2;
+        return result;
+    }
+
+    // Test for separation on the axis C0 + t*A0[2]xA1[0].
+    r = std::abs(dotDA0[1] * dot01[0][0] - dotDA0[0] * dot01[1][0]);
+    r0 = E0[0] * absDot01[1][0] + E0[1] * absDot01[0][0];
+    r1 = E1[1] * absDot01[2][2] + E1[2] * absDot01[2][1];
+    r01 = r0 + r1;
+    if (r > r01)
+    {
+        result.intersect = false;
+        result.separating[0] = 2;
+        result.separating[1] = 0;
+        return result;
+    }
+
+    // Test for separation on the axis C0 + t*A0[2]xA1[1].
+    r = std::abs(dotDA0[1] * dot01[0][1] - dotDA0[0] * dot01[1][1]);
+    r0 = E0[0] * absDot01[1][1] + E0[1] * absDot01[0][1];
+    r1 = E1[0] * absDot01[2][2] + E1[2] * absDot01[2][0];
+    r01 = r0 + r1;
+    if (r > r01)
+    {
+        result.intersect = false;
+        result.separating[0] = 2;
+        result.separating[1] = 1;
+        return result;
+    }
+
+    // Test for separation on the axis C0 + t*A0[2]xA1[2].
+    r = std::abs(dotDA0[1] * dot01[0][2] - dotDA0[0] * dot01[1][2]);
+    r0 = E0[0] * absDot01[1][2] + E0[1] * absDot01[0][2];
+    r1 = E1[0] * absDot01[2][1] + E1[1] * absDot01[2][0];
+    r01 = r0 + r1;
+    if (r > r01)
+    {
+        result.intersect = false;
+        result.separating[0] = 2;
+        result.separating[1] = 2;
+        return result;
+    }
+
+    result.intersect = true;
+    return result;
+}
+
+
+}
