move GeometryProcessing plugin from NotForLicensees to Experimental

diff --git a/Engine/Plugins/Experimental/GeometryProcessing/Source/GeometryAlgorithms/Private/GTEngine/Mathematics/GteIntrOrientedBox3Sphere3.h b/Engine/Plugins/Experimental/GeometryProcessing/Source/GeometryAlgorithms/Private/GTEngine/Mathematics/GteIntrOrientedBox3Sphere3.h
new file mode 100644
--- /dev/null
+++ b/Engine/Plugins/Experimental/GeometryProcessing/Source/GeometryAlgorithms/Private/GTEngine/Mathematics/GteIntrOrientedBox3Sphere3.h
@@ -0,0 +1,94 @@
+// David Eberly, Geometric Tools, Redmond WA 98052
+// Copyright (c) 1998-2019
+// Distributed under the Boost Software License, Version 1.0.
+// http://www.boost.org/LICENSE_1_0.txt
+// http://www.geometrictools.com/License/Boost/LICENSE_1_0.txt
+// File Version: 3.0.2 (2018/10/02)
+
+#pragma once
+
+#include <GTEngine/Mathematics/GteDistPointOrientedBox.h>
+#include <GTEngine/Mathematics/GteIntrAlignedBox3Sphere3.h>
+
+namespace gte
+{
+    template <typename Real>
+    class TIQuery<Real, OrientedBox3<Real>, Sphere3<Real>>
+        :
+        public TIQuery<Real, AlignedBox3<Real>, Sphere3<Real>>
+    {
+    public:
+        // The intersection query considers the box and sphere to be solids.
+        // For example, if the sphere is strictly inside the box (does not touch
+        // the box faces), the objects intersect.
+        struct Result
+            :
+            public TIQuery<Real, AlignedBox3<Real>, Sphere3<Real>>::Result
+        {
+            // No additional information to compute.
+        };
+
+        Result operator()(OrientedBox3<Real> const& box, Sphere3<Real> const& sphere)
+        {
+            DCPQuery<Real, Vector3<Real>, OrientedBox3<Real>> pbQuery;
+            auto pbResult = pbQuery(sphere.center, box);
+            Result result;
+            result.intersect = (pbResult.sqrDistance <= sphere.radius * sphere.radius);
+            return result;
+        }
+    };
+
+    template <typename Real>
+    class FIQuery<Real, OrientedBox3<Real>, Sphere3<Real>>
+        :
+        public FIQuery<Real, AlignedBox3<Real>, Sphere3<Real>>
+    {
+    public:
+        // Currently, only a dynamic query is supported.  The static query must
+        // compute the intersection set of (solid) box and sphere.
+        struct Result
+            :
+            public FIQuery<Real, AlignedBox3<Real>, Sphere3<Real>>::Result
+        {
+            // No additional information to compute.
+        };
+
+        Result operator()(OrientedBox3<Real> const& box, Vector3<Real> const& boxVelocity,
+            Sphere3<Real> const& sphere, Vector3<Real> const& sphereVelocity)
+        {
+            Result result;
+            result.intersectionType = 0;
+            result.contactTime = (Real)0;
+            result.contactPoint = { (Real)0, (Real)0, (Real)0 };
+
+            // Transform the sphere and box so that the box center becomes
+            // the origin and the box is axis aligned.  Compute the velocity
+            // of the sphere relative to the box.
+            Vector3<Real> temp = sphere.center - box.center;
+            Vector3<Real> C
+            {
+                Dot(temp, box.axis[0]),
+                Dot(temp, box.axis[1]),
+                Dot(temp, box.axis[2])
+            };
+
+            temp = sphereVelocity - boxVelocity;
+            Vector3<Real> V
+            {
+                Dot(temp, box.axis[0]),
+                Dot(temp, box.axis[1]),
+                Dot(temp, box.axis[2])
+            };
+
+            this->DoQuery(box.extent, C, sphere.radius, V, result);
+
+            // Transform back to the original coordinate system.
+            if (result.intersectionType != 0)
+            {
+                auto& P = result.contactPoint;
+                P = box.center + P[0] * box.axis[0] + P[1] * box.axis[1] + P[2] * box.axis[2];
+            }
+            return result;
+        }
+    };
+}
diff --git a/Engine/Plugins/Experimental/GeometryProcessing/Source/GeometryAlgorithms/Private/GTEngine/Mathematics/GteIntrPlane3Capsule3.h b/Engine/Plugins/Experimental/GeometryProcessing/Source/GeometryAlgorithms/Private/GTEngine/Mathematics/GteIntrPlane3Capsule3.h
new file mode 100644
--- /dev/null
+++ b/Engine/Plugins/Experimental/GeometryProcessing/Source/GeometryAlgorithms/Private/GTEngine/Mathematics/GteIntrPlane3Capsule3.h
@@ -0,0 +1,58 @@
+// David Eberly, Geometric Tools, Redmond WA 98052
+// Copyright (c) 1998-2019
+// Distributed under the Boost Software License, Version 1.0.
+// http://www.boost.org/LICENSE_1_0.txt
+// http://www.geometrictools.com/License/Boost/LICENSE_1_0.txt
+// File Version: 3.0.0 (2016/06/19)
+
+#pragma once
+
+#include <GTEngine/Mathematics/GteDistPoint3Plane3.h>
+#include <GTEngine/Mathematics/GteCapsule.h>
+#include <GTEngine/Mathematics/GteTIQuery.h>
+
+namespace gte
+{
+
+template <typename Real>
+class TIQuery<Real, Plane3<Real>, Capsule3<Real>>
+{
+public:
+    struct Result
+    {
+        bool intersect;
+    };
+
+    Result operator()(Plane3<Real> const& plane,
+        Capsule3<Real> const& capsule);
+};
+
+
+template <typename Real>
+typename TIQuery<Real, Plane3<Real>, Capsule3<Real>>::Result
+TIQuery<Real, Plane3<Real>, Capsule3<Real>>::operator()(
+    Plane3<Real> const& plane, Capsule3<Real> const& capsule)
+{
+    Result result;
+
+    DCPQuery<Real, Vector3<Real>, Plane3<Real>> vpQuery;
+    Real sdistance0 = vpQuery(capsule.segment.p[0], plane).signedDistance;
+    Real sdistance1 = vpQuery(capsule.segment.p[1], plane).signedDistance;
+    if (sdistance0 * sdistance1 <= (Real)0)
+    {
+        // A capsule segment endpoint is on the plane or the two endpoints
+        // are on opposite sides of the plane.
+        result.intersect = true;
+        return result;
+    }
+
+    // The endpoints on same side of plane, but the endpoint spheres might
+    // intersect the plane.
+    result.intersect =
+        std::abs(sdistance0) <= capsule.radius ||
+        std::abs(sdistance1) <= capsule.radius;
+    return result;
+}
+
+
+}
diff --git a/Engine/Plugins/Experimental/GeometryProcessing/Source/GeometryAlgorithms/Private/GTEngine/Mathematics/GteIntrPlane3Circle3.h b/Engine/Plugins/Experimental/GeometryProcessing/Source/GeometryAlgorithms/Private/GTEngine/Mathematics/GteIntrPlane3Circle3.h
new file mode 100644
--- /dev/null
+++ b/Engine/Plugins/Experimental/GeometryProcessing/Source/GeometryAlgorithms/Private/GTEngine/Mathematics/GteIntrPlane3Circle3.h
@@ -0,0 +1,166 @@
+// David Eberly, Geometric Tools, Redmond WA 98052
+// Copyright (c) 1998-2019
+// Distributed under the Boost Software License, Version 1.0.
+// http://www.boost.org/LICENSE_1_0.txt
+// http://www.geometrictools.com/License/Boost/LICENSE_1_0.txt
+// File Version: 3.0.1 (2018/10/05)
+
+#pragma once
+
+#include <GTEngine/Mathematics/GteCircle3.h>
+#include <GTEngine/Mathematics/GteIntrPlane3Plane3.h>
+
+namespace gte
+{
+
+template <typename Real>
+class TIQuery<Real, Plane3<Real>, Circle3<Real>>
+{
+public:
+    struct Result
+    {
+        bool intersect;
+    };
+
+    Result operator()(Plane3<Real> const& plane, Circle3<Real> const& circle);
+};
+
+template <typename Real>
+class FIQuery<Real, Plane3<Real>, Circle3<Real>>
+{
+public:
+    struct Result
+    {
+        bool intersect;
+
+        // If 'intersect' is true, the intersection is either 1 or 2 points
+        // or the entire circle.  When points, 'numIntersections' and
+        // 'point' are valid.  When a circle, 'circle' is set to the incoming
+        // circle.
+        bool isPoints;
+        int numIntersections;
+        Vector3<Real> point[2];
+        Circle3<Real> circle;
+    };
+
+    Result operator()(Plane3<Real> const& plane, Circle3<Real> const& circle);
+};
+
+
+template <typename Real>
+typename TIQuery<Real, Plane3<Real>, Circle3<Real>>::Result
+TIQuery<Real, Plane3<Real>, Circle3<Real>>::operator()(
+    Plane3<Real> const& plane, Circle3<Real> const& circle)
+{
+    Result result;
+
+    // Construct the plane of the circle.
+    Plane3<Real> cPlane(circle.normal, circle.center);
+
+    // Compute the intersection of this plane with the input plane.
+    FIQuery<Real, Plane3<Real>, Plane3<Real>> ppQuery;
+    auto ppResult = ppQuery(plane, cPlane);
+    if (!ppResult.intersect)
+    {
+        // Planes are parallel and nonintersecting.
+        result.intersect = false;
+        return result;
+    }
+
+    if (!ppResult.isLine)
+    {
+        // Planes are the same, the circle is the common intersection set.
+        result.intersect = true;
+        return result;
+    }
+
+    // The planes intersect in a line.  Locate one or two points that are on
+    // the circle and line.  If the line is t*D+P, the circle center is C,
+    // and the circle radius is r, then
+    //   r^2 = |t*D+P-C|^2 = |D|^2*t^2 + 2*Dot(D,P-C)*t + |P-C|^2
+    // This is a quadratic equation of the form a2*t^2 + 2*a1*t + a0 = 0.
+    Vector3<Real> diff = ppResult.line.origin - circle.center;
+    Real a2 = Dot(ppResult.line.direction, ppResult.line.direction);
+    Real a1 = Dot(diff, ppResult.line.direction);
+    Real a0 = Dot(diff, diff) - circle.radius*circle.radius;
+
+    // Real-valued roots imply an intersection.
+    Real discr = a1*a1 - a0*a2;
+    result.intersect = (discr >= (Real)0);
+    return result;
+}
+
+
+
+template <typename Real>
+typename FIQuery<Real, Plane3<Real>, Circle3<Real>>::Result
+FIQuery<Real, Plane3<Real>, Circle3<Real>>::operator()(
+    Plane3<Real> const& plane, Circle3<Real> const& circle)
+{
+    Result result;
+
+    // Construct the plane of the circle.
+    Plane3<Real> cPlane(circle.normal, circle.center);
+
+    // Compute the intersection of this plane with the input plane.
+    FIQuery<Real, Plane3<Real>, Plane3<Real>> ppQuery;
+    auto ppResult = ppQuery(plane, cPlane);
+    if (!ppResult.intersect)
+    {
+        // Planes are parallel and nonintersecting.
+        result.intersect = false;
+        return result;
+    }
+
+    if (!ppResult.isLine)
+    {
+        // Planes are the same, the circle is the common intersection set.
+        result.intersect = true;
+        result.isPoints = false;
+        result.circle = circle;
+        return result;
+    }
+
+    // The planes intersect in a line.  Locate one or two points that are on
+    // the circle and line.  If the line is t*D+P, the circle center is C,
+    // and the circle radius is r, then
+    //   r^2 = |t*D+P-C|^2 = |D|^2*t^2 + 2*Dot(D,P-C)*t + |P-C|^2
+    // This is a quadratic equation of the form a2*t^2 + 2*a1*t + a0 = 0.
+    Vector3<Real> diff = ppResult.line.origin - circle.center;
+    Real a2 = Dot(ppResult.line.direction, ppResult.line.direction);
+    Real a1 = Dot(diff, ppResult.line.direction);
+    Real a0 = Dot(diff, diff) - circle.radius*circle.radius;
+
+    Real discr = a1*a1 - a0*a2;
+    if (discr < (Real)0)
+    {
+        // No real roots, the circle does not intersect the plane.
+        result.intersect = false;
+        return result;
+    }
+
+    result.isPoints = true;
+
+    Real inv = ((Real)1) / a2;
+    if (discr == (Real)0)
+    {
+        // One repeated root, the circle just touches the plane.
+        result.numIntersections = 1;
+        result.point[0] =
+            ppResult.line.origin - (a1*inv)*ppResult.line.direction;
+        return result;
+    }
+
+    // Two distinct, real-valued roots, the circle intersects the plane in
+    // two points.
+    Real root = std::sqrt(discr);
+    result.numIntersections = 2;
+    result.point[0] =
+        ppResult.line.origin - ((a1 + root)*inv)*ppResult.line.direction;
+    result.point[1] =
+        ppResult.line.origin - ((a1 - root)*inv)*ppResult.line.direction;
+    return result;
+}
+
+
+}
diff --git a/Engine/Plugins/Experimental/GeometryProcessing/Source/GeometryAlgorithms/Private/GTEngine/Mathematics/GteIntrPlane3Cylinder3.h b/Engine/Plugins/Experimental/GeometryProcessing/Source/GeometryAlgorithms/Private/GTEngine/Mathematics/GteIntrPlane3Cylinder3.h
new file mode 100644
--- /dev/null
+++ b/Engine/Plugins/Experimental/GeometryProcessing/Source/GeometryAlgorithms/Private/GTEngine/Mathematics/GteIntrPlane3Cylinder3.h
@@ -0,0 +1,165 @@
+// David Eberly, Geometric Tools, Redmond WA 98052
+// Copyright (c) 1998-2019
+// Distributed under the Boost Software License, Version 1.0.
+// http://www.boost.org/LICENSE_1_0.txt
+// http://www.geometrictools.com/License/Boost/LICENSE_1_0.txt
+// File Version: 3.0.1 (2018/10/05)
+
+#pragma once
+
+#include <GTEngine/LowLevel/GteLogger.h>
+#include <GTEngine/Mathematics/GteDistPoint3Plane3.h>
+#include <GTEngine/Mathematics/GteCircle3.h>
+#include <GTEngine/Mathematics/GteCylinder3.h>
+#include <GTEngine/Mathematics/GteEllipse3.h>
+#include <GTEngine/Mathematics/GteLine.h>
+#include <GTEngine/Mathematics/GteFIQuery.h>
+#include <GTEngine/Mathematics/GteTIQuery.h>
+
+namespace gte
+{
+
+template <typename Real>
+class TIQuery<Real, Plane3<Real>, Cylinder3<Real>>
+{
+public:
+    struct Result
+    {
+        bool intersect;
+    };
+
+    // The cylinder must have finite height.
+    Result operator()(Plane3<Real> const& plane,
+        Cylinder3<Real> const& cylinder);
+};
+
+template <typename Real>
+class FIQuery<Real, Plane3<Real>, Cylinder3<Real>>
+{
+public:
+    struct Result
+    {
+        bool intersect;
+
+        // The type of intersection.
+        //   0: none
+        //   1: single line (cylinder is tangent to plane), line[0] valid
+        //   2: two parallel lines (plane cuts cylinder in two lines)
+        //   3: circle (cylinder axis perpendicular to plane)
+        //   4: ellipse (cylinder axis neither parallel nor perpendicular)
+        int type;
+        Line3<Real> line[2];
+        Circle3<Real> circle;
+        Ellipse3<Real> ellipse;
+    };
+
+    // The cylinder must have infinite height.
+    Result operator()(Plane3<Real> const& plane,
+        Cylinder3<Real> const& cylinder);
+};
+
+
+template <typename Real>
+typename TIQuery<Real, Plane3<Real>, Cylinder3<Real>>::Result
+TIQuery<Real, Plane3<Real>, Cylinder3<Real>>::operator()(
+    Plane3<Real> const& plane, Cylinder3<Real> const& cylinder)
+{
+    LogAssert(cylinder.height != std::numeric_limits<Real>::max(),
+        "Cylinder height must be finite for TIQuery.");
+
+    Result result;
+
+    // Compute extremes of signed distance Dot(N,X)-d for points on the
+    // cylinder.  These are
+    //   min = (Dot(N,C)-d) - r*sqrt(1-Dot(N,W)^2) - (h/2)*|Dot(N,W)|
+    //   max = (Dot(N,C)-d) + r*sqrt(1-Dot(N,W)^2) + (h/2)*|Dot(N,W)|
+    DCPQuery<Real, Vector3<Real>, Plane3<Real>> vpQuery;
+    Real distance = vpQuery(cylinder.axis.origin, plane).distance;
+    Real absNdW = std::abs(Dot(plane.normal, cylinder.axis.direction));
+    Real root = std::sqrt(std::max((Real)1 - absNdW*absNdW, (Real)0));
+    Real term = cylinder.radius*root + ((Real)0.5)*cylinder.height*absNdW;
+
+    // Intersection occurs if and only if 0 is in the interval [min,max].
+    result.intersect = (distance <= term);
+    return result;
+}
+
+template <typename Real>
+typename FIQuery<Real, Plane3<Real>, Cylinder3<Real>>::Result
+FIQuery<Real, Plane3<Real>, Cylinder3<Real>>::operator()(
+    Plane3<Real> const& plane, Cylinder3<Real> const& cylinder)
+{
+    LogAssert(cylinder.height == std::numeric_limits<Real>::max(),
+        "Cylinder height must be infinite for FIQuery.");
+
+    Result result;
+
+    DCPQuery<Real, Vector3<Real>, Plane3<Real>> vpQuery;
+    Real sdistance = vpQuery(cylinder.axis.origin, plane).signedDistance;
+    Vector3<Real> center = cylinder.axis.origin - sdistance*plane.normal;
+    Real cosTheta = Dot(cylinder.axis.direction, plane.normal);
+    Real absCosTheta = std::abs(cosTheta);
+
+    if (absCosTheta > (Real)0)
+    {
+        // The cylinder axis intersects the plane in a unique point.
+        result.intersect = true;
+        if (absCosTheta < (Real)1)
+        {
+            result.type = 4;
+            result.ellipse.normal = plane.normal;
+            result.ellipse.center = cylinder.axis.origin -
+                (sdistance / cosTheta)*cylinder.axis.direction;
+            result.ellipse.axis[0] = cylinder.axis.direction -
+                cosTheta*plane.normal;
+            Normalize(result.ellipse.axis[0]);
+            result.ellipse.axis[1] = UnitCross(plane.normal,
+                result.ellipse.axis[0]);
+            result.ellipse.extent[0] = cylinder.radius / absCosTheta;
+            result.ellipse.extent[1] = cylinder.radius;
+        }
+        else
+        {
+            result.type = 3;
+            result.circle.normal = plane.normal;
+            result.circle.center = center;
+            result.circle.radius = cylinder.radius;
+        }
+    }
+    else
+    {
+        // The cylinder is parallel to the plane.
+        Real distance = std::abs(sdistance);
+        if (distance < cylinder.radius)
+        {
+            result.intersect = true;
+            result.type = 2;
+
+            Vector3<Real> offset = Cross(cylinder.axis.direction,
+                plane.normal);
+            Real extent = std::sqrt(cylinder.radius*cylinder.radius - sdistance*sdistance);
+
+            result.line[0].origin = center - extent*offset;
+            result.line[0].direction = cylinder.axis.direction;
+            result.line[1].origin = center + extent*offset;
+            result.line[1].direction = cylinder.axis.direction;
+        }
+        else if (distance == cylinder.radius)
+        {
+            result.intersect = true;
+            result.type = 1;
+            result.line[0].origin = center;
+            result.line[0].direction = cylinder.axis.direction;
+        }
+        else
+        {
+            result.intersect = false;
+            result.type = 0;
+        }
+    }
+
+    return result;
+}
+
+
+}
diff --git a/Engine/Plugins/Experimental/GeometryProcessing/Source/GeometryAlgorithms/Private/GTEngine/Mathematics/GteIntrPlane3Ellipsoid3.h b/Engine/Plugins/Experimental/GeometryProcessing/Source/GeometryAlgorithms/Private/GTEngine/Mathematics/GteIntrPlane3Ellipsoid3.h
new file mode 100644
--- /dev/null
+++ b/Engine/Plugins/Experimental/GeometryProcessing/Source/GeometryAlgorithms/Private/GTEngine/Mathematics/GteIntrPlane3Ellipsoid3.h
@@ -0,0 +1,49 @@
+// David Eberly, Geometric Tools, Redmond WA 98052
+// Copyright (c) 1998-2019
+// Distributed under the Boost Software License, Version 1.0.
+// http://www.boost.org/LICENSE_1_0.txt
+// http://www.geometrictools.com/License/Boost/LICENSE_1_0.txt
+// File Version: 3.0.1 (2018/10/05)
+
+#pragma once
+
+#include <GTEngine/Mathematics/GteMatrix3x3.h>
+#include <GTEngine/Mathematics/GteDistPoint3Plane3.h>
+#include <GTEngine/Mathematics/GteHyperellipsoid.h>
+#include <GTEngine/Mathematics/GteTIQuery.h>
+
+namespace gte
+{
+
+template <typename Real>
+class TIQuery<Real, Plane3<Real>, Ellipsoid3<Real>>
+{
+public:
+    struct Result
+    {
+        bool intersect;
+    };
+
+    Result operator()(Plane3<Real> const& plane,
+        Ellipsoid3<Real> const& ellipsoid);
+};
+
+
+template <typename Real>
+typename TIQuery<Real, Plane3<Real>, Ellipsoid3<Real>>::Result
+TIQuery<Real, Plane3<Real>, Ellipsoid3<Real>>::operator()(
+    Plane3<Real> const& plane, Ellipsoid3<Real> const& ellipsoid)
+{
+    Result result;
+    Matrix3x3<Real> MInverse;
+    ellipsoid.GetMInverse(MInverse);
+    Real discr = Dot(plane.normal, MInverse * plane.normal);
+    Real root = std::sqrt(std::max(discr, (Real)0));
+    DCPQuery<Real, Vector3<Real>, Plane3<Real>> vpQuery;
+    Real distance = vpQuery(ellipsoid.center, plane).distance;
+    result.intersect = (distance <= root);
+    return result;
+}
+
+
+}
diff --git a/Engine/Plugins/Experimental/GeometryProcessing/Source/GeometryAlgorithms/Private/GTEngine/Mathematics/GteIntrPlane3OrientedBox3.h b/Engine/Plugins/Experimental/GeometryProcessing/Source/GeometryAlgorithms/Private/GTEngine/Mathematics/GteIntrPlane3OrientedBox3.h
new file mode 100644
--- /dev/null
+++ b/Engine/Plugins/Experimental/GeometryProcessing/Source/GeometryAlgorithms/Private/GTEngine/Mathematics/GteIntrPlane3OrientedBox3.h
@@ -0,0 +1,51 @@
+// David Eberly, Geometric Tools, Redmond WA 98052
+// Copyright (c) 1998-2019
+// Distributed under the Boost Software License, Version 1.0.
+// http://www.boost.org/LICENSE_1_0.txt
+// http://www.geometrictools.com/License/Boost/LICENSE_1_0.txt
+// File Version: 3.0.0 (2016/06/19)
+
+#pragma once
+
+#include <GTEngine/Mathematics/GteDistPoint3Plane3.h>
+#include <GTEngine/Mathematics/GteOrientedBox.h>
+#include <GTEngine/Mathematics/GteFIQuery.h>
+#include <GTEngine/Mathematics/GteTIQuery.h>
+
+namespace gte
+{
+
+template <typename Real>
+class TIQuery<Real, Plane3<Real>, OrientedBox3<Real>>
+{
+public:
+    struct Result
+    {
+        bool intersect;
+    };
+
+    Result operator()(Plane3<Real> const& plane,
+        OrientedBox3<Real> const& box);
+};
+
+
+template <typename Real>
+typename TIQuery<Real, Plane3<Real>, OrientedBox3<Real>>::Result
+TIQuery<Real, Plane3<Real>, OrientedBox3<Real>>::operator()(
+    Plane3<Real> const& plane, OrientedBox3<Real> const& box)
+{
+    Result result;
+
+    Real radius =
+        std::abs(box.extent[0] * Dot(plane.normal, box.axis[0])) +
+        std::abs(box.extent[1] * Dot(plane.normal, box.axis[1])) +
+        std::abs(box.extent[2] * Dot(plane.normal, box.axis[2]));
+
+    DCPQuery<Real, Vector3<Real>, Plane3<Real>> ppQuery;
+    auto ppResult = ppQuery(box.center, plane);
+    result.intersect = (ppResult.distance <= radius);
+    return result;
+}
+
+
+}
diff --git a/Engine/Plugins/Experimental/GeometryProcessing/Source/GeometryAlgorithms/Private/GTEngine/Mathematics/GteIntrPlane3Plane3.h b/Engine/Plugins/Experimental/GeometryProcessing/Source/GeometryAlgorithms/Private/GTEngine/Mathematics/GteIntrPlane3Plane3.h
new file mode 100644
--- /dev/null
+++ b/Engine/Plugins/Experimental/GeometryProcessing/Source/GeometryAlgorithms/Private/GTEngine/Mathematics/GteIntrPlane3Plane3.h
@@ -0,0 +1,155 @@
+// David Eberly, Geometric Tools, Redmond WA 98052
+// Copyright (c) 1998-2019
+// Distributed under the Boost Software License, Version 1.0.
+// http://www.boost.org/LICENSE_1_0.txt
+// http://www.geometrictools.com/License/Boost/LICENSE_1_0.txt
+// File Version: 3.0.0 (2016/06/19)
+
+#pragma once
+
+#include <GTEngine/Mathematics/GteVector3.h>
+#include <GTEngine/Mathematics/GteHyperplane.h>
+#include <GTEngine/Mathematics/GteLine.h>
+#include <GTEngine/Mathematics/GteFIQuery.h>
+#include <GTEngine/Mathematics/GteTIQuery.h>
+
+namespace gte
+{
+
+template <typename Real>
+class TIQuery<Real, Plane3<Real>, Plane3<Real>>
+{
+public:
+    struct Result
+    {
+        bool intersect;
+    };
+
+    Result operator()(Plane3<Real> const& plane0, Plane3<Real> const& plane1);
+};
+
+template <typename Real>
+class FIQuery<Real, Plane3<Real>, Plane3<Real>>
+{
+public:
+    struct Result
+    {
+        bool intersect;
+
+        // If 'intersect' is true, the intersection is either a line or the
+        // planes are the same.  When a line, 'line' is valid.  When a plane,
+        // 'plane' is set to one of the planes.
+        bool isLine;
+        Line3<Real> line;
+        Plane3<Real> plane;
+    };
+
+    Result operator()(Plane3<Real> const& plane0, Plane3<Real> const& plane1);
+};
+
+
+template <typename Real>
+typename TIQuery<Real, Plane3<Real>, Plane3<Real>>::Result
+TIQuery<Real, Plane3<Real>, Plane3<Real>>::operator()(
+    Plane3<Real> const& plane0, Plane3<Real> const& plane1)
+{
+    // If Cross(N0,N1) is zero, then either planes are parallel and separated
+    // or the same plane.  In both cases, 'false' is returned.  Otherwise, the
+    // planes intersect.  To avoid subtle differences in reporting between
+    // Test() and Find(), the same parallel test is used.  Mathematically,
+    //   |Cross(N0,N1)|^2 = Dot(N0,N0)*Dot(N1,N1)-Dot(N0,N1)^2
+    //                    = 1 - Dot(N0,N1)^2
+    // The last equality is true since planes are required to have unit-length
+    // normal vectors.  The test |Cross(N0,N1)| = 0 is the same as
+    // |Dot(N0,N1)| = 1.
+
+    Result result;
+    Real dot = Dot(plane0.normal, plane1.normal);
+    if (std::abs(dot) < (Real)1)
+    {
+        result.intersect = true;
+        return result;
+    }
+
+    // The planes are parallel.  Check whether they are coplanar.
+    Real cDiff;
+    if (dot >= (Real)0)
+    {
+        // Normals are in same direction, need to look at c0-c1.
+        cDiff = plane0.constant - plane1.constant;
+    }
+    else
+    {
+        // Normals are in opposite directions, need to look at c0+c1.
+        cDiff = plane0.constant + plane1.constant;
+    }
+
+    result.intersect = (std::abs(cDiff) == (Real)0);
+    return result;
+}
+
+
+
+template <typename Real>
+typename FIQuery<Real, Plane3<Real>, Plane3<Real>>::Result
+FIQuery<Real, Plane3<Real>, Plane3<Real>>::operator()(
+    Plane3<Real> const& plane0, Plane3<Real> const& plane1)
+{
+    // If N0 and N1 are parallel, either the planes are parallel and separated
+    // or the same plane.  In both cases, 'false' is returned.  Otherwise,
+    // the intersection line is
+    //   L(t) = t*Cross(N0,N1)/|Cross(N0,N1)| + c0*N0 + c1*N1
+    // for some coefficients c0 and c1 and for t any real number (the line
+    // parameter).  Taking dot products with the normals,
+    //   d0 = Dot(N0,L) = c0*Dot(N0,N0) + c1*Dot(N0,N1) = c0 + c1*d
+    //   d1 = Dot(N1,L) = c0*Dot(N0,N1) + c1*Dot(N1,N1) = c0*d + c1
+    // where d = Dot(N0,N1).  These are two equations in two unknowns.  The
+    // solution is
+    //   c0 = (d0 - d*d1)/det
+    //   c1 = (d1 - d*d0)/det
+    // where det = 1 - d^2.
+
+    Result result;
+
+    Real dot = Dot(plane0.normal, plane1.normal);
+    if (std::abs(dot) >= (Real)1)
+    {
+        // The planes are parallel.  Check if they are coplanar.
+        Real cDiff;
+        if (dot >= (Real)0)
+        {
+            // Normals are in same direction, need to look at c0-c1.
+            cDiff = plane0.constant - plane1.constant;
+        }
+        else
+        {
+            // Normals are in opposite directions, need to look at c0+c1.
+            cDiff = plane0.constant + plane1.constant;
+        }
+
+        if (std::abs(cDiff) == (Real)0)
+        {
+            // The planes are coplanar.
+            result.intersect = true;
+            result.isLine = false;
+            result.plane = plane0;
+            return result;
+        }
+
+        // The planes are parallel but distinct.
+        result.intersect = false;
+        return result;
+    }
+
+    Real invDet = ((Real)1) / ((Real)1 - dot*dot);
+    Real c0 = (plane0.constant - dot*plane1.constant)*invDet;
+    Real c1 = (plane1.constant - dot*plane0.constant)*invDet;
+    result.intersect = true;
+    result.isLine = true;
+    result.line.origin = c0*plane0.normal + c1*plane1.normal;
+    result.line.direction = UnitCross(plane0.normal, plane1.normal);
+    return result;
+}
+
+
+}
diff --git a/Engine/Plugins/Experimental/GeometryProcessing/Source/GeometryAlgorithms/Private/GTEngine/Mathematics/GteIntrPlane3Sphere3.h b/Engine/Plugins/Experimental/GeometryProcessing/Source/GeometryAlgorithms/Private/GTEngine/Mathematics/GteIntrPlane3Sphere3.h
new file mode 100644
--- /dev/null
+++ b/Engine/Plugins/Experimental/GeometryProcessing/Source/GeometryAlgorithms/Private/GTEngine/Mathematics/GteIntrPlane3Sphere3.h
@@ -0,0 +1,98 @@
+// David Eberly, Geometric Tools, Redmond WA 98052
+// Copyright (c) 1998-2019
+// Distributed under the Boost Software License, Version 1.0.
+// http://www.boost.org/LICENSE_1_0.txt
+// http://www.geometrictools.com/License/Boost/LICENSE_1_0.txt
+// File Version: 3.0.1 (2018/10/05)
+
+#pragma once
+
+#include <GTEngine/Mathematics/GteDistPoint3Plane3.h>
+#include <GTEngine/Mathematics/GteCircle3.h>
+#include <GTEngine/Mathematics/GteHypersphere.h>
+#include <GTEngine/Mathematics/GteFIQuery.h>
+#include <GTEngine/Mathematics/GteTIQuery.h>
+
+namespace gte
+{
+
+template <typename Real>
+class TIQuery<Real, Plane3<Real>, Sphere3<Real>>
+{
+public:
+    struct Result
+    {
+        bool intersect;
+    };
+
+    Result operator()(Plane3<Real> const& plane, Sphere3<Real> const& sphere);
+};
+
+template <typename Real>
+class FIQuery<Real, Plane3<Real>, Sphere3<Real>>
+{
+public:
+    struct Result
+    {
+        bool intersect;
+
+        // If 'intersect' is true, the intersection is either a point or a
+        // circle.  When 'isCircle' is true, 'circle' is valid.  When
+        // 'isCircle' is false, 'point' is valid.
+        bool isCircle;
+        Circle3<Real> circle;
+        Vector3<Real> point;
+    };
+
+    Result operator()(Plane3<Real> const& plane, Sphere3<Real> const& sphere);
+};
+
+
+template <typename Real>
+typename TIQuery<Real, Plane3<Real>, Sphere3<Real>>::Result
+TIQuery<Real, Plane3<Real>, Sphere3<Real>>::operator()(
+    Plane3<Real> const& plane, Sphere3<Real> const& sphere)
+{
+    Result result;
+    DCPQuery<Real, Vector3<Real>, Plane3<Real>> ppQuery;
+    auto ppResult = ppQuery(sphere.center, plane);
+    result.intersect = (ppResult.distance <= sphere.radius);
+    return result;
+}
+
+
+template <typename Real>
+typename FIQuery<Real, Plane3<Real>, Sphere3<Real>>::Result
+FIQuery<Real, Plane3<Real>, Sphere3<Real>>::operator()(
+    Plane3<Real> const& plane, Sphere3<Real> const& sphere)
+{
+    Result result;
+    DCPQuery<Real, Vector3<Real>, Plane3<Real>> ppQuery;
+    auto ppResult = ppQuery(sphere.center, plane);
+    if (ppResult.distance < sphere.radius)
+    {
+        result.intersect = true;
+        result.isCircle = true;
+        result.circle.center = sphere.center - ppResult.signedDistance * plane.normal;
+        result.circle.normal = plane.normal;
+        Real sum = sphere.radius + ppResult.distance;  // > 0
+        Real dif = sphere.radius - ppResult.distance;  // > 0
+        Real arg = sum * dif;  // sqr(sphere.radius) - sqr(ppResult.distance)
+        result.circle.radius = std::sqrt(arg);
+        return result;
+    }
+    else if (ppResult.distance == sphere.radius)
+    {
+        result.intersect = true;
+        result.isCircle = false;
+        result.point = sphere.center - ppResult.signedDistance * plane.normal;
+        return result;
+    }
+    else
+    {
+        result.intersect = false;
+        return result;
+    }
+}
+
+}
diff --git a/Engine/Plugins/Experimental/GeometryProcessing/Source/GeometryAlgorithms/Private/GTEngine/Mathematics/GteIntrPlane3Triangle3.h b/Engine/Plugins/Experimental/GeometryProcessing/Source/GeometryAlgorithms/Private/GTEngine/Mathematics/GteIntrPlane3Triangle3.h
new file mode 100644
--- /dev/null
+++ b/Engine/Plugins/Experimental/GeometryProcessing/Source/GeometryAlgorithms/Private/GTEngine/Mathematics/GteIntrPlane3Triangle3.h
@@ -0,0 +1,289 @@
+// David Eberly, Geometric Tools, Redmond WA 98052
+// Copyright (c) 1998-2019
+// Distributed under the Boost Software License, Version 1.0.
+// http://www.boost.org/LICENSE_1_0.txt
+// http://www.geometrictools.com/License/Boost/LICENSE_1_0.txt
+// File Version: 3.0.0 (2016/06/19)
+
+#pragma once
+
+#include <GTEngine/Mathematics/GteVector3.h>
+#include <GTEngine/Mathematics/GteHyperplane.h>
+#include <GTEngine/Mathematics/GteTriangle.h>
+#include <GTEngine/Mathematics/GteFIQuery.h>
+#include <GTEngine/Mathematics/GteTIQuery.h>
+
+namespace gte
+{
+
+template <typename Real>
+class TIQuery<Real, Plane3<Real>, Triangle3<Real>>
+{
+public:
+    struct Result
+    {
+        bool intersect;
+
+        // The number is 0 (no intersection), 1 (plane and triangle intersect
+        // at a single point [vertex]), 2 (plane and triangle intersect in a
+        // segment), or 3 (triangle is in the plane).  When the number is 2,
+        // the segment is either interior to the triangle or is an edge of the
+        // triangle, the distinction stored in 'isInterior'.
+        int numIntersections;
+        bool isInterior;
+    };
+
+    Result operator()(Plane3<Real> const& plane,
+        Triangle3<Real> const& triangle);
+};
+
+template <typename Real>
+class FIQuery<Real, Plane3<Real>, Triangle3<Real>>
+{
+public:
+    struct Result
+    {
+        bool intersect;
+
+        // The number is 0 (no intersection), 1 (plane and triangle intersect
+        // at a single point [vertex]), 2 (plane and triangle intersect in a
+        // segment), or 3 (triangle is in the plane).  When the number is 2,
+        // the segment is either interior to the triangle or is an edge of the
+        // triangle, the distinction stored in 'isInterior'.
+        int numIntersections;
+        bool isInterior;
+        Vector3<Real> point[3];
+    };
+
+    Result operator()(Plane3<Real> const& plane,
+        Triangle3<Real> const& triangle);
+};
+
+
+template <typename Real>
+typename TIQuery<Real, Plane3<Real>, Triangle3<Real>>::Result
+TIQuery<Real, Plane3<Real>, Triangle3<Real>>::operator()(
+    Plane3<Real> const& plane, Triangle3<Real> const& triangle)
+{
+    Result result;
+
+    // Determine on which side of the plane the vertices lie.  The table of
+    // possibilities is listed next with n = numNegative, p = numPositive, and
+    // z = numZero.
+    //
+    //   n p z  intersection
+    //   ------------------------------------
+    //   0 3 0  none
+    //   0 2 1  vertex
+    //   0 1 2  edge
+    //   0 0 3  triangle in the plane
+    //   1 2 0  segment (2 edges clipped)
+    //   1 1 1  segment (1 edge clipped)
+    //   1 0 2  edge
+    //   2 1 0  segment (2 edges clipped)
+    //   2 0 1  vertex
+    //   3 0 0  none
+
+    Real s[3];
+    int numPositive = 0, numNegative = 0, numZero = 0;
+    for (int i = 0; i < 3; ++i)
+    {
+        s[i] = Dot(plane.normal, triangle.v[i]) - plane.constant;
+        if (s[i] >(Real)0)
+        {
+            ++numPositive;
+        }
+        else if (s[i] < (Real)0)
+        {
+            ++numNegative;
+        }
+        else
+        {
+            ++numZero;
+        }
+    }
+
+    if (numZero == 0 && numPositive > 0 && numNegative > 0)
+    {
+        result.intersect = true;
+        result.numIntersections = 2;
+        result.isInterior = true;
+        return result;
+    }
+
+    if (numZero == 1)
+    {
+        result.intersect = true;
+        for (int i = 0; i < 3; ++i)
+        {
+            if (s[i] == (Real)0)
+            {
+                if (numPositive == 2 || numNegative == 2)
+                {
+                    result.numIntersections = 1;
+                }
+                else
+                {
+                    result.numIntersections = 2;
+                    result.isInterior = true;
+                }
+                break;
+            }
+        }
+        return result;
+    }
+
+    if (numZero == 2)
+    {
+        result.intersect = true;
+        result.numIntersections = 2;
+        result.isInterior = false;
+        return result;
+    }
+
+    if (numZero == 3)
+    {
+        result.intersect = true;
+        result.numIntersections = 3;
+    }
+    else
+    {
+        result.intersect = false;
+        result.numIntersections = 0;
+
+    }
+    return result;
+}
+
+
+
+template <typename Real>
+typename FIQuery<Real, Plane3<Real>, Triangle3<Real>>::Result
+FIQuery<Real, Plane3<Real>, Triangle3<Real>>::operator()(
+    Plane3<Real> const& plane, Triangle3<Real> const& triangle)
+{
+    Result result;
+
+    // Determine on which side of the plane the vertices lie.  The table of
+    // possibilities is listed next with n = numNegative, p = numPositive, and
+    // z = numZero.
+    //
+    //   n p z  intersection
+    //   ------------------------------------
+    //   0 3 0  none
+    //   0 2 1  vertex
+    //   0 1 2  edge
+    //   0 0 3  triangle in the plane
+    //   1 2 0  segment (2 edges clipped)
+    //   1 1 1  segment (1 edge clipped)
+    //   1 0 2  edge
+    //   2 1 0  segment (2 edges clipped)
+    //   2 0 1  vertex
+    //   3 0 0  none
+
+    Real s[3];
+    int numPositive = 0, numNegative = 0, numZero = 0;
+    for (int i = 0; i < 3; ++i)
+    {
+        s[i] = Dot(plane.normal, triangle.v[i]) - plane.constant;
+        if (s[i] >(Real)0)
+        {
+            ++numPositive;
+        }
+        else if (s[i] < (Real)0)
+        {
+            ++numNegative;
+        }
+        else
+        {
+            ++numZero;
+        }
+    }
+
+    if (numZero == 0 && numPositive > 0 && numNegative > 0)
+    {
+        result.intersect = true;
+        result.numIntersections = 2;
+        result.isInterior = true;
+        Real sign = (Real)3 - numPositive * (Real)2;
+        for (int i0 = 0; i0 < 3; ++i0)
+        {
+            if (sign * s[i0] >(Real)0)
+            {
+                int i1 = (i0 + 1) % 3, i2 = (i0 + 2) % 3;
+                Real t1 = s[i1] / (s[i1] - s[i0]);
+                Real t2 = s[i2] / (s[i2] - s[i0]);
+                result.point[0] = triangle.v[i1] + t1 *
+                    (triangle.v[i0] - triangle.v[i1]);
+                result.point[1] = triangle.v[i2] + t2 *
+                    (triangle.v[i0] - triangle.v[i2]);
+                break;
+            }
+        }
+        return result;
+    }
+
+    if (numZero == 1)
+    {
+        result.intersect = true;
+        for (int i0 = 0; i0 < 3; ++i0)
+        {
+            if (s[i0] == (Real)0)
+            {
+                int i1 = (i0 + 1) % 3, i2 = (i0 + 2) % 3;
+                result.point[0] = triangle.v[i0];
+                if (numPositive == 2 || numNegative == 2)
+                {
+                    result.numIntersections = 1;
+                }
+                else
+                {
+                    result.numIntersections = 2;
+                    result.isInterior = true;
+                    Real t = s[i1] / (s[i1] - s[i2]);
+                    result.point[1] = triangle.v[i1] + t *
+                        (triangle.v[i2] - triangle.v[i1]);
+                }
+                break;
+            }
+        }
+        return result;
+    }
+
+    if (numZero == 2)
+    {
+        result.intersect = true;
+        result.numIntersections = 2;
+        result.isInterior = false;
+        for (int i0 = 0; i0 < 3; ++i0)
+        {
+            if (s[i0] != (Real)0)
+            {
+                int i1 = (i0 + 1) % 3, i2 = (i0 + 2) % 3;
+                result.point[0] = triangle.v[i1];
+                result.point[1] = triangle.v[i2];
+                break;
+            }
+        }
+        return result;
+    }
+
+    if (numZero == 3)
+    {
+        result.intersect = true;
+        result.numIntersections = 3;
+        result.point[0] = triangle.v[0];
+        result.point[1] = triangle.v[1];
+        result.point[2] = triangle.v[2];
+    }
+    else
+    {
+        result.intersect = false;
+        result.numIntersections = 0;
+
+    }
+    return result;
+}
+
+
+}
diff --git a/Engine/Plugins/Experimental/GeometryProcessing/Source/GeometryAlgorithms/Private/GTEngine/Mathematics/GteIntrRay2AlignedBox2.h b/Engine/Plugins/Experimental/GeometryProcessing/Source/GeometryAlgorithms/Private/GTEngine/Mathematics/GteIntrRay2AlignedBox2.h
new file mode 100644
--- /dev/null
+++ b/Engine/Plugins/Experimental/GeometryProcessing/Source/GeometryAlgorithms/Private/GTEngine/Mathematics/GteIntrRay2AlignedBox2.h
@@ -0,0 +1,150 @@
+// David Eberly, Geometric Tools, Redmond WA 98052
+// Copyright (c) 1998-2019
+// Distributed under the Boost Software License, Version 1.0.
+// http://www.boost.org/LICENSE_1_0.txt
+// http://www.geometrictools.com/License/Boost/LICENSE_1_0.txt
+// File Version: 3.0.0 (2016/06/19)
+
+#pragma once
+
+#include <GTEngine/Mathematics/GteVector2.h>
+#include <GTEngine/Mathematics/GteRay.h>
+#include <GTEngine/Mathematics/GteIntrIntervals.h>
+#include <GTEngine/Mathematics/GteIntrLine2AlignedBox2.h>
+
+// The queries consider the box to be a solid.
+//
+// The test-intersection queries use the method of separating axes.  The
+// find-intersection queries use parametric clipping against the four edges of
+// the box.
+
+namespace gte
+{
+
+template <typename Real>
+class TIQuery<Real, Ray2<Real>, AlignedBox2<Real>>
+    :
+    public TIQuery<Real, Line2<Real>, AlignedBox2<Real>>
+{
+public:
+    struct Result
+        :
+        public TIQuery<Real, Line2<Real>, AlignedBox2<Real>>::Result
+    {
+        // No additional information to compute.
+    };
+
+    Result operator()(Ray2<Real> const& ray, AlignedBox2<Real> const& box);
+
+protected:
+    void DoQuery(Vector2<Real> const& rayOrigin,
+        Vector2<Real> const& rayDirection, Vector2<Real> const& boxExtent,
+        Result& result);
+};
+
+template <typename Real>
+class FIQuery<Real, Ray2<Real>, AlignedBox2<Real>>
+    :
+    public FIQuery<Real, Line2<Real>, AlignedBox2<Real>>
+{
+public:
+    struct Result
+        :
+        public FIQuery<Real, Line2<Real>, AlignedBox2<Real>>::Result
+    {
+        // No additional information to compute.
+    };
+
+    Result operator()(Ray2<Real> const& ray, AlignedBox2<Real> const& box);
+
+protected:
+    void DoQuery(Vector2<Real> const& rayOrigin,
+        Vector2<Real> const& rayDirection, Vector2<Real> const& boxExtent,
+        Result& result);
+};
+
+
+template <typename Real>
+typename TIQuery<Real, Ray2<Real>, AlignedBox2<Real>>::Result
+TIQuery<Real, Ray2<Real>, AlignedBox2<Real>>::operator()(
+    Ray2<Real> const& ray, AlignedBox2<Real> const& box)
+{
+    // Get the centered form of the aligned box.  The axes are implicitly
+    // Axis[d] = Vector2<Real>::Unit(d).
+    Vector2<Real> boxCenter, boxExtent;
+    box.GetCenteredForm(boxCenter, boxExtent);
+
+    // Transform the ray to the aligned-box coordinate system.
+    Vector2<Real> rayOrigin = ray.origin - boxCenter;
+
+    Result result;
+    DoQuery(rayOrigin, ray.direction, boxExtent, result);
+    return result;
+}
+
+template <typename Real>
+void TIQuery<Real, Ray2<Real>, AlignedBox2<Real>>::DoQuery(
+    Vector2<Real> const& rayOrigin, Vector2<Real> const& rayDirection,
+    Vector2<Real> const& boxExtent, Result& result)
+{
+    for (int i = 0; i < 2; ++i)
+    {
+        if (std::abs(rayOrigin[i]) > boxExtent[i]
+            && rayOrigin[i] * rayDirection[i] >= (Real)0)
+        {
+            result.intersect = false;
+            return;
+        }
+    }
+
+    TIQuery<Real, Line2<Real>, AlignedBox2<Real>>::DoQuery(rayOrigin,
+        rayDirection, boxExtent, result);
+}
+
+template <typename Real>
+typename FIQuery<Real, Ray2<Real>, AlignedBox2<Real>>::Result
+FIQuery<Real, Ray2<Real>, AlignedBox2<Real>>::operator()(
+    Ray2<Real> const& ray, AlignedBox2<Real> const& box)
+{
+    // Get the centered form of the aligned box.  The axes are implicitly
+    // Axis[d] = Vector2<Real>::Unit(d).
+    Vector2<Real> boxCenter, boxExtent;
+    box.GetCenteredForm(boxCenter, boxExtent);
+
+    // Transform the ray to the aligned-box coordinate system.
+    Vector2<Real> rayOrigin = ray.origin - boxCenter;
+
+    Result result;
+    DoQuery(rayOrigin, ray.direction, boxExtent, result);
+    for (int i = 0; i < result.numIntersections; ++i)
+    {
+        result.point[i] = ray.origin + result.parameter[i] * ray.direction;
+    }
+    return result;
+}
+
+template <typename Real>
+void FIQuery<Real, Ray2<Real>, AlignedBox2<Real>>::DoQuery(
+    Vector2<Real> const& rayOrigin, Vector2<Real> const& rayDirection,
+    Vector2<Real> const& boxExtent, Result& result)
+{
+    FIQuery<Real, Line2<Real>, AlignedBox2<Real>>::DoQuery(rayOrigin,
+        rayDirection, boxExtent, result);
+
+    if (result.intersect)
+    {
+        // The line containing the ray intersects the box; the t-interval is
+        // [t0,t1].  The ray intersects the box as long as [t0,t1] overlaps
+        // the ray t-interval [0,+infinity).
+        std::array<Real, 2> rayInterval =
+            { (Real)0, std::numeric_limits<Real>::max() };
+        FIQuery<Real, std::array<Real, 2>, std::array<Real, 2>> iiQuery;
+        auto iiResult = iiQuery(result.parameter, rayInterval);
+        result.intersect = iiResult.intersect;
+        result.numIntersections = iiResult.numIntersections;
+        result.parameter = iiResult.overlap;
+    }
+}
+
+
+}
diff --git a/Engine/Plugins/Experimental/GeometryProcessing/Source/GeometryAlgorithms/Private/GTEngine/Mathematics/GteIntrRay2Arc2.h b/Engine/Plugins/Experimental/GeometryProcessing/Source/GeometryAlgorithms/Private/GTEngine/Mathematics/GteIntrRay2Arc2.h
new file mode 100644
--- /dev/null
+++ b/Engine/Plugins/Experimental/GeometryProcessing/Source/GeometryAlgorithms/Private/GTEngine/Mathematics/GteIntrRay2Arc2.h
@@ -0,0 +1,94 @@
+// David Eberly, Geometric Tools, Redmond WA 98052
+// Copyright (c) 1998-2019
+// Distributed under the Boost Software License, Version 1.0.
+// http://www.boost.org/LICENSE_1_0.txt
+// http://www.geometrictools.com/License/Boost/LICENSE_1_0.txt
+// File Version: 3.0.0 (2016/06/19)
+
+#pragma once
+
+#include <GTEngine/Mathematics/GteArc2.h>
+#include <GTEngine/Mathematics/GteIntrRay2Circle2.h>
+
+// The queries consider the arc to be a 1-dimensional object.
+
+namespace gte
+{
+
+template <typename Real>
+class TIQuery<Real, Ray2<Real>, Arc2<Real>>
+{
+public:
+    struct Result
+    {
+        bool intersect;
+    };
+
+    Result operator()(Ray2<Real> const& ray, Arc2<Real> const& arc);
+};
+
+template <typename Real>
+class FIQuery<Real, Ray2<Real>, Arc2<Real>>
+{
+public:
+    struct Result
+    {
+        bool intersect;
+        int numIntersections;
+        std::array<Real, 2> parameter;
+        std::array<Vector2<Real>, 2> point;
+    };
+
+    Result operator()(Ray2<Real> const& ray, Arc2<Real> const& arc);
+};
+
+
+template <typename Real>
+typename TIQuery<Real, Ray2<Real>, Arc2<Real>>::Result
+TIQuery<Real, Ray2<Real>, Arc2<Real>>::operator()(
+    Ray2<Real> const& ray, Arc2<Real> const& arc)
+{
+    Result result;
+    FIQuery<Real, Ray2<Real>, Arc2<Real>> raQuery;
+    auto raResult = raQuery(ray, arc);
+    result.intersect = raResult.intersect;
+    return result;
+}
+
+template <typename Real>
+typename FIQuery<Real, Ray2<Real>, Arc2<Real>>::Result
+FIQuery<Real, Ray2<Real>, Arc2<Real>>::operator()(
+    Ray2<Real> const& ray, Arc2<Real> const& arc)
+{
+    Result result;
+
+    FIQuery<Real, Ray2<Real>, Circle2<Real>> rcQuery;
+    Circle2<Real> circle(arc.center, arc.radius);
+    auto rcResult = rcQuery(ray, circle);
+    if (rcResult.intersect)
+    {
+        // Test whether ray-circle intersections are on the arc.
+        result.numIntersections = 0;
+        for (int i = 0; i < rcResult.numIntersections; ++i)
+        {
+            if (arc.Contains(rcResult.point[i]))
+            {
+                result.intersect = true;
+                result.parameter[result.numIntersections]
+                    = rcResult.parameter[i];
+                result.point[result.numIntersections++]
+                    = rcResult.point[i];
+            }
+        }
+    }
+    else
+    {
+        result.intersect = false;
+        result.numIntersections = 0;
+    }
+
+    return result;
+}
+
+
+}
diff --git a/Engine/Plugins/Experimental/GeometryProcessing/Source/GeometryAlgorithms/Private/GTEngine/Mathematics/GteIntrRay2Circle2.h b/Engine/Plugins/Experimental/GeometryProcessing/Source/GeometryAlgorithms/Private/GTEngine/Mathematics/GteIntrRay2Circle2.h
new file mode 100644
--- /dev/null
+++ b/Engine/Plugins/Experimental/GeometryProcessing/Source/GeometryAlgorithms/Private/GTEngine/Mathematics/GteIntrRay2Circle2.h
@@ -0,0 +1,100 @@
+// David Eberly, Geometric Tools, Redmond WA 98052
+// Copyright (c) 1998-2019
+// Distributed under the Boost Software License, Version 1.0.
+// http://www.boost.org/LICENSE_1_0.txt
+// http://www.geometrictools.com/License/Boost/LICENSE_1_0.txt
+// File Version: 3.0.0 (2016/06/19)
+
+#pragma once
+
+#include <GTEngine/Mathematics/GteRay.h>
+#include <GTEngine/Mathematics/GteIntrIntervals.h>
+#include <GTEngine/Mathematics/GteIntrLine2Circle2.h>
+
+// The queries consider the circle to be a solid (disk).
+
+namespace gte
+{
+
+template <typename Real>
+class TIQuery<Real, Ray2<Real>, Circle2<Real>>
+{
+public:
+    struct Result
+    {
+        bool intersect;
+    };
+
+    Result operator()(Ray2<Real> const& ray, Circle2<Real> const& circle);
+};
+
+template <typename Real>
+class FIQuery<Real, Ray2<Real>, Circle2<Real>>
+    :
+    public FIQuery<Real, Line2<Real>, Circle2<Real>>
+{
+public:
+    struct Result
+        :
+        public FIQuery<Real, Line2<Real>, Circle2<Real>>::Result
+    {
+        // No additional information to compute.
+    };
+
+    Result operator()(Ray2<Real> const& ray, Circle2<Real> const& circle);
+
+protected:
+    void DoQuery(Vector2<Real> const& rayOrigin,
+        Vector2<Real> const& rayDirection, Circle2<Real> const& circle,
+        Result& result);
+};
+
+
+template <typename Real>
+typename TIQuery<Real, Ray2<Real>, Circle2<Real>>::Result
+TIQuery<Real, Ray2<Real>, Circle2<Real>>::operator()(
+    Ray2<Real> const& ray, Circle2<Real> const& circle)
+{
+    Result result;
+    FIQuery<Real, Ray2<Real>, Circle2<Real>> rcQuery;
+    result.intersect = rcQuery(ray, circle).intersect;
+    return result;
+}
+
+template <typename Real>
+typename FIQuery<Real, Ray2<Real>, Circle2<Real>>::Result
+FIQuery<Real, Ray2<Real>, Circle2<Real>>::operator()(
+    Ray2<Real> const& ray, Circle2<Real> const& circle)
+{
+    Result result;
+    DoQuery(ray.origin, ray.direction, circle, result);
+    for (int i = 0; i < result.numIntersections; ++i)
+    {
+        result.point[i] = ray.origin + result.parameter[i] * ray.direction;
+    }
+    return result;
+}
+
+template <typename Real>
+void FIQuery<Real, Ray2<Real>, Circle2<Real>>::DoQuery(
+    Vector2<Real> const& rayOrigin, Vector2<Real> const& rayDirection,
+    Circle2<Real> const& circle, Result& result)
+{
+    FIQuery<Real, Line2<Real>, Circle2<Real>>::DoQuery(rayOrigin,
+        rayDirection, circle, result);
+
+    if (result.intersect)
+    {
+        // The line containing the ray intersects the disk; the t-interval is
+        // [t0,t1].  The ray intersects the disk as long as [t0,t1] overlaps
+        // the ray t-interval [0,+infinity).
+        std::array<Real, 2> rayInterval = { (Real)0, std::numeric_limits<Real>::max() };
+        FIQuery<Real, std::array<Real, 2>, std::array<Real, 2>> iiQuery;
+        auto iiResult = iiQuery(result.parameter, rayInterval);
+        result.intersect = iiResult.intersect;
+        result.numIntersections = iiResult.numIntersections;
+        result.parameter = iiResult.overlap;
+    }
+}
+
+}
diff --git a/Engine/Plugins/Experimental/GeometryProcessing/Source/GeometryAlgorithms/Private/GTEngine/Mathematics/GteIntrRay2OrientedBox2.h b/Engine/Plugins/Experimental/GeometryProcessing/Source/GeometryAlgorithms/Private/GTEngine/Mathematics/GteIntrRay2OrientedBox2.h
new file mode 100644
--- /dev/null
+++ b/Engine/Plugins/Experimental/GeometryProcessing/Source/GeometryAlgorithms/Private/GTEngine/Mathematics/GteIntrRay2OrientedBox2.h
@@ -0,0 +1,106 @@
+// David Eberly, Geometric Tools, Redmond WA 98052
+// Copyright (c) 1998-2019
+// Distributed under the Boost Software License, Version 1.0.
+// http://www.boost.org/LICENSE_1_0.txt
+// http://www.geometrictools.com/License/Boost/LICENSE_1_0.txt
+// File Version: 3.0.0 (2016/06/19)
+
+#pragma once
+
+#include <GTEngine/Mathematics/GteOrientedBox.h>
+#include <GTEngine/Mathematics/GteIntrRay2AlignedBox2.h>
+
+// The queries consider the box to be a solid.
+//
+// The test-intersection queries use the method of separating axes.  The
+// find-intersection queries use parametric clipping against the four edges of
+// the box.
+
+namespace gte
+{
+
+template <typename Real>
+class TIQuery<Real, Ray2<Real>, OrientedBox2<Real>>
+    :
+    public TIQuery<Real, Ray2<Real>, AlignedBox2<Real>>
+{
+public:
+    struct Result
+        :
+        public TIQuery<Real, Ray2<Real>, AlignedBox2<Real>>::Result
+    {
+        // No additional information to compute.
+    };
+
+    Result operator()(Ray2<Real> const& ray, OrientedBox2<Real> const& box);
+};
+
+template <typename Real>
+class FIQuery<Real, Ray2<Real>, OrientedBox2<Real>>
+    :
+    public FIQuery<Real, Ray2<Real>, AlignedBox2<Real>>
+{
+public:
+    struct Result
+        :
+        public FIQuery<Real, Ray2<Real>, AlignedBox2<Real>>::Result
+    {
+        // No additional information to compute.
+    };
+
+    Result operator()(Ray2<Real> const& ray, OrientedBox2<Real> const& box);
+};
+
+
+template <typename Real>
+typename TIQuery<Real, Ray2<Real>, OrientedBox2<Real>>::Result
+TIQuery<Real, Ray2<Real>, OrientedBox2<Real>>::operator()(
+    Ray2<Real> const& ray, OrientedBox2<Real> const& box)
+{
+    // Transform the ray to the oriented-box coordinate system.
+    Vector2<Real> diff = ray.origin - box.center;
+    Vector2<Real> rayOrigin
+    {
+        Dot(diff, box.axis[0]),
+        Dot(diff, box.axis[1])
+    };
+    Vector2<Real> rayDirection = Vector2<Real>
+    {
+        Dot(ray.direction, box.axis[0]),
+            Dot(ray.direction, box.axis[1])
+    };
+
+    Result result;
+    this->DoQuery(rayOrigin, rayDirection, box.extent, result);
+    return result;
+}
+
+template <typename Real>
+typename FIQuery<Real, Ray2<Real>, OrientedBox2<Real>>::Result
+FIQuery<Real, Ray2<Real>, OrientedBox2<Real>>::operator()(
+    Ray2<Real> const& ray, OrientedBox2<Real> const& box)
+{
+    // Transform the ray to the oriented-box coordinate system.
+    Vector2<Real> diff = ray.origin - box.center;
+    Vector2<Real> rayOrigin
+    {
+        Dot(diff, box.axis[0]),
+        Dot(diff, box.axis[1])
+    };
+    Vector2<Real> rayDirection = Vector2<Real>
+    {
+        Dot(ray.direction, box.axis[0]),
+            Dot(ray.direction, box.axis[1])
+    };
+
+    Result result;
+    this->DoQuery(rayOrigin, rayDirection, box.extent, result);
+    for (int i = 0; i < result.numIntersections; ++i)
+    {
+        result.point[i] = ray.origin + result.parameter[i] * ray.direction;
+    }
+    return result;
+}
+
+
+}
diff --git a/Engine/Plugins/Experimental/GeometryProcessing/Source/GeometryAlgorithms/Private/GTEngine/Mathematics/GteIntrRay2Ray2.h b/Engine/Plugins/Experimental/GeometryProcessing/Source/GeometryAlgorithms/Private/GTEngine/Mathematics/GteIntrRay2Ray2.h
new file mode 100644
--- /dev/null
+++ b/Engine/Plugins/Experimental/GeometryProcessing/Source/GeometryAlgorithms/Private/GTEngine/Mathematics/GteIntrRay2Ray2.h
@@ -0,0 +1,232 @@
+// David Eberly, Geometric Tools, Redmond WA 98052
+// Copyright (c) 1998-2019
+// Distributed under the Boost Software License, Version 1.0.
+// http://www.boost.org/LICENSE_1_0.txt
+// http://www.geometrictools.com/License/Boost/LICENSE_1_0.txt
+// File Version: 3.0.0 (2016/06/19)
+
+#pragma once
+
+#include <GTEngine/Mathematics/GteRay.h>
+#include <GTEngine/Mathematics/GteIntrLine2Line2.h>
+
+namespace gte
+{
+
+template <typename Real>
+class TIQuery<Real, Ray2<Real>, Ray2<Real>>
+{
+public:
+    struct Result
+    {
+        bool intersect;
+
+        // The number is 0 (no intersection), 1 (rays intersect in a
+        // single point), 2 (rays are collinear and intersect in a segment;
+        // ray directions are opposite of each other), or
+        // std::numeric_limits<int>::max() (intersection is a ray; ray
+        // directions are the same).
+        int numIntersections;
+    };
+
+    Result operator()(Ray2<Real> const& ray0, Ray2<Real> const& ray1);
+};
+
+template <typename Real>
+class FIQuery<Real, Ray2<Real>, Ray2<Real>>
+{
+public:
+    struct Result
+    {
+        bool intersect;
+
+        // The number is 0 (no intersection), 1 (rays intersect in a
+        // single point), 2 (rays are collinear and intersect in a segment;
+        // ray directions are opposite of each other), or
+        // std::numeric_limits<int>::max() (intersection is a ray; ray
+        // directions are the same).
+        int numIntersections;
+
+        // If numIntersections is 1, the intersection is
+        //   point[0] = ray0.origin + ray0Parameter[0] * ray0.direction
+        //            = ray1.origin + ray1Parameter[0] * ray1.direction
+        // If numIntersections is 2, the segment of intersection is formed by
+        // the ray origins,
+        //   ray0Parameter[0] = ray1Parameter[0] = 0
+        //   point[0] = ray0.origin
+        //            = ray1.origin + ray1Parameter[1] * ray1.direction
+        //   point[1] = ray1.origin
+        //            = ray0.origin + ray0Parameter[1] * ray0.direction
+        // where ray0Parameter[1] >= 0 and ray1Parameter[1] >= 0.
+        // If numIntersections is maxInt, let
+        //   ray1.origin = ray0.origin + t * ray0.direction
+        // then
+        //   ray0Parameter[] = { max(t,0), +maxReal }
+        //   ray1Parameter[] = { -min(t,0), +maxReal }
+        //   point[0] = ray0.origin + ray0Parameter[0] * ray0.direction
+        Real ray0Parameter[2], ray1Parameter[2];
+        Vector2<Real> point[2];
+    };
+
+    Result operator()(Ray2<Real> const& ray0, Ray2<Real> const& ray1);
+};
+
+
+template <typename Real>
+typename TIQuery<Real, Ray2<Real>, Ray2<Real>>::Result
+TIQuery<Real, Ray2<Real>, Ray2<Real>>::operator()(
+    Ray2<Real> const& ray0, Ray2<Real> const& ray1)
+{
+    Result result;
+    FIQuery<Real, Line2<Real>, Line2<Real>> llQuery;
+    Line2<Real> line0(ray0.origin, ray0.direction);
+    Line2<Real> line1(ray1.origin, ray1.direction);
+    auto llResult = llQuery(line0, line1);
+    if (llResult.numIntersections == 1)
+    {
+        // Test whether the line-line intersection is on the rays.
+        if (llResult.line0Parameter[0] >= (Real)0
+            && llResult.line1Parameter[0] >= (Real)0)
+        {
+            result.intersect = true;
+            result.numIntersections = 1;
+        }
+        else
+        {
+            result.intersect = false;
+            result.numIntersections = 0;
+        }
+    }
+    else if (llResult.numIntersections == std::numeric_limits<int>::max())
+    {
+        if (Dot(ray0.direction, ray1.direction) > (Real)0)
+        {
+            // The rays are collinear and in the same direction, so they must
+            // overlap.
+            result.intersect = true;
+            result.numIntersections = std::numeric_limits<int>::max();
+        }
+        else
+        {
+            // The rays are collinear but in opposite directions.  Test
+            // whether they overlap.  Ray0 has interval [0,+infinity) and
+            // ray1 has interval (-infinity,t] relative to ray0.direction.
+            Vector2<Real> diff = ray1.origin - ray0.origin;
+            Real t = Dot(ray0.direction, diff);
+            if (t > (Real)0)
+            {
+                result.intersect = true;
+                result.numIntersections = 2;
+            }
+            else if (t < (Real)0)
+            {
+                result.intersect = false;
+                result.numIntersections = 0;
+            }
+            else  // t == 0
+            {
+                result.intersect = true;
+                result.numIntersections = 1;
+            }
+        }
+    }
+    else
+    {
+        result.intersect = false;
+        result.numIntersections = 0;
+    }
+
+    return result;
+}
+
+template <typename Real>
+typename FIQuery<Real, Ray2<Real>, Ray2<Real>>::Result
+FIQuery<Real, Ray2<Real>, Ray2<Real>>::operator()(
+    Ray2<Real> const& ray0, Ray2<Real> const& ray1)
+{
+    Result result;
+    FIQuery<Real, Line2<Real>, Line2<Real>> llQuery;
+    Line2<Real> line0(ray0.origin, ray0.direction);
+    Line2<Real> line1(ray1.origin, ray1.direction);
+    auto llResult = llQuery(line0, line1);
+    if (llResult.numIntersections == 1)
+    {
+        // Test whether the line-line intersection is on the rays.
+        if (llResult.line0Parameter[0] >= (Real)0
+            && llResult.line1Parameter[0] >= (Real)0)
+        {
+            result.intersect = true;
+            result.numIntersections = 1;
+            result.ray0Parameter[0] = llResult.line0Parameter[0];
+            result.ray1Parameter[0] = llResult.line1Parameter[0];
+            result.point[0] = llResult.point;
+        }
+        else
+        {
+            result.intersect = false;
+            result.numIntersections = 0;
+        }
+    }
+    else if (llResult.numIntersections == std::numeric_limits<int>::max())
+    {
+        // Compute t for which ray1.origin = ray0.origin + t*ray0.direction.
+        Real maxReal = std::numeric_limits<Real>::max();
+        Vector2<Real> diff = ray1.origin - ray0.origin;
+        Real t = Dot(ray0.direction, diff);
+        if (Dot(ray0.direction, ray1.direction) > (Real)0)
+        {
+            // The rays are collinear and in the same direction, so they must
+            // overlap.
+            result.intersect = true;
+            result.numIntersections = std::numeric_limits<int>::max();
+            if (t >= (Real)0)
+            {
+                result.ray0Parameter[0] = t;
+                result.ray0Parameter[1] = maxReal;
+                result.ray1Parameter[0] = (Real)0;
+                result.ray1Parameter[1] = maxReal;
+                result.point[0] = ray1.origin;
+            }
+            else
+            {
+                result.ray0Parameter[0] = (Real)0;
+                result.ray0Parameter[1] = maxReal;
+                result.ray1Parameter[0] = -t;
+                result.ray1Parameter[1] = maxReal;
+                result.point[0] = ray0.origin;
+            }
+        }
+        else
+        {
+            // The rays are collinear but in opposite directions.  Test
+            // whether they overlap.  Ray0 has interval [0,+infinity) and
+            // ray1 has interval (-infinity,t1] relative to ray0.direction.
+            if (t >= (Real)0)
+            {
+                result.intersect = true;
+                result.numIntersections = 2;
+                result.ray0Parameter[0] = (Real)0;
+                result.ray0Parameter[1] = t;
+                result.ray1Parameter[0] = (Real)0;
+                result.ray1Parameter[1] = t;
+                result.point[0] = ray0.origin;
+                result.point[1] = ray1.origin;
+            }
+            else
+            {
+                result.intersect = false;
+                result.numIntersections = 0;
+            }
+        }
+    }
+    else
+    {
+        result.intersect = false;
+        result.numIntersections = 0;
+    }
+
+    return result;
+}
+
+
+}
diff --git a/Engine/Plugins/Experimental/GeometryProcessing/Source/GeometryAlgorithms/Private/GTEngine/Mathematics/GteIntrRay2Segment2.h b/Engine/Plugins/Experimental/GeometryProcessing/Source/GeometryAlgorithms/Private/GTEngine/Mathematics/GteIntrRay2Segment2.h
new file mode 100644
--- /dev/null
+++ b/Engine/Plugins/Experimental/GeometryProcessing/Source/GeometryAlgorithms/Private/GTEngine/Mathematics/GteIntrRay2Segment2.h
@@ -0,0 +1,202 @@
+// David Eberly, Geometric Tools, Redmond WA 98052
+// Copyright (c) 1998-2019
+// Distributed under the Boost Software License, Version 1.0.
+// http://www.boost.org/LICENSE_1_0.txt
+// http://www.geometrictools.com/License/Boost/LICENSE_1_0.txt
+// File Version: 3.0.0 (2016/06/19)
+
+#pragma once
+
+#include <GTEngine/Mathematics/GteIntrIntervals.h>
+#include <GTEngine/Mathematics/GteIntrLine2Line2.h>
+#include <GTEngine/Mathematics/GteRay.h>
+#include <GTEngine/Mathematics/GteSegment.h>
+
+namespace gte
+{
+
+template <typename Real>
+class TIQuery<Real, Ray2<Real>, Segment2<Real>>
+{
+public:
+    struct Result
+    {
+        bool intersect;
+
+        // The number is 0 (no intersection), 1 (ray and segment intersect in
+        // a single point), or 2 (ray and segment are collinear and intersect
+        // in a segment).
+        int numIntersections;
+    };
+
+    Result operator()(Ray2<Real> const& ray, Segment2<Real> const& segment);
+};
+
+template <typename Real>
+class FIQuery<Real, Ray2<Real>, Segment2<Real>>
+{
+public:
+    struct Result
+    {
+        bool intersect;
+
+        // The number is 0 (no intersection), 1 (ray and segment intersect in
+        // a single point), or 2 (ray and segment are collinear and intersect
+        // in a segment).
+        int numIntersections;
+
+        // If numIntersections is 1, the intersection is
+        //   point[0] = ray.origin + rayParameter[0] * ray.direction
+        //     = segment.center + segmentParameter[0] * segment.direction
+        // If numIntersections is 2, the endpoints of the segment of
+        // intersection are
+        //   point[i] = ray.origin + rayParameter[i] * ray.direction
+        //     = segment.center + segmentParameter[i] * segment.direction
+        // with rayParameter[0] <= rayParameter[1] and
+        // segmentParameter[0] <= segmentParameter[1].
+        Real rayParameter[2], segmentParameter[2];
+        Vector2<Real> point[2];
+    };
+
+    Result operator()(Ray2<Real> const& ray, Segment2<Real> const& segment);
+};
+
+
+template <typename Real>
+typename TIQuery<Real, Ray2<Real>, Segment2<Real>>::Result
+TIQuery<Real, Ray2<Real>, Segment2<Real>>::operator()(
+    Ray2<Real> const& ray, Segment2<Real> const& segment)
+{
+    Result result;
+    Vector2<Real> segOrigin, segDirection;
+    Real segExtent;
+    segment.GetCenteredForm(segOrigin, segDirection, segExtent);
+
+    FIQuery<Real, Line2<Real>, Line2<Real>> llQuery;
+    Line2<Real> line0(ray.origin, ray.direction);
+    Line2<Real> line1(segOrigin, segDirection);
+    auto llResult = llQuery(line0, line1);
+    if (llResult.numIntersections == 1)
+    {
+        // Test whether the line-line intersection is on the ray and segment.
+        if (llResult.line0Parameter[0] >= (Real)0
+            && std::abs(llResult.line1Parameter[0]) <= segExtent)
+        {
+            result.intersect = true;
+            result.numIntersections = 1;
+        }
+        else
+        {
+            result.intersect = false;
+            result.numIntersections = 0;
+        }
+    }
+    else if (llResult.numIntersections == std::numeric_limits<int>::max())
+    {
+        // Compute the location of the right-most point of the segment
+        // relative to the ray direction.
+        Vector2<Real> diff = segOrigin - ray.origin;
+        Real t = Dot(ray.direction, diff) + segExtent;
+        if (t > (Real)0)
+        {
+            result.intersect = true;
+            result.numIntersections = 2;
+        }
+        else if (t < (Real)0)
+        {
+            result.intersect = false;
+            result.numIntersections = 0;
+        }
+        else  // t == 0
+        {
+            result.intersect = true;
+            result.numIntersections = 1;
+        }
+    }
+    else
+    {
+        result.intersect = false;
+        result.numIntersections = 0;
+    }
+
+    return result;
+}
+
+template <typename Real>
+typename FIQuery<Real, Ray2<Real>, Segment2<Real>>::Result
+FIQuery<Real, Ray2<Real>, Segment2<Real>>::operator()(
+    Ray2<Real> const& ray, Segment2<Real> const& segment)
+{
+    Result result;
+    Vector2<Real> segOrigin, segDirection;
+    Real segExtent;
+    segment.GetCenteredForm(segOrigin, segDirection, segExtent);
+
+    FIQuery<Real, Line2<Real>, Line2<Real>> llQuery;
+    Line2<Real> line0(ray.origin, ray.direction);
+    Line2<Real> line1(segOrigin, segDirection);
+    auto llResult = llQuery(line0, line1);
+    if (llResult.numIntersections == 1)
+    {
+        // Test whether the line-line intersection is on the ray and segment.
+        if (llResult.line0Parameter[0] >= (Real)0
+            && std::abs(llResult.line1Parameter[0]) <= segExtent)
+        {
+            result.intersect = true;
+            result.numIntersections = 1;
+            result.rayParameter[0] = llResult.line0Parameter[0];
+            result.segmentParameter[0] = llResult.line1Parameter[0];
+            result.point[0] = llResult.point;
+        }
+        else
+        {
+            result.intersect = false;
+            result.numIntersections = 0;
+        }
+    }
+    else if (llResult.numIntersections == std::numeric_limits<int>::max())
+    {
+        // Compute t for which segment.origin = ray.origin + t*ray.direction.
+        Vector2<Real> diff = segOrigin - ray.origin;
+        Real t = Dot(ray.direction, diff);
+
+        // Get the ray interval.
+        std::array<Real, 2> interval0 =
+        {
+            (Real)0, std::numeric_limits<Real>::max()
+        };
+
+        // Compute the location of the segment endpoints relative to the ray.
+        std::array<Real, 2> interval1 = { t - segExtent, t + segExtent };
+
+        // Compute the intersection of [0,+infinity) and [tmin,tmax].
+        FIQuery<Real, std::array<Real, 2>, std::array<Real, 2>> iiQuery;
+        auto iiResult = iiQuery(interval0, interval1);
+        if (iiResult.intersect)
+        {
+            result.intersect = true;
+            result.numIntersections = iiResult.numIntersections;
+            for (int i = 0; i < iiResult.numIntersections; ++i)
+            {
+                result.rayParameter[i] = iiResult.overlap[i];
+                result.segmentParameter[i] = iiResult.overlap[i] - t;
+                result.point[i] = ray.origin +
+                    result.rayParameter[i] * ray.direction;
+            }
+        }
+        else
+        {
+            result.intersect = false;
+            result.numIntersections = 0;
+        }
+    }
+    else
+    {
+        result.intersect = false;
+        result.numIntersections = 0;
+    }
+    return result;
+}
+
+
+}
diff --git a/Engine/Plugins/Experimental/GeometryProcessing/Source/GeometryAlgorithms/Private/GTEngine/Mathematics/GteIntrRay2Triangle2.h b/Engine/Plugins/Experimental/GeometryProcessing/Source/GeometryAlgorithms/Private/GTEngine/Mathematics/GteIntrRay2Triangle2.h
new file mode 100644
--- /dev/null
+++ b/Engine/Plugins/Experimental/GeometryProcessing/Source/GeometryAlgorithms/Private/GTEngine/Mathematics/GteIntrRay2Triangle2.h
@@ -0,0 +1,99 @@
+// David Eberly, Geometric Tools, Redmond WA 98052
+// Copyright (c) 1998-2019
+// Distributed under the Boost Software License, Version 1.0.
+// http://www.boost.org/LICENSE_1_0.txt
+// http://www.geometrictools.com/License/Boost/LICENSE_1_0.txt
+// File Version: 3.0.0 (2016/06/19)
+
+#pragma once
+
+#include <GTEngine/Mathematics/GteRay.h>
+#include <GTEngine/Mathematics/GteIntrIntervals.h>
+#include <GTEngine/Mathematics/GteIntrLine2Triangle2.h>
+
+// The queries consider the triangle to be a solid.
+
+namespace gte
+{
+
+template <typename Real>
+class TIQuery<Real, Ray2<Real>, Triangle2<Real>>
+{
+public:
+    struct Result
+    {
+        bool intersect;
+    };
+
+    Result operator()(Ray2<Real> const& ray, Triangle2<Real> const& triangle);
+};
+
+template <typename Real>
+class FIQuery<Real, Ray2<Real>, Triangle2<Real>>
+    :
+    public FIQuery<Real, Line2<Real>, Triangle2<Real>>
+{
+public:
+    struct Result
+        :
+        public FIQuery<Real, Line2<Real>, Triangle2<Real>>::Result
+    {
+        // No additional information to compute.
+    };
+
+    Result operator()(Ray2<Real> const& ray, Triangle2<Real> const& triangle);
+
+protected:
+    void DoQuery(Vector2<Real> const& rayOrigin,
+        Vector2<Real> const& rayDirection, Triangle2<Real> const& triangle,
+        Result& result);
+};
+
+
+template <typename Real>
+typename TIQuery<Real, Ray2<Real>, Triangle2<Real>>::Result
+TIQuery<Real, Ray2<Real>, Triangle2<Real>>::operator()(
+    Ray2<Real> const& ray, Triangle2<Real> const& triangle)
+{
+    Result result;
+    FIQuery<Real, Ray2<Real>, Triangle2<Real>> rtQuery;
+    result.intersect = rtQuery(ray, triangle).intersect;
+    return result;
+}
+
+template <typename Real>
+typename FIQuery<Real, Ray2<Real>, Triangle2<Real>>::Result
+FIQuery<Real, Ray2<Real>, Triangle2<Real>>::operator()(
+    Ray2<Real> const& ray, Triangle2<Real> const& triangle)
+{
+    Result result;
+    DoQuery(ray.origin, ray.direction, triangle, result);
+    for (int i = 0; i < result.numIntersections; ++i)
+    {
+        result.point[i] = ray.origin + result.parameter[i] * ray.direction;
+    }
+    return result;
+}
+
+template <typename Real>
+void FIQuery<Real, Ray2<Real>, Triangle2<Real>>::DoQuery(
+    Vector2<Real> const& rayOrigin, Vector2<Real> const& rayDirection,
+    Triangle2<Real> const& triangle, Result& result)
+{
+    FIQuery<Real, Line2<Real>, Triangle2<Real>>::DoQuery(rayOrigin,
+        rayDirection, triangle, result);
+
+    if (result.intersect)
+    {
+        // The line containing the ray intersects the disk; the t-interval is
+        // [t0,t1].  The ray intersects the disk as long as [t0,t1] overlaps
+        // the ray t-interval [0,+infinity).
+        std::array<Real, 2> rayInterval =
+        { (Real)0, std::numeric_limits<Real>::max() };
+        FIQuery<Real, std::array<Real, 2>, std::array<Real, 2>> iiQuery;
+        result.parameter = iiQuery(result.parameter, rayInterval).overlap;
+    }
+}
+
+
+}
diff --git a/Engine/Plugins/Experimental/GeometryProcessing/Source/GeometryAlgorithms/Private/GTEngine/Mathematics/GteIntrRay3AlignedBox3.h b/Engine/Plugins/Experimental/GeometryProcessing/Source/GeometryAlgorithms/Private/GTEngine/Mathematics/GteIntrRay3AlignedBox3.h
new file mode 100644
--- /dev/null
+++ b/Engine/Plugins/Experimental/GeometryProcessing/Source/GeometryAlgorithms/Private/GTEngine/Mathematics/GteIntrRay3AlignedBox3.h
@@ -0,0 +1,123 @@
+// David Eberly, Geometric Tools, Redmond WA 98052
+// Copyright (c) 1998-2019
+// Distributed under the Boost Software License, Version 1.0.
+// http://www.boost.org/LICENSE_1_0.txt
+// http://www.geometrictools.com/License/Boost/LICENSE_1_0.txt
+// File Version: 3.0.1 (2018/11/28)
+
+#pragma once
+
+#include <GTEngine/Mathematics/GteRay.h>
+#include <GTEngine/Mathematics/GteIntrLine3AlignedBox3.h>
+
+// The test-intersection queries use the method of separating axes.  The
+// find-intersection queries use parametric clipping against the six faces of
+// the box.  The find-intersection queries use Liang-Barsky clipping.  The
+// queries consider the box to be a solid.  The algorithms are described in
+//   https://www.geometrictools.com/Documentation/IntersectionLineBox.pdf
+
+namespace gte
+{
+    template <typename Real>
+    class TIQuery<Real, Ray3<Real>, AlignedBox3<Real>>
+        :
+        public TIQuery<Real, Line3<Real>, AlignedBox3<Real>>
+    {
+    public:
+        struct Result
+            :
+            public TIQuery<Real, Line3<Real>, AlignedBox3<Real>>::Result
+        {
+            // No additional information to compute.
+        };
+
+        Result operator()(Ray3<Real> const& ray, AlignedBox3<Real> const& box)
+        {
+            // Get the centered form of the aligned box.  The axes are
+            // implicitly Axis[d] = Vector3<Real>::Unit(d).
+            Vector3<Real> boxCenter, boxExtent;
+            box.GetCenteredForm(boxCenter, boxExtent);
+
+            // Transform the ray to the aligned-box coordinate system.
+            Vector3<Real> rayOrigin = ray.origin - boxCenter;
+
+            Result result;
+            DoQuery(rayOrigin, ray.direction, boxExtent, result);
+            return result;
+        }
+
+    protected:
+        void DoQuery(Vector3<Real> const& rayOrigin, Vector3<Real> const& rayDirection,
+            Vector3<Real> const& boxExtent, Result& result)
+        {
+            for (int i = 0; i < 3; ++i)
+            {
+                if (std::abs(rayOrigin[i]) > boxExtent[i] && rayOrigin[i] * rayDirection[i] >= (Real)0)
+                {
+                    result.intersect = false;
+                    return;
+                }
+            }
+
+            TIQuery<Real, Line3<Real>, AlignedBox3<Real>>::DoQuery(rayOrigin, rayDirection, boxExtent, result);
+        }
+    };
+
+    template <typename Real>
+    class FIQuery<Real, Ray3<Real>, AlignedBox3<Real>>
+        :
+        public FIQuery<Real, Line3<Real>, AlignedBox3<Real>>
+    {
+    public:
+        struct Result
+            :
+            public FIQuery<Real, Line3<Real>, AlignedBox3<Real>>::Result
+        {
+            // No additional information to compute.
+        };
+
+        Result operator()(Ray3<Real> const& ray, AlignedBox3<Real> const& box)
+        {
+            // Get the centered form of the aligned box.  The axes are
+            // implicitly Axis[d] = Vector3<Real>::Unit(d).
+            Vector3<Real> boxCenter, boxExtent;
+            box.GetCenteredForm(boxCenter, boxExtent);
+
+            // Transform the ray to the aligned-box coordinate system.
+            Vector3<Real> rayOrigin = ray.origin - boxCenter;
+
+            Result result;
+            DoQuery(rayOrigin, ray.direction, boxExtent, result);
+            for (int i = 0; i < result.numPoints; ++i)
+            {
+                result.point[i] = ray.origin + result.lineParameter[i] * ray.direction;
+            }
+            return result;
+        }
+
+    protected:
+        void DoQuery(Vector3<Real> const& rayOrigin, Vector3<Real> const& rayDirection,
+            Vector3<Real> const& boxExtent, Result& result)
+        {
+            FIQuery<Real, Line3<Real>, AlignedBox3<Real>>::DoQuery(rayOrigin, rayDirection, boxExtent, result);
+            if (result.intersect)
+            {
+                // The line containing the ray intersects the box; the
+                // t-interval is [t0,t1].  The ray intersects the box as long
+                // as [t0,t1] overlaps the ray t-interval (0,+infinity).
+                if (result.lineParameter[1] >= (Real)0)
+                {
+                    if (result.lineParameter[0] < (Real)0)
+                    {
+                        result.lineParameter[0] = (Real)0;
+                    }
+                }
+                else
+                {
+                    result.intersect = false;
+                    result.numPoints = 0;
+                }
+            }
+        }
+    };
+}
diff --git a/Engine/Plugins/Experimental/GeometryProcessing/Source/GeometryAlgorithms/Private/GTEngine/Mathematics/GteIntrRay3Capsule3.h b/Engine/Plugins/Experimental/GeometryProcessing/Source/GeometryAlgorithms/Private/GTEngine/Mathematics/GteIntrRay3Capsule3.h
new file mode 100644
--- /dev/null
+++ b/Engine/Plugins/Experimental/GeometryProcessing/Source/GeometryAlgorithms/Private/GTEngine/Mathematics/GteIntrRay3Capsule3.h
@@ -0,0 +1,112 @@
+// David Eberly, Geometric Tools, Redmond WA 98052
+// Copyright (c) 1998-2019
+// Distributed under the Boost Software License, Version 1.0.
+// http://www.boost.org/LICENSE_1_0.txt
+// http://www.geometrictools.com/License/Boost/LICENSE_1_0.txt
+// File Version: 3.0.0 (2016/06/19)
+
+#pragma once
+
+#include <GTEngine/Mathematics/GteDistRaySegment.h>
+#include <GTEngine/Mathematics/GteIntrIntervals.h>
+#include <GTEngine/Mathematics/GteIntrLine3Capsule3.h>
+
+// The queries consider the capsule to be a solid.
+//
+// The test-intersection queries are based on distance computations.
+
+namespace gte
+{
+
+template <typename Real>
+class TIQuery<Real, Ray3<Real>, Capsule3<Real>>
+{
+public:
+    struct Result
+    {
+        bool intersect;
+    };
+
+    Result operator()(Ray3<Real> const& ray, Capsule3<Real> const& capsule);
+};
+
+template <typename Real>
+class FIQuery<Real, Ray3<Real>, Capsule3<Real>>
+    :
+    public FIQuery<Real, Line3<Real>, Capsule3<Real>>
+{
+public:
+    struct Result
+        :
+        public FIQuery<Real, Line3<Real>, Capsule3<Real>>::Result
+    {
+        // No additional information to compute.
+    };
+
+    Result operator()(Ray3<Real> const& ray, Capsule3<Real> const& capsule);
+
+protected:
+    void DoQuery(Vector3<Real> const& rayOrigin,
+        Vector3<Real> const& rayDirection, Capsule3<Real> const& capsule,
+        Result& result);
+};
+
+
+template <typename Real>
+typename TIQuery<Real, Ray3<Real>, Capsule3<Real>>::Result
+TIQuery<Real, Ray3<Real>, Capsule3<Real>>::operator()(
+    Ray3<Real> const& ray, Capsule3<Real> const& capsule)
+{
+    Result result;
+    DCPQuery<Real, Ray3<Real>, Segment3<Real>> rsQuery;
+    auto rsResult = rsQuery(ray, capsule.segment);
+    result.intersect = (rsResult.distance <= capsule.radius);
+    return result;
+}
+
+template <typename Real>
+typename FIQuery<Real, Ray3<Real>, Capsule3<Real>>::Result
+FIQuery<Real, Ray3<Real>, Capsule3<Real>>::operator()(
+    Ray3<Real> const& ray, Capsule3<Real> const& capsule)
+{
+    Result result;
+    DoQuery(ray.origin, ray.direction, capsule, result);
+    for (int i = 0; i < result.numIntersections; ++i)
+    {
+        result.point[i] = ray.origin + result.parameter[i] * ray.direction;
+    }
+    return result;
+}
+
+template <typename Real>
+void FIQuery<Real, Ray3<Real>, Capsule3<Real>>::DoQuery(
+    Vector3<Real> const& rayOrigin, Vector3<Real> const& rayDirection,
+    Capsule3<Real> const& capsule, Result& result)
+{
+    FIQuery<Real, Line3<Real>, Capsule3<Real>>::DoQuery(rayOrigin,
+        rayDirection, capsule, result);
+
+    if (result.intersect)
+    {
+        // The line containing the ray intersects the capsule; the t-interval
+        // is [t0,t1].  The ray intersects the capsule as long as [t0,t1]
+        // overlaps the ray t-interval [0,+infinity).
+        std::array<Real, 2> rayInterval =
+        { (Real)0, std::numeric_limits<Real>::max() };
+        FIQuery<Real, std::array<Real, 2>, std::array<Real, 2>> iiQuery;
+        auto iiResult = iiQuery(result.parameter, rayInterval);
+        if (iiResult.intersect)
+        {
+            result.numIntersections = iiResult.numIntersections;
+            result.parameter = iiResult.overlap;
+        }
+        else
+        {
+            result.intersect = false;
+            result.numIntersections = 0;
+        }
+    }
+}
+
+
+}
diff --git a/Engine/Plugins/Experimental/GeometryProcessing/Source/GeometryAlgorithms/Private/GTEngine/Mathematics/GteIntrRay3Cone3.h b/Engine/Plugins/Experimental/GeometryProcessing/Source/GeometryAlgorithms/Private/GTEngine/Mathematics/GteIntrRay3Cone3.h
new file mode 100644
--- /dev/null
+++ b/Engine/Plugins/Experimental/GeometryProcessing/Source/GeometryAlgorithms/Private/GTEngine/Mathematics/GteIntrRay3Cone3.h
@@ -0,0 +1,105 @@
+// David Eberly, Geometric Tools, Redmond WA 98052
+// Copyright (c) 1998-2019
+// Distributed under the Boost Software License, Version 1.0.
+// http://www.boost.org/LICENSE_1_0.txt
+// http://www.geometrictools.com/License/Boost/LICENSE_1_0.txt
+// File Version: 3.0.0 (2016/06/19)
+
+#pragma once
+
+#include <GTEngine/Mathematics/GteRay.h>
+#include <GTEngine/Mathematics/GteIntrLine3Cone3.h>
+
+// The queries consider the cone to be single sided and solid.
+
+namespace gte
+{
+
+template <typename Real>
+class FIQuery<Real, Ray3<Real>, Cone3<Real>>
+    :
+    public FIQuery<Real, Line3<Real>, Cone3<Real>>
+{
+public:
+    struct Result
+        :
+        public FIQuery<Real, Line3<Real>, Cone3<Real>>::Result
+    {
+        // No additional information to compute.
+    };
+
+    Result operator()(Ray3<Real> const& ray, Cone3<Real> const& cone);
+
+protected:
+    void DoQuery(Vector3<Real> const& rayOrigin,
+        Vector3<Real> const& rayDirection, Cone3<Real> const& cone,
+        Result& result);
+};
+
+
+template <typename Real>
+typename FIQuery<Real, Ray3<Real>, Cone3<Real>>::Result
+FIQuery<Real, Ray3<Real>, Cone3<Real>>::operator()(Ray3<Real> const& ray,
+    Cone3<Real> const& cone)
+{
+    Result result;
+    DoQuery(ray.origin, ray.direction, cone, result);
+    switch (result.type)
+    {
+    case 1:  // point
+        result.point[0] = ray.origin + result.parameter[0] * ray.direction;
+        result.point[1] = result.point[0];
+        break;
+    case 2:  // segment
+        result.point[0] = ray.origin + result.parameter[0] * ray.direction;
+        result.point[1] = ray.origin + result.parameter[1] * ray.direction;
+        break;
+    case 3:  // ray
+        result.point[0] = ray.origin + result.parameter[0] * ray.direction;
+        result.point[1] = ray.direction;
+        break;
+    default:  // no intersection
+        break;
+    }
+    return result;
+}
+
+template <typename Real>
+void FIQuery<Real, Ray3<Real>, Cone3<Real>>::DoQuery(
+    Vector3<Real> const& rayOrigin, Vector3<Real> const& rayDirection,
+    Cone3<Real> const& cone, Result& result)
+{
+    FIQuery<Real, Line3<Real>, Cone3<Real>>::DoQuery(rayOrigin,
+        rayDirection, cone, result);
+
+    if (result.intersect)
+    {
+        // The line containing the ray intersects the cone; the t-interval
+        // is [t0,t1].  The ray intersects the cone as long as [t0,t1]
+        // overlaps the ray t-interval [0,+infinity).
+        std::array<Real, 2> rayInterval = {
+            (Real)0, std::numeric_limits<Real>::max() };
+        FIIntervalInterval<Real> iiQuery;
+        auto iiResult = iiQuery(result.parameter, rayInterval);
+        if (iiResult.intersect)
+        {
+            result.parameter = iiResult.overlap;
+            if (result.parameter[1] < std::numeric_limits<Real>::max())
+            {
+                result.type = iiResult.numIntersections;
+            }
+            else
+            {
+                result.type = 3;
+            }
+        }
+        else
+        {
+            result.intersect = false;
+            result.type = 0;
+        }
+    }
+}
+
+
+}
diff --git a/Engine/Plugins/Experimental/GeometryProcessing/Source/GeometryAlgorithms/Private/GTEngine/Mathematics/GteIntrRay3Cylinder3.h b/Engine/Plugins/Experimental/GeometryProcessing/Source/GeometryAlgorithms/Private/GTEngine/Mathematics/GteIntrRay3Cylinder3.h
new file mode 100644
--- /dev/null
+++ b/Engine/Plugins/Experimental/GeometryProcessing/Source/GeometryAlgorithms/Private/GTEngine/Mathematics/GteIntrRay3Cylinder3.h
@@ -0,0 +1,86 @@
+// David Eberly, Geometric Tools, Redmond WA 98052
+// Copyright (c) 1998-2019
+// Distributed under the Boost Software License, Version 1.0.
+// http://www.boost.org/LICENSE_1_0.txt
+// http://www.geometrictools.com/License/Boost/LICENSE_1_0.txt
+// File Version: 3.0.0 (2016/06/19)
+
+#pragma once
+
+#include <GTEngine/Mathematics/GteRay.h>
+#include <GTEngine/Mathematics/GteIntrIntervals.h>
+#include <GTEngine/Mathematics/GteIntrLine3Cylinder3.h>
+
+// The queries consider the cylinder to be a solid.
+
+namespace gte
+{
+
+template <typename Real>
+class FIQuery<Real, Ray3<Real>, Cylinder3<Real>>
+    :
+    public FIQuery<Real, Line3<Real>, Cylinder3<Real>>
+{
+public:
+    struct Result
+        :
+        public FIQuery<Real, Line3<Real>, Cylinder3<Real>>::Result
+    {
+        // No additional information to compute.
+    };
+
+    Result operator()(Ray3<Real> const& ray, Cylinder3<Real> const& cylinder);
+
+protected:
+    void DoQuery(Vector3<Real> const& rayOrigin,
+        Vector3<Real> const& rayDirection, Cylinder3<Real> const& cylinder,
+        Result& result);
+};
+
+
+template <typename Real>
+typename FIQuery<Real, Ray3<Real>, Cylinder3<Real>>::Result
+FIQuery<Real, Ray3<Real>, Cylinder3<Real>>::operator()(
+    Ray3<Real> const& ray, Cylinder3<Real> const& cylinder)
+{
+    Result result;
+    DoQuery(ray.origin, ray.direction, cylinder, result);
+    for (int i = 0; i < result.numIntersections; ++i)
+    {
+        result.point[i] = ray.origin + result.parameter[i] * ray.direction;
+    }
+    return result;
+}
+
+template <typename Real>
+void FIQuery<Real, Ray3<Real>, Cylinder3<Real>>::DoQuery(
+    Vector3<Real> const& rayOrigin, Vector3<Real> const& rayDirection,
+    Cylinder3<Real> const& cylinder, Result& result)
+{
+    FIQuery<Real, Line3<Real>, Cylinder3<Real>>::DoQuery(rayOrigin,
+        rayDirection, cylinder, result);
+
+    if (result.intersect)
+    {
+        // The line containing the ray intersects the cylinder; the t-interval
+        // is [t0,t1].  The ray intersects the cylinder as long as [t0,t1]
+        // overlaps the ray t-interval [0,+infinity).
+        std::array<Real, 2> rayInterval =
+        { (Real)0, std::numeric_limits<Real>::max() };
+        FIQuery<Real, std::array<Real, 2>, std::array<Real, 2>> iiQuery;
+        auto iiResult = iiQuery(result.parameter, rayInterval);
+        if (iiResult.intersect)
+        {
+            result.numIntersections = iiResult.numIntersections;
+            result.parameter = iiResult.overlap;
+        }
+        else
+        {
+            result.intersect = false;
+            result.numIntersections = 0;
+        }
+    }
+}
+
+
+}
diff --git a/Engine/Plugins/Experimental/GeometryProcessing/Source/GeometryAlgorithms/Private/GTEngine/Mathematics/GteIntrRay3Ellipsoid3.h b/Engine/Plugins/Experimental/GeometryProcessing/Source/GeometryAlgorithms/Private/GTEngine/Mathematics/GteIntrRay3Ellipsoid3.h
new file mode 100644
--- /dev/null
+++ b/Engine/Plugins/Experimental/GeometryProcessing/Source/GeometryAlgorithms/Private/GTEngine/Mathematics/GteIntrRay3Ellipsoid3.h
@@ -0,0 +1,148 @@
+// David Eberly, Geometric Tools, Redmond WA 98052
+// Copyright (c) 1998-2019
+// Distributed under the Boost Software License, Version 1.0.
+// http://www.boost.org/LICENSE_1_0.txt
+// http://www.geometrictools.com/License/Boost/LICENSE_1_0.txt
+// File Version: 3.0.0 (2016/06/19)
+
+#pragma once
+
+#include <GTEngine/Mathematics/GteMatrix3x3.h>
+#include <GTEngine/Mathematics/GteRay.h>
+#include <GTEngine/Mathematics/GteIntrIntervals.h>
+#include <GTEngine/Mathematics/GteIntrLine3Ellipsoid3.h>
+
+// The queries consider the ellipsoid to be a solid.
+
+namespace gte
+{
+
+template <typename Real>
+class TIQuery<Real, Ray3<Real>, Ellipsoid3<Real>>
+{
+public:
+    struct Result
+    {
+        bool intersect;
+    };
+
+    Result operator()(Ray3<Real> const& ray,
+        Ellipsoid3<Real> const& ellipsoid);
+};
+
+template <typename Real>
+class FIQuery<Real, Ray3<Real>, Ellipsoid3<Real>>
+    :
+    public FIQuery<Real, Line3<Real>, Ellipsoid3<Real>>
+{
+public:
+    struct Result
+        :
+        public FIQuery<Real, Line3<Real>, Ellipsoid3<Real>>::Result
+    {
+        // No additional information to compute.
+    };
+
+    Result operator()(Ray3<Real> const& ray,
+        Ellipsoid3<Real> const& ellipsoid);
+
+protected:
+    void DoQuery(Vector3<Real> const& rayOrigin,
+        Vector3<Real> const& rayDirection, Ellipsoid3<Real> const& ellipsoid,
+        Result& result);
+};
+
+
+template <typename Real>
+typename TIQuery<Real, Ray3<Real>, Ellipsoid3<Real>>::Result
+TIQuery<Real, Ray3<Real>, Ellipsoid3<Real>>::operator()(
+    Ray3<Real> const& ray, Ellipsoid3<Real> const& ellipsoid)
+{
+    // The ellipsoid is (X-K)^T*M*(X-K)-1 = 0 and the line is X = P+t*D.
+    // Substitute the line equation into the ellipsoid equation to obtain
+    // a quadratic equation Q(t) = a2*t^2 + 2*a1*t + a0 = 0, where
+    // a2 = D^T*M*D, a1 = D^T*M*(P-K), and a0 = (P-K)^T*M*(P-K)-1.
+    Result result;
+
+    Matrix3x3<Real> M;
+    ellipsoid.GetM(M);
+
+    Vector3<Real> diff = ray.origin - ellipsoid.center;
+    Vector3<Real> matDir = M*ray.direction;
+    Vector3<Real> matDiff = M*diff;
+    Real a2 = Dot(ray.direction, matDir);
+    Real a1 = Dot(ray.direction, matDiff);
+    Real a0 = Dot(diff, matDiff) - (Real)1;
+
+    Real discr = a1*a1 - a0*a2;
+    if (discr >= (Real)0)
+    {
+        // Test whether ray origin is inside ellipsoid.
+        if (a0 <= (Real)0)
+        {
+            result.intersect = true;
+        }
+        else
+        {
+            // At this point, Q(0) = a0 > 0 and Q(t) has real roots.  It is
+            // also the case that a2 > 0, since M is positive definite,
+            // implying that D^T*M*D > 0 for any nonzero vector D.  Thus,
+            // an intersection occurs only when Q'(0) < 0.
+            result.intersect = (a1 < (Real)0);
+        }
+    }
+    else
+    {
+        // No intersection if Q(t) has no real roots.
+        result.intersect = false;
+    }
+
+    return result;
+}
+
+template <typename Real>
+typename FIQuery<Real, Ray3<Real>, Ellipsoid3<Real>>::Result
+FIQuery<Real, Ray3<Real>, Ellipsoid3<Real>>::operator()(
+    Ray3<Real> const& ray, Ellipsoid3<Real> const& ellipsoid)
+{
+    Result result;
+    DoQuery(ray.origin, ray.direction, ellipsoid, result);
+    for (int i = 0; i < result.numIntersections; ++i)
+    {
+        result.point[i] = ray.origin + result.parameter[i] * ray.direction;
+    }
+    return result;
+}
+
+template <typename Real>
+void FIQuery<Real, Ray3<Real>, Ellipsoid3<Real>>::DoQuery(
+    Vector3<Real> const& rayOrigin, Vector3<Real> const& rayDirection,
+    Ellipsoid3<Real> const& ellipsoid, Result& result)
+{
+    FIQuery<Real, Line3<Real>, Ellipsoid3<Real>>::DoQuery(rayOrigin,
+        rayDirection, ellipsoid, result);
+
+    if (result.intersect)
+    {
+        // The line containing the ray intersects the ellipsoid; the
+        // t-interval is [t0,t1].  The ray intersects the capsule as long as
+        // [t0,t1] overlaps the ray t-interval [0,+infinity).
+        std::array<Real, 2> rayInterval =
+        { (Real)0, std::numeric_limits<Real>::max() };
+        FIQuery<Real, std::array<Real, 2>, std::array<Real, 2>> iiQuery;
+        auto iiResult = iiQuery(result.parameter, rayInterval);
+        if (iiResult.intersect)
+        {
+            result.numIntersections = iiResult.numIntersections;
+            result.parameter = iiResult.overlap;
+        }
+        else
+        {
+            result.intersect = false;
+            result.numIntersections = 0;
+        }
+    }
+}
+
+
+}
diff --git a/Engine/Plugins/Experimental/GeometryProcessing/Source/GeometryAlgorithms/Private/GTEngine/Mathematics/GteIntrRay3OrientedBox3.h b/Engine/Plugins/Experimental/GeometryProcessing/Source/GeometryAlgorithms/Private/GTEngine/Mathematics/GteIntrRay3OrientedBox3.h
new file mode 100644
--- /dev/null
+++ b/Engine/Plugins/Experimental/GeometryProcessing/Source/GeometryAlgorithms/Private/GTEngine/Mathematics/GteIntrRay3OrientedBox3.h
@@ -0,0 +1,96 @@
+// David Eberly, Geometric Tools, Redmond WA 98052
+// Copyright (c) 1998-2019
+// Distributed under the Boost Software License, Version 1.0.
+// http://www.boost.org/LICENSE_1_0.txt
+// http://www.geometrictools.com/License/Boost/LICENSE_1_0.txt
+// File Version: 3.0.1 (2018/11/28)
+
+#pragma once
+
+#include <GTEngine/Mathematics/GteOrientedBox.h>
+#include <GTEngine/Mathematics/GteIntrRay3AlignedBox3.h>
+
+// The test-intersection queries use the method of separating axes.  The
+// find-intersection queries use parametric clipping against the six faces of
+// the box.  The find-intersection queries use Liang-Barsky clipping.  The
+// queries consider the box to be a solid.  The algorithms are described in
+//   https://www.geometrictools.com/Documentation/IntersectionLineBox.pdf
+
+namespace gte
+{
+    template <typename Real>
+    class TIQuery<Real, Ray3<Real>, OrientedBox3<Real>>
+        :
+        public TIQuery<Real, Ray3<Real>, AlignedBox3<Real>>
+    {
+    public:
+        struct Result
+            :
+            public TIQuery<Real, Ray3<Real>, AlignedBox3<Real>>::Result
+        {
+            // No additional information to compute.
+        };
+
+        Result operator()(Ray3<Real> const& ray, OrientedBox3<Real> const& box)
+        {
+            // Transform the ray to the oriented-box coordinate system.
+            Vector3<Real> diff = ray.origin - box.center;
+            Vector3<Real> rayOrigin
+            {
+                Dot(diff, box.axis[0]),
+                Dot(diff, box.axis[1]),
+                Dot(diff, box.axis[2])
+            };
+            Vector3<Real> rayDirection = Vector3<Real>
+            {
+                Dot(ray.direction, box.axis[0]),
+                Dot(ray.direction, box.axis[1]),
+                Dot(ray.direction, box.axis[2])
+            };
+
+            Result result;
+            this->DoQuery(rayOrigin, rayDirection, box.extent, result);
+            return result;
+        }
+    };
+
+    template <typename Real>
+    class FIQuery<Real, Ray3<Real>, OrientedBox3<Real>>
+        :
+        public FIQuery<Real, Ray3<Real>, AlignedBox3<Real>>
+    {
+    public:
+        struct Result
+            :
+            public FIQuery<Real, Ray3<Real>, AlignedBox3<Real>>::Result
+        {
+            // No additional information to compute.
+        };
+
+        Result operator()(Ray3<Real> const& ray, OrientedBox3<Real> const& box)
+        {
+            // Transform the ray to the oriented-box coordinate system.
+            Vector3<Real> diff = ray.origin - box.center;
+            Vector3<Real> rayOrigin
+            {
+                Dot(diff, box.axis[0]),
+                Dot(diff, box.axis[1]),
+                Dot(diff, box.axis[2])
+            };
+            Vector3<Real> rayDirection = Vector3<Real>
+            {
+                Dot(ray.direction, box.axis[0]),
+                Dot(ray.direction, box.axis[1]),
+                Dot(ray.direction, box.axis[2])
+            };
+
+            Result result;
+            this->DoQuery(rayOrigin, rayDirection, box.extent, result);
+            for (int i = 0; i < result.numPoints; ++i)
+            {
+                result.point[i] = ray.origin + result.lineParameter[i] * ray.direction;
+            }
+            return result;
+        }
+    };
+}
