move GeometryProcessing plugin from NotForLicensees to Experimental

diff --git a/Engine/Plugins/Experimental/GeometryProcessing/Source/GeometryAlgorithms/Private/GTEngine/Mathematics/GteIntrRay3Plane3.h b/Engine/Plugins/Experimental/GeometryProcessing/Source/GeometryAlgorithms/Private/GTEngine/Mathematics/GteIntrRay3Plane3.h
new file mode 100644
--- /dev/null
+++ b/Engine/Plugins/Experimental/GeometryProcessing/Source/GeometryAlgorithms/Private/GTEngine/Mathematics/GteIntrRay3Plane3.h
@@ -0,0 +1,117 @@
+// David Eberly, Geometric Tools, Redmond WA 98052
+// Copyright (c) 1998-2019
+// Distributed under the Boost Software License, Version 1.0.
+// http://www.boost.org/LICENSE_1_0.txt
+// http://www.geometrictools.com/License/Boost/LICENSE_1_0.txt
+// File Version: 3.0.0 (2016/06/19)
+
+#pragma once
+
+#include <GTEngine/Mathematics/GteRay.h>
+#include <GTEngine/Mathematics/GteIntrLine3Plane3.h>
+
+namespace gte
+{
+
+template <typename Real>
+class TIQuery<Real, Ray3<Real>, Plane3<Real>>
+{
+public:
+    struct Result
+    {
+        bool intersect;
+    };
+
+    Result operator()(Ray3<Real> const& ray, Plane3<Real> const& plane);
+};
+
+template <typename Real>
+class FIQuery<Real, Ray3<Real>, Plane3<Real>>
+    :
+    public FIQuery<Real, Line3<Real>, Plane3<Real>>
+{
+public:
+    struct Result
+        :
+        public FIQuery<Real, Line3<Real>, Plane3<Real>>::Result
+    {
+        // No additional information to compute.
+    };
+
+    Result operator()(Ray3<Real> const& ray, Plane3<Real> const& plane);
+
+protected:
+    void DoQuery(Vector3<Real> const& rayOrigin,
+        Vector3<Real> const& rayDirection, Plane3<Real> const& plane,
+        Result& result);
+};
+
+
+template <typename Real>
+typename TIQuery<Real, Ray3<Real>, Plane3<Real>>::Result
+TIQuery<Real, Ray3<Real>, Plane3<Real>>::operator()(
+    Ray3<Real> const& ray, Plane3<Real> const& plane)
+{
+    Result result;
+
+    // Compute the (signed) distance from the ray origin to the plane.
+    DCPQuery<Real, Vector3<Real>, Plane3<Real>> vpQuery;
+    auto vpResult = vpQuery(ray.origin, plane);
+
+    Real DdN = Dot(ray.direction, plane.normal);
+    if (DdN > (Real)0)
+    {
+        // The ray is not parallel to the plane and is directed toward the
+        // +normal side of the plane.
+        result.intersect = (vpResult.signedDistance <= (Real)0);
+    }
+    else if (DdN < (Real)0)
+    {
+        // The ray is not parallel to the plane and is directed toward the
+        // -normal side of the plane.
+        result.intersect = (vpResult.signedDistance >= (Real)0);
+    }
+    else
+    {
+        // The ray and plane are parallel.
+        result.intersect = (vpResult.distance == (Real)0);
+    }
+
+    return result;
+}
+
+template <typename Real>
+typename FIQuery<Real, Ray3<Real>, Plane3<Real>>::Result
+FIQuery<Real, Ray3<Real>, Plane3<Real>>::operator()(
+    Ray3<Real> const& ray, Plane3<Real> const& plane)
+{
+    Result result;
+    DoQuery(ray.origin, ray.direction, plane, result);
+    if (result.intersect)
+    {
+        result.point = ray.origin + result.parameter * ray.direction;
+    }
+    return result;
+}
+
+template <typename Real>
+void FIQuery<Real, Ray3<Real>, Plane3<Real>>::DoQuery(
+    Vector3<Real> const& rayOrigin, Vector3<Real> const& rayDirection,
+    Plane3<Real> const& plane, Result& result)
+{
+    FIQuery<Real, Line3<Real>, Plane3<Real>>::DoQuery(rayOrigin,
+        rayDirection, plane, result);
+    if (result.intersect)
+    {
+        // The line intersects the plane in a point that might not be on the
+        // ray.
+        if (result.parameter < (Real)0)
+        {
+            result.intersect = false;
+            result.numIntersections = 0;
+        }
+    }
+}
+
+
+}
diff --git a/Engine/Plugins/Experimental/GeometryProcessing/Source/GeometryAlgorithms/Private/GTEngine/Mathematics/GteIntrRay3Sphere3.h b/Engine/Plugins/Experimental/GeometryProcessing/Source/GeometryAlgorithms/Private/GTEngine/Mathematics/GteIntrRay3Sphere3.h
new file mode 100644
--- /dev/null
+++ b/Engine/Plugins/Experimental/GeometryProcessing/Source/GeometryAlgorithms/Private/GTEngine/Mathematics/GteIntrRay3Sphere3.h
@@ -0,0 +1,130 @@
+// David Eberly, Geometric Tools, Redmond WA 98052
+// Copyright (c) 1998-2019
+// Distributed under the Boost Software License, Version 1.0.
+// http://www.boost.org/LICENSE_1_0.txt
+// http://www.geometrictools.com/License/Boost/LICENSE_1_0.txt
+// File Version: 3.0.0 (2016/06/19)
+
+#pragma once
+
+#include <GTEngine/Mathematics/GteRay.h>
+#include <GTEngine/Mathematics/GteIntrIntervals.h>
+#include <GTEngine/Mathematics/GteIntrLine3Sphere3.h>
+
+namespace gte
+{
+
+template <typename Real>
+class TIQuery<Real, Ray3<Real>, Sphere3<Real>>
+{
+public:
+    struct Result
+    {
+        bool intersect;
+    };
+
+    Result operator()(Ray3<Real> const& ray, Sphere3<Real> const& sphere);
+};
+
+template <typename Real>
+class FIQuery<Real, Ray3<Real>, Sphere3<Real>>
+    :
+    public FIQuery<Real, Line3<Real>, Sphere3<Real>>
+{
+public:
+    struct Result
+        :
+        public FIQuery<Real, Line3<Real>, Sphere3<Real>>::Result
+    {
+        // No additional information to compute.
+    };
+
+    Result operator()(Ray3<Real> const& ray, Sphere3<Real> const& sphere);
+
+protected:
+    void DoQuery(Vector3<Real> const& rayOrigin,
+        Vector3<Real> const& rayDirection, Sphere3<Real> const& sphere,
+        Result& result);
+};
+
+
+template <typename Real>
+typename TIQuery<Real, Ray3<Real>, Sphere3<Real>>::Result
+TIQuery<Real, Ray3<Real>, Sphere3<Real>>::operator()(
+    Ray3<Real> const& ray, Sphere3<Real> const& sphere)
+{
+    // The sphere is (X-C)^T*(X-C)-1 = 0 and the line is X = P+t*D.
+    // Substitute the line equation into the sphere equation to obtain a
+    // quadratic equation Q(t) = t^2 + 2*a1*t + a0 = 0, where a1 = D^T*(P-C),
+    // and a0 = (P-C)^T*(P-C)-1.
+    Result result;
+
+    Vector3<Real> diff = ray.origin - sphere.center;
+    Real a0 = Dot(diff, diff) - sphere.radius * sphere.radius;
+    if (a0 <= (Real)0)
+    {
+        // P is inside the sphere.
+        result.intersect = true;
+        return result;
+    }
+    // else: P is outside the sphere
+
+    Real a1 = Dot(ray.direction, diff);
+    if (a1 >= (Real)0)
+    {
+        result.intersect = false;
+        return result;
+    }
+
+    // Intersection occurs when Q(t) has real roots.
+    Real discr = a1*a1 - a0;
+    result.intersect = (discr >= (Real)0);
+    return result;
+}
+
+template <typename Real>
+typename FIQuery<Real, Ray3<Real>, Sphere3<Real>>::Result
+FIQuery<Real, Ray3<Real>, Sphere3<Real>>::operator()(
+    Ray3<Real> const& ray, Sphere3<Real> const& sphere)
+{
+    Result result;
+    DoQuery(ray.origin, ray.direction, sphere, result);
+    for (int i = 0; i < result.numIntersections; ++i)
+    {
+        result.point[i] = ray.origin + result.parameter[i] * ray.direction;
+    }
+    return result;
+}
+
+template <typename Real>
+void FIQuery<Real, Ray3<Real>, Sphere3<Real>>::DoQuery(
+    Vector3<Real> const& rayOrigin, Vector3<Real> const& rayDirection,
+    Sphere3<Real> const& sphere, Result& result)
+{
+    FIQuery<Real, Line3<Real>, Sphere3<Real>>::DoQuery(rayOrigin,
+        rayDirection, sphere, result);
+
+    if (result.intersect)
+    {
+        // The line containing the ray intersects the sphere; the t-interval
+        // is [t0,t1].  The ray intersects the sphere as long as [t0,t1]
+        // overlaps the ray t-interval [0,+infinity).
+        std::array<Real, 2> rayInterval =
+        { (Real)0, std::numeric_limits<Real>::max() };
+        FIQuery<Real, std::array<Real, 2>, std::array<Real, 2>> iiQuery;
+        auto iiResult = iiQuery(result.parameter, rayInterval);
+        if (iiResult.intersect)
+        {
+            result.numIntersections = iiResult.numIntersections;
+            result.parameter = iiResult.overlap;
+        }
+        else
+        {
+            result.intersect = false;
+            result.numIntersections = 0;
+        }
+    }
+}
+
+
+}
diff --git a/Engine/Plugins/Experimental/GeometryProcessing/Source/GeometryAlgorithms/Private/GTEngine/Mathematics/GteIntrRay3Triangle3.h b/Engine/Plugins/Experimental/GeometryProcessing/Source/GeometryAlgorithms/Private/GTEngine/Mathematics/GteIntrRay3Triangle3.h
new file mode 100644
--- /dev/null
+++ b/Engine/Plugins/Experimental/GeometryProcessing/Source/GeometryAlgorithms/Private/GTEngine/Mathematics/GteIntrRay3Triangle3.h
@@ -0,0 +1,198 @@
+// David Eberly, Geometric Tools, Redmond WA 98052
+// Copyright (c) 1998-2019
+// Distributed under the Boost Software License, Version 1.0.
+// http://www.boost.org/LICENSE_1_0.txt
+// http://www.geometrictools.com/License/Boost/LICENSE_1_0.txt
+// File Version: 3.0.0 (2016/06/19)
+
+#pragma once
+
+#include <GTEngine/Mathematics/GteVector3.h>
+#include <GTEngine/Mathematics/GteRay.h>
+#include <GTEngine/Mathematics/GteTriangle.h>
+#include <GTEngine/Mathematics/GteFIQuery.h>
+#include <GTEngine/Mathematics/GteTIQuery.h>
+
+namespace gte
+{
+
+template <typename Real>
+class TIQuery<Real, Ray3<Real>, Triangle3<Real>>
+{
+public:
+    struct Result
+    {
+        bool intersect;
+    };
+
+    Result operator()(Ray3<Real> const& ray, Triangle3<Real> const& triangle);
+};
+
+template <typename Real>
+class FIQuery<Real, Ray3<Real>, Triangle3<Real>>
+{
+public:
+    struct Result
+    {
+        Result();
+
+        bool intersect;
+        Real parameter;
+        Real triangleBary[3];
+        Vector3<Real> point;
+    };
+
+    Result operator()(Ray3<Real> const& ray, Triangle3<Real> const& triangle);
+};
+
+
+template <typename Real>
+typename TIQuery<Real, Ray3<Real>, Triangle3<Real>>::Result
+TIQuery<Real, Ray3<Real>, Triangle3<Real>>::operator()(
+    Ray3<Real> const& ray, Triangle3<Real> const& triangle)
+{
+    Result result;
+
+    // Compute the offset origin, edges, and normal.
+    Vector3<Real> diff = ray.origin - triangle.v[0];
+    Vector3<Real> edge1 = triangle.v[1] - triangle.v[0];
+    Vector3<Real> edge2 = triangle.v[2] - triangle.v[0];
+    Vector3<Real> normal = Cross(edge1, edge2);
+
+    // Solve Q + t*D = b1*E1 + b2*E2 (Q = kDiff, D = ray direction,
+    // E1 = edge1, E2 = edge2, N = Cross(E1,E2)) by
+    //   |Dot(D,N)|*b1 = sign(Dot(D,N))*Dot(D,Cross(Q,E2))
+    //   |Dot(D,N)|*b2 = sign(Dot(D,N))*Dot(D,Cross(E1,Q))
+    //   |Dot(D,N)|*t = -sign(Dot(D,N))*Dot(Q,N)
+    Real DdN = Dot(ray.direction, normal);
+    Real sign;
+    if (DdN > (Real)0)
+    {
+        sign = (Real)1;
+    }
+    else if (DdN < (Real)0)
+    {
+        sign = (Real)-1;
+        DdN = -DdN;
+    }
+    else
+    {
+        // Ray and triangle are parallel, call it a "no intersection"
+        // even if the ray does intersect.
+        result.intersect = false;
+        return result;
+    }
+
+    Real DdQxE2 = sign*DotCross(ray.direction, diff, edge2);
+    if (DdQxE2 >= (Real)0)
+    {
+        Real DdE1xQ = sign*DotCross(ray.direction, edge1, diff);
+        if (DdE1xQ >= (Real)0)
+        {
+            if (DdQxE2 + DdE1xQ <= DdN)
+            {
+                // Line intersects triangle, check whether ray does.
+                Real QdN = -sign*Dot(diff, normal);
+                if (QdN >= (Real)0)
+                {
+                    // Ray intersects triangle.
+                    result.intersect = true;
+                    return result;
+                }
+                // else: t < 0, no intersection
+            }
+            // else: b1+b2 > 1, no intersection
+        }
+        // else: b2 < 0, no intersection
+    }
+    // else: b1 < 0, no intersection
+
+    result.intersect = false;
+    return result;
+}
+
+template <typename Real>
+FIQuery<Real, Ray3<Real>, Triangle3<Real>>::Result::Result()
+    :
+    parameter((Real)0)
+{
+    triangleBary[0] = (Real)0;
+    triangleBary[1] = (Real)0;
+    triangleBary[2] = (Real)0;
+}
+
+template <typename Real>
+typename FIQuery<Real, Ray3<Real>, Triangle3<Real>>::Result
+FIQuery<Real, Ray3<Real>, Triangle3<Real>>::operator()(
+    Ray3<Real> const& ray, Triangle3<Real> const& triangle)
+{
+    Result result;
+
+    // Compute the offset origin, edges, and normal.
+    Vector3<Real> diff = ray.origin - triangle.v[0];
+    Vector3<Real> edge1 = triangle.v[1] - triangle.v[0];
+    Vector3<Real> edge2 = triangle.v[2] - triangle.v[0];
+    Vector3<Real> normal = Cross(edge1, edge2);
+
+    // Solve Q + t*D = b1*E1 + b2*E2 (Q = kDiff, D = ray direction,
+    // E1 = edge1, E2 = edge2, N = Cross(E1,E2)) by
+    //   |Dot(D,N)|*b1 = sign(Dot(D,N))*Dot(D,Cross(Q,E2))
+    //   |Dot(D,N)|*b2 = sign(Dot(D,N))*Dot(D,Cross(E1,Q))
+    //   |Dot(D,N)|*t = -sign(Dot(D,N))*Dot(Q,N)
+    Real DdN = Dot(ray.direction, normal);
+    Real sign;
+    if (DdN > (Real)0)
+    {
+        sign = (Real)1;
+    }
+    else if (DdN < (Real)0)
+    {
+        sign = (Real)-1;
+        DdN = -DdN;
+    }
+    else
+    {
+        // Ray and triangle are parallel, call it a "no intersection"
+        // even if the ray does intersect.
+        result.intersect = false;
+        return result;
+    }
+
+    Real DdQxE2 = sign*DotCross(ray.direction, diff, edge2);
+    if (DdQxE2 >= (Real)0)
+    {
+        Real DdE1xQ = sign*DotCross(ray.direction, edge1, diff);
+        if (DdE1xQ >= (Real)0)
+        {
+            if (DdQxE2 + DdE1xQ <= DdN)
+            {
+                // Line intersects triangle, check whether ray does.
+                Real QdN = -sign*Dot(diff, normal);
+                if (QdN >= (Real)0)
+                {
+                    // Ray intersects triangle.
+                    result.intersect = true;
+                    Real inv = ((Real)1) / DdN;
+                    result.parameter = QdN*inv;
+                    result.triangleBary[1] = DdQxE2*inv;
+                    result.triangleBary[2] = DdE1xQ*inv;
+                    result.triangleBary[0] = (Real)1 - result.triangleBary[1]
+                        - result.triangleBary[2];
+                    result.point = ray.origin +
+                        result.parameter * ray.direction;
+                    return result;
+                }
+                // else: t < 0, no intersection
+            }
+            // else: b1+b2 > 1, no intersection
+        }
+        // else: b2 < 0, no intersection
+    }
+    // else: b1 < 0, no intersection
+
+    result.intersect = false;
+    return result;
+}
+
+
+}
diff --git a/Engine/Plugins/Experimental/GeometryProcessing/Source/GeometryAlgorithms/Private/GTEngine/Mathematics/GteIntrSegment2AlignedBox2.h b/Engine/Plugins/Experimental/GeometryProcessing/Source/GeometryAlgorithms/Private/GTEngine/Mathematics/GteIntrSegment2AlignedBox2.h
new file mode 100644
--- /dev/null
+++ b/Engine/Plugins/Experimental/GeometryProcessing/Source/GeometryAlgorithms/Private/GTEngine/Mathematics/GteIntrSegment2AlignedBox2.h
@@ -0,0 +1,175 @@
+// David Eberly, Geometric Tools, Redmond WA 98052
+// Copyright (c) 1998-2019
+// Distributed under the Boost Software License, Version 1.0.
+// http://www.boost.org/LICENSE_1_0.txt
+// http://www.geometrictools.com/License/Boost/LICENSE_1_0.txt
+// File Version: 3.0.1 (2018/10/04)
+
+#pragma once
+
+#include <GTEngine/Mathematics/GteVector2.h>
+#include <GTEngine/Mathematics/GteSegment.h>
+#include <GTEngine/Mathematics/GteIntrIntervals.h>
+#include <GTEngine/Mathematics/GteIntrLine2AlignedBox2.h>
+
+// The queries consider the box to be a solid.
+//
+// The test-intersection queries use the method of separating axes.  The
+// find-intersection queries use parametric clipping against the four edges of
+// the box.
+
+namespace gte
+{
+
+
+template <typename Real>
+class TIQuery<Real, Segment2<Real>, AlignedBox2<Real>>
+    :
+    public TIQuery<Real, Line2<Real>, AlignedBox2<Real>>
+{
+public:
+    struct Result
+        :
+        public TIQuery<Real, Line2<Real>, AlignedBox2<Real>>::Result
+    {
+        // No additional information to compute.
+    };
+
+    Result operator()(Segment2<Real> const& segment,
+        AlignedBox2<Real> const& box);
+
+protected:
+    void DoQuery(Vector2<Real> const& segOrigin,
+        Vector2<Real> const& segDirection, Real segExtent,
+        Vector2<Real> const& boxExtent, Result& result);
+};
+
+template <typename Real>
+class FIQuery<Real, Segment2<Real>, AlignedBox2<Real>>
+    :
+    public FIQuery<Real, Line2<Real>, AlignedBox2<Real>>
+{
+public:
+    struct Result
+        :
+        public FIQuery<Real, Line2<Real>, AlignedBox2<Real>>::Result
+    {
+        // The base class parameter[] values are t-values for the
+        // segment parameterization (1-t)*p[0] + t*p[1], where t in [0,1].
+        // The values in this class are s-values for the centered form
+        // C + s * D, where s in [-e,e] and e is the extent of the segment.
+        std::array<Real, 2> cdeParameter;
+    };
+
+    Result operator()(Segment2<Real> const& segment,
+        AlignedBox2<Real> const& box);
+
+protected:
+    void DoQuery(Vector2<Real> const& segOrigin,
+        Vector2<Real> const& segDirection, Real segExtent,
+        Vector2<Real> const& boxExtent, Result& result);
+};
+
+
+template <typename Real>
+typename TIQuery<Real, Segment2<Real>, AlignedBox2<Real>>::Result
+TIQuery<Real, Segment2<Real>, AlignedBox2<Real>>::operator()(
+    Segment2<Real> const& segment, AlignedBox2<Real> const& box)
+{
+    // Get the centered form of the aligned box.  The axes are implicitly
+    // Axis[d] = Vector2<Real>::Unit(d).
+    Vector2<Real> boxCenter, boxExtent;
+    box.GetCenteredForm(boxCenter, boxExtent);
+
+    // Transform the segment to a centered form in the aligned-box coordinate
+    // system.
+    Vector2<Real> transformedP0 = segment.p[0] - boxCenter;
+    Vector2<Real> transformedP1 = segment.p[1] - boxCenter;
+    Segment2<Real> transformedSegment(transformedP0, transformedP1);
+    Vector2<Real> segOrigin, segDirection;
+    Real segExtent;
+    transformedSegment.GetCenteredForm(segOrigin, segDirection, segExtent);
+
+    Result result;
+    DoQuery(segOrigin, segDirection, segExtent, boxExtent, result);
+    return result;
+}
+
+template <typename Real>
+void TIQuery<Real, Segment2<Real>, AlignedBox2<Real>>::DoQuery(
+    Vector2<Real> const& segOrigin, Vector2<Real> const& segDirection,
+    Real segExtent, Vector2<Real> const& boxExtent, Result& result)
+{
+    for (int i = 0; i < 2; ++i)
+    {
+        Real lhs = std::abs(segOrigin[i]);
+        Real rhs = boxExtent[i] + segExtent * std::abs(segDirection[i]);
+        if (lhs > rhs)
+        {
+            result.intersect = false;
+            return;
+        }
+    }
+
+    TIQuery<Real, Line2<Real>, AlignedBox2<Real>>::DoQuery(segOrigin,
+        segDirection, boxExtent, result);
+}
+
+template <typename Real>
+typename FIQuery<Real, Segment2<Real>, AlignedBox2<Real>>::Result
+FIQuery<Real, Segment2<Real>, AlignedBox2<Real>>::operator()(
+    Segment2<Real> const& segment, AlignedBox2<Real> const& box)
+{
+    // Get the centered form of the aligned box.  The axes are implicitly
+    // Axis[d] = Vector2<Real>::Unit(d).
+    Vector2<Real> boxCenter, boxExtent;
+    box.GetCenteredForm(boxCenter, boxExtent);
+
+    // Transform the segment to a centered form in the aligned-box coordinate
+    // system.
+    Vector2<Real> transformedP0 = segment.p[0] - boxCenter;
+    Vector2<Real> transformedP1 = segment.p[1] - boxCenter;
+    Segment2<Real> transformedSegment(transformedP0, transformedP1);
+    Vector2<Real> segOrigin, segDirection;
+    Real segExtent;
+    transformedSegment.GetCenteredForm(segOrigin, segDirection, segExtent);
+
+    Result result;
+    DoQuery(segOrigin, segDirection, segExtent, boxExtent, result);
+    for (int i = 0; i < result.numIntersections; ++i)
+    {
+        // Compute the segment in the aligned-box coordinate system and then
+        // translate it back to the original coordinates using the box cener.
+        result.point[i] = boxCenter + (segOrigin + result.parameter[i] * segDirection);
+        result.cdeParameter[i] = result.parameter[i];
+
+        // Convert the parameters from the centered form to the endpoint form.
+        result.parameter[i] = (result.parameter[i] / segExtent + (Real)1) * (Real)0.5;
+    }
+    return result;
+}
+
+template <typename Real>
+void FIQuery<Real, Segment2<Real>, AlignedBox2<Real>>::DoQuery(
+    Vector2<Real> const& segOrigin, Vector2<Real> const& segDirection,
+    Real segExtent, Vector2<Real> const& boxExtent, Result& result)
+{
+    FIQuery<Real, Line2<Real>, AlignedBox2<Real>>::DoQuery(segOrigin,
+        segDirection, boxExtent, result);
+
+    if (result.intersect)
+    {
+        // The line containing the segment intersects the box; the t-interval
+        // is [t0,t1].  The segment intersects the box as long as [t0,t1]
+        // overlaps the segment t-interval [-segExtent,+segExtent].
+        std::array<Real, 2> segInterval = { -segExtent, segExtent };
+        FIQuery<Real, std::array<Real, 2>, std::array<Real, 2>> iiQuery;
+        auto iiResult = iiQuery(result.parameter, segInterval);
+        result.intersect = iiResult.intersect;
+        result.numIntersections = iiResult.numIntersections;
+        result.parameter = iiResult.overlap;
+    }
+}
+
+
+}
diff --git a/Engine/Plugins/Experimental/GeometryProcessing/Source/GeometryAlgorithms/Private/GTEngine/Mathematics/GteIntrSegment2Arc2.h b/Engine/Plugins/Experimental/GeometryProcessing/Source/GeometryAlgorithms/Private/GTEngine/Mathematics/GteIntrSegment2Arc2.h
new file mode 100644
--- /dev/null
+++ b/Engine/Plugins/Experimental/GeometryProcessing/Source/GeometryAlgorithms/Private/GTEngine/Mathematics/GteIntrSegment2Arc2.h
@@ -0,0 +1,89 @@
+// David Eberly, Geometric Tools, Redmond WA 98052
+// Copyright (c) 1998-2019
+// Distributed under the Boost Software License, Version 1.0.
+// http://www.boost.org/LICENSE_1_0.txt
+// http://www.geometrictools.com/License/Boost/LICENSE_1_0.txt
+// File Version: 3.0.1 (2018/06/22)
+
+#pragma once
+
+#include <GTEngine/Mathematics/GteArc2.h>
+#include <GTEngine/Mathematics/GteIntrSegment2Circle2.h>
+
+// The queries consider the arc to be a 1-dimensional object.
+
+namespace gte
+{
+
+template <typename Real>
+class TIQuery<Real, Segment2<Real>, Arc2<Real>>
+{
+public:
+    struct Result
+    {
+        bool intersect;
+    };
+
+    Result operator()(Segment2<Real> const& segment, Arc2<Real> const& arc);
+};
+
+template <typename Real>
+class FIQuery<Real, Segment2<Real>, Arc2<Real>>
+{
+public:
+    struct Result
+    {
+        bool intersect;
+        int numIntersections;
+        std::array<Real, 2> parameter;
+        std::array<Vector2<Real>, 2> point;
+    };
+
+    Result operator()(Segment2<Real> const& segment, Arc2<Real> const& arc);
+};
+
+
+template <typename Real>
+typename TIQuery<Real, Segment2<Real>, Arc2<Real>>::Result
+TIQuery<Real, Segment2<Real>, Arc2<Real>>::operator()(
+    Segment2<Real> const& segment, Arc2<Real> const& arc)
+{
+    Result result;
+    FIQuery<Real, Segment2<Real>, Arc2<Real>> saQuery;
+    auto saResult = saQuery(segment, arc);
+    result.intersect = saResult.intersect;
+    return result;
+}
+
+template <typename Real>
+typename FIQuery<Real, Segment2<Real>, Arc2<Real>>::Result
+FIQuery<Real, Segment2<Real>, Arc2<Real>>::operator()(
+    Segment2<Real> const& segment, Arc2<Real> const& arc)
+{
+    Result result;
+    result.intersect = false;
+    result.numIntersections = 0;
+
+    FIQuery<Real, Segment2<Real>, Circle2<Real>> scQuery;
+    Circle2<Real> circle(arc.center, arc.radius);
+    auto scResult = scQuery(segment, circle);
+    if (scResult.intersect)
+    {
+        // Test whether line-circle intersections are on the arc.
+        for (int i = 0; i < scResult.numIntersections; ++i)
+        {
+            if (arc.Contains(scResult.point[i]))
+            {
+                result.intersect = true;
+                result.parameter[result.numIntersections]
+                    = scResult.parameter[i];
+                result.point[result.numIntersections++]
+                    = scResult.point[i];
+            }
+        }
+    }
+
+    return result;
+}
+
+}
diff --git a/Engine/Plugins/Experimental/GeometryProcessing/Source/GeometryAlgorithms/Private/GTEngine/Mathematics/GteIntrSegment2Circle2.h b/Engine/Plugins/Experimental/GeometryProcessing/Source/GeometryAlgorithms/Private/GTEngine/Mathematics/GteIntrSegment2Circle2.h
new file mode 100644
--- /dev/null
+++ b/Engine/Plugins/Experimental/GeometryProcessing/Source/GeometryAlgorithms/Private/GTEngine/Mathematics/GteIntrSegment2Circle2.h
@@ -0,0 +1,107 @@
+// David Eberly, Geometric Tools, Redmond WA 98052
+// Copyright (c) 1998-2019
+// Distributed under the Boost Software License, Version 1.0.
+// http://www.boost.org/LICENSE_1_0.txt
+// http://www.geometrictools.com/License/Boost/LICENSE_1_0.txt
+// File Version: 3.0.0 (2016/06/19)
+
+#pragma once
+
+#include <GTEngine/Mathematics/GteSegment.h>
+#include <GTEngine/Mathematics/GteIntrIntervals.h>
+#include <GTEngine/Mathematics/GteIntrLine2Circle2.h>
+
+// The queries consider the circle to be a solid (disk).
+
+namespace gte
+{
+
+template <typename Real>
+class TIQuery<Real, Segment2<Real>, Circle2<Real>>
+{
+public:
+    struct Result
+    {
+        bool intersect;
+    };
+
+    Result operator()(Segment2<Real> const& segment,
+        Circle2<Real> const& circle);
+};
+
+template <typename Real>
+class FIQuery<Real, Segment2<Real>, Circle2<Real>>
+    :
+    public FIQuery<Real, Line2<Real>, Circle2<Real>>
+{
+public:
+    struct Result
+        :
+        public FIQuery<Real, Line2<Real>, Circle2<Real>>::Result
+    {
+        // No additional information to compute.
+    };
+
+    Result operator()(Segment2<Real> const& segment,
+        Circle2<Real> const& circle);
+
+protected:
+    void DoQuery(Vector2<Real> const& segyOrigin,
+        Vector2<Real> const& segyDirection, Real segExtent,
+        Circle2<Real> const& circle, Result& result);
+};
+
+
+template <typename Real>
+typename TIQuery<Real, Segment2<Real>, Circle2<Real>>::Result
+TIQuery<Real, Segment2<Real>, Circle2<Real>>::operator()(
+    Segment2<Real> const& segment, Circle2<Real> const& circle)
+{
+    Result result;
+    FIQuery<Real, Segment2<Real>, Circle2<Real>> scQuery;
+    result.intersect = scQuery(segment, circle).intersect;
+    return result;
+}
+
+template <typename Real>
+typename FIQuery<Real, Segment2<Real>, Circle2<Real>>::Result
+FIQuery<Real, Segment2<Real>, Circle2<Real>>::operator()(
+    Segment2<Real> const& segment, Circle2<Real> const& circle)
+{
+    Vector2<Real> segOrigin, segDirection;
+    Real segExtent;
+    segment.GetCenteredForm(segOrigin, segDirection, segExtent);
+
+    Result result;
+    DoQuery(segOrigin, segDirection, segExtent, circle, result);
+    for (int i = 0; i < result.numIntersections; ++i)
+    {
+        result.point[i] = segOrigin + result.parameter[i] * segDirection;
+    }
+    return result;
+}
+
+template <typename Real>
+void FIQuery<Real, Segment2<Real>, Circle2<Real>>::DoQuery(
+    Vector2<Real> const& segOrigin, Vector2<Real> const& segDirection,
+    Real segExtent, Circle2<Real> const& circle, Result& result)
+{
+    FIQuery<Real, Line2<Real>, Circle2<Real>>::DoQuery(segOrigin,
+        segDirection, circle, result);
+
+    if (result.intersect)
+    {
+        // The line containing the segment intersects the disk; the t-interval
+        // is [t0,t1].  The segment intersects the disk as long as [t0,t1]
+        // overlaps the segment t-interval [-segExtent,+segExtent].
+        std::array<Real, 2> segInterval = { -segExtent, segExtent };
+        FIQuery<Real, std::array<Real, 2>, std::array<Real, 2>> iiQuery;
+        auto iiResult = iiQuery(result.parameter, segInterval);
+        result.intersect = iiResult.intersect;
+        result.numIntersections = iiResult.numIntersections;
+        result.parameter = iiResult.overlap;
+    }
+}
+
+
+}
diff --git a/Engine/Plugins/Experimental/GeometryProcessing/Source/GeometryAlgorithms/Private/GTEngine/Mathematics/GteIntrSegment2OrientedBox2.h b/Engine/Plugins/Experimental/GeometryProcessing/Source/GeometryAlgorithms/Private/GTEngine/Mathematics/GteIntrSegment2OrientedBox2.h
new file mode 100644
--- /dev/null
+++ b/Engine/Plugins/Experimental/GeometryProcessing/Source/GeometryAlgorithms/Private/GTEngine/Mathematics/GteIntrSegment2OrientedBox2.h
@@ -0,0 +1,124 @@
+// David Eberly, Geometric Tools, Redmond WA 98052
+// Copyright (c) 1998-2019
+// Distributed under the Boost Software License, Version 1.0.
+// http://www.boost.org/LICENSE_1_0.txt
+// http://www.geometrictools.com/License/Boost/LICENSE_1_0.txt
+// File Version: 3.0.0 (2016/06/19)
+
+#pragma once
+
+#include <GTEngine/Mathematics/GteOrientedBox.h>
+#include <GTEngine/Mathematics/GteIntrSegment2AlignedBox2.h>
+
+// The queries consider the box to be a solid.
+//
+// The test-intersection queries use the method of separating axes.  The
+// find-intersection queries use parametric clipping against the four edges of
+// the box.
+
+namespace gte
+{
+
+template <typename Real>
+class TIQuery<Real, Segment2<Real>, OrientedBox2<Real>>
+    :
+    public TIQuery<Real, Segment2<Real>, AlignedBox2<Real>>
+{
+public:
+    struct Result
+        :
+        public TIQuery<Real, Segment2<Real>, AlignedBox2<Real>>::Result
+    {
+        // No additional information to compute.
+    };
+
+    Result operator()(Segment2<Real> const& segment,
+        OrientedBox2<Real> const& box);
+};
+
+template <typename Real>
+class FIQuery<Real, Segment2<Real>, OrientedBox2<Real>>
+    :
+    public FIQuery<Real, Segment2<Real>, AlignedBox2<Real>>
+{
+public:
+    struct Result
+        :
+        public FIQuery<Real, Segment2<Real>, AlignedBox2<Real>>::Result
+    {
+        // The base class parameter[] values are t-values for the
+        // segment parameterization (1-t)*p[0] + t*p[1], where t in [0,1].
+        // The values in this class are s-values for the centered form
+        // C + s * D, where s in [-e,e] and e is the extent of the segment.
+        std::array<Real, 2> cdeParameter;
+    };
+
+    Result operator()(Segment2<Real> const& segment,
+        OrientedBox2<Real> const& box);
+};
+
+
+template <typename Real>
+typename TIQuery<Real, Segment2<Real>, OrientedBox2<Real>>::Result
+TIQuery<Real, Segment2<Real>, OrientedBox2<Real>>::operator()(
+    Segment2<Real> const& segment, OrientedBox2<Real> const& box)
+{
+    // Transform the segment to the oriented-box coordinate system.
+    Vector2<Real> tmpOrigin, tmpDirection;
+    Real segExtent;
+    segment.GetCenteredForm(tmpOrigin, tmpDirection, segExtent);
+    Vector2<Real> diff = tmpOrigin - box.center;
+    Vector2<Real> segOrigin
+    {
+        Dot(diff, box.axis[0]),
+        Dot(diff, box.axis[1])
+    };
+    Vector2<Real> segDirection
+    {
+        Dot(tmpDirection, box.axis[0]),
+        Dot(tmpDirection, box.axis[1])
+    };
+
+    Result result;
+    this->DoQuery(segOrigin, segDirection, segExtent, box.extent, result);
+    return result;
+}
+
+template <typename Real>
+typename FIQuery<Real, Segment2<Real>, OrientedBox2<Real>>::Result
+FIQuery<Real, Segment2<Real>, OrientedBox2<Real>>::operator()(
+    Segment2<Real> const& segment, OrientedBox2<Real> const& box)
+{
+    // Transform the segment to the oriented-box coordinate system.
+    Vector2<Real> tmpOrigin, tmpDirection;
+    Real segExtent;
+    segment.GetCenteredForm(tmpOrigin, tmpDirection, segExtent);
+    Vector2<Real> diff = tmpOrigin - box.center;
+    Vector2<Real> segOrigin
+    {
+        Dot(diff, box.axis[0]),
+        Dot(diff, box.axis[1])
+    };
+    Vector2<Real> segDirection
+    {
+        Dot(tmpDirection, box.axis[0]),
+        Dot(tmpDirection, box.axis[1])
+    };
+
+    Result result;
+    this->DoQuery(segOrigin, segDirection, segExtent, box.extent, result);
+    for (int i = 0; i < result.numIntersections; ++i)
+    {
+        // Compute the segment in the aligned-box coordinate system and then
+        // translate it back to the original coordinates using the box cener.
+        result.point[i] = box.center + (segOrigin + result.parameter[i] * segDirection);
+        result.cdeParameter[i] = result.parameter[i];
+
+        // Convert the parameters from the centered form to the endpoint form.
+        result.parameter[i] = (result.parameter[i] / segExtent + (Real)1) * (Real)0.5;
+    }
+    return result;
+}
+
+
+}
diff --git a/Engine/Plugins/Experimental/GeometryProcessing/Source/GeometryAlgorithms/Private/GTEngine/Mathematics/GteIntrSegment2Segment2.h b/Engine/Plugins/Experimental/GeometryProcessing/Source/GeometryAlgorithms/Private/GTEngine/Mathematics/GteIntrSegment2Segment2.h
new file mode 100644
--- /dev/null
+++ b/Engine/Plugins/Experimental/GeometryProcessing/Source/GeometryAlgorithms/Private/GTEngine/Mathematics/GteIntrSegment2Segment2.h
@@ -0,0 +1,197 @@
+// David Eberly, Geometric Tools, Redmond WA 98052
+// Copyright (c) 1998-2019
+// Distributed under the Boost Software License, Version 1.0.
+// http://www.boost.org/LICENSE_1_0.txt
+// http://www.geometrictools.com/License/Boost/LICENSE_1_0.txt
+// File Version: 3.0.1 (2019/01/11)
+
+#pragma once
+
+#include <GTEngine/Mathematics/GteSegment.h>
+#include <GTEngine/Mathematics/GteIntrIntervals.h>
+#include <GTEngine/Mathematics/GteIntrLine2Line2.h>
+
+namespace gte
+{
+
+template <typename Real>
+class TIQuery<Real, Segment2<Real>, Segment2<Real>>
+{
+public:
+    struct Result
+    {
+        bool intersect;
+
+        // The number is 0 (no intersection), 1 (segments intersect in a
+        // single point), or 2 (segments are collinear and intersect in a
+        // segment).
+        int numIntersections;
+    };
+
+    Result operator()(Segment2<Real> const& segment0,
+        Segment2<Real> const& segment1);
+};
+
+template <typename Real>
+class FIQuery<Real, Segment2<Real>, Segment2<Real>>
+{
+public:
+    struct Result
+    {
+        bool intersect;
+
+        // The number is 0 (no intersection), 1 (segments intersect in a
+        // a single point), or 2 (segments are collinear and intersect
+        // in a segment).
+        int numIntersections;
+
+        // If numIntersections is 1, the intersection is
+        //   point[0]
+        //     = segment0.origin + segment0Parameter[0] * segment0.direction
+        //     = segment1.origin + segment1Parameter[0] * segment1.direction
+        // If numIntersections is 2, the endpoints of the segment of
+        // intersection are
+        //   point[i]
+        //     = segment0.origin + segment0Parameter[i] * segment0.direction
+        //     = segment1.origin + segment1Parameter[i] * segment1.direction
+        // with segment0Parameter[0] <= segment0Parameter[1] and
+        // segment1Parameter[0] <= segment1Parameter[1].
+        Real segment0Parameter[2], segment1Parameter[2];
+        Vector2<Real> point[2];
+    };
+
+    Result operator()(Segment2<Real> const& segment0,
+        Segment2<Real> const& segment1);
+};
+
+
+template <typename Real>
+typename TIQuery<Real, Segment2<Real>, Segment2<Real>>::Result
+TIQuery<Real, Segment2<Real>, Segment2<Real>>::operator()(
+    Segment2<Real> const& segment0, Segment2<Real> const& segment1)
+{
+    Result result;
+    Vector2<Real> seg0Origin, seg0Direction, seg1Origin, seg1Direction;
+    Real seg0Extent, seg1Extent;
+    segment0.GetCenteredForm(seg0Origin, seg0Direction, seg0Extent);
+    segment1.GetCenteredForm(seg1Origin, seg1Direction, seg1Extent);
+
+    FIQuery<Real, Line2<Real>, Line2<Real>> llQuery;
+    Line2<Real> line0(seg0Origin, seg0Direction);
+    Line2<Real> line1(seg1Origin, seg1Direction);
+    auto llResult = llQuery(line0, line1);
+    if (llResult.numIntersections == 1)
+    {
+        // Test whether the line-line intersection is on the segments.
+        if (std::abs(llResult.line0Parameter[0]) <= seg0Extent
+            && std::abs(llResult.line1Parameter[0]) <= seg1Extent)
+        {
+            result.intersect = true;
+            result.numIntersections = 1;
+        }
+        else
+        {
+            result.intersect = false;
+            result.numIntersections = 0;
+        }
+    }
+    else if (llResult.numIntersections == std::numeric_limits<int>::max())
+    {
+        // Compute the location of segment1 endpoints relative to segment0.
+        Vector2<Real> diff = seg1Origin - seg0Origin;
+        Real t = Dot(seg0Direction, diff);
+
+        // Get the parameter intervals of the segments relative to segment0.
+        std::array<Real, 2> interval0 = { -seg0Extent, seg0Extent };
+        std::array<Real, 2> interval1 = { t - seg1Extent, t + seg1Extent };
+
+        // Compute the intersection of the intervals.
+        FIQuery<Real, std::array<Real, 2>, std::array<Real, 2>> iiQuery;
+        auto iiResult = iiQuery(interval0, interval1);
+        result.intersect = iiResult.intersect;
+        result.numIntersections = iiResult.numIntersections;
+    }
+    else
+    {
+        result.intersect = false;
+        result.numIntersections = 0;
+    }
+
+    return result;
+}
+
+template <typename Real>
+typename FIQuery<Real, Segment2<Real>, Segment2<Real>>::Result
+FIQuery<Real, Segment2<Real>, Segment2<Real>>::operator()(
+    Segment2<Real> const& segment0, Segment2<Real> const& segment1)
+{
+    Result result;
+    Vector2<Real> seg0Origin, seg0Direction, seg1Origin, seg1Direction;
+    Real seg0Extent, seg1Extent;
+    segment0.GetCenteredForm(seg0Origin, seg0Direction, seg0Extent);
+    segment1.GetCenteredForm(seg1Origin, seg1Direction, seg1Extent);
+
+    FIQuery<Real, Line2<Real>, Line2<Real>> llQuery;
+    Line2<Real> line0(seg0Origin, seg0Direction);
+    Line2<Real> line1(seg1Origin, seg1Direction);
+    auto llResult = llQuery(line0, line1);
+    if (llResult.numIntersections == 1)
+    {
+        // Test whether the line-line intersection is on the segments.
+        if (std::abs(llResult.line0Parameter[0]) <= seg0Extent
+            && std::abs(llResult.line1Parameter[0]) <= seg1Extent)
+        {
+            result.intersect = true;
+            result.numIntersections = 1;
+            result.segment0Parameter[0] = llResult.line0Parameter[0];
+            result.segment1Parameter[0] = llResult.line1Parameter[0];
+            result.point[0] = llResult.point;
+        }
+        else
+        {
+            result.intersect = false;
+            result.numIntersections = 0;
+        }
+    }
+    else if (llResult.numIntersections == std::numeric_limits<int>::max())
+    {
+        // Compute the location of segment1 endpoints relative to segment0.
+        Vector2<Real> diff = seg1Origin - seg0Origin;
+        Real t = Dot(seg0Direction, diff);
+
+        // Get the parameter intervals of the segments relative to segment0.
+        std::array<Real, 2> interval0 = { -seg0Extent, seg0Extent };
+        std::array<Real, 2> interval1 = { t - seg1Extent, t + seg1Extent };
+
+        // Compute the intersection of the intervals.
+        FIQuery<Real, std::array<Real, 2>, std::array<Real, 2>> iiQuery;
+        auto iiResult = iiQuery(interval0, interval1);
+        if (iiResult.intersect)
+        {
+            result.intersect = true;
+            result.numIntersections = iiResult.numIntersections;
+            for (int i = 0; i < iiResult.numIntersections; ++i)
+            {
+                result.segment0Parameter[i] = iiResult.overlap[i];
+                result.segment1Parameter[i] = iiResult.overlap[i] - t;
+                result.point[i] = seg0Origin +
+                    result.segment0Parameter[i] * seg0Direction;
+            }
+        }
+        else
+        {
+            result.intersect = false;
+            result.numIntersections = 0;
+        }
+    }
+    else
+    {
+        result.intersect = false;
+        result.numIntersections = 0;
+    }
+
+    return result;
+}
+
+
+}
diff --git a/Engine/Plugins/Experimental/GeometryProcessing/Source/GeometryAlgorithms/Private/GTEngine/Mathematics/GteIntrSegment2Triangle2.h b/Engine/Plugins/Experimental/GeometryProcessing/Source/GeometryAlgorithms/Private/GTEngine/Mathematics/GteIntrSegment2Triangle2.h
new file mode 100644
--- /dev/null
+++ b/Engine/Plugins/Experimental/GeometryProcessing/Source/GeometryAlgorithms/Private/GTEngine/Mathematics/GteIntrSegment2Triangle2.h
@@ -0,0 +1,104 @@
+// David Eberly, Geometric Tools, Redmond WA 98052
+// Copyright (c) 1998-2019
+// Distributed under the Boost Software License, Version 1.0.
+// http://www.boost.org/LICENSE_1_0.txt
+// http://www.geometrictools.com/License/Boost/LICENSE_1_0.txt
+// File Version: 3.0.0 (2016/06/19)
+
+#pragma once
+
+#include <GTEngine/Mathematics/GteSegment.h>
+#include <GTEngine/Mathematics/GteIntrIntervals.h>
+#include <GTEngine/Mathematics/GteIntrLine2Triangle2.h>
+
+// The queries consider the triangle to be a solid.
+
+namespace gte
+{
+
+template <typename Real>
+class TIQuery<Real, Segment2<Real>, Triangle2<Real>>
+{
+public:
+    struct Result
+    {
+        bool intersect;
+    };
+
+    Result operator()(Segment2<Real> const& segment,
+        Triangle2<Real> const& triangle);
+};
+
+template <typename Real>
+class FIQuery <Real, Segment2<Real>, Triangle2<Real>>
+    :
+    public FIQuery<Real, Line2<Real>, Triangle2<Real>>
+{
+public:
+    struct Result
+        :
+        public FIQuery<Real, Line2<Real>, Triangle2<Real>>::Result
+    {
+        // No additional information to compute.
+    };
+
+    Result operator()(Segment2<Real> const& segment,
+        Triangle2<Real> const& triangle);
+
+protected:
+    void DoQuery(Vector2<Real> const& segOrigin,
+        Vector2<Real> const& segDirection, Real segExtent,
+        Triangle2<Real> const& triangle, Result& result);
+};
+
+
+template <typename Real>
+typename TIQuery<Real, Segment2<Real>, Triangle2<Real>>::Result
+TIQuery<Real, Segment2<Real>, Triangle2<Real>>::operator()(
+    Segment2<Real> const& segment, Triangle2<Real> const& triangle)
+{
+    Result result;
+    FIQuery<Real, Segment2<Real>, Triangle2<Real>> stQuery;
+    result.intersect = stQuery(segment, triangle).intersect;
+    return result;
+}
+
+template <typename Real>
+typename FIQuery<Real, Segment2<Real>, Triangle2<Real>>::Result
+FIQuery<Real, Segment2<Real>, Triangle2<Real>>::operator()(
+    Segment2<Real> const& segment, Triangle2<Real> const& triangle)
+{
+    Vector2<Real> segOrigin, segDirection;
+    Real segExtent;
+    segment.GetCenteredForm(segOrigin, segDirection, segExtent);
+
+    Result result;
+    DoQuery(segOrigin, segDirection, segExtent, triangle, result);
+    for (int i = 0; i < result.numIntersections; ++i)
+    {
+        result.point[i] = segOrigin + result.parameter[i] * segDirection;
+    }
+    return result;
+}
+
+template <typename Real>
+void FIQuery<Real, Segment2<Real>, Triangle2<Real>>::DoQuery(
+    Vector2<Real> const& segOrigin, Vector2<Real> const& segDirection,
+    Real segExtent, Triangle2<Real> const& triangle, Result& result)
+{
+    FIQuery<Real, Line2<Real>, Triangle2<Real>>::DoQuery(segOrigin,
+        segDirection, triangle, result);
+
+    if (result.intersect)
+    {
+        // The line containing the segment intersects the disk; the t-interval
+        // is [t0,t1].  The segment intersects the disk as long as [t0,t1]
+        // overlaps the segment t-interval [-segExtent,+segExtent].
+        std::array<Real, 2> segInterval = { -segExtent, segExtent };
+        FIQuery<Real, std::array<Real, 2>, std::array<Real, 2>> iiQuery;
+        result.parameter = iiQuery(result.parameter, segInterval).overlap;
+    }
+}
+
+
+}
diff --git a/Engine/Plugins/Experimental/GeometryProcessing/Source/GeometryAlgorithms/Private/GTEngine/Mathematics/GteIntrSegment3AlignedBox3.h b/Engine/Plugins/Experimental/GeometryProcessing/Source/GeometryAlgorithms/Private/GTEngine/Mathematics/GteIntrSegment3AlignedBox3.h
new file mode 100644
--- /dev/null
+++ b/Engine/Plugins/Experimental/GeometryProcessing/Source/GeometryAlgorithms/Private/GTEngine/Mathematics/GteIntrSegment3AlignedBox3.h
@@ -0,0 +1,155 @@
+// David Eberly, Geometric Tools, Redmond WA 98052
+// Copyright (c) 1998-2019
+// Distributed under the Boost Software License, Version 1.0.
+// http://www.boost.org/LICENSE_1_0.txt
+// http://www.geometrictools.com/License/Boost/LICENSE_1_0.txt
+// File Version: 3.0.2 (2018/11/28)
+
+#pragma once
+
+#include <GTEngine/Mathematics/GteSegment.h>
+#include <GTEngine/Mathematics/GteIntrIntervals.h>
+#include <GTEngine/Mathematics/GteIntrLine3AlignedBox3.h>
+
+// The test-intersection queries use the method of separating axes.  The
+// find-intersection queries use parametric clipping against the six faces of
+// the box.  The find-intersection queries use Liang-Barsky clipping.  The
+// queries consider the box to be a solid.  The algorithms are described in
+//   https://www.geometrictools.com/Documentation/IntersectionLineBox.pdf
+
+namespace gte
+{
+    template <typename Real>
+    class TIQuery<Real, Segment3<Real>, AlignedBox3<Real>>
+        :
+        public TIQuery<Real, Line3<Real>, AlignedBox3<Real>>
+    {
+    public:
+        struct Result
+            :
+            public TIQuery<Real, Line3<Real>, AlignedBox3<Real>>::Result
+        {
+            // No additional information to compute.
+        };
+
+        Result operator()(Segment3<Real> const& segment, AlignedBox3<Real> const& box)
+        {
+            // Get the centered form of the aligned box.  The axes are
+            // implicitly Axis[d] = Vector3<Real>::Unit(d).
+            Vector3<Real> boxCenter, boxExtent;
+            box.GetCenteredForm(boxCenter, boxExtent);
+
+            // Transform the segment to a centered form in the aligned-box
+            // coordinate system.
+            Vector3<Real> transformedP0 = segment.p[0] - boxCenter;
+            Vector3<Real> transformedP1 = segment.p[1] - boxCenter;
+            Segment3<Real> transformedSegment(transformedP0, transformedP1);
+            Vector3<Real> segOrigin, segDirection;
+            Real segExtent;
+            transformedSegment.GetCenteredForm(segOrigin, segDirection, segExtent);
+
+            Result result;
+            DoQuery(segOrigin, segDirection, segExtent, boxExtent, result);
+            return result;
+        }
+
+    protected:
+        void DoQuery(Vector3<Real> const& segOrigin, Vector3<Real> const& segDirection,
+            Real segExtent, Vector3<Real> const& boxExtent, Result& result)
+        {
+            for (int i = 0; i < 3; ++i)
+            {
+                if (std::abs(segOrigin[i]) > boxExtent[i] + segExtent * std::abs(segDirection[i]))
+                {
+                    result.intersect = false;
+                    return;
+                }
+            }
+
+            TIQuery<Real, Line3<Real>, AlignedBox3<Real>>::DoQuery(segOrigin, segDirection, boxExtent, result);
+        }
+    };
+
+    template <typename Real>
+    class FIQuery<Real, Segment3<Real>, AlignedBox3<Real>>
+        :
+        public FIQuery<Real, Line3<Real>, AlignedBox3<Real>>
+    {
+    public:
+        struct Result
+            :
+            public FIQuery<Real, Line3<Real>, AlignedBox3<Real>>::Result
+        {
+            // No additional information to compute.
+        };
+
+        Result operator()(Segment3<Real> const& segment, AlignedBox3<Real> const& box)
+        {
+            // Get the centered form of the aligned box.  The axes are
+            // implicitly Axis[d] = Vector3<Real>::Unit(d).
+            Vector3<Real> boxCenter, boxExtent;
+            box.GetCenteredForm(boxCenter, boxExtent);
+
+            // Transform the segment to a centered form in the aligned-box
+            // coordinate system.
+            Vector3<Real> transformedP0 = segment.p[0] - boxCenter;
+            Vector3<Real> transformedP1 = segment.p[1] - boxCenter;
+            Segment3<Real> transformedSegment(transformedP0, transformedP1);
+            Vector3<Real> segOrigin, segDirection;
+            Real segExtent;
+            transformedSegment.GetCenteredForm(segOrigin, segDirection, segExtent);
+
+            Result result;
+            DoQuery(segOrigin, segDirection, segExtent, boxExtent, result);
+
+            // The segment origin is in aligned-box coordinates.  Transform it
+            // back to the original space.
+            segOrigin += boxCenter;
+            for (int i = 0; i < result.numPoints; ++i)
+            {
+                result.point[i] = segOrigin + result.lineParameter[i] * segDirection;
+            }
+            return result;
+        }
+
+    protected:
+        void DoQuery(Vector3<Real> const& segOrigin, Vector3<Real> const& segDirection,
+            Real segExtent, Vector3<Real> const& boxExtent, Result& result)
+        {
+            FIQuery<Real, Line3<Real>, AlignedBox3<Real>>::DoQuery(segOrigin, segDirection, boxExtent, result);
+            if (result.intersect)
+            {
+                // The line containing the segment intersects the box; the
+                // t-interval is [t0,t1].  The segment intersects the box as
+                // long as [t0,t1] overlaps the segment t-interval
+                // [-segExtent,+segExtent].
+                FIQuery<Real, std::array<Real, 2>, std::array<Real, 2>> iiQuery;
+
+                std::array<Real, 2> interval0 =
+                {
+                    result.lineParameter[0], result.lineParameter[1]
+                };
+
+                std::array<Real, 2> interval1 =
+                {
+                    -segExtent, segExtent
+                };
+
+                auto iiResult = iiQuery(interval0, interval1);
+                if (iiResult.numIntersections > 0)
+                {
+                    result.numPoints = iiResult.numIntersections;
+                    for (int i = 0; i < result.numPoints; ++i)
+                    {
+                        result.lineParameter[i] = iiResult.overlap[i];
+                    }
+                }
+                else
+                {
+                    result.intersect = false;
+                    result.numPoints = 0;
+                }
+            }
+        }
+    };
+}
diff --git a/Engine/Plugins/Experimental/GeometryProcessing/Source/GeometryAlgorithms/Private/GTEngine/Mathematics/GteIntrSegment3Capsule3.h b/Engine/Plugins/Experimental/GeometryProcessing/Source/GeometryAlgorithms/Private/GTEngine/Mathematics/GteIntrSegment3Capsule3.h
new file mode 100644
--- /dev/null
+++ b/Engine/Plugins/Experimental/GeometryProcessing/Source/GeometryAlgorithms/Private/GTEngine/Mathematics/GteIntrSegment3Capsule3.h
@@ -0,0 +1,118 @@
+// David Eberly, Geometric Tools, Redmond WA 98052
+// Copyright (c) 1998-2019
+// Distributed under the Boost Software License, Version 1.0.
+// http://www.boost.org/LICENSE_1_0.txt
+// http://www.geometrictools.com/License/Boost/LICENSE_1_0.txt
+// File Version: 3.0.0 (2016/06/19)
+
+#pragma once
+
+#include <GTEngine/Mathematics/GteDistSegmentSegment.h>
+#include <GTEngine/Mathematics/GteIntrIntervals.h>
+#include <GTEngine/Mathematics/GteIntrLine3Capsule3.h>
+
+// The queries consider the capsule to be a solid.
+//
+// The test-intersection queries are based on distance computations.
+
+namespace gte
+{
+
+template <typename Real>
+class TIQuery<Real, Segment3<Real>, Capsule3<Real>>
+{
+public:
+    struct Result
+    {
+        bool intersect;
+    };
+
+    Result operator()(Segment3<Real> const& segment,
+        Capsule3<Real> const& capsule);
+};
+
+template <typename Real>
+class FIQuery<Real, Segment3<Real>, Capsule3<Real>>
+    :
+    public FIQuery<Real, Line3<Real>, Capsule3<Real>>
+{
+public:
+    struct Result
+        :
+        public FIQuery<Real, Line3<Real>, Capsule3<Real>>::Result
+    {
+        // No additional information to compute.
+    };
+
+    Result operator()(Segment3<Real> const& segment,
+        Capsule3<Real> const& capsule);
+
+protected:
+    void DoQuery(Vector3<Real> const& segOrigin,
+        Vector3<Real> const& segDirection, Real segExtent,
+        Capsule3<Real> const& capsule, Result& result);
+};
+
+
+template <typename Real>
+typename TIQuery<Real, Segment3<Real>, Capsule3<Real>>::Result
+TIQuery<Real, Segment3<Real>, Capsule3<Real>>::operator()(
+    Segment3<Real> const& segment, Capsule3<Real> const& capsule)
+{
+    Result result;
+    DCPQuery<Real, Segment3<Real>, Segment3<Real>> ssQuery;
+    auto ssResult = ssQuery(segment, capsule.segment);
+    result.intersect = (ssResult.distance <= capsule.radius);
+    return result;
+}
+
+template <typename Real>
+typename FIQuery<Real, Segment3<Real>, Capsule3<Real>>::Result
+FIQuery<Real, Segment3<Real>, Capsule3<Real>>::operator()(
+    Segment3<Real> const& segment, Capsule3<Real> const& capsule)
+{
+    Vector3<Real> segOrigin, segDirection;
+    Real segExtent;
+    segment.GetCenteredForm(segOrigin, segDirection, segExtent);
+
+    Result result;
+    DoQuery(segOrigin, segDirection, segExtent, capsule, result);
+    for (int i = 0; i < result.numIntersections; ++i)
+    {
+        result.point[i] = segOrigin + result.parameter[i] * segDirection;
+    }
+    return result;
+}
+
+template <typename Real>
+void FIQuery<Real, Segment3<Real>, Capsule3<Real>>::DoQuery(
+    Vector3<Real> const& segOrigin, Vector3<Real> const& segDirection,
+    Real segExtent, Capsule3<Real> const& capsule, Result& result)
+{
+    FIQuery<Real, Line3<Real>, Capsule3<Real>>::DoQuery(segOrigin,
+        segDirection, capsule, result);
+
+    if (result.intersect)
+    {
+        // The line containing the segment intersects the capsule; the
+        // t-interval is [t0,t1].  The segment intersects the capsule as
+        // long as [t0,t1] overlaps the segment t-interval
+        // [-segExtent,+segExtent].
+        std::array<Real, 2> segInterval = { -segExtent, segExtent };
+        FIQuery<Real, std::array<Real, 2>, std::array<Real, 2>> iiQuery;
+        auto iiResult = iiQuery(result.parameter, segInterval);
+        if (iiResult.intersect)
+        {
+            result.numIntersections = iiResult.numIntersections;
+            result.parameter = iiResult.overlap;
+        }
+        else
+        {
+            result.intersect = false;
+            result.numIntersections = 0;
+        }
+    }
+}
+
+
+}
diff --git a/Engine/Plugins/Experimental/GeometryProcessing/Source/GeometryAlgorithms/Private/GTEngine/Mathematics/GteIntrSegment3Cone3.h b/Engine/Plugins/Experimental/GeometryProcessing/Source/GeometryAlgorithms/Private/GTEngine/Mathematics/GteIntrSegment3Cone3.h
new file mode 100644
--- /dev/null
+++ b/Engine/Plugins/Experimental/GeometryProcessing/Source/GeometryAlgorithms/Private/GTEngine/Mathematics/GteIntrSegment3Cone3.h
@@ -0,0 +1,98 @@
+// David Eberly, Geometric Tools, Redmond WA 98052
+// Copyright (c) 1998-2019
+// Distributed under the Boost Software License, Version 1.0.
+// http://www.boost.org/LICENSE_1_0.txt
+// http://www.geometrictools.com/License/Boost/LICENSE_1_0.txt
+// File Version: 3.0.0 (2016/06/19)
+
+#pragma once
+
+#include <GTEngine/Mathematics/GteSegment.h>
+#include <GTEngine/Mathematics/GteIntrLine3Cone3.h>
+
+// The queries consider the cone to be single sided and solid.
+
+namespace gte
+{
+
+template <typename Real>
+class FIQuery<Real, Segment3<Real>, Cone3<Real>>
+    :
+    public FIQuery<Real, Line3<Real>, Cone3<Real>>
+{
+public:
+    struct Result
+        :
+        public FIQuery<Real, Line3<Real>, Cone3<Real>>::Result
+    {
+        // No additional information to compute.
+    };
+
+    Result operator()(Segment3<Real> const& segment, Cone3<Real> const& cone);
+
+protected:
+    void DoQuery(Vector3<Real> const& segOrigin,
+        Vector3<Real> const& segDirection, Real segExtent,
+        Cone3<Real> const& cone, Result& result);
+};
+
+
+template <typename Real>
+typename FIQuery<Real, Segment3<Real>, Cone3<Real>>::Result
+FIQuery<Real, Segment3<Real>, Cone3<Real>>::operator()(
+    Segment3<Real> const& segment, Cone3<Real> const& cone)
+{
+    Vector3<Real> segOrigin, segDirection;
+    Real segExtent;
+    segment.GetCenteredForm(segOrigin, segDirection, segExtent);
+
+    Result result;
+    DoQuery(segOrigin, segDirection, segExtent, cone, result);
+    switch (result.type)
+    {
+    case 1:  // point
+        result.point[0] = segOrigin + result.parameter[0] * segDirection;
+        result.point[1] = result.point[0];
+        break;
+    case 2:  // segment
+        result.point[0] = segOrigin + result.parameter[0] * segDirection;
+        result.point[1] = segOrigin + result.parameter[1] * segDirection;
+        break;
+    default:  // no intersection
+        break;
+    }
+    return result;
+}
+
+template <typename Real>
+void FIQuery<Real, Segment3<Real>, Cone3<Real>>::DoQuery(
+    Vector3<Real> const& segOrigin, Vector3<Real> const& segDirection,
+    Real segExtent, Cone3<Real> const& cone, Result& result)
+{
+    FIQuery<Real, Line3<Real>, Cone3<Real>>::DoQuery(segOrigin,
+        segDirection, cone, result);
+
+    if (result.intersect)
+    {
+        // The line containing the segment intersects the cone; the
+        // t-interval is [t0,t1].  The segment intersects the cone as
+        // long as [t0,t1] overlaps the segment t-interval
+        // [-segExtent,+segExtent].
+        std::array<Real, 2> segInterval = { -segExtent, segExtent };
+        FIIntervalInterval<Real> iiQuery;
+        auto iiResult = iiQuery(result.parameter, segInterval);
+        if (iiResult.intersect)
+        {
+            result.parameter = iiResult.overlap;
+            result.type = iiResult.numIntersections;
+        }
+        else
+        {
+            result.intersect = false;
+            result.type = 0;
+        }
+    }
+}
+
+
+}
diff --git a/Engine/Plugins/Experimental/GeometryProcessing/Source/GeometryAlgorithms/Private/GTEngine/Mathematics/GteIntrSegment3Cylinder3.h b/Engine/Plugins/Experimental/GeometryProcessing/Source/GeometryAlgorithms/Private/GTEngine/Mathematics/GteIntrSegment3Cylinder3.h
new file mode 100644
--- /dev/null
+++ b/Engine/Plugins/Experimental/GeometryProcessing/Source/GeometryAlgorithms/Private/GTEngine/Mathematics/GteIntrSegment3Cylinder3.h
@@ -0,0 +1,91 @@
+// David Eberly, Geometric Tools, Redmond WA 98052
+// Copyright (c) 1998-2019
+// Distributed under the Boost Software License, Version 1.0.
+// http://www.boost.org/LICENSE_1_0.txt
+// http://www.geometrictools.com/License/Boost/LICENSE_1_0.txt
+// File Version: 3.0.0 (2016/06/19)
+
+#pragma once
+
+#include <GTEngine/Mathematics/GteSegment.h>
+#include <GTEngine/Mathematics/GteIntrIntervals.h>
+#include <GTEngine/Mathematics/GteIntrLine3Cylinder3.h>
+
+// The queries consider the cylinder to be a solid.
+
+namespace gte
+{
+
+template <typename Real>
+class FIQuery<Real, Segment3<Real>, Cylinder3<Real>>
+    :
+    public FIQuery<Real, Line3<Real>, Cylinder3<Real>>
+{
+public:
+    struct Result
+        :
+        public FIQuery<Real, Line3<Real>, Cylinder3<Real>>::Result
+    {
+        // No additional information to compute.
+    };
+
+    Result operator()(Segment3<Real> const& segment,
+        Cylinder3<Real> const& cylinder);
+
+protected:
+    void DoQuery(Vector3<Real> const& segOrigin,
+        Vector3<Real> const& segDirection, Real segExtent,
+        Cylinder3<Real> const& cylinder, Result& result);
+};
+
+
+template <typename Real>
+typename FIQuery<Real, Segment3<Real>, Cylinder3<Real>>::Result
+FIQuery<Real, Segment3<Real>, Cylinder3<Real>>::operator()(
+    Segment3<Real> const& segment, Cylinder3<Real> const& cylinder)
+{
+    Vector3<Real> segOrigin, segDirection;
+    Real segExtent;
+    segment.GetCenteredForm(segOrigin, segDirection, segExtent);
+
+    Result result;
+    DoQuery(segOrigin, segDirection, segExtent, cylinder, result);
+    for (int i = 0; i < result.numIntersections; ++i)
+    {
+        result.point[i] = segOrigin + result.parameter[i] * segDirection;
+    }
+    return result;
+}
+
+template <typename Real>
+void FIQuery<Real, Segment3<Real>, Cylinder3<Real>>::DoQuery(
+    Vector3<Real> const& segOrigin, Vector3<Real> const& segDirection,
+    Real segExtent, Cylinder3<Real> const& cylinder, Result& result)
+{
+    FIQuery<Real, Line3<Real>, Cylinder3<Real>>::DoQuery(segOrigin,
+        segDirection, cylinder, result);
+
+    if (result.intersect)
+    {
+        // The line containing the segment intersects the cylinder; the
+        // t-interval is [t0,t1].  The segment intersects the cylinder as
+        // long as [t0,t1] overlaps the segment t-interval
+        // [-segExtent,+segExtent].
+        std::array<Real, 2> segInterval = { -segExtent, segExtent };
+        FIQuery<Real, std::array<Real, 2>, std::array<Real, 2>> iiQuery;
+        auto iiResult = iiQuery(result.parameter, segInterval);
+        if (iiResult.intersect)
+        {
+            result.numIntersections = iiResult.numIntersections;
+            result.parameter = iiResult.overlap;
+        }
+        else
+        {
+            result.intersect = false;
+            result.numIntersections = 0;
+        }
+    }
+}
+
+
+}
diff --git a/Engine/Plugins/Experimental/GeometryProcessing/Source/GeometryAlgorithms/Private/GTEngine/Mathematics/GteIntrSegment3Ellipsoid3.h b/Engine/Plugins/Experimental/GeometryProcessing/Source/GeometryAlgorithms/Private/GTEngine/Mathematics/GteIntrSegment3Ellipsoid3.h
new file mode 100644
--- /dev/null
+++ b/Engine/Plugins/Experimental/GeometryProcessing/Source/GeometryAlgorithms/Private/GTEngine/Mathematics/GteIntrSegment3Ellipsoid3.h
@@ -0,0 +1,187 @@
+// David Eberly, Geometric Tools, Redmond WA 98052
+// Copyright (c) 1998-2019
+// Distributed under the Boost Software License, Version 1.0.
+// http://www.boost.org/LICENSE_1_0.txt
+// http://www.geometrictools.com/License/Boost/LICENSE_1_0.txt
+// File Version: 3.0.0 (2016/06/19)
+
+#pragma once
+
+#include <GTEngine/Mathematics/GteMatrix3x3.h>
+#include <GTEngine/Mathematics/GteSegment.h>
+#include <GTEngine/Mathematics/GteIntrIntervals.h>
+#include <GTEngine/Mathematics/GteIntrLine3Ellipsoid3.h>
+
+// The queries consider the ellipsoid to be a solid.
+
+namespace gte
+{
+
+template <typename Real>
+class TIQuery<Real, Segment3<Real>, Ellipsoid3<Real>>
+{
+public:
+    struct Result
+    {
+        bool intersect;
+    };
+
+    Result operator()(Segment3<Real> const& segment,
+        Ellipsoid3<Real> const& ellipsoid);
+};
+
+template <typename Real>
+class FIQuery<Real, Segment3<Real>, Ellipsoid3<Real>>
+    :
+    public FIQuery<Real, Line3<Real>, Ellipsoid3<Real>>
+{
+public:
+    struct Result
+        :
+        public FIQuery<Real, Line3<Real>, Ellipsoid3<Real>>::Result
+    {
+        // No additional information to compute.
+    };
+
+    Result operator()(Segment3<Real> const& segment,
+        Ellipsoid3<Real> const& ellipsoid);
+
+protected:
+    void DoQuery(Vector3<Real> const& segOrigin,
+        Vector3<Real> const& segDirection, Real segExtent,
+        Ellipsoid3<Real> const& ellipsoid, Result& result);
+};
+
+
+template <typename Real>
+typename TIQuery<Real, Segment3<Real>, Ellipsoid3<Real>>::Result
+TIQuery<Real, Segment3<Real>, Ellipsoid3<Real>>::operator()(
+    Segment3<Real> const& segment, Ellipsoid3<Real> const& ellipsoid)
+{
+    // The ellipsoid is (X-K)^T*M*(X-K)-1 = 0 and the line is X = P+t*D.
+    // Substitute the line equation into the ellipsoid equation to obtain
+    // a quadratic equation Q(t) = a2*t^2 + 2*a1*t + a0 = 0, where
+    // a2 = D^T*M*D, a1 = D^T*M*(P-K), and a0 = (P-K)^T*M*(P-K)-1.
+    Result result;
+
+    Vector3<Real> segOrigin, segDirection;
+    Real segExtent;
+    segment.GetCenteredForm(segOrigin, segDirection, segExtent);
+
+    Matrix3x3<Real> M;
+    ellipsoid.GetM(M);
+
+    Vector3<Real> diff = segOrigin - ellipsoid.center;
+    Vector3<Real> matDir = M*segDirection;
+    Vector3<Real> matDiff = M*diff;
+    Real a2 = Dot(segDirection, matDir);
+    Real a1 = Dot(segDirection, matDiff);
+    Real a0 = Dot(diff, matDiff) - (Real)1;
+
+    Real discr = a1*a1 - a0*a2;
+    if (discr >= (Real)0)
+    {
+        // Test whether ray origin is inside ellipsoid.
+        if (a0 <= (Real)0)
+        {
+            result.intersect = true;
+        }
+        else
+        {
+            // At this point, Q(0) = a0 > 0 and Q(t) has real roots.  It is
+            // also the case that a2 > 0, since M is positive definite,
+            // implying that D^T*M*D > 0 for any nonzero vector D.
+            Real q, qder;
+            if (a1 >= (Real)0)
+            {
+                // Roots are possible only on [-e,0], e is the segment extent.
+                // At least one root occurs if Q(-e) <= 0 or if Q(-e) > 0 and
+                // Q'(-e) < 0.
+                q = a0 + segExtent*(((Real)-2)*a1 + a2*segExtent);
+                if (q <= (Real)0)
+                {
+                    result.intersect = true;
+                }
+                else
+                {
+                    qder = a1 - a2*segExtent;
+                    result.intersect = (qder < (Real)0);
+                }
+            }
+            else
+            {
+                // Roots are only possible on [0,e], e is the segment extent.
+                // At least one root occurs if Q(e) <= 0 or if Q(e) > 0 and
+                // Q'(e) > 0.
+                q = a0 + segExtent*(((Real)2)*a1 + a2*segExtent);
+                if (q <= (Real)0.0)
+                {
+                    result.intersect = true;
+                }
+                else
+                {
+                    qder = a1 + a2*segExtent;
+                    result.intersect = (qder < (Real)0);
+                }
+            }
+        }
+    }
+    else
+    {
+        // No intersection if Q(t) has no real roots.
+        result.intersect = false;
+    }
+
+    return result;
+}
+
+template <typename Real>
+typename FIQuery<Real, Segment3<Real>, Ellipsoid3<Real>>::Result
+FIQuery<Real, Segment3<Real>, Ellipsoid3<Real>>::operator()(
+    Segment3<Real> const& segment, Ellipsoid3<Real> const& ellipsoid)
+{
+    Vector3<Real> segOrigin, segDirection;
+    Real segExtent;
+    segment.GetCenteredForm(segOrigin, segDirection, segExtent);
+
+    Result result;
+    DoQuery(segOrigin, segDirection, segExtent, ellipsoid, result);
+    for (int i = 0; i < result.numIntersections; ++i)
+    {
+        result.point[i] = segOrigin + result.parameter[i] * segDirection;
+    }
+    return result;
+}
+
+template <typename Real>
+void FIQuery<Real, Segment3<Real>, Ellipsoid3<Real>>::DoQuery(
+    Vector3<Real> const& segOrigin, Vector3<Real> const& segDirection,
+    Real segExtent, Ellipsoid3<Real> const& ellipsoid, Result& result)
+{
+    FIQuery<Real, Line3<Real>, Ellipsoid3<Real>>::DoQuery(segOrigin,
+        segDirection, ellipsoid, result);
+
+    if (result.intersect)
+    {
+        // The line containing the segment intersects the ellipsoid; the
+        // t-interval is [t0,t1].  The segment intersects the ellipsoid as
+        // long as [t0,t1] overlaps the segment t-interval
+        // [-segExtent,+segExtent].
+        std::array<Real, 2> segInterval = { -segExtent, segExtent };
+        FIQuery<Real, std::array<Real, 2>, std::array<Real, 2>> iiQuery;
+        auto iiResult = iiQuery(result.parameter, segInterval);
+        if (iiResult.intersect)
+        {
+            result.numIntersections = iiResult.numIntersections;
+            result.parameter = iiResult.overlap;
+        }
+        else
+        {
+            result.intersect = false;
+            result.numIntersections = 0;
+        }
+    }
+}
+
+
+}
diff --git a/Engine/Plugins/Experimental/GeometryProcessing/Source/GeometryAlgorithms/Private/GTEngine/Mathematics/GteIntrSegment3OrientedBox3.h b/Engine/Plugins/Experimental/GeometryProcessing/Source/GeometryAlgorithms/Private/GTEngine/Mathematics/GteIntrSegment3OrientedBox3.h
new file mode 100644
--- /dev/null
+++ b/Engine/Plugins/Experimental/GeometryProcessing/Source/GeometryAlgorithms/Private/GTEngine/Mathematics/GteIntrSegment3OrientedBox3.h
@@ -0,0 +1,113 @@
+// David Eberly, Geometric Tools, Redmond WA 98052
+// Copyright (c) 1998-2019
+// Distributed under the Boost Software License, Version 1.0.
+// http://www.boost.org/LICENSE_1_0.txt
+// http://www.geometrictools.com/License/Boost/LICENSE_1_0.txt
+// File Version: 3.0.2 (2018/11/28)
+
+#pragma once
+
+#include <GTEngine/Mathematics/GteOrientedBox.h>
+#include <GTEngine/Mathematics/GteIntrSegment3AlignedBox3.h>
+
+// The test-intersection queries use the method of separating axes.  The
+// find-intersection queries use parametric clipping against the six faces of
+// the box.  The find-intersection queries use Liang-Barsky clipping.  The
+// queries consider the box to be a solid.  The algorithms are described in
+//   https://www.geometrictools.com/Documentation/IntersectionLineBox.pdf
+
+namespace gte
+{
+    template <typename Real>
+    class TIQuery<Real, Segment3<Real>, OrientedBox3<Real>>
+        :
+        public TIQuery<Real, Segment3<Real>, AlignedBox3<Real>>
+    {
+    public:
+        struct Result
+            :
+            public TIQuery<Real, Segment3<Real>, AlignedBox3<Real>>::Result
+        {
+            // No additional information to compute.
+        };
+
+        Result operator()(Segment3<Real> const& segment, OrientedBox3<Real> const& box)
+        {
+            // Transform the segment to the oriented-box coordinate system.
+            Vector3<Real> tmpOrigin, tmpDirection;
+            Real segExtent;
+            segment.GetCenteredForm(tmpOrigin, tmpDirection, segExtent);
+            Vector3<Real> diff = tmpOrigin - box.center;
+            Vector3<Real> segOrigin
+            {
+                Dot(diff, box.axis[0]),
+                Dot(diff, box.axis[1]),
+                Dot(diff, box.axis[2])
+            };
+            Vector3<Real> segDirection
+            {
+                Dot(tmpDirection, box.axis[0]),
+                Dot(tmpDirection, box.axis[1]),
+                Dot(tmpDirection, box.axis[2])
+            };
+
+            Result result;
+            this->DoQuery(segOrigin, segDirection, segExtent, box.extent, result);
+            return result;
+        }
+    };
+
+    template <typename Real>
+    class FIQuery<Real, Segment3<Real>, OrientedBox3<Real>>
+        :
+        public FIQuery<Real, Segment3<Real>, AlignedBox3<Real>>
+    {
+    public:
+        struct Result
+            :
+            public FIQuery<Real, Segment3<Real>, AlignedBox3<Real>>::Result
+        {
+            // No additional relevant information to compute.
+        };
+
+        Result operator()(Segment3<Real> const& segment, OrientedBox3<Real> const& box)
+        {
+            // Transform the segment to the oriented-box coordinate system.
+            Vector3<Real> tmpOrigin, tmpDirection;
+            Real segExtent;
+            segment.GetCenteredForm(tmpOrigin, tmpDirection, segExtent);
+            Vector3<Real> diff = tmpOrigin - box.center;
+            Vector3<Real> segOrigin
+            {
+                Dot(diff, box.axis[0]),
+                Dot(diff, box.axis[1]),
+                Dot(diff, box.axis[2])
+            };
+            Vector3<Real> segDirection
+            {
+                Dot(tmpDirection, box.axis[0]),
+                Dot(tmpDirection, box.axis[1]),
+                Dot(tmpDirection, box.axis[2])
+            };
+
+            Result result;
+            this->DoQuery(segOrigin, segDirection, segExtent, box.extent, result);
+            for (int i = 0; i < result.numPoints; ++i)
+            {
+                // Compute the intersection point in the oriented-box
+                // coordinate system.
+                Vector3<Real> y = segOrigin + result.lineParameter[i] * segDirection;
+
+                // Transform the intersection point to the original coordinate
+                // system.
+                result.point[i] = box.center;
+                for (int j = 0; j < 3; ++j)
+                {
+                    result.point[i] += y[j] * box.axis[j];
+                }
+            }
+
+            return result;
+        }
+    };
+}
diff --git a/Engine/Plugins/Experimental/GeometryProcessing/Source/GeometryAlgorithms/Private/GTEngine/Mathematics/GteIntrSegment3Plane3.h b/Engine/Plugins/Experimental/GeometryProcessing/Source/GeometryAlgorithms/Private/GTEngine/Mathematics/GteIntrSegment3Plane3.h
new file mode 100644
--- /dev/null
+++ b/Engine/Plugins/Experimental/GeometryProcessing/Source/GeometryAlgorithms/Private/GTEngine/Mathematics/GteIntrSegment3Plane3.h
@@ -0,0 +1,120 @@
+// David Eberly, Geometric Tools, Redmond WA 98052
+// Copyright (c) 1998-2019
+// Distributed under the Boost Software License, Version 1.0.
+// http://www.boost.org/LICENSE_1_0.txt
+// http://www.geometrictools.com/License/Boost/LICENSE_1_0.txt
+// File Version: 3.0.0 (2016/06/19)
+
+#pragma once
+
+#include <GTEngine/Mathematics/GteSegment.h>
+#include <GTEngine/Mathematics/GteIntrLine3Plane3.h>
+
+namespace gte
+{
+
+template <typename Real>
+class TIQuery<Real, Segment3<Real>, Plane3<Real>>
+{
+public:
+    struct Result
+    {
+        bool intersect;
+    };
+
+    Result operator()(Segment3<Real> const& segment,
+        Plane3<Real> const& plane);
+};
+
+template <typename Real>
+class FIQuery<Real, Segment3<Real>, Plane3<Real>>
+    :
+    public FIQuery<Real, Line3<Real>, Plane3<Real>>
+{
+public:
+    struct Result
+        :
+        public FIQuery<Real, Line3<Real>, Plane3<Real>>::Result
+    {
+        // No additional information to compute.
+    };
+
+    Result operator()(Segment3<Real> const& segment,
+        Plane3<Real> const& plane);
+
+protected:
+    void DoQuery(Vector3<Real> const& segOrigin,
+        Vector3<Real> const& segDirection, Real segExtent,
+        Plane3<Real> const& plane, Result& result);
+};
+
+
+template <typename Real>
+typename TIQuery<Real, Segment3<Real>, Plane3<Real>>::Result
+TIQuery<Real, Segment3<Real>, Plane3<Real>>::operator()(
+    Segment3<Real> const& segment, Plane3<Real> const& plane)
+{
+    Result result;
+
+    // Compute the (signed) distance from the segment endpoints to the plane.
+    DCPQuery<Real, Vector3<Real>, Plane3<Real>> vpQuery;
+    Real sdistance0 = vpQuery(segment.p[0], plane).signedDistance;
+    if (sdistance0 == (Real)0)
+    {
+        // Endpoint p[0] is on the plane.
+        result.intersect = true;
+        return result;
+    }
+
+    Real sdistance1 = vpQuery(segment.p[1], plane).signedDistance;
+    if (sdistance1 == (Real)0)
+    {
+        // Endpoint p[1] is on the plane.
+        result.intersect = true;
+        return result;
+    }
+
+    // Test whether the segment transversely intersects the plane.
+    result.intersect = (sdistance0 * sdistance1 < (Real)0);
+    return result;
+}
+
+template <typename Real>
+typename FIQuery<Real, Segment3<Real>, Plane3<Real>>::Result
+FIQuery<Real, Segment3<Real>, Plane3<Real>>::operator()(
+    Segment3<Real> const& segment, Plane3<Real> const& plane)
+{
+    Vector3<Real> segOrigin, segDirection;
+    Real segExtent;
+    segment.GetCenteredForm(segOrigin, segDirection, segExtent);
+
+    Result result;
+    DoQuery(segOrigin, segDirection, segExtent, plane, result);
+    if (result.intersect)
+    {
+        result.point = segOrigin + result.parameter * segDirection;
+    }
+    return result;
+}
+
+template <typename Real>
+void FIQuery<Real, Segment3<Real>, Plane3<Real>>::DoQuery(
+    Vector3<Real> const& segOrigin, Vector3<Real> const& segDirection,
+    Real segExtent, Plane3<Real> const& plane, Result& result)
+{
+    FIQuery<Real, Line3<Real>, Plane3<Real>>::DoQuery(segOrigin,
+        segDirection, plane, result);
+    if (result.intersect)
+    {
+        // The line intersects the plane in a point that might not be on the
+        // segment.
+        if (std::abs(result.parameter) > segExtent)
+        {
+            result.intersect = false;
+            result.numIntersections = 0;
+        }
+    }
+}
+
+
+}
diff --git a/Engine/Plugins/Experimental/GeometryProcessing/Source/GeometryAlgorithms/Private/GTEngine/Mathematics/GteIntrSegment3Sphere3.h b/Engine/Plugins/Experimental/GeometryProcessing/Source/GeometryAlgorithms/Private/GTEngine/Mathematics/GteIntrSegment3Sphere3.h
new file mode 100644
--- /dev/null
+++ b/Engine/Plugins/Experimental/GeometryProcessing/Source/GeometryAlgorithms/Private/GTEngine/Mathematics/GteIntrSegment3Sphere3.h
@@ -0,0 +1,141 @@
+// David Eberly, Geometric Tools, Redmond WA 98052
+// Copyright (c) 1998-2019
+// Distributed under the Boost Software License, Version 1.0.
+// http://www.boost.org/LICENSE_1_0.txt
+// http://www.geometrictools.com/License/Boost/LICENSE_1_0.txt
+// File Version: 3.0.0 (2016/06/19)
+
+#pragma once
+
+#include <GTEngine/Mathematics/GteSegment.h>
+#include <GTEngine/Mathematics/GteIntrIntervals.h>
+#include <GTEngine/Mathematics/GteIntrLine3Sphere3.h>
+
+namespace gte
+{
+
+template <typename Real>
+class TIQuery<Real, Segment3<Real>, Sphere3<Real>>
+{
+public:
+    struct Result
+    {
+        bool intersect;
+    };
+
+    Result operator()(Segment3<Real> const& segment,
+        Sphere3<Real> const& sphere);
+};
+
+template <typename Real>
+class FIQuery<Real, Segment3<Real>, Sphere3<Real>>
+    :
+    public FIQuery<Real, Line3<Real>, Sphere3<Real>>
+{
+public:
+    struct Result
+        :
+        public FIQuery<Real, Line3<Real>, Sphere3<Real>>::Result
+    {
+        // No additional information to compute.
+    };
+
+    Result operator()(Segment3<Real> const& segment,
+        Sphere3<Real> const& sphere);
+
+protected:
+    void DoQuery(Vector3<Real> const& segOrigin,
+        Vector3<Real> const& segDirection, Real segExtent,
+        Sphere3<Real> const& sphere, Result& result);
+};
+
+
+template <typename Real>
+typename TIQuery<Real, Segment3<Real>, Sphere3<Real>>::Result
+TIQuery<Real, Segment3<Real>, Sphere3<Real>>::operator()(
+    Segment3<Real> const& segment, Sphere3<Real> const& sphere)
+{
+    // The sphere is (X-C)^T*(X-C)-1 = 0 and the line is X = P+t*D.
+    // Substitute the line equation into the sphere equation to obtain a
+    // quadratic equation Q(t) = t^2 + 2*a1*t + a0 = 0, where a1 = D^T*(P-C),
+    // and a0 = (P-C)^T*(P-C)-1.
+    Result result;
+
+    Vector3<Real> segOrigin, segDirection;
+    Real segExtent;
+    segment.GetCenteredForm(segOrigin, segDirection, segExtent);
+
+    Vector3<Real> diff = segOrigin - sphere.center;
+    Real a0 = Dot(diff, diff) - sphere.radius * sphere.radius;
+    Real a1 = Dot(segDirection, diff);
+    Real discr = a1*a1 - a0;
+    if (discr < (Real)0)
+    {
+        result.intersect = false;
+        return result;
+    }
+
+    Real tmp0 = segExtent*segExtent + a0;
+    Real tmp1 = ((Real)2)*a1*segExtent;
+    Real qm = tmp0 - tmp1;
+    Real qp = tmp0 + tmp1;
+    if (qm*qp <= (Real)0)
+    {
+        result.intersect = true;
+        return result;
+    }
+
+    result.intersect = (qm > (Real)0 && std::abs(a1) < segExtent);
+    return result;
+}
+
+template <typename Real>
+typename FIQuery<Real, Segment3<Real>, Sphere3<Real>>::Result
+FIQuery<Real, Segment3<Real>, Sphere3<Real>>::operator()(
+    Segment3<Real> const& segment, Sphere3<Real> const& sphere)
+{
+    Vector3<Real> segOrigin, segDirection;
+    Real segExtent;
+    segment.GetCenteredForm(segOrigin, segDirection, segExtent);
+
+    Result result;
+    DoQuery(segOrigin, segDirection, segExtent, sphere, result);
+    for (int i = 0; i < result.numIntersections; ++i)
+    {
+        result.point[i] = segOrigin + result.parameter[i] * segDirection;
+    }
+    return result;
+}
+
+template <typename Real>
+void FIQuery<Real, Segment3<Real>, Sphere3<Real>>::DoQuery(
+    Vector3<Real> const& segOrigin, Vector3<Real> const& segDirection,
+    Real segExtent, Sphere3<Real> const& sphere, Result& result)
+{
+    FIQuery<Real, Line3<Real>, Sphere3<Real>>::DoQuery(segOrigin,
+        segDirection, sphere, result);
+
+    if (result.intersect)
+    {
+        // The line containing the segment intersects the sphere; the
+        // t-interval is [t0,t1].  The segment intersects the sphere as
+        // long as [t0,t1] overlaps the segment t-interval
+        // [-segExtent,+segExtent].
+        std::array<Real, 2> segInterval = { -segExtent, segExtent };
+        FIQuery<Real, std::array<Real, 2>, std::array<Real, 2>> iiQuery;
+        auto iiResult = iiQuery(result.parameter, segInterval);
+        if (iiResult.intersect)
+        {
+            result.numIntersections = iiResult.numIntersections;
+            result.parameter = iiResult.overlap;
+        }
+        else
+        {
+            result.intersect = false;
+            result.numIntersections = 0;
+        }
+    }
+}
+
+
+}
diff --git a/Engine/Plugins/Experimental/GeometryProcessing/Source/GeometryAlgorithms/Private/GTEngine/Mathematics/GteIntrSegment3Triangle3.h b/Engine/Plugins/Experimental/GeometryProcessing/Source/GeometryAlgorithms/Private/GTEngine/Mathematics/GteIntrSegment3Triangle3.h
new file mode 100644
--- /dev/null
+++ b/Engine/Plugins/Experimental/GeometryProcessing/Source/GeometryAlgorithms/Private/GTEngine/Mathematics/GteIntrSegment3Triangle3.h
@@ -0,0 +1,210 @@
+// David Eberly, Geometric Tools, Redmond WA 98052
+// Copyright (c) 1998-2019
+// Distributed under the Boost Software License, Version 1.0.
+// http://www.boost.org/LICENSE_1_0.txt
+// http://www.geometrictools.com/License/Boost/LICENSE_1_0.txt
+// File Version: 3.0.0 (2016/06/19)
+
+#pragma once
+
+#include <GTEngine/Mathematics/GteVector3.h>
+#include <GTEngine/Mathematics/GteSegment.h>
+#include <GTEngine/Mathematics/GteTriangle.h>
+#include <GTEngine/Mathematics/GteFIQuery.h>
+#include <GTEngine/Mathematics/GteTIQuery.h>
+
+namespace gte
+{
+
+template <typename Real>
+class TIQuery<Real, Segment3<Real>, Triangle3<Real>>
+{
+public:
+    struct Result
+    {
+        bool intersect;
+    };
+
+    Result operator()(Segment3<Real> const& segment,
+        Triangle3<Real> const& triangle);
+};
+
+template <typename Real>
+class FIQuery<Real, Segment3<Real>, Triangle3<Real>>
+{
+public:
+    struct Result
+    {
+        Result();
+
+        bool intersect;
+        Real parameter;
+        Real triangleBary[3];
+        Vector3<Real> point;
+    };
+
+    Result operator()(Segment3<Real> const& segment,
+        Triangle3<Real> const& triangle);
+};
+
+
+template <typename Real>
+typename TIQuery<Real, Segment3<Real>, Triangle3<Real>>::Result
+TIQuery<Real, Segment3<Real>, Triangle3<Real>>::operator()(
+    Segment3<Real> const& segment, Triangle3<Real> const& triangle)
+{
+    Result result;
+
+    Vector3<Real> segOrigin, segDirection;
+    Real segExtent;
+    segment.GetCenteredForm(segOrigin, segDirection, segExtent);
+
+    // Compute the offset origin, edges, and normal.
+    Vector3<Real> diff = segOrigin - triangle.v[0];
+    Vector3<Real> edge1 = triangle.v[1] - triangle.v[0];
+    Vector3<Real> edge2 = triangle.v[2] - triangle.v[0];
+    Vector3<Real> normal = Cross(edge1, edge2);
+
+    // Solve Q + t*D = b1*E1 + b2*E2 (Q = diff, D = segment direction,
+    // E1 = edge1, E2 = edge2, N = Cross(E1,E2)) by
+    //   |Dot(D,N)|*b1 = sign(Dot(D,N))*Dot(D,Cross(Q,E2))
+    //   |Dot(D,N)|*b2 = sign(Dot(D,N))*Dot(D,Cross(E1,Q))
+    //   |Dot(D,N)|*t = -sign(Dot(D,N))*Dot(Q,N)
+    Real DdN = Dot(segDirection, normal);
+    Real sign;
+    if (DdN > (Real)0)
+    {
+        sign = (Real)1;
+    }
+    else if (DdN < (Real)0)
+    {
+        sign = (Real)-1;
+        DdN = -DdN;
+    }
+    else
+    {
+        // Segment and triangle are parallel, call it a "no intersection"
+        // even if the segment does intersect.
+        result.intersect = false;
+        return result;
+    }
+
+    Real DdQxE2 = sign*DotCross(segDirection, diff, edge2);
+    if (DdQxE2 >= (Real)0)
+    {
+        Real DdE1xQ = sign*DotCross(segDirection, edge1, diff);
+        if (DdE1xQ >= (Real)0)
+        {
+            if (DdQxE2 + DdE1xQ <= DdN)
+            {
+                // Line intersects triangle, check whether segment does.
+                Real QdN = -sign*Dot(diff, normal);
+                Real extDdN = segExtent*DdN;
+                if (-extDdN <= QdN && QdN <= extDdN)
+                {
+                    // Segment intersects triangle.
+                    result.intersect = true;
+                    return result;
+                }
+                // else: |t| > extent, no intersection
+            }
+            // else: b1+b2 > 1, no intersection
+        }
+        // else: b2 < 0, no intersection
+    }
+    // else: b1 < 0, no intersection
+
+    result.intersect = false;
+    return result;
+}
+
+template <typename Real>
+FIQuery<Real, Segment3<Real>, Triangle3<Real>>::Result::Result()
+    :
+    parameter((Real)0)
+{
+    triangleBary[0] = (Real)0;
+    triangleBary[1] = (Real)0;
+    triangleBary[2] = (Real)0;
+}
+
+template <typename Real>
+typename FIQuery<Real, Segment3<Real>, Triangle3<Real>>::Result
+FIQuery<Real, Segment3<Real>, Triangle3<Real>>::operator()(
+    Segment3<Real> const& segment, Triangle3<Real> const& triangle)
+{
+    Result result;
+
+    Vector3<Real> segOrigin, segDirection;
+    Real segExtent;
+    segment.GetCenteredForm(segOrigin, segDirection, segExtent);
+
+    // Compute the offset origin, edges, and normal.
+    Vector3<Real> diff = segOrigin - triangle.v[0];
+    Vector3<Real> edge1 = triangle.v[1] - triangle.v[0];
+    Vector3<Real> edge2 = triangle.v[2] - triangle.v[0];
+    Vector3<Real> normal = Cross(edge1, edge2);
+
+    // Solve Q + t*D = b1*E1 + b2*E2 (Q = diff, D = segment direction,
+    // E1 = edge1, E2 = edge2, N = Cross(E1,E2)) by
+    //   |Dot(D,N)|*b1 = sign(Dot(D,N))*Dot(D,Cross(Q,E2))
+    //   |Dot(D,N)|*b2 = sign(Dot(D,N))*Dot(D,Cross(E1,Q))
+    //   |Dot(D,N)|*t = -sign(Dot(D,N))*Dot(Q,N)
+    Real DdN = Dot(segDirection, normal);
+    Real sign;
+    if (DdN > (Real)0)
+    {
+        sign = (Real)1;
+    }
+    else if (DdN < (Real)0)
+    {
+        sign = (Real)-1;
+        DdN = -DdN;
+    }
+    else
+    {
+        // Segment and triangle are parallel, call it a "no intersection"
+        // even if the segment does intersect.
+        result.intersect = false;
+        return result;
+    }
+
+    Real DdQxE2 = sign*DotCross(segDirection, diff, edge2);
+    if (DdQxE2 >= (Real)0)
+    {
+        Real DdE1xQ = sign*DotCross(segDirection, edge1, diff);
+        if (DdE1xQ >= (Real)0)
+        {
+            if (DdQxE2 + DdE1xQ <= DdN)
+            {
+                // Line intersects triangle, check whether segment does.
+                Real QdN = -sign*Dot(diff, normal);
+                Real extDdN = segExtent*DdN;
+                if (-extDdN <= QdN && QdN <= extDdN)
+                {
+                    // Segment intersects triangle.
+                    result.intersect = true;
+                    Real inv = ((Real)1) / DdN;
+                    result.parameter = QdN*inv;
+                    result.triangleBary[1] = DdQxE2*inv;
+                    result.triangleBary[2] = DdE1xQ*inv;
+                    result.triangleBary[0] = (Real)1 - result.triangleBary[1]
+                        - result.triangleBary[2];
+                    result.point = segOrigin +
+                        result.parameter * segDirection;
+                    return result;
+                }
+                // else: |t| > extent, no intersection
+            }
+            // else: b1+b2 > 1, no intersection
+        }
+        // else: b2 < 0, no intersection
+    }
+    // else: b1 < 0, no intersection
+
+    result.intersect = false;
+    return result;
+}
+
+
+}
