move GeometryProcessing plugin from NotForLicensees to Experimental

diff --git a/Engine/Plugins/Experimental/GeometryProcessing/Source/GeometryAlgorithms/Private/GTEngine/Mathematics/GteIntrSphere3Cone3.h b/Engine/Plugins/Experimental/GeometryProcessing/Source/GeometryAlgorithms/Private/GTEngine/Mathematics/GteIntrSphere3Cone3.h
new file mode 100644
--- /dev/null
+++ b/Engine/Plugins/Experimental/GeometryProcessing/Source/GeometryAlgorithms/Private/GTEngine/Mathematics/GteIntrSphere3Cone3.h
@@ -0,0 +1,341 @@
+// David Eberly, Geometric Tools, Redmond WA 98052
+// Copyright (c) 1998-2019
+// Distributed under the Boost Software License, Version 1.0.
+// http://www.boost.org/LICENSE_1_0.txt
+// http://www.geometrictools.com/License/Boost/LICENSE_1_0.txt
+// File Version: 3.0.2 (2019/02/15)
+
+#pragma once
+
+#include <GTEngine/Mathematics/GteVector3.h>
+#include <GTEngine/Mathematics/GteHypersphere.h>
+#include <GTEngine/Mathematics/GteCone.h>
+#include <GTEngine/Mathematics/GteFIQuery.h>
+#include <GTEngine/Mathematics/GteTIQuery.h>
+
+// The test-intersection query is based on the document
+// https://www.geometrictools.com/Documentation/IntersectionSphereCone.pdf
+//
+// The find-intersection returns a single point in the set of intersection
+// when that intersection is not empty.
+
+namespace gte
+{
+    template <typename Real>
+    class TIQuery<Real, Sphere3<Real>, Cone3<Real>>
+    {
+    public:
+        struct Result
+        {
+            bool intersect;
+        };
+
+        Result operator()(Sphere3<Real> const& sphere, Cone3<Real> const& cone)
+        {
+            Result result;
+            if (cone.minHeight > (Real)0)
+            {
+                if (cone.maxHeight < std::numeric_limits<Real>::max())
+                {
+                    result.intersect = DoQueryConeFrustum(sphere, cone);
+                }
+                else
+                {
+                    result.intersect = DoQueryInfiniteTruncatedCone(sphere, cone);
+                }
+            }
+            else
+            {
+                if (cone.maxHeight < std::numeric_limits<Real>::max())
+                {
+                    result.intersect = DoQueryFiniteCone(sphere, cone);
+                }
+                else
+                {
+                    result.intersect = DoQueryInfiniteCone(sphere, cone);
+                }
+            }
+            return result;
+        }
+
+    private:
+        bool DoQueryInfiniteCone(Sphere3<Real> const& sphere, Cone3<Real> const& cone)
+        {
+            Vector3<Real> U = cone.ray.origin - (sphere.radius * cone.invSinAngle) * cone.ray.direction;
+            Vector3<Real> CmU = sphere.center - U;
+            Real AdCmU = Dot(cone.ray.direction, CmU);
+            if (AdCmU > (Real)0)
+            {
+                Real sqrLengthCmU = Dot(CmU, CmU);
+                if (AdCmU * AdCmU >= sqrLengthCmU * cone.cosAngleSqr)
+                {
+                    Vector3<Real> CmV = sphere.center - cone.ray.origin;
+                    Real AdCmV = Dot(cone.ray.direction, CmV);
+                    if (AdCmV < -sphere.radius)
+                    {
+                        return false;
+                    }
+
+                    Real rSinAngle = sphere.radius * cone.sinAngle;
+                    if (AdCmV >= -rSinAngle)
+                    {
+                        return true;
+                    }
+
+                    Real sqrLengthCmV = Dot(CmV, CmV);
+                    return sqrLengthCmV <= sphere.radius * sphere.radius;
+                }
+            }
+
+            return false;
+        }
+
+        bool DoQueryInfiniteTruncatedCone(Sphere3<Real> const& sphere, Cone3<Real> const& cone)
+        {
+            Vector3<Real> U = cone.ray.origin - (sphere.radius * cone.invSinAngle) * cone.ray.direction;
+            Vector3<Real> CmU = sphere.center - U;
+            Real AdCmU = Dot(cone.ray.direction, CmU);
+            if (AdCmU > (Real)0)
+            {
+                Real sqrLengthCmU = Dot(CmU, CmU);
+                if (AdCmU * AdCmU >= sqrLengthCmU * cone.cosAngleSqr)
+                {
+                    Vector3<Real> CmV = sphere.center - cone.ray.origin;
+                    Real AdCmV = Dot(cone.ray.direction, CmV);
+                    if (AdCmV < cone.minHeight - sphere.radius)
+                    {
+                        return false;
+                    }
+
+                    Real rSinAngle = sphere.radius * cone.sinAngle;
+                    if (AdCmV >= -rSinAngle)
+                    {
+                        return true;
+                    }
+
+                    Vector3<Real> D = CmV - cone.minHeight * cone.ray.direction;
+                    Real lengthAxD = Length(Cross(cone.ray.direction, D));
+                    Real hminTanAngle = cone.minHeight * cone.tanAngle;
+                    if (lengthAxD <= hminTanAngle)
+                    {
+                        return true;
+                    }
+
+                    Real AdD = AdCmV - cone.minHeight;
+                    Real diff = lengthAxD - hminTanAngle;
+                    Real sqrLengthCmK = AdD * AdD + diff * diff;
+                    return sqrLengthCmK <= sphere.radius * sphere.radius;
+                }
+            }
+
+            return false;
+        }
+
+        bool DoQueryFiniteCone(Sphere3<Real> const& sphere, Cone3<Real> const& cone)
+        {
+            Vector3<Real> U = cone.ray.origin - (sphere.radius * cone.invSinAngle) * cone.ray.direction;
+            Vector3<Real> CmU = sphere.center - U;
+            Real AdCmU = Dot(cone.ray.direction, CmU);
+            if (AdCmU > (Real)0)
+            {
+                Real sqrLengthCmU = Dot(CmU, CmU);
+                if (AdCmU * AdCmU >= sqrLengthCmU * cone.cosAngleSqr)
+                {
+                    Vector3<Real> CmV = sphere.center - cone.ray.origin;
+                    Real AdCmV = Dot(cone.ray.direction, CmV);
+                    if (AdCmV < -sphere.radius)
+                    {
+                        return false;
+                    }
+
+                    if (AdCmV > cone.maxHeight + sphere.radius)
+                    {
+                        return false;
+                    }
+
+                    Real rSinAngle = sphere.radius * cone.sinAngle;
+                    if (AdCmV >= -rSinAngle)
+                    {
+                        if (AdCmV <= cone.maxHeight - rSinAngle)
+                        {
+                            return true;
+                        }
+                        else
+                        {
+                            Vector3<Real> barD = CmV - cone.maxHeight * cone.ray.direction;
+                            Real lengthAxBarD = Length(Cross(cone.ray.direction, barD));
+                            Real hmaxTanAngle = cone.maxHeight * cone.tanAngle;
+                            if (lengthAxBarD <= hmaxTanAngle)
+                            {
+                                return true;
+                            }
+
+                            Real AdBarD = AdCmV - cone.maxHeight;
+                            Real diff = lengthAxBarD - hmaxTanAngle;
+                            Real sqrLengthCmBarK = AdBarD * AdBarD + diff * diff;
+                            return sqrLengthCmBarK <= sphere.radius * sphere.radius;
+                        }
+                    }
+                    else
+                    {
+                        Real sqrLengthCmV = Dot(CmV, CmV);
+                        return sqrLengthCmV <= sphere.radius * sphere.radius;
+                    }
+                }
+            }
+
+            return false;
+        }
+
+        bool DoQueryConeFrustum(Sphere3<Real> const& sphere, Cone3<Real> const& cone)
+        {
+            Vector3<Real> U = cone.ray.origin - (sphere.radius * cone.invSinAngle) * cone.ray.direction;
+            Vector3<Real> CmU = sphere.center - U;
+            Real AdCmU = Dot(cone.ray.direction, CmU);
+            if (AdCmU > (Real)0)
+            {
+                Real sqrLengthCmU = Dot(CmU, CmU);
+                if (AdCmU * AdCmU >= sqrLengthCmU * cone.cosAngleSqr)
+                {
+                    Vector3<Real> CmV = sphere.center - cone.ray.origin;
+                    Real AdCmV = Dot(cone.ray.direction, CmV);
+                    if (AdCmV < cone.minHeight - sphere.radius)
+                    {
+                        return false;
+                    }
+
+                    if (AdCmV > cone.maxHeight + sphere.radius)
+                    {
+                        return false;
+                    }
+
+                    Real rSinAngle = sphere.radius * cone.sinAngle;
+                    if (AdCmV >= cone.minHeight - rSinAngle)
+                    {
+                        if (AdCmV <= cone.maxHeight - rSinAngle)
+                        {
+                            return true;
+                        }
+                        else
+                        {
+                            Vector3<Real> barD = CmV - cone.maxHeight * cone.ray.direction;
+                            Real lengthAxBarD = Length(Cross(cone.ray.direction, barD));
+                            Real hmaxTanAngle = cone.maxHeight * cone.tanAngle;
+                            if (lengthAxBarD <= hmaxTanAngle)
+                            {
+                                return true;
+                            }
+
+                            Real AdBarD = AdCmV - cone.maxHeight;
+                            Real diff = lengthAxBarD - hmaxTanAngle;
+                            Real sqrLengthCmBarK = AdBarD * AdBarD + diff * diff;
+                            return sqrLengthCmBarK <= sphere.radius * sphere.radius;
+                        }
+                    }
+                    else
+                    {
+                        Vector3<Real> D = CmV - cone.minHeight * cone.ray.direction;
+                        Real lengthAxD = Length(Cross(cone.ray.direction, D));
+                        Real hminTanAngle = cone.minHeight * cone.tanAngle;
+                        if (lengthAxD <= hminTanAngle)
+                        {
+                            return true;
+                        }
+
+                        Real AdD = AdCmV - cone.minHeight;
+                        Real diff = lengthAxD - hminTanAngle;
+                        Real sqrLengthCmK = AdD * AdD + diff * diff;
+                        return sqrLengthCmK <= sphere.radius * sphere.radius;
+                    }
+                }
+            }
+
+            return false;
+        }
+    };
+
+    template <typename Real>
+    class FIQuery<Real, Sphere3<Real>, Cone3<Real>>
+    {
+    public:
+        struct Result
+        {
+            // If an intersection occurs, it is potentially an infinite set.
+            // If the cone vertex is inside the sphere, 'point' is set to the
+            // cone vertex.  If the sphere center is inside the cone, 'point'
+            // is set to the sphere center. Otherwise, 'point' is set to the
+            // cone point that is closest to the cone vertex and inside the
+            // sphere.
+            bool intersect;
+            Vector3<Real> point;
+        };
+
+        Result operator()(Sphere3<Real> const& sphere, Cone3<Real> const& cone)
+        {
+            Result result;
+
+            // Test whether the cone vertex is inside the sphere.
+            Vector3<Real> diff = sphere.center - cone.ray.origin;
+            Real rSqr = sphere.radius * sphere.radius;
+            Real lenSqr = Dot(diff, diff);
+            if (lenSqr <= rSqr)
+            {
+                // The cone vertex is inside the sphere, so the sphere and
+                // cone intersect.
+                result.intersect = true;
+                result.point = cone.ray.origin;
+                return result;
+            }
+
+            // Test whether the sphere center is inside the cone.
+            Real dot = Dot(diff, cone.ray.direction);
+            Real dotSqr = dot * dot;
+            if (dotSqr >= lenSqr * cone.cosAngleSqr && dot > (Real)0)
+            {
+                // The sphere center is inside cone, so the sphere and cone
+                // intersect.
+                result.intersect = true;
+                result.point = sphere.center;
+                return result;
+            }
+
+            // The sphere center is outside the cone.  The problem now reduces
+            // to computing an intersection between the circle and the ray in
+            // the plane containing the cone vertex and spanned by the cone
+            // axis and vector from the cone vertex to the sphere center.
+
+            // The ray is parameterized by t * D + V with t >= 0, |D| = 1 and
+            // dot(A,D) = cos(angle).  Also, D = e * A + f * (C - V).
+            // Substituting the ray equation into the sphere equation yields
+            // R^2 = |t * D + V - C|^2, so the quadratic for intersections is
+            // t^2 - 2 * dot(D, C - V) * t + |C - V|^2 - R^2 = 0.  An
+            // intersection occurs if and only if the discriminant is
+            // nonnegative.  This test becomes
+            //     dot(D, C - V)^2 >= dot(C - V, C - V) - R^2
+            // Note that if the right-hand side is nonpositive, then the
+            // inequality is true (the sphere contains V).  This is already
+            // ruled out in the first block of code in this function.
+
+            Real uLen = std::sqrt(std::max(lenSqr - dotSqr, (Real)0));
+            Real test = cone.cosAngle * dot + cone.sinAngle * uLen;
+            Real discr = test * test - lenSqr + rSqr;
+
+            if (discr >= (Real)0 && test >= (Real)0)
+            {
+                // Compute the point of intersection closest to the cone
+                // vertex.
+                result.intersect = true;
+                Real t = test - std::sqrt(std::max(discr, (Real)0));
+                Vector3<Real> B = diff - dot * cone.ray.direction;
+                Real tmp = cone.sinAngle / uLen;
+                result.point = t * (cone.cosAngle * cone.ray.direction + tmp * B);
+            }
+            else
+            {
+                result.intersect = false;
+            }
+
+            return result;
+        }
+    };
+}
diff --git a/Engine/Plugins/Experimental/GeometryProcessing/Source/GeometryAlgorithms/Private/GTEngine/Mathematics/GteIntrSphere3Frustum3.h b/Engine/Plugins/Experimental/GeometryProcessing/Source/GeometryAlgorithms/Private/GTEngine/Mathematics/GteIntrSphere3Frustum3.h
new file mode 100644
--- /dev/null
+++ b/Engine/Plugins/Experimental/GeometryProcessing/Source/GeometryAlgorithms/Private/GTEngine/Mathematics/GteIntrSphere3Frustum3.h
@@ -0,0 +1,44 @@
+// David Eberly, Geometric Tools, Redmond WA 98052
+// Copyright (c) 1998-2019
+// Distributed under the Boost Software License, Version 1.0.
+// http://www.boost.org/LICENSE_1_0.txt
+// http://www.geometrictools.com/License/Boost/LICENSE_1_0.txt
+// File Version: 3.0.0 (2016/06/19)
+
+#pragma once
+
+#include <GTEngine/Mathematics/GteDistPoint3Frustum3.h>
+#include <GTEngine/Mathematics/GteHypersphere.h>
+#include <GTEngine/Mathematics/GteTIQuery.h>
+
+namespace gte
+{
+
+template <typename Real>
+class TIQuery<Real, Sphere3<Real>, Frustum3<Real>>
+{
+public:
+    struct Result
+    {
+        bool intersect;
+    };
+
+    Result operator()(Sphere3<Real> const& sphere,
+        Frustum3<Real> const& frustum);
+};
+
+
+template <typename Real>
+typename TIQuery<Real, Sphere3<Real>, Frustum3<Real>>::Result
+TIQuery<Real, Sphere3<Real>, Frustum3<Real>>::operator()(
+    Sphere3<Real> const& sphere, Frustum3<Real> const& frustum)
+{
+    Result result;
+    DCPQuery<Real, Vector3<Real>, Frustum3<Real>> vfQuery;
+    Real distance = vfQuery(sphere.center, frustum).distance;
+    result.intersect = (distance <= sphere.radius);
+    return result;
+}
+
+
+}
diff --git a/Engine/Plugins/Experimental/GeometryProcessing/Source/GeometryAlgorithms/Private/GTEngine/Mathematics/GteIntrSphere3Sphere3.h b/Engine/Plugins/Experimental/GeometryProcessing/Source/GeometryAlgorithms/Private/GTEngine/Mathematics/GteIntrSphere3Sphere3.h
new file mode 100644
--- /dev/null
+++ b/Engine/Plugins/Experimental/GeometryProcessing/Source/GeometryAlgorithms/Private/GTEngine/Mathematics/GteIntrSphere3Sphere3.h
@@ -0,0 +1,152 @@
+// David Eberly, Geometric Tools, Redmond WA 98052
+// Copyright (c) 1998-2019
+// Distributed under the Boost Software License, Version 1.0.
+// http://www.boost.org/LICENSE_1_0.txt
+// http://www.geometrictools.com/License/Boost/LICENSE_1_0.txt
+// File Version: 3.0.1 (2018/10/05)
+
+#pragma once
+
+#include <GTEngine/Mathematics/GteCircle3.h>
+#include <GTEngine/Mathematics/GteHypersphere.h>
+#include <GTEngine/Mathematics/GteFIQuery.h>
+#include <GTEngine/Mathematics/GteTIQuery.h>
+
+// The queries consider the spheres to be solids.
+
+namespace gte
+{
+
+template <typename Real>
+class TIQuery<Real, Sphere3<Real>, Sphere3<Real>>
+{
+public:
+    struct Result
+    {
+        bool intersect;
+    };
+
+    Result operator()(Sphere3<Real> const& sphere0,
+        Sphere3<Real> const& sphere1);
+};
+
+template <typename Real>
+class FIQuery<Real, Sphere3<Real>, Sphere3<Real>>
+{
+public:
+    struct Result
+    {
+        bool intersect;
+
+        // The type of intersection.
+        //   0: spheres are disjoint and separated
+        //   1: spheres touch at point, each sphere outside the other
+        //   2: spheres intersect in a circle
+        //   3: sphere0 strictly contained in sphere1
+        //   4: sphere0 contained in sphere1, share common point
+        //   5: sphere1 strictly contained in sphere0
+        //   6: sphere1 contained in sphere0, share common point
+        int type;
+        Vector3<Real> point;    // types 1, 4, 6
+        Circle3<Real> circle;   // type 2
+    };
+
+    Result operator()(Sphere3<Real> const& sphere0,
+        Sphere3<Real> const& sphere1);
+};
+
+
+template <typename Real>
+typename TIQuery<Real, Sphere3<Real>, Sphere3<Real>>::Result
+TIQuery<Real, Sphere3<Real>, Sphere3<Real>>::operator()(
+    Sphere3<Real> const& sphere0, Sphere3<Real> const& sphere1)
+{
+    Result result;
+    Vector3<Real> diff = sphere1.center - sphere0.center;
+    Real rSum = sphere0.radius + sphere1.radius;
+    result.intersect = (Dot(diff, diff) <= rSum*rSum);
+    return result;
+}
+
+template <typename Real>
+typename FIQuery<Real, Sphere3<Real>, Sphere3<Real>>::Result
+FIQuery<Real, Sphere3<Real>, Sphere3<Real>>::operator()(
+    Sphere3<Real> const& sphere0, Sphere3<Real> const& sphere1)
+{
+    Result result;
+
+    // The plane of intersection must have C1-C0 as its normal direction.
+    Vector3<Real> C1mC0 = sphere1.center - sphere0.center;
+    Real sqrLen = Dot(C1mC0, C1mC0);
+    Real r0 = sphere0.radius, r1 = sphere1.radius;
+    Real rSum = r0 + r1;
+    Real rSumSqr = rSum * rSum;
+
+    if (sqrLen > rSumSqr)
+    {
+        // The spheres are disjoint/separated.
+        result.intersect = false;
+        result.type = 0;
+        return result;
+    }
+
+    if (sqrLen == rSumSqr)
+    {
+        // The spheres are just touching with each sphere outside the other.
+        Normalize(C1mC0);
+        result.intersect = true;
+        result.type = 1;
+        result.point = sphere0.center + r0 * C1mC0;
+        return result;
+    }
+
+    Real rDif = r0 - r1;
+    Real rDifSqr = rDif*rDif;
+    if (sqrLen < rDifSqr)
+    {
+        // One sphere is strictly contained in the other.  Compute a point in
+        // the intersection set.
+        result.intersect = true;
+        result.type = (rDif <= (Real)0 ? 3 : 5);
+        result.point = ((Real)0.5)*(sphere0.center + sphere1.center);
+        return result;
+    }
+    if (sqrLen == rDifSqr)
+    {
+        // One sphere is contained in the other sphere but with a single point
+        // of contact.
+        Normalize(C1mC0);
+        result.intersect = true;
+        if (rDif <= (Real)0)
+        {
+            result.type = 4;
+            result.point = sphere1.center + r1 * C1mC0;
+        }
+        else
+        {
+            result.type = 6;
+            result.point = sphere0.center + r0 * C1mC0;
+        }
+        return result;
+    }
+
+    // Compute t for which the circle of intersection has center
+    // K = C0 + t*(C1 - C0).
+    Real t = ((Real)0.5) * ((Real)1 + rDif * rSum / sqrLen);
+
+    // Compute the center and radius of the circle of intersection.
+    result.circle.center = sphere0.center + t * C1mC0;
+    result.circle.radius = std::sqrt(std::max(r0*r0 - t*t*sqrLen, (Real)0));
+
+    // Compute the normal for the plane of the circle.
+    Normalize(C1mC0);
+    result.circle.normal = C1mC0;
+
+    // The intersection is a circle.
+    result.intersect = true;
+    result.type = 2;
+    return result;
+}
+
+
+}
diff --git a/Engine/Plugins/Experimental/GeometryProcessing/Source/GeometryAlgorithms/Private/GTEngine/Mathematics/GteIntrSphere3Triangle3.h b/Engine/Plugins/Experimental/GeometryProcessing/Source/GeometryAlgorithms/Private/GTEngine/Mathematics/GteIntrSphere3Triangle3.h
new file mode 100644
--- /dev/null
+++ b/Engine/Plugins/Experimental/GeometryProcessing/Source/GeometryAlgorithms/Private/GTEngine/Mathematics/GteIntrSphere3Triangle3.h
@@ -0,0 +1,591 @@
+// David Eberly, Geometric Tools, Redmond WA 98052
+// Copyright (c) 1998-2019
+// Distributed under the Boost Software License, Version 1.0.
+// http://www.boost.org/LICENSE_1_0.txt
+// http://www.geometrictools.com/License/Boost/LICENSE_1_0.txt
+// File Version: 3.21.0 (2019/01/21)
+
+#pragma once
+
+#include <GTEngine/Mathematics/GteDistPointTriangle.h>
+#include <GTEngine/Mathematics/GteHypersphere.h>
+#include <GTEngine/Mathematics/GteVector3.h>
+#include <GTEngine/Mathematics/GteQuadraticField.h>
+#include <GTEngine/Mathematics/GteFIQuery.h>
+
+namespace gte
+{
+    // Currently, only a dynamic query is supported.  A static query will
+    // need to compute the intersection set of triangle and sphere.
+
+    template <typename Real>
+    class FIQuery<Real, Sphere3<Real>, Triangle3<Real>>
+    {
+    public:
+        // The implementation for floating-point types.
+        struct Result
+        {
+            // The cases are
+            // 1. Objects initially overlapping.  The contactPoint is only one
+            //    of infinitely many points in the overlap.
+            //      intersectionType = -1
+            //      contactTime = 0
+            //      contactPoint = triangle point closest to sphere.center
+            // 2. Objects initially separated but do not intersect later.  The
+            //      contactTime and contactPoint are invalid.
+            //      intersectionType = 0
+            //      contactTime = 0
+            //      contactPoint = (0,0,0)
+            // 3. Objects initially separated but intersect later.
+            //      intersectionType = +1
+            //      contactTime = first time T > 0
+            //      contactPoint = corresponding first contact
+            int intersectionType;
+            Real contactTime;
+            Vector3<Real> contactPoint;
+        };
+
+        template <typename Dummy = Real>
+        typename std::enable_if<!is_arbitrary_precision<Dummy>::value, Result>::type
+        operator()(Sphere3<Real> const& sphere, Vector3<Real> const& sphereVelocity,
+            Triangle3<Real> const& triangle, Vector3<Real> const& triangleVelocity)
+        {
+            Result result = { 0, (Real)0, { (Real)0, (Real)0, (Real)0 } };
+
+            // Test for initial overlap or contact.
+            DCPQuery<Real, Vector3<Real>, Triangle3<Real>> ptQuery;
+            auto ptResult = ptQuery(sphere.center, triangle);
+            Real rsqr = sphere.radius * sphere.radius;
+            if (ptResult.sqrDistance <= rsqr)
+            {
+                result.intersectionType = (ptResult.sqrDistance < rsqr ? -1 : +1);
+                result.contactTime = (Real)0;
+                result.contactPoint = ptResult.closest;
+                return result;
+            }
+
+            // To reach here, the sphere and triangle are initially separated.
+            // Compute the velocity of the sphere relative to the triangle.
+            Vector3<Real> V = sphereVelocity - triangleVelocity;
+            Real sqrLenV = Dot(V, V);
+            if (sqrLenV == (Real)0)
+            {
+                // The sphere and triangle are separated and the sphere is not
+                // moving relative to the triangle, so there is no contact.
+                // The 'result' is already set to the correct state for this
+                // case.
+                return result;
+            }
+
+            // Compute the triangle edge directions E[], the vector U normal
+            // to the plane of the triangle,  and compute the normals to the
+            // edges in the plane of the triangle.  TODO: For a nondeforming
+            // triangle (or mesh of triangles), these quantities can all be
+            // precomputed to reduce the computational cost of the query.  Add
+            // another operator()-query that accepts the precomputed values.
+            // TODO: When the triangle is deformable, these quantities must be
+            // computed, either by the caller or here.  Optimize the code to
+            // compute the quantities on-demand (i.e. only when they are
+            // needed, but cache them for later use).
+            Vector3<Real> E[3] =
+            {
+                triangle.v[1] - triangle.v[0],
+                triangle.v[2] - triangle.v[1],
+                triangle.v[0] - triangle.v[2]
+            };
+            Real sqrLenE[3] = { Dot(E[0], E[0]), Dot(E[1], E[1]), Dot(E[2], E[2]) };
+            Vector3<Real> U = UnitCross(E[0], E[1]);
+            Vector3<Real> ExU[3] =
+            {
+                Cross(E[0], U),
+                Cross(E[1], U),
+                Cross(E[2], U)
+            };
+
+            // Compute the vectors from the triangle vertices to the sphere
+            // center.
+            Vector3<Real> Delta[3] =
+            {
+                sphere.center - triangle.v[0],
+                sphere.center - triangle.v[1],
+                sphere.center - triangle.v[2]
+            };
+
+            // Determine where the sphere center is located relative to the
+            // planes of the triangle offset faces of the sphere-swept volume.
+            Real dotUDelta0 = Dot(U, Delta[0]);
+            if (dotUDelta0 >= sphere.radius)
+            {
+                // The sphere is on the positive side of Dot(U,X-C) = r.  If
+                // the sphere will contact the sphere-swept volume at a
+                // triangular face, it can do so only on the face of the
+                // aforementioned plane.
+                Real dotUV = Dot(U, V);
+                if (dotUV >= (Real)0)
+                {
+                    // The sphere is moving away from, or parallel to, the
+                    // plane of the triangle.  The 'result' is already set to
+                    // the correct state for this case.
+                    return result;
+                }
+
+                Real tbar = (sphere.radius - dotUDelta0) / dotUV;
+                bool foundContact = true;
+                for (int i = 0; i < 3; ++i)
+                {
+                    Real phi = Dot(ExU[i], Delta[i]);
+                    Real psi = Dot(ExU[i], V);
+                    if (phi + psi * tbar > (Real)0)
+                    {
+                        foundContact = false;
+                        break;
+                    }
+                }
+                if (foundContact)
+                {
+                    result.intersectionType = 1;
+                    result.contactTime = tbar;
+                    result.contactPoint = sphere.center + tbar * sphereVelocity;
+                    return result;
+                }
+            }
+            else if (dotUDelta0 <= -sphere.radius)
+            {
+                // The sphere is on the positive side of Dot(-U,X-C) = r.  If
+                // the sphere will contact the sphere-swept volume at a
+                // triangular face, it can do so only on the face of the
+                // aforementioned plane.
+                Real dotUV = Dot(U, V);
+                if (dotUV <= (Real)0)
+                {
+                    // The sphere is moving away from, or parallel to, the
+                    // plane of the triangle.  The 'result' is already set to
+                    // the correct state for this case.
+                    return result;
+                }
+
+                Real tbar = (-sphere.radius - dotUDelta0) / dotUV;
+                bool foundContact = true;
+                for (int i = 0; i < 3; ++i)
+                {
+                    Real phi = Dot(ExU[i], Delta[i]);
+                    Real psi = Dot(ExU[i], V);
+                    if (phi + psi * tbar > (Real)0)
+                    {
+                        foundContact = false;
+                        break;
+                    }
+                }
+                if (foundContact)
+                {
+                    result.intersectionType = 1;
+                    result.contactTime = tbar;
+                    result.contactPoint = sphere.center + tbar * sphereVelocity;
+                    return result;
+                }
+            }
+            // else: The ray-sphere-swept-volume contact point (if any) cannot
+            // be on a triangular face of the sphere-swept-volume.
+
+            // The sphere is moving towards the slab between the two planes
+            // of the sphere-swept volume triangular faces.  Determine whether
+            // the ray intersects the half cylinders or sphere wedges of the
+            // sphere-swept volume.
+
+            // Test for contact with half cylinders of the sphere-swept
+            // volume.  First, precompute some dot products required in the
+            // computations.  TODO: Optimize the code to compute the quantities
+            // on-demand (i.e. only when they are needed, but cache them for
+            // later use).
+            Real del[3], delp[3], nu[3];
+            for (int im1 = 2, i = 0; i < 3; im1 = i++)
+            {
+                del[i] = Dot(E[i], Delta[i]);
+                delp[im1] = Dot(E[im1], Delta[i]);
+                nu[i] = Dot(E[i], V);
+            }
+
+            for (int i = 2, ip1 = 0; ip1 < 3; i = ip1++)
+            {
+                Vector3<Real> hatV = V - E[i] * nu[i] / sqrLenE[i];
+                Real sqrLenHatV = Dot(hatV, hatV);
+                if (sqrLenHatV > (Real)0)
+                {
+                    Vector3<Real> hatDelta = Delta[i] - E[i] * del[i] / sqrLenE[i];
+                    Real alpha = -Dot(hatV, hatDelta);
+                    if (alpha >= (Real)0)
+                    {
+                        Real sqrLenHatDelta = Dot(hatDelta, hatDelta);
+                        Real beta = alpha * alpha - sqrLenHatV * (sqrLenHatDelta - rsqr);
+                        if (beta >= (Real)0)
+                        {
+                            Real tbar = (alpha - std::sqrt(beta)) / sqrLenHatV;
+
+                            Real mu = Dot(ExU[i], Delta[i]);
+                            Real omega = Dot(ExU[i], hatV);
+                            if (mu + omega * tbar >= (Real)0)
+                            {
+                                if (del[i] + nu[i] * tbar >= (Real)0)
+                                {
+                                    if (delp[i] + nu[i] * tbar <= (Real)0)
+                                    {
+                                        // The constraints are satisfied, so
+                                        // tbar is the first time of contact.
+                                        result.intersectionType = 1;
+                                        result.contactTime = tbar;
+                                        result.contactPoint = sphere.center + tbar * sphereVelocity;
+                                        return result;
+                                    }
+                                }
+                            }
+                        }
+                    }
+                }
+            }
+
+            // Test for contact with sphere wedges of the sphere-swept
+            // volume.  We know that |V|^2 > 0 because of a previous
+            // early-exit test.
+            for (int im1 = 2, i = 0; i < 3; im1 = i++)
+            {
+                Real alpha = -Dot(V, Delta[i]);
+                if (alpha >= (Real)0)
+                {
+                    Real sqrLenDelta = Dot(Delta[i], Delta[i]);
+                    Real beta = alpha * alpha - sqrLenV * (sqrLenDelta - rsqr);
+                    if (beta >= (Real)0)
+                    {
+                        Real tbar = (alpha - std::sqrt(beta)) / sqrLenV;
+                        if (delp[im1] + nu[im1] * tbar >= (Real)0)
+                        {
+                            if (del[i] + nu[i] * tbar <= (Real)0)
+                            {
+                                // The constraints are satisfied, so tbar
+                                // is the first time of contact.
+                                result.intersectionType = 1;
+                                result.contactTime = tbar;
+                                result.contactPoint = sphere.center + tbar * sphereVelocity;
+                                return result;
+                            }
+                        }
+                    }
+                }
+            }
+
+            // The ray and sphere-swept volume do not intersect, so the sphere
+            // and triangle do not come into contact.  The 'result' is already
+            // set to the correct state for this case.
+            return result;
+        }
+
+
+        // The implementation for arbitrary-precision types.
+        typedef typename QuadraticField<Real>::Element QFElement;
+
+        struct ExactResult
+        {
+            // The cases are
+            // 1. Objects initially overlapping.  The contactPoint is only one
+            //    of infinitely many points in the overlap.
+            //      intersectionType = -1
+            //      contactTime = 0
+            //      contactPoint = triangle point closest to sphere.center
+            // 2. Objects initially separated but do not intersect later.  The
+            //      contactTime and contactPoint are invalid.
+            //      intersectionType = 0
+            //      contactTime = 0
+            //      contactPoint = (0,0,0)
+            // 3. Objects initially separated but intersect later.
+            //      intersectionType = +1
+            //      contactTime = first time T > 0
+            //      contactPoint = corresponding first contact
+            int intersectionType;
+
+            // The exact representation of the contact time and point.  To
+            // convert to a floating-point type, use
+            //   FloatType contactTime = field.Convert(result.contactTime);
+            //   Vector3<FloatType> contactPoint
+            //   {
+            //       field.Convert(result.contactPoint[0]),
+            //       field.Convert(result.contactPoint[1]),
+            //       field.Convert(result.contactPoint[2])
+            //   };
+            QuadraticField<Real> field;
+            QFElement contactTime;
+            Vector3<QFElement> contactPoint;
+        };
+
+        template <typename Dummy = Real>
+        typename std::enable_if<is_arbitrary_precision<Dummy>::value, ExactResult>::type
+        operator()(Sphere3<Real> const& sphere, Vector3<Real> const& sphereVelocity,
+            Triangle3<Real> const& triangle, Vector3<Real> const& triangleVelocity)
+        {
+            // The default constructors for the members of 'result' set their
+            // own members to zero.
+            ExactResult result;
+
+            // Test for initial overlap or contact.
+            DCPQuery<Real, Vector3<Real>, Triangle3<Real>> ptQuery;
+            auto ptResult = ptQuery(sphere.center, triangle);
+            Real rsqr = sphere.radius * sphere.radius;
+            if (ptResult.sqrDistance <= rsqr)
+            {
+                // The values result.field, result.contactTime and
+                // result.contactPoint[] are all zero, so we need only set
+                // the result.contactPoint[].x values.
+                result.intersectionType = (ptResult.sqrDistance < rsqr ? -1 : +1);
+                for (int j = 0; j < 3; ++j)
+                {
+                    result.contactPoint[j].x = ptResult.closest[j];
+                }
+                return result;
+            }
+
+            // To reach here, the sphere and triangle are initially separated.
+            // Compute the velocity of the sphere relative to the triangle.
+            Vector3<Real> V = sphereVelocity - triangleVelocity;
+            Real sqrLenV = Dot(V, V);
+            if (sqrLenV == (Real)0)
+            {
+                // The sphere and triangle are separated and the sphere is not
+                // moving relative to the triangle, so there is no contact.
+                // The 'result' is already set to the correct state for this
+                // case.
+                return result;
+            }
+
+            // Compute the triangle edge directions E[], the vector U normal
+            // to the plane of the triangle,  and compute the normals to the
+            // edges in the plane of the triangle.  TODO: For a nondeforming
+            // triangle (or mesh of triangles), these quantities can all be
+            // precomputed to reduce the computational cost of the query.  Add
+            // another operator()-query that accepts the precomputed values.
+            // TODO: When the triangle is deformable, these quantities must be
+            // computed, either by the caller or here.  Optimize the code to
+            // compute the quantities on-demand (i.e. only when they are
+            // needed, but cache them for later use).
+            Vector3<Real> E[3] =
+            {
+                triangle.v[1] - triangle.v[0],
+                triangle.v[2] - triangle.v[1],
+                triangle.v[0] - triangle.v[2]
+            };
+            Real sqrLenE[3] = { Dot(E[0], E[0]), Dot(E[1], E[1]), Dot(E[2], E[2]) };
+            // Use an unnormalized U for the plane of the triangle.  This
+            // allows us to use quadratic fields for the comparisons of the
+            // constraints.
+            Vector3<Real> U = Cross(E[0], E[1]);
+            Real sqrLenU = Dot(U, U);
+            Vector3<Real> ExU[3] =
+            {
+                Cross(E[0], U),
+                Cross(E[1], U),
+                Cross(E[2], U)
+            };
+
+            // Compute the vectors from the triangle vertices to the sphere
+            // center.
+            Vector3<Real> Delta[3] =
+            {
+                sphere.center - triangle.v[0],
+                sphere.center - triangle.v[1],
+                sphere.center - triangle.v[2]
+            };
+
+            // Determine where the sphere center is located relative to the
+            // planes of the triangle offset faces of the sphere-swept volume.
+            QuadraticField<Real> ufield(sqrLenU);
+            QFElement element(Dot(U, Delta[0]), -sphere.radius);
+            if (ufield.GreaterThanOrEqualZero(element))
+            {
+                // The sphere is on the positive side of Dot(U,X-C) = r|U|.
+                // If the sphere will contact the sphere-swept volume at a
+                // triangular face, it can do so only on the face of the
+                // aforementioned plane.
+                Real dotUV = Dot(U, V);
+                if (dotUV >= (Real)0)
+                {
+                    // The sphere is moving away from, or parallel to, the
+                    // plane of the triangle.  The 'result' is already set
+                    // to the correct state for this case.
+                    return result;
+                }
+
+                bool foundContact = true;
+                for (int i = 0; i < 3; ++i)
+                {
+                    Real phi = Dot(ExU[i], Delta[i]);
+                    Real psi = Dot(ExU[i], V);
+                    QFElement arg(psi * element.x - phi * dotUV, psi * element.y);
+                    if (ufield.GreaterThanZero(arg))
+                    {
+                        foundContact = false;
+                        break;
+                    }
+                }
+                if (foundContact)
+                {
+                    result.intersectionType = 1;
+                    result.field = ufield;
+                    result.contactTime.x = -element.x / dotUV;
+                    result.contactTime.y = -element.y / dotUV;
+                    for (int j = 0; j < 3; ++j)
+                    {
+                        result.contactPoint[j].x = sphere.center[j] + result.contactTime.x * sphereVelocity[j];
+                        result.contactPoint[j].y = result.contactTime.y * sphereVelocity[j];
+                    }
+                    return result;
+                }
+            }
+            else
+            {
+                element.y = -element.y;
+                if (ufield.LessThanOrEqualZero(element))
+                {
+                    // The sphere is on the positive side of Dot(-U,X-C) = r|U|.
+                    // If the sphere will contact the sphere-swept volume at a
+                    // triangular face, it can do so only on the face of the
+                    // aforementioned plane.
+                    Real dotUV = Dot(U, V);
+                    if (dotUV <= (Real)0)
+                    {
+                        // The sphere is moving away from, or parallel to, the
+                        // plane of the triangle.  The 'result' is already set
+                        // to the correct state for this case.
+                        return result;
+                    }
+
+                    bool foundContact = true;
+                    for (int i = 0; i < 3; ++i)
+                    {
+                        Real phi = Dot(ExU[i], Delta[i]);
+                        Real psi = Dot(ExU[i], V);
+                        QFElement arg(phi * dotUV - psi * element.x, -psi * element.y);
+                        if (ufield.GreaterThanZero(arg))
+                        {
+                            foundContact = false;
+                            break;
+                        }
+                    }
+                    if (foundContact)
+                    {
+                        result.intersectionType = 1;
+                        result.field = ufield;
+                        result.contactTime.x = -element.x / dotUV;
+                        result.contactTime.y = -element.y / dotUV;
+                        for (int j = 0; j < 3; ++j)
+                        {
+                            result.contactPoint[j].x = sphere.center[j] + result.contactTime.x * sphereVelocity[j];
+                            result.contactPoint[j].y = result.contactTime.y * sphereVelocity[j];
+                        }
+                        return result;
+                    }
+                }
+                // else: The ray-sphere-swept-volume contact point (if any)
+                // cannot be on a triangular face of the sphere-swept-volume.
+            }
+
+            // The sphere is moving towards the slab between the two planes
+            // of the sphere-swept volume triangular faces.  Determine whether
+            // the ray intersects the half cylinders or sphere wedges of the
+            // sphere-swept volume.
+
+            // Test for contact with half cylinders of the sphere-swept
+            // volume.  First, precompute some dot products required in the
+            // computations.  TODO: Optimize the code to compute the quantities
+            // on-demand (i.e. only when they are needed, but cache them for
+            // later use).
+            Real del[3], delp[3], nu[3];
+            for (int im1 = 2, i = 0; i < 3; im1 = i++)
+            {
+                del[i] = Dot(E[i], Delta[i]);
+                delp[im1] = Dot(E[im1], Delta[i]);
+                nu[i] = Dot(E[i], V);
+            }
+
+            for (int i = 2, ip1 = 0; ip1 < 3; i = ip1++)
+            {
+                Vector3<Real> hatV = V - E[i] * nu[i] / sqrLenE[i];
+                Real sqrLenHatV = Dot(hatV, hatV);
+                if (sqrLenHatV > (Real)0)
+                {
+                    Vector3<Real> hatDelta = Delta[i] - E[i] * del[i] / sqrLenE[i];
+                    Real alpha = -Dot(hatV, hatDelta);
+                    if (alpha >= (Real)0)
+                    {
+                        Real sqrLenHatDelta = Dot(hatDelta, hatDelta);
+                        Real beta = alpha * alpha - sqrLenHatV * (sqrLenHatDelta - rsqr);
+                        if (beta >= (Real)0)
+                        {
+                            QuadraticField<Real> bfield(beta);
+                            Real mu = Dot(ExU[i], Delta[i]);
+                            Real omega = Dot(ExU[i], hatV);
+                            QFElement arg0(mu * sqrLenHatV + omega * alpha, -omega);
+                            if (bfield.GreaterThanOrEqualZero(arg0))
+                            {
+                                QFElement arg1(del[i] * sqrLenHatV + nu[i] * alpha, -nu[i]);
+                                if (bfield.GreaterThanOrEqualZero(arg1))
+                                {
+                                    QFElement arg2(delp[i] * sqrLenHatV + nu[i] * alpha, -nu[i]);
+                                    if (bfield.LessThanOrEqualZero(arg2))
+                                    {
+                                        result.intersectionType = 1;
+                                        result.field = bfield;
+                                        result.contactTime.x = alpha / sqrLenHatV;
+                                        result.contactTime.y = (Real)-1 / sqrLenHatV;
+                                        for (int j = 0; j < 3; ++j)
+                                        {
+                                            result.contactPoint[j].x = sphere.center[j] + result.contactTime.x * sphereVelocity[j];
+                                            result.contactPoint[j].y = result.contactTime.y * sphereVelocity[j];
+                                        }
+                                        return result;
+                                    }
+                                }
+                            }
+                        }
+                    }
+                }
+            }
+
+            // Test for contact with sphere wedges of the sphere-swept
+            // volume.  We know that |V|^2 > 0 because of a previous
+            // early-exit test.
+            for (int im1 = 2, i = 0; i < 3; im1 = i++)
+            {
+                Real alpha = -Dot(V, Delta[i]);
+                if (alpha >= (Real)0)
+                {
+                    Real sqrLenDelta = Dot(Delta[i], Delta[i]);
+                    Real beta = alpha * alpha - sqrLenV * (sqrLenDelta - rsqr);
+                    if (beta >= (Real)0)
+                    {
+                        QuadraticField<Real> bfield(beta);
+                        QFElement arg0(delp[im1] * sqrLenV + nu[im1] * alpha, -nu[im1]);
+                        if (bfield.GreaterThanOrEqualZero(arg0))
+                        {
+                            QFElement arg1(del[i] * sqrLenV + nu[i] * alpha, -nu[i]);
+                            if (bfield.LessThanOrEqualZero(arg1))
+                            {
+                                result.intersectionType = 1;
+                                result.field = bfield;
+                                result.contactTime.x = alpha / sqrLenV;
+                                result.contactTime.y = (Real)-1 / sqrLenV;
+                                for (int j = 0; j < 3; ++j)
+                                {
+                                    result.contactPoint[j].x = sphere.center[j] + result.contactTime.x * sphereVelocity[j];
+                                    result.contactPoint[j].y = result.contactTime.y * sphereVelocity[j];
+                                }
+                                return result;
+                            }
+                        }
+                    }
+                }
+            }
+
+            // The ray and sphere-swept volume do not intersect, so the sphere
+            // and triangle do not come into contact.  The 'result' is already
+            // set to the correct state for this case.
+            return result;
+        }
+    };
+}
diff --git a/Engine/Plugins/Experimental/GeometryProcessing/Source/GeometryAlgorithms/Private/GTEngine/Mathematics/GteIntrTriangle3OrientedBox3.h b/Engine/Plugins/Experimental/GeometryProcessing/Source/GeometryAlgorithms/Private/GTEngine/Mathematics/GteIntrTriangle3OrientedBox3.h
new file mode 100644
--- /dev/null
+++ b/Engine/Plugins/Experimental/GeometryProcessing/Source/GeometryAlgorithms/Private/GTEngine/Mathematics/GteIntrTriangle3OrientedBox3.h
@@ -0,0 +1,207 @@
+// David Eberly, Geometric Tools, Redmond WA 98052
+// Copyright (c) 1998-2019
+// Distributed under the Boost Software License, Version 1.0.
+// http://www.boost.org/LICENSE_1_0.txt
+// http://www.geometrictools.com/License/Boost/LICENSE_1_0.txt
+// File Version: 3.0.2 (2019/02/13)
+
+#pragma once
+
+#include <GTEngine/Mathematics/GteTIQuery.h>
+#include <GTEngine/Mathematics/GteFIQuery.h>
+#include <GTEngine/Mathematics/GteTriangle.h>
+#include <GTEngine/Mathematics/GteOrientedBox.h>
+#include <GTEngine/Mathematics/GteIntrConvexPolygonHyperplane.h>
+#include <GTEngine/Mathematics/GteVector3.h>
+
+// The test-intersection query is based on the document
+// https://www.geometrictools.com/Documentation/MethodOfSeparatingAxes.pdf
+// The find-intersection query clips the triangle against the faces of
+// the oriented box.
+
+namespace gte
+{
+    template <typename Real>
+    class TIQuery<Real, Triangle3<Real>, OrientedBox3<Real>>
+    {
+    public:
+        struct Result
+        {
+            bool intersect;
+        };
+
+        Result operator()(Triangle3<Real> const& triangle, OrientedBox3<Real> const& box)
+        {
+            Result result;
+
+            Real min0, max0, min1, max1;
+            Vector3<Real> D, edge[3];
+
+            // Test direction of triangle normal.
+            edge[0] = triangle.v[1] - triangle.v[0];
+            edge[1] = triangle.v[2] - triangle.v[0];
+            D = Cross(edge[0], edge[1]);
+            min0 = Dot(D, triangle.v[0]);
+            max0 = min0;
+            GetProjection(D, box, min1, max1);
+            if (max1 < min0 || max0 < min1)
+            {
+                result.intersect = false;
+                return result;
+            }
+
+            // Test direction of box faces.
+            for (int i = 0; i < 3; ++i)
+            {
+                D = box.axis[i];
+                GetProjection(D, triangle, min0, max0);
+                Real DdC = Dot(D, box.center);
+                min1 = DdC - box.extent[i];
+                max1 = DdC + box.extent[i];
+                if (max1 < min0 || max0 < min1)
+                {
+                    result.intersect = false;
+                    return result;
+                }
+            }
+
+            // Test direction of triangle-box edge cross products.
+            edge[2] = edge[1] - edge[0];
+            for (int i0 = 0; i0 < 3; ++i0)
+            {
+                for (int i1 = 0; i1 < 3; ++i1)
+                {
+                    D = Cross(edge[i0], box.axis[i1]);
+                    GetProjection(D, triangle, min0, max0);
+                    GetProjection(D, box, min1, max1);
+                    if (max1 < min0 || max0 < min1)
+                    {
+                        result.intersect = false;
+                        return result;
+                    }
+                }
+            }
+
+            result.intersect = true;
+            return result;
+        }
+
+    private:
+        void GetProjection(Vector3<Real> const& axis, Triangle3<Real> const& triangle, Real& imin, Real& imax)
+        {
+            Real dot[3] =
+            {
+                Dot(axis, triangle.v[0]),
+                Dot(axis, triangle.v[1]),
+                Dot(axis, triangle.v[2])
+            };
+
+            imin = dot[0];
+            imax = imin;
+
+            if (dot[1] < imin)
+            {
+                imin = dot[1];
+            }
+            else if (dot[1] > imax)
+            {
+                imax = dot[1];
+            }
+
+            if (dot[2] < imin)
+            {
+                imin = dot[2];
+            }
+            else if (dot[2] > imax)
+            {
+                imax = dot[2];
+            }
+        }
+
+        void GetProjection(Vector3<Real> const& axis, OrientedBox3<Real> const& box, Real& imin, Real& imax)
+        {
+            Real origin = Dot(axis, box.center);
+            Real maximumExtent =
+                std::abs(box.extent[0] * Dot(axis, box.axis[0])) +
+                std::abs(box.extent[1] * Dot(axis, box.axis[1])) +
+                std::abs(box.extent[2] * Dot(axis, box.axis[2]));
+
+            imin = origin - maximumExtent;
+            imax = origin + maximumExtent;
+        }
+    };
+
+    template <typename Real>
+    class FIQuery<Real, Triangle3<Real>, OrientedBox3<Real>>
+    {
+    public:
+        struct Result
+        {
+            std::vector<Vector3<Real>> insidePolygon;
+            std::vector<std::vector<Vector3<Real>>> outsidePolygons;
+        };
+
+        Result operator()(Triangle3<Real> const& triangle, OrientedBox3<Real> const& box)
+        {
+            Result result;
+
+            // Start with the triangle and clip it against each face of the
+            // box.  The largest number of vertices for the polygon of
+            // intersection is 7.
+            result.insidePolygon.resize(3);
+            for (int i = 0; i < 3; ++i)
+            {
+                result.insidePolygon[i] = triangle.v[i];
+            }
+
+            typedef FIQuery<Real, std::vector<Vector<3, Real>>, Hyperplane<3, Real>> PPQuery;
+
+            Plane3<Real> plane;
+            PPQuery ppQuery;
+            typename PPQuery::Result ppResult;
+            for (int dir = -1; dir <= 1; dir += 2)
+            {
+                for (int side = 0; side < 3; ++side)
+                {
+                    // Create a plane for the box face that points inside the box.
+                    plane.normal = ((Real)dir) * box.axis[side];
+                    plane.constant = Dot(plane.normal, box.center) - box.extent[side];
+
+                    ppResult = ppQuery(result.insidePolygon, plane);
+                    switch (ppResult.configuration)
+                    {
+                    case PPQuery::Configuration::SPLIT:
+                        result.insidePolygon = ppResult.positivePolygon;
+                        result.outsidePolygons.push_back(ppResult.negativePolygon);
+                        break;
+                    case PPQuery::Configuration::POSITIVE_SIDE_VERTEX:
+                    case PPQuery::Configuration::POSITIVE_SIDE_EDGE:
+                    case PPQuery::Configuration::POSITIVE_SIDE_STRICT:
+                        // The result.insidePolygon is already
+                        // ppResult.positivePolygon, but to make it clear,
+                        // assign it here.
+                        result.insidePolygon = ppResult.positivePolygon;
+                        break;
+                    case PPQuery::Configuration::NEGATIVE_SIDE_VERTEX:
+                    case PPQuery::Configuration::NEGATIVE_SIDE_EDGE:
+                    case PPQuery::Configuration::NEGATIVE_SIDE_STRICT:
+                        result.insidePolygon.clear();
+                        result.outsidePolygons.push_back(ppResult.negativePolygon);
+                        return result;
+                    case PPQuery::Configuration::CONTAINED:
+                        // A triangle coplanar with a box face will be
+                        // processed as if it is inside the box.
+                        result.insidePolygon = ppResult.intersection;
+                        break;
+                    default:
+                        result.insidePolygon.clear();
+                        result.outsidePolygons.clear();
+                        break;
+                    }
+                }
+            }
+
+            return result;
+        }
+    };
+}
diff --git a/Engine/Plugins/Experimental/GeometryProcessing/Source/GeometryAlgorithms/Private/GTEngine/Mathematics/GteInvSqrtEstimate.h b/Engine/Plugins/Experimental/GeometryProcessing/Source/GeometryAlgorithms/Private/GTEngine/Mathematics/GteInvSqrtEstimate.h
new file mode 100644
--- /dev/null
+++ b/Engine/Plugins/Experimental/GeometryProcessing/Source/GeometryAlgorithms/Private/GTEngine/Mathematics/GteInvSqrtEstimate.h
@@ -0,0 +1,193 @@
+// David Eberly, Geometric Tools, Redmond WA 98052
+// Copyright (c) 1998-2019
+// Distributed under the Boost Software License, Version 1.0.
+// http://www.boost.org/LICENSE_1_0.txt
+// http://www.geometrictools.com/License/Boost/LICENSE_1_0.txt
+// File Version: 3.0.1 (2018/10/05)
+
+#pragma once
+
+#include <GTEngine/Mathematics/GteMath.h>
+
+// Minimax polynomial approximations to 1/sqrt(x).  The polynomial p(x) of
+// degree D minimizes the quantity maximum{|1/sqrt(x) - p(x)| : x in [1,2]}
+// over all polynomials of degree D.
+
+namespace gte
+{
+
+template <typename Real>
+class InvSqrtEstimate
+{
+public:
+    // The input constraint is x in [1,2].  For example,
+    //   float x; // in [1,2]
+    //   float result = InvSqrtEstimate<float>::Degree<3>(x);
+    template <int D>
+    inline static Real Degree(Real x);
+
+    // The input constraint is x > 0.  Range reduction is used to generate a
+    // value y in [1,2], call Evaluate(y), and combine the output with the
+    // proper exponent to obtain the approximation.  For example,
+    //   float x;  // x > 0
+    //   float result = InvSqrtEstimate<float>::DegreeRR<3>(x);
+    template <int D>
+    inline static Real DegreeRR(Real x);
+
+private:
+    // Metaprogramming and private implementation to allow specialization of
+    // a template member function.
+    template <int D> struct degree {};
+    inline static Real Evaluate(degree<1>, Real t);
+    inline static Real Evaluate(degree<2>, Real t);
+    inline static Real Evaluate(degree<3>, Real t);
+    inline static Real Evaluate(degree<4>, Real t);
+    inline static Real Evaluate(degree<5>, Real t);
+    inline static Real Evaluate(degree<6>, Real t);
+    inline static Real Evaluate(degree<7>, Real t);
+    inline static Real Evaluate(degree<8>, Real t);
+
+    // Support for range reduction.
+    inline static void Reduce(Real x, Real& adj, Real& y, int& p);
+    inline static Real Combine(Real adj, Real y, int p);
+};
+
+
+template <typename Real>
+template <int D>
+inline Real InvSqrtEstimate<Real>::Degree(Real x)
+{
+    Real t = x - (Real)1;  // t in (0,1]
+    return Evaluate(degree<D>(), t);
+}
+
+template <typename Real>
+template <int D>
+inline Real InvSqrtEstimate<Real>::DegreeRR(Real x)
+{
+    Real adj, y;
+    int p;
+    Reduce(x, adj, y, p);
+    Real poly = Degree<D>(y);
+    Real result = Combine(adj, poly, p);
+    return result;
+}
+
+template <typename Real>
+inline Real InvSqrtEstimate<Real>::Evaluate(degree<1>, Real t)
+{
+    Real poly;
+    poly = (Real)GTE_C_INVSQRT_DEG1_C1;
+    poly = (Real)GTE_C_INVSQRT_DEG1_C0 + poly * t;
+    return poly;
+}
+
+template <typename Real>
+inline Real InvSqrtEstimate<Real>::Evaluate(degree<2>, Real t)
+{
+    Real poly;
+    poly = (Real)GTE_C_INVSQRT_DEG2_C2;
+    poly = (Real)GTE_C_INVSQRT_DEG2_C1 + poly * t;
+    poly = (Real)GTE_C_INVSQRT_DEG2_C0 + poly * t;
+    return poly;
+}
+
+template <typename Real>
+inline Real InvSqrtEstimate<Real>::Evaluate(degree<3>, Real t)
+{
+    Real poly;
+    poly = (Real)GTE_C_INVSQRT_DEG3_C3;
+    poly = (Real)GTE_C_INVSQRT_DEG3_C2 + poly * t;
+    poly = (Real)GTE_C_INVSQRT_DEG3_C1 + poly * t;
+    poly = (Real)GTE_C_INVSQRT_DEG3_C0 + poly * t;
+    return poly;
+}
+
+template <typename Real>
+inline Real InvSqrtEstimate<Real>::Evaluate(degree<4>, Real t)
+{
+    Real poly;
+    poly = (Real)GTE_C_INVSQRT_DEG4_C4;
+    poly = (Real)GTE_C_INVSQRT_DEG4_C3 + poly * t;
+    poly = (Real)GTE_C_INVSQRT_DEG4_C2 + poly * t;
+    poly = (Real)GTE_C_INVSQRT_DEG4_C1 + poly * t;
+    poly = (Real)GTE_C_INVSQRT_DEG4_C0 + poly * t;
+    return poly;
+}
+
+template <typename Real>
+inline Real InvSqrtEstimate<Real>::Evaluate(degree<5>, Real t)
+{
+    Real poly;
+    poly = (Real)GTE_C_INVSQRT_DEG5_C5;
+    poly = (Real)GTE_C_INVSQRT_DEG5_C4 + poly * t;
+    poly = (Real)GTE_C_INVSQRT_DEG5_C3 + poly * t;
+    poly = (Real)GTE_C_INVSQRT_DEG5_C2 + poly * t;
+    poly = (Real)GTE_C_INVSQRT_DEG5_C1 + poly * t;
+    poly = (Real)GTE_C_INVSQRT_DEG5_C0 + poly * t;
+    return poly;
+}
+
+template <typename Real>
+inline Real InvSqrtEstimate<Real>::Evaluate(degree<6>, Real t)
+{
+    Real poly;
+    poly = (Real)GTE_C_INVSQRT_DEG6_C6;
+    poly = (Real)GTE_C_INVSQRT_DEG6_C5 + poly * t;
+    poly = (Real)GTE_C_INVSQRT_DEG6_C4 + poly * t;
+    poly = (Real)GTE_C_INVSQRT_DEG6_C3 + poly * t;
+    poly = (Real)GTE_C_INVSQRT_DEG6_C2 + poly * t;
+    poly = (Real)GTE_C_INVSQRT_DEG6_C1 + poly * t;
+    poly = (Real)GTE_C_INVSQRT_DEG6_C0 + poly * t;
+    return poly;
+}
+
+template <typename Real>
+inline Real InvSqrtEstimate<Real>::Evaluate(degree<7>, Real t)
+{
+    Real poly;
+    poly = (Real)GTE_C_INVSQRT_DEG7_C7;
+    poly = (Real)GTE_C_INVSQRT_DEG7_C6 + poly * t;
+    poly = (Real)GTE_C_INVSQRT_DEG7_C5 + poly * t;
+    poly = (Real)GTE_C_INVSQRT_DEG7_C4 + poly * t;
+    poly = (Real)GTE_C_INVSQRT_DEG7_C3 + poly * t;
+    poly = (Real)GTE_C_INVSQRT_DEG7_C2 + poly * t;
+    poly = (Real)GTE_C_INVSQRT_DEG7_C1 + poly * t;
+    poly = (Real)GTE_C_INVSQRT_DEG7_C0 + poly * t;
+    return poly;
+}
+
+template <typename Real>
+inline Real InvSqrtEstimate<Real>::Evaluate(degree<8>, Real t)
+{
+    Real poly;
+    poly = (Real)GTE_C_INVSQRT_DEG8_C8;
+    poly = (Real)GTE_C_INVSQRT_DEG8_C7 + poly * t;
+    poly = (Real)GTE_C_INVSQRT_DEG8_C6 + poly * t;
+    poly = (Real)GTE_C_INVSQRT_DEG8_C5 + poly * t;
+    poly = (Real)GTE_C_INVSQRT_DEG8_C4 + poly * t;
+    poly = (Real)GTE_C_INVSQRT_DEG8_C3 + poly * t;
+    poly = (Real)GTE_C_INVSQRT_DEG8_C2 + poly * t;
+    poly = (Real)GTE_C_INVSQRT_DEG8_C1 + poly * t;
+    poly = (Real)GTE_C_INVSQRT_DEG8_C0 + poly * t;
+    return poly;
+}
+
+template <typename Real>
+inline void InvSqrtEstimate<Real>::Reduce(Real x, Real& adj, Real& y, int& p)
+{
+    y = std::frexp(x, &p);  // y in [1/2,1)
+    y = ((Real)2)*y;  // y in [1,2)
+    --p;
+    adj = (1 & p)*(Real)GTE_C_INV_SQRT_2 + (1 & ~p)*(Real)1;
+    p = -(p >> 1);
+}
+
+template <typename Real>
+inline Real InvSqrtEstimate<Real>::Combine(Real adj, Real y, int p)
+{
+    return adj * std::ldexp(y, p);
+}
+
+
+}
diff --git a/Engine/Plugins/Experimental/GeometryProcessing/Source/GeometryAlgorithms/Private/GTEngine/Mathematics/GteIsPlanarGraph.h b/Engine/Plugins/Experimental/GeometryProcessing/Source/GeometryAlgorithms/Private/GTEngine/Mathematics/GteIsPlanarGraph.h
new file mode 100644
--- /dev/null
+++ b/Engine/Plugins/Experimental/GeometryProcessing/Source/GeometryAlgorithms/Private/GTEngine/Mathematics/GteIsPlanarGraph.h
@@ -0,0 +1,526 @@
+// David Eberly, Geometric Tools, Redmond WA 98052
+// Copyright (c) 1998-2019
+// Distributed under the Boost Software License, Version 1.0.
+// http://www.boost.org/LICENSE_1_0.txt
+// http://www.geometrictools.com/License/Boost/LICENSE_1_0.txt
+// File Version: 3.0.0 (2016/06/19)
+
+#pragma once
+
+#include <algorithm>
+#include <array>
+#include <map>
+#include <set>
+#include <vector>
+
+// Test whether an undirected graph is planar.  The input positions must be
+// unique and the input edges must be unique, so the number of positions is
+// at least 2 and the number of edges is at least one.  The elements of the
+// edges array must be indices in {0,..,positions.size()-1}.
+// 
+// A sort-and-sweep algorithm is used to determine edge-edge intersections.
+// If none of the intersections occur at edge-interior points, the graph is
+// planar.  See Game Physics (2nd edition), Section 6.2.2: Culling with
+// Axis-Aligned Bounding Boxes for such an algorithm.  The operator()
+// returns 'true' when the graph is planar.  If it returns 'false', the
+// 'invalidIntersections' set contains pairs of edges that intersect at an
+// edge-interior point (that by definition is not a graph vertex).  Each set
+// element is a pair of indices into the 'edges' array; the indices are
+// ordered so that element[0] < element[1].  The Real type must be chosen
+// to guarantee exact computation of edge-edge intersections.
+
+namespace gte
+{
+
+template <typename Real>
+class IsPlanarGraph
+{
+public:
+    IsPlanarGraph();
+
+    enum
+    {
+        IPG_IS_PLANAR_GRAPH = 0,
+        IPG_INVALID_INPUT_SIZES = 1,
+        IPG_DUPLICATED_POSITIONS = 2,
+        IPG_DUPLICATED_EDGES = 4,
+        IPG_DEGENERATE_EDGES = 8,
+        IPG_EDGES_WITH_INVALID_VERTICES = 16,
+        IPG_INVALID_INTERSECTIONS = 32
+    };
+
+    // The function returns a combination of the IPG_* flags listed in the
+    // previous enumeration.  A combined value of 0 indicates the input
+    // forms a planar graph.  If the combined value is not zero, you may
+    // examine the flags for the failure conditions and use the Get* member
+    // accessors to obtain specific information about the failure.  If the
+    // positions.size() < 2 or edges.size() == 0, the IPG_INVALID_INPUT_SIZES
+    // flag is set.
+
+    int operator()(std::vector<std::array<Real, 2>> const& positions,
+        std::vector<std::array<int, 2>> const& edges);
+
+    // A pair of indices (v0,v1) into the positions array is stored as
+    // (min(v0,v1), max(v0,v1)).  This supports sorted containers of edges.
+    struct OrderedEdge
+    {
+        OrderedEdge(int v0 = -1, int v1 = -1);
+        bool operator<(OrderedEdge const& edge) const;
+        std::array<int, 2> V;
+    };
+
+    inline std::vector<std::vector<int>> const& GetDuplicatedPositions() const;
+    inline std::vector<std::vector<int>> const& GetDuplicatedEdges() const;
+    inline std::vector<int> const& GetDegenerateEdges() const;
+    inline std::vector<int> const& GetEdgesWithInvalidVertices() const;
+    inline std::vector<OrderedEdge> const& GetInvalidIntersections() const;
+
+private:
+    class Endpoint
+    {
+    public:
+        Real value;     // endpoint value
+        int type;       // '0' if interval min, '1' if interval max.
+        int index;      // index of interval containing this endpoint
+
+        // Comparison operator for sorting.
+        bool operator<(Endpoint const& endpoint) const;
+    };
+
+    int IsValidTopology(std::vector<std::array<Real, 2>> const& positions,
+        std::vector<std::array<int, 2>> const& edges);
+
+    void ComputeOverlappingRectangles(std::vector<std::array<Real, 2>> const& positions,
+        std::vector<std::array<int, 2>> const& edges,
+        std::set<OrderedEdge>& overlappingRectangles) const;
+
+    bool InvalidSegmentIntersection(
+        std::array<Real, 2> const& p0, std::array<Real, 2> const& p1,
+        std::array<Real, 2> const& q0, std::array<Real, 2> const& q1) const;
+
+    std::vector<std::vector<int>> mDuplicatedPositions;
+    std::vector<std::vector<int>> mDuplicatedEdges;
+    std::vector<int> mDegenerateEdges;
+    std::vector<int> mEdgesWithInvalidVertices;
+    std::vector<OrderedEdge> mInvalidIntersections;
+    Real mZero, mOne;
+};
+
+
+template <typename Real>
+IsPlanarGraph<Real>::IsPlanarGraph()
+    :
+    mZero(0),
+    mOne(1)
+{
+}
+
+template <typename Real>
+int IsPlanarGraph<Real>::operator()(
+    std::vector<std::array<Real, 2>> const& positions,
+    std::vector<std::array<int, 2>> const& edges)
+{
+    mDuplicatedPositions.clear();
+    mDuplicatedEdges.clear();
+    mDegenerateEdges.clear();
+    mEdgesWithInvalidVertices.clear();
+    mInvalidIntersections.clear();
+
+    int flags = IsValidTopology(positions, edges);
+    if (flags == IPG_INVALID_INPUT_SIZES)
+    {
+        return flags;
+    }
+
+    std::set<OrderedEdge> overlappingRectangles;
+    ComputeOverlappingRectangles(positions, edges, overlappingRectangles);
+    for (auto key : overlappingRectangles)
+    {
+        // Get the endpoints of the line segments for the edges whose bounding
+        // rectangles overlapped.  Determine whether the line segments
+        // intersect.  If they do, determine how they intersect.
+        std::array<int, 2> e0 = edges[key.V[0]];
+        std::array<int, 2> e1 = edges[key.V[1]];
+        std::array<Real, 2> const& p0 = positions[e0[0]];
+        std::array<Real, 2> const& p1 = positions[e0[1]];
+        std::array<Real, 2> const& q0 = positions[e1[0]];
+        std::array<Real, 2> const& q1 = positions[e1[1]];
+        if (InvalidSegmentIntersection(p0, p1, q0, q1))
+        {
+            mInvalidIntersections.push_back(key);
+        }
+    }
+
+    if (mInvalidIntersections.size() > 0)
+    {
+        flags |= IPG_INVALID_INTERSECTIONS;
+    }
+
+    return flags;
+}
+
+template <typename Real>
+inline std::vector<std::vector<int>> const&
+IsPlanarGraph<Real>::GetDuplicatedPositions() const
+{
+    return mDuplicatedPositions;
+}
+
+template <typename Real>
+inline std::vector<std::vector<int>> const&
+IsPlanarGraph<Real>::GetDuplicatedEdges() const
+{
+    return mDuplicatedEdges;
+}
+
+template <typename Real>
+inline std::vector<int> const&
+IsPlanarGraph<Real>::GetDegenerateEdges() const
+{
+    return mDegenerateEdges;
+}
+
+template <typename Real>
+inline std::vector<int> const&
+IsPlanarGraph<Real>::GetEdgesWithInvalidVertices() const
+{
+    return mEdgesWithInvalidVertices;
+}
+
+template <typename Real>
+inline std::vector<typename IsPlanarGraph<Real>::OrderedEdge> const&
+IsPlanarGraph<Real>::GetInvalidIntersections() const
+{
+    return mInvalidIntersections;
+}
+
+template <typename Real>
+int IsPlanarGraph<Real>::IsValidTopology(std::vector<std::array<Real, 2>> const& positions,
+    std::vector<std::array<int, 2>> const& edges)
+{
+    int const numPositions = static_cast<int>(positions.size());
+    int const numEdges = static_cast<int>(edges.size());
+    if (numPositions < 2 || numEdges == 0)
+    {
+        // The graph must have at least one edge.
+        return IPG_INVALID_INPUT_SIZES;
+    }
+
+    // The positions must be unique.
+    int flags = IPG_IS_PLANAR_GRAPH;
+    std::map<std::array<Real, 2>, std::vector<int>> uniquePositions;
+    for (int i = 0; i < numPositions; ++i)
+    {
+        std::array<Real, 2> p = positions[i];
+        auto iter = uniquePositions.find(p);
+        if (iter == uniquePositions.end())
+        {
+            std::vector<int> indices;
+            indices.push_back(i);
+            uniquePositions.insert(std::make_pair(p, indices));
+        }
+        else
+        {
+            iter->second.push_back(i);
+        }
+    }
+    if (uniquePositions.size() < positions.size())
+    {
+        // At least two positions are duplicated.
+        for (auto const& element : uniquePositions)
+        {
+            if (element.second.size() > 1)
+            {
+                mDuplicatedPositions.push_back(element.second);
+            }
+        }
+        flags |= IPG_DUPLICATED_POSITIONS;
+    }
+
+    // The edges must be unique.
+    std::map<OrderedEdge, std::vector<int>> uniqueEdges;
+    for (int i = 0; i < numEdges; ++i)
+    {
+        OrderedEdge key(edges[i][0], edges[i][1]);
+        auto iter = uniqueEdges.find(key);
+        if (iter == uniqueEdges.end())
+        {
+            std::vector<int> indices;
+            indices.push_back(i);
+            uniqueEdges.insert(std::make_pair(key, indices));
+        }
+        else
+        {
+            iter->second.push_back(i);
+        }
+    }
+    if (uniqueEdges.size() < edges.size())
+    {
+        // At least two edges are duplicated, possibly even a pair of
+        // edges (v0,v1) and (v1,v0) which is not allowed because the
+        // graph is undirected.
+        for (auto const& element : uniqueEdges)
+        {
+            if (element.second.size() > 1)
+            {
+                mDuplicatedEdges.push_back(element.second);
+            }
+        }
+        flags |= IPG_DUPLICATED_EDGES;
+    }
+
+    // The edges are represented as pairs of indices into the 'positions'
+    // array.  The indices for a single edge must be different (no edges
+    // allowed from a vertex to itself) and all indices must be valid.
+    // At the same time, keep track of unique edges.
+    for (int i = 0; i < numEdges; ++i)
+    {
+        std::array<int, 2> e = edges[i];
+        if (e[0] == e[1])
+        {
+            // The edge is degenerate, originating and terminating at the
+            // same vertex.
+            mDegenerateEdges.push_back(i);
+            flags |= IPG_DEGENERATE_EDGES;
+        }
+
+        if (e[0] < 0 || e[0] >= numPositions || e[1] < 0 || e[1] >= numPositions)
+        {
+            // The edge has an index that references a non-existant vertex.
+            mEdgesWithInvalidVertices.push_back(i);
+            flags |= IPG_EDGES_WITH_INVALID_VERTICES;
+        }
+    }
+
+    return flags;
+}
+
+template <typename Real>
+void IsPlanarGraph<Real>::ComputeOverlappingRectangles(std::vector<std::array<Real, 2>> const& positions,
+    std::vector<std::array<int, 2>> const& edges, std::set<OrderedEdge>& overlappingRectangles) const
+{
+    // Compute axis-aligned bounding rectangles for the edges.
+    int const numEdges = static_cast<int>(edges.size());
+    std::vector<std::array<Real, 2>> emin(numEdges);
+    std::vector<std::array<Real, 2>> emax(numEdges);
+    for (int i = 0; i < numEdges; ++i)
+    {
+        std::array<int, 2> e = edges[i];
+        std::array<Real, 2> const& p0 = positions[e[0]];
+        std::array<Real, 2> const& p1 = positions[e[1]];
+
+        for (int j = 0; j < 2; ++j)
+        {
+            if (p0[j] < p1[j])
+            {
+                emin[i][j] = p0[j];
+                emax[i][j] = p1[j];
+            }
+            else
+            {
+                emin[i][j] = p1[j];
+                emax[i][j] = p0[j];
+            }
+        }
+    }
+
+    // Get the rectangle endpoints.
+    int const numEndpoints = 2 * numEdges;
+    std::vector<Endpoint> xEndpoints(numEndpoints);
+    std::vector<Endpoint> yEndpoints(numEndpoints);
+    for (int i = 0, j = 0; i < numEdges; ++i)
+    {
+        xEndpoints[j].type = 0;
+        xEndpoints[j].value = emin[i][0];
+        xEndpoints[j].index = i;
+        yEndpoints[j].type = 0;
+        yEndpoints[j].value = emin[i][1];
+        yEndpoints[j].index = i;
+        ++j;
+
+        xEndpoints[j].type = 1;
+        xEndpoints[j].value = emax[i][0];
+        xEndpoints[j].index = i;
+        yEndpoints[j].type = 1;
+        yEndpoints[j].value = emax[i][1];
+        yEndpoints[j].index = i;
+        ++j;
+    }
+
+    // Sort the rectangle endpoints.
+    std::sort(xEndpoints.begin(), xEndpoints.end());
+    std::sort(yEndpoints.begin(), yEndpoints.end());
+
+    // Sweep through the endpoints to determine overlapping x-intervals.  Use
+    // an active set of rectangles to reduce the complexity of the algorithm.
+    std::set<int> active;
+    for (int i = 0; i < numEndpoints; ++i)
+    {
+        Endpoint const& endpoint = xEndpoints[i];
+        int index = endpoint.index;
+        if (endpoint.type == 0)  // an interval 'begin' value
+        {
+            // In the 1D problem, the current interval overlaps with all the
+            // active intervals.  In 2D this we also need to check for
+            // y-overlap.
+            for (auto activeIndex : active)
+            {
+                // Rectangles activeIndex and index overlap in the
+                // x-dimension.  Test for overlap in the y-dimension.
+                std::array<Real, 2> const& r0min = emin[activeIndex];
+                std::array<Real, 2> const& r0max = emax[activeIndex];
+                std::array<Real, 2> const& r1min = emin[index];
+                std::array<Real, 2> const& r1max = emax[index];
+                if (r0max[1] >= r1min[1] && r0min[1] <= r1max[1])
+                {
+                    if (activeIndex < index)
+                    {
+                        overlappingRectangles.insert(OrderedEdge(activeIndex, index));
+                    }
+                    else
+                    {
+                        overlappingRectangles.insert(OrderedEdge(index, activeIndex));
+                    }
+                }
+            }
+            active.insert(index);
+        }
+        else  // an interval 'end' value
+        {
+            active.erase(index);
+        }
+    }
+}
+
+template <typename Real>
+bool IsPlanarGraph<Real>::InvalidSegmentIntersection(
+    std::array<Real, 2> const& p0, std::array<Real, 2> const& p1,
+    std::array<Real, 2> const& q0, std::array<Real, 2> const& q1) const
+{
+    // We must solve the two linear equations
+    //   p0 + t0 * (p1 - p0) = q0 + t1 * (q1 - q0)
+    // for the unknown variables t0 and t1.  These may be written as
+    //   t0 * (p1 - p0) - t1 * (q1 - q0) = q0 - p0
+    // If denom = Dot(p1 - p0, Perp(q1 - q0)) is not zero, then
+    //   t0 = Dot(q0 - p0, Perp(q1 - q0)) / denom = numer0 / denom
+    //   t1 = Dot(q0 - p0, Perp(p1 - p0)) / denom = numer1 / denom
+    // For an invalid intersection, we need (t0,t1) with:
+    // ((0 < t0 < 1) and (0 <= t1 <= 1)) or ((0 <= t0 <= 1) and (0 < t1 < 1).
+    // If denom is zero, then
+
+    std::array<Real, 2> p1mp0, q1mq0, q0mp0;
+    for (int j = 0; j < 2; ++j)
+    {
+        p1mp0[j] = p1[j] - p0[j];
+        q1mq0[j] = q1[j] - q0[j];
+        q0mp0[j] = q0[j] - p0[j];
+    }
+
+    Real denom = p1mp0[0] * q1mq0[1] - p1mp0[1] * q1mq0[0];
+    Real numer0 = q0mp0[0] * q1mq0[1] - q0mp0[1] * q1mq0[0];
+    Real numer1 = q0mp0[0] * p1mp0[1] - q0mp0[1] * p1mp0[0];
+
+    if (denom != mZero)
+    {
+        // The lines of the segments are not parallel.
+        if (denom > mZero)
+        {
+            if (mZero <= numer0 && numer0 <= denom && mZero <= numer1 && numer1 <= denom)
+            {
+                // The segments intersect.
+                return (numer0 != mZero && numer0 != denom) || (numer1 != mZero && numer1 != denom);
+            }
+            else
+            {
+                return false;
+            }
+        }
+        else  // denom < mZero
+        {
+            if (mZero >= numer0 && numer0 >= denom && mZero >= numer1 && numer1 >= denom)
+            {
+                // The segments intersect.
+                return (numer0 != mZero && numer0 != denom) || (numer1 != mZero && numer1 != denom);
+            }
+            else
+            {
+                return false;
+            }
+        }
+    }
+    else
+    {
+        // The lines of the segments are parallel.
+        if (numer0 != mZero || numer1 != mZero)
+        {
+            // The lines of the segments are separated.
+            return false;
+        }
+        else
+        {
+            // The segments lie on the same line.  Compute the parameter
+            // intervals for the segments in terms of the t0-parameter and
+            // determine their overlap (if any).
+            std::array<Real, 2> q1mp0;
+            for (int j = 0; j < 2; ++j)
+            {
+                q1mp0[j] = q1[j] - p0[j];
+            }
+            Real sqrLenP1mP0 = p1mp0[0] * p1mp0[0] + p1mp0[1] * p1mp0[1];
+            Real value0 = q0mp0[0] * p1mp0[0] + q0mp0[1] * p1mp0[1];
+            Real value1 = q1mp0[0] * p1mp0[0] + q1mp0[1] * p1mp0[1];
+            if ((value0 >= sqrLenP1mP0 && value1 >= sqrLenP1mP0)
+                || (value0 <= mZero && value1 <= mZero))
+            {
+                // If the segments intersect, they must do so at endpoints of
+                // the segments.
+                return false;
+            }
+            else
+            {
+                // The segments overlap in a positive-length interval.
+                return true;
+            }
+        }
+    }
+}
+
+template <typename Real>
+bool IsPlanarGraph<Real>::Endpoint::operator<(Endpoint const& endpoint) const
+{
+    if (value < endpoint.value)
+    {
+        return true;
+    }
+    if (value > endpoint.value)
+    {
+        return false;
+    }
+    return type < endpoint.type;
+}
+
+template <typename Real>
+IsPlanarGraph<Real>::OrderedEdge::OrderedEdge(int v0, int v1)
+{
+    if (v0 < v1)
+    {
+        // v0 is minimum
+        V[0] = v0;
+        V[1] = v1;
+    }
+    else
+    {
+        // v1 is minimum
+        V[0] = v1;
+        V[1] = v0;
+    }
+}
+
+template <typename Real>
+bool IsPlanarGraph<Real>::OrderedEdge::operator<(OrderedEdge const& edge) const
+{
+    // Lexicographical ordering used by std::array<int,2>.
+    return V < edge.V;
+}
+
+}
