move GeometryProcessing plugin from NotForLicensees to Experimental

diff --git a/Engine/Plugins/Experimental/GeometryProcessing/Source/GeometryAlgorithms/Private/GTEngine/Mathematics/GteLCPSolver.h b/Engine/Plugins/Experimental/GeometryProcessing/Source/GeometryAlgorithms/Private/GTEngine/Mathematics/GteLCPSolver.h
new file mode 100644
--- /dev/null
+++ b/Engine/Plugins/Experimental/GeometryProcessing/Source/GeometryAlgorithms/Private/GTEngine/Mathematics/GteLCPSolver.h
@@ -0,0 +1,616 @@
+// David Eberly, Geometric Tools, Redmond WA 98052
+// Copyright (c) 1998-2019
+// Distributed under the Boost Software License, Version 1.0.
+// http://www.boost.org/LICENSE_1_0.txt
+// http://www.geometrictools.com/License/Boost/LICENSE_1_0.txt
+// File Version: 3.4.2 (2016/11/24)
+
+#pragma once
+
+#include <algorithm>
+#include <array>
+#include <vector>
+
+// A class for solving the Linear Complementarity Problem (LCP)
+// w = q + M * z, w^T * z = 0, w >= 0, z >= 0.  The vectors q, w, and z are
+// n-tuples and the matrix M is n-by-n.  The inputs to Solve(...) are q and M.
+// The outputs are w and z, which are valid when the returned bool is true but
+// are invalid when the returned bool is false.
+//
+// The comments at the end of this file explain what the preprocessor symbol
+// means regarding the LCP solver implementation.  If the algorithm fails to
+// converge within the specified maximum number of iterations, consider
+// increasing the number and calling Solve(...) again.
+//#define GTE_REPORT_FAILED_TO_CONVERGE
+
+namespace gte
+{
+
+// Support templates for number of dimensions known at compile time or known
+// only at run time.
+template <typename Real, int... Dimensions>
+class LCPSolver {};
+
+
+template <typename Real>
+class LCPSolverShared
+{
+protected:
+    // Abstract base class construction.  A virtual destructor is not provided
+    // because there are no required side effects when destroying objects from
+    // the derived classes.  The member mMaxIterations is set by this call to
+    // the default value of n*n.
+    LCPSolverShared(int n);
+
+public:
+    // Theoretically, when there is a solution the algorithm must converge
+    // in a finite number of iterations.  The number of iterations depends
+    // on the problem at hand, but we need to guard against an infinite loop
+    // by limiting the number.  The implementation uses a maximum number of
+    // n*n (chosen arbitrarily).  You can set the number yourself, perhaps
+    // when a call to Solve fails--increase the number of iterations and call
+    // Solve again.
+    inline void SetMaxIterations(int maxIterations);
+    inline int GetMaxIterations() const;
+
+    // Access the actual number of iterations used in a call to Solve.
+    inline int GetNumIterations() const;
+
+    enum Result
+    {
+        HAS_TRIVIAL_SOLUTION,
+        HAS_NONTRIVIAL_SOLUTION,
+        NO_SOLUTION,
+        FAILED_TO_CONVERGE,
+        INVALID_INPUT
+    };
+
+protected:
+    // Bookkeeping of variables during the iterations of the solver.  The name
+    // is either 'w' or 'z' and is used for human-readable debugging help.
+    // The 'index' is that for the original variables w[index] or z[index].
+    // The 'complementary' index is the location of the complementary variable
+    // in mVarBasic[] or in mVarNonbasic[].  The 'tuple' is a pointer to
+    // &w[0] or &z[0], the choice based on name of 'w' or 'z', and is used to
+    // fill in the solution values (the variables are permuted during the
+    // pivoting algorithm).
+    struct Variable
+    {
+        char name;
+        int index;
+        int complementary;
+        Real* tuple;
+    };
+
+    // The augmented problem is w = q + M*z + z[n]*U = 0, where U is an
+    // n-tuple of 1-values.  We manipulate the augmented matrix [M | U | p(t)]
+    // where p(t) is a column vector of polynomials of at most degree n.  If
+    // p[r](t) is the polynomial for row r, then p[r](0) = q[r].  These are
+    // perturbations of q[r] designed so that the algorithm avoids degeneracies
+    // (a q-term becomes zero during the iterations).  The basic variables are
+    // w[0] through w[n-1] and the nonbasic variables are z[0] through z[n].
+    // The returned z consists only of z[0] through z[n-1].
+
+    // The derived classes ensure that the pointers point to the correct
+    // of elements for each array.  The matrix M must be stored in row-major
+    // order.
+    bool Solve(Real const* q, Real const* M, Real* w, Real* z, Result* result);
+
+    // Access mAugmented as a 2-dimensional array.
+    inline Real const& Augmented(int row, int col) const;
+    inline Real& Augmented(int row, int col);
+
+    // Support for polynomials with n+1 coefficients and degree no larger
+    // than n.
+    void MakeZero(Real* poly);
+    void Copy(Real const* poly0, Real* poly1);
+    bool LessThan(Real const* poly0, Real const* poly1);
+    bool LessThanZero(Real const* poly);
+    void Multiply(Real const* poly, Real scalar, Real* product);
+
+    int mDimension;
+    int mMaxIterations;
+    int mNumIterations;
+
+    // These pointers are set by the derived-class constructors to arrays
+    // that have the correct number of elements.  The arrays mVarBasic,
+    // mVarNonbasic, mQMin, mMinRatio, and mRatio each have n+1 elements.
+    // The mAugmented array has n rows and 2*(n+1) columns stored in
+    // row-major order in a 1-dimensional array.  The array of pointers
+    // mPoly has n elements.
+    Variable* mVarBasic;
+    Variable* mVarNonbasic;
+    int mNumCols;
+    Real* mAugmented;
+    Real* mQMin;
+    Real* mMinRatio;
+    Real* mRatio;
+    Real** mPoly;
+};
+
+
+template <typename Real, int n>
+class LCPSolver<Real, n> : public LCPSolverShared<Real>
+{
+public:
+    // Construction.  The member mMaxIterations is set by this call to the
+    // default value of n*n.
+    LCPSolver();
+
+    // If you want to know specifically why 'true' or 'false' was returned,
+    // pass the address of a Result variable as the last parameter.
+    bool Solve(std::array<Real, n> const& q, std::array<std::array<Real, n>, n> const& M,
+        std::array<Real, n>& w, std::array<Real, n>& z,
+        typename LCPSolverShared<Real>::Result* result = nullptr);
+
+private:
+    std::array<typename LCPSolverShared<Real>::Variable, n + 1> mArrayVarBasic;
+    std::array<typename LCPSolverShared<Real>::Variable, n + 1> mArrayVarNonbasic;
+    std::array<Real, 2 * (n + 1) * n> mArrayAugmented;
+    std::array<Real, n + 1> mArrayQMin;
+    std::array<Real, n + 1> mArrayMinRatio;
+    std::array<Real, n + 1> mArrayRatio;
+    std::array<Real*, n> mArrayPoly;
+};
+
+
+template <typename Real>
+class LCPSolver<Real> : public LCPSolverShared<Real>
+{
+public:
+    // Construction.  The member mMaxIterations is set by this call to the
+    // default value of n*n.
+    LCPSolver(int n);
+
+    // The input q must have n elements and the input M must be an n-by-n
+    // matrix stored in row-major order.  The outputs w and z have n elements.
+    // If you want to know specifically why 'true' or 'false' was returned,
+    // pass the address of a Result variable as the last parameter.
+    bool Solve(std::vector<Real> const& q, std::vector<Real> const& M,
+        std::vector<Real>& w, std::vector<Real>& z,
+        typename LCPSolverShared<Real>::Result* result = nullptr);
+
+private:
+    std::vector<typename LCPSolverShared<Real>::Variable> mVectorVarBasic;
+    std::vector<typename LCPSolverShared<Real>::Variable> mVectorVarNonbasic;
+    std::vector<Real> mVectorAugmented;
+    std::vector<Real> mVectorQMin;
+    std::vector<Real> mVectorMinRatio;
+    std::vector<Real> mVectorRatio;
+    std::vector<Real*> mVectorPoly;
+};
+
+
+template <typename Real>
+LCPSolverShared<Real>::LCPSolverShared(int n)
+    :
+    mNumIterations(0),
+    mVarBasic(nullptr),
+    mVarNonbasic(nullptr),
+    mNumCols(0),
+    mAugmented(nullptr),
+    mQMin(nullptr),
+    mMinRatio(nullptr),
+    mRatio(nullptr),
+    mPoly(nullptr)
+{
+    if (n > 0)
+    {
+        mDimension = n;
+        mMaxIterations = n * n;
+    }
+    else
+    {
+        mDimension = 0;
+        mMaxIterations = 0;
+    }
+}
+
+template <typename Real>
+inline void LCPSolverShared<Real>::SetMaxIterations(int maxIterations)
+{
+    mMaxIterations = (maxIterations > 0 ? maxIterations : mDimension * mDimension);
+}
+
+template <typename Real>
+inline int LCPSolverShared<Real>::GetMaxIterations() const
+{
+    return mMaxIterations;
+}
+
+template <typename Real>
+inline int LCPSolverShared<Real>::GetNumIterations() const
+{
+    return mNumIterations;
+}
+
+template <typename Real>
+bool LCPSolverShared<Real>::Solve(Real const* q, Real const* M,
+    Real* w, Real* z, Result* result)
+{
+    // Perturb the q[r] constants to be polynomials of degree r+1 represented
+    // as an array of n+1 coefficients.  The coefficient with index r+1 is 1
+    // and the coefficients with indices larger than r+1 are 0.
+    for (int r = 0; r < mDimension; ++r)
+    {
+        mPoly[r] = &Augmented(r, mDimension + 1);
+        MakeZero(mPoly[r]);
+        mPoly[r][0] = q[r];
+        mPoly[r][r + 1] = (Real)1;
+    }
+
+    // Determine whether there is the trivial solution w = z = 0.
+    Copy(mPoly[0], mQMin);
+    int basic = 0;
+    for (int r = 1; r < mDimension; ++r)
+    {
+        if (LessThan(mPoly[r], mQMin))
+        {
+            Copy(mPoly[r], mQMin);
+            basic = r;
+        }
+    }
+
+    if (!LessThanZero(mQMin))
+    {
+        for (int r = 0; r < mDimension; ++r)
+        {
+            w[r] = q[r];
+            z[r] = (Real)0;
+        }
+
+        if (result)
+        {
+            *result = HAS_TRIVIAL_SOLUTION;
+        }
+        return true;
+    }
+
+    // Initialize the remainder of the augmented matrix with M and U.
+    for (int r = 0; r < mDimension; ++r)
+    {
+        for (int c = 0; c < mDimension; ++c)
+        {
+            Augmented(r, c) = M[c + mDimension * r];
+        }
+        Augmented(r, mDimension) = (Real)1;
+    }
+
+    // Keep track of when the variables enter and exit the dictionary,
+    // including where complementary variables are relocated.
+    for (int i = 0; i <= mDimension; ++i)
+    {
+        mVarBasic[i].name = 'w';
+        mVarBasic[i].index = i;
+        mVarBasic[i].complementary = i;
+        mVarBasic[i].tuple = w;
+        mVarNonbasic[i].name = 'z';
+        mVarNonbasic[i].index = i;
+        mVarNonbasic[i].complementary = i;
+        mVarNonbasic[i].tuple = z;
+    }
+
+    // The augmented variable z[n] is the initial driving variable for
+    // pivoting.  The equation 'basic' is the one to solve for z[n] and
+    // pivoting with w[basic].  The last column of M remains all 1-values
+    // for this initial step, so no algebraic computations occur for M[r][n].
+    int driving = mDimension;
+    for (int r = 0; r < mDimension; ++r)
+    {
+        if (r != basic)
+        {
+            for (int c = 0; c < mNumCols; ++c)
+            {
+                if (c != mDimension)
+                {
+                    Augmented(r, c) -= Augmented(basic, c);
+                }
+            }
+        }
+    }
+
+    for (int c = 0; c < mNumCols; ++c)
+    {
+        if (c != mDimension)
+        {
+            Augmented(basic, c) = -Augmented(basic, c);
+        }
+    }
+
+    mNumIterations = 0;
+    for (int i = 0; i < mMaxIterations; ++i, ++mNumIterations)
+    {
+        // The basic variable of equation 'basic' exited the dictionary, so
+        // its complementary (nonbasic) variable must become the next driving
+        // variable in order for it to enter the dictionary.
+        int nextDriving = mVarBasic[basic].complementary;
+        mVarNonbasic[nextDriving].complementary = driving;
+        std::swap(mVarBasic[basic], mVarNonbasic[driving]);
+        if (mVarNonbasic[driving].index == mDimension)
+        {
+            // The algorithm has converged.
+            for (int r = 0; r < mDimension; ++r)
+            {
+                mVarBasic[r].tuple[mVarBasic[r].index] = mPoly[r][0];
+            }
+            for (int c = 0; c <= mDimension; ++c)
+            {
+                int index = mVarNonbasic[c].index;
+                if (index < mDimension)
+                {
+                    mVarNonbasic[c].tuple[index] = (Real)0;
+                }
+            }
+            if (result)
+            {
+                *result = HAS_NONTRIVIAL_SOLUTION;
+            }
+            return true;
+        }
+
+        // Determine the 'basic' equation for which the ratio -q[r]/M(r,driving)
+        // is minimized among all equations r with M(r,driving) < 0.
+        driving = nextDriving;
+        basic = -1;
+        for (int r = 0; r < mDimension; ++r)
+        {
+            if (Augmented(r, driving) < (Real)0)
+            {
+                Real factor = (Real)-1 / Augmented(r, driving);
+                Multiply(mPoly[r], factor, mRatio);
+                if (basic == -1 || LessThan(mRatio, mMinRatio))
+                {
+                    Copy(mRatio, mMinRatio);
+                    basic = r;
+                }
+            }
+        }
+
+        if (basic == -1)
+        {
+            // The coefficients of z[driving] in all the equations are
+            // nonnegative, so the z[driving] variable cannot leave the
+            // dictionary.  There is no solution to the LCP.
+            for (int r = 0; r < mDimension; ++r)
+            {
+                w[r] = (Real)0;
+                z[r] = (Real)0;
+            }
+
+            if (result)
+            {
+                *result = NO_SOLUTION;
+            }
+            return false;
+        }
+
+        // Solve the basic equation so that z[driving] enters the dictionary
+        // and w[basic] exits the dictionary.
+        Real invDenom = (Real)1 / Augmented(basic, driving);
+        for (int r = 0; r < mDimension; ++r)
+        {
+            if (r != basic && Augmented(r, driving) != (Real)0)
+            {
+                Real multiplier = Augmented(r, driving) * invDenom;
+                for (int c = 0; c < mNumCols; ++c)
+                {
+                    if (c != driving)
+                    {
+                        Augmented(r, c) -= Augmented(basic, c) * multiplier;
+                    }
+                    else
+                    {
+                        Augmented(r, driving) = multiplier;
+                    }
+                }
+            }
+        }
+
+        for (int c = 0; c < mNumCols; ++c)
+        {
+            if (c != driving)
+            {
+                Augmented(basic, c) = -Augmented(basic, c) * invDenom;
+            }
+            else
+            {
+                Augmented(basic, driving) = invDenom;
+            }
+        }
+    }
+
+    // Numerical round-off errors can cause the Lemke algorithm not to
+    // converge.  In particular, the code above has a test
+    //   if (mAugmented[r][driving] < (Real)0) { ... }
+    // to determine the 'basic' equation with which to pivot.  It is
+    // possible that theoretically mAugmented[r][driving]is zero but
+    // rounding errors cause it to be slightly negative.  If theoretically
+    // all mAugmented[r][driving] >= 0, there is no solution to the LCP.
+    // With the rounding errors, if the algorithm fails to converge within
+    // the specified number of iterations, NO_SOLUTION is returned, which
+    // is hopefully the correct result. It is also possible that the rounding
+    // errors lead to a NO_SOLUTION (returned from inside the loop) when in
+    // fact there is a solution. [When the LCP solver is used by intersection
+    // testing algorithms, the hope is that misclassifications occur only
+    // when the two objects are nearly in tangential contact.]
+    //
+    // To determine whether the rounding errors are the problem, you can
+    // execute the query using exact arithmetic with the following type
+    // used for 'Real' (replacing 'float' or 'double') of
+    // BSRational<UIntegerAP32> Rational.
+    //
+    // That said, if the algorithm fails to converge and you believe that
+    // the rounding errors are not causing this, please file a bug report
+    // and provide the input data to the solver.
+
+#if defined(GTE_REPORT_FAILED_TO_CONVERGE)
+    LogInformation("The algorithm failed to converge.");
+#endif
+
+    if (result)
+    {
+        *result = FAILED_TO_CONVERGE;
+    }
+    return false;
+}
+
+template <typename Real>
+inline Real const& LCPSolverShared<Real>::Augmented(int row, int col) const
+{
+    return mAugmented[col + mNumCols * row];
+}
+
+template <typename Real>
+inline Real& LCPSolverShared<Real>::Augmented(int row, int col)
+{
+    return mAugmented[col + mNumCols * row];
+}
+
+
+template <typename Real>
+void LCPSolverShared<Real>::MakeZero(Real* poly)
+{
+    for (int i = 0; i <= mDimension; ++i)
+    {
+        poly[i] = (Real)0;
+    }
+}
+
+template <typename Real>
+void LCPSolverShared<Real>::Copy(Real const* poly0, Real* poly1)
+{
+    for (int i = 0; i <= mDimension; ++i)
+    {
+        poly1[i] = poly0[i];
+    }
+}
+
+template <typename Real>
+bool LCPSolverShared<Real>::LessThan(Real const* poly0, Real const* poly1)
+{
+    for (int i = 0; i <= mDimension; ++i)
+    {
+        if (poly0[i] < poly1[i])
+        {
+            return true;
+        }
+
+        if (poly0[i] > poly1[i])
+        {
+            return false;
+        }
+    }
+
+    return false;
+}
+
+template <typename Real>
+bool LCPSolverShared<Real>::LessThanZero(Real const* poly)
+{
+    for (int i = 0; i <= mDimension; ++i)
+    {
+        if (poly[i] < (Real)0)
+        {
+            return true;
+        }
+
+        if (poly[i] > (Real)0)
+        {
+            return false;
+        }
+    }
+
+    return false;
+}
+
+template <typename Real>
+void LCPSolverShared<Real>::Multiply(Real const* poly, Real scalar, Real* product)
+{
+    for (int i = 0; i <= mDimension; ++i)
+    {
+        product[i] = poly[i] * scalar;
+    }
+}
+
+
+template <typename Real, int n>
+LCPSolver<Real, n>::LCPSolver()
+    :
+    LCPSolverShared<Real>(n)
+{
+    this->mVarBasic = mArrayVarBasic.data();
+    this->mVarNonbasic = mArrayVarNonbasic.data();
+    this->mNumCols = 2 * (n + 1);
+    this->mAugmented = mArrayAugmented.data();
+    this->mQMin = mArrayQMin.data();
+    this->mMinRatio = mArrayMinRatio.data();
+    this->mRatio = mArrayRatio.data();
+    this->mPoly = mArrayPoly.data();
+}
+
+template <typename Real, int n>
+bool LCPSolver<Real, n>::Solve(
+    std::array<Real, n> const& q, std::array<std::array<Real, n>, n> const& M,
+    std::array<Real, n>& w, std::array<Real, n>& z,
+    typename LCPSolverShared<Real>::Result* result)
+{
+    return LCPSolverShared<Real>::Solve(q.data(), M.front().data(), w.data(), z.data(), result);
+}
+
+
+template <typename Real>
+LCPSolver<Real>::LCPSolver(int n)
+    :
+    LCPSolverShared<Real>(n)
+{
+    if (n > 0)
+    {
+        mVectorVarBasic.resize(n + 1);
+        mVectorVarNonbasic.resize(n + 1);
+        mVectorAugmented.resize(2 * (n + 1) * n);
+        mVectorQMin.resize(n + 1);
+        mVectorMinRatio.resize(n + 1);
+        mVectorRatio.resize(n + 1);
+        mVectorPoly.resize(n);
+
+        this->mVarBasic = mVectorVarBasic.data();
+        this->mVarNonbasic = mVectorVarNonbasic.data();
+        this->mNumCols = 2 * (n + 1);
+        this->mAugmented = mVectorAugmented.data();
+        this->mQMin = mVectorQMin.data();
+        this->mMinRatio = mVectorMinRatio.data();
+        this->mRatio = mVectorRatio.data();
+        this->mPoly = mVectorPoly.data();
+    }
+}
+
+template <typename Real>
+bool LCPSolver<Real>::Solve(
+    std::vector<Real> const& q, std::vector<Real> const& M,
+    std::vector<Real>& w, std::vector<Real>& z,
+    typename LCPSolverShared<Real>::Result* result)
+{
+    if (this->mDimension > static_cast<int>(q.size())
+        || this->mDimension * this->mDimension > static_cast<int>(M.size()))
+    {
+        if (result)
+        {
+            *result = this->INVALID_INPUT;
+        }
+        return false;
+    }
+
+    if (this->mDimension > static_cast<int>(w.size()))
+    {
+        w.resize(this->mDimension);
+    }
+
+    if (this->mDimension > static_cast<int>(z.size()))
+    {
+        z.resize(this->mDimension);
+    }
+
+    return LCPSolverShared<Real>::Solve(q.data(), M.data(), w.data(), z.data(), result);
+}
+
+}
diff --git a/Engine/Plugins/Experimental/GeometryProcessing/Source/GeometryAlgorithms/Private/GTEngine/Mathematics/GteLevenbergMarquardtMinimizer.h b/Engine/Plugins/Experimental/GeometryProcessing/Source/GeometryAlgorithms/Private/GTEngine/Mathematics/GteLevenbergMarquardtMinimizer.h
new file mode 100644
--- /dev/null
+++ b/Engine/Plugins/Experimental/GeometryProcessing/Source/GeometryAlgorithms/Private/GTEngine/Mathematics/GteLevenbergMarquardtMinimizer.h
@@ -0,0 +1,292 @@
+// David Eberly, Geometric Tools, Redmond WA 98052
+// Copyright (c) 1998-2019
+// Distributed under the Boost Software License, Version 1.0.
+// http://www.boost.org/LICENSE_1_0.txt
+// http://www.geometrictools.com/License/Boost/LICENSE_1_0.txt
+// File Version: 3.14.0 (2018/07/16)
+
+#pragma once
+
+#include <GTEngine/Mathematics/GteCholeskyDecomposition.h>
+#include <functional>
+
+// See GteGaussNewtonMinimizer.h for a formulation of the minimization
+// problem and how Levenberg-Marquardt relates to Gauss-Newton.
+
+namespace gte
+{
+    template <typename Real>
+    class LevenbergMarquardtMinimizer
+    {
+    public:
+        // Convenient types for the domain vectors, the range vectors, the
+        // function F and the Jacobian J.
+        typedef GVector<Real> DVector;  // NumPDimensions elements
+        typedef GVector<Real> RVector;  // NumFDImensions elements
+        typedef GMatrix<Real> JMatrix;  // NumFDimensions-by-NumPDimensions elements
+        typedef GMatrix<Real> JTJMatrix;  // NumPDimensions-by-NumPDimensions elements
+        typedef GVector<Real> JTFVector;  // NumPDimensions elements
+        typedef std::function<void(DVector const&, RVector&)> FFunction;
+        typedef std::function<void(DVector const&, JMatrix&)> JFunction;
+        typedef std::function<void(DVector const&, JTJMatrix&, JTFVector&)> JPlusFunction;
+
+        // NOTE: The C++ compiler for Microsoft Visual Studio 12.0.21005.1 REL
+        // (MSVS 2013) appears to have a bug regarding passing std::function
+        // arguments.  The unit-test code for the first LevenbergMarquardtMinimizer
+        // constructor was passed std::function objects of type FFunction and
+        // JFunction.  The compiler claimed that the constructor call is
+        // ambiguous because the JFunction portion of the first constructor's
+        // signature matches the JPlusFunction portion of the second
+        // constructor's signature, which is clearly not correct.  When using
+        // MSVS 2013, you need to explicitly typecast using
+        //   LevenbergMarquardtMinimizer<Real> minimizer(
+        //       numPDimensions, numFDimensions,
+        //       (LevenbergMarquardtMinimizer<Real>::FFunction const&) myFFunction,
+        //       (LevenbergMarquardtMinimizer<Real>::JFunction const&) myJFunction);
+        // The same typecasting is also required for the second constructor.
+
+        // Create the minimizer that computes F(p) and J(p) directly.
+        LevenbergMarquardtMinimizer(int numPDimensions, int numFDimensions,
+            FFunction const& inFFunction, JFunction const& inJFunction)
+            :
+            mNumPDimensions(numPDimensions),
+            mNumFDimensions(numFDimensions),
+            mFFunction(inFFunction),
+            mJFunction(inJFunction),
+            mF(mNumFDimensions),
+            mJ(mNumFDimensions, mNumPDimensions),
+            mJTJ(mNumPDimensions, mNumPDimensions),
+            mNegJTF(mNumPDimensions),
+            mDecomposer(mNumPDimensions),
+            mUseJFunction(true)
+        {
+            LogAssert(mNumPDimensions > 0 && mNumFDimensions > 0, "Invalid dimensions.");
+        }
+
+        // Create the minimizer that computes J^T(p)*J(p) and -J(p)*F(p).
+        LevenbergMarquardtMinimizer(int numPDimensions, int numFDimensions,
+            FFunction const& inFFunction, JPlusFunction const& inJPlusFunction)
+            :
+            mNumPDimensions(numPDimensions),
+            mNumFDimensions(numFDimensions),
+            mFFunction(inFFunction),
+            mJPlusFunction(inJPlusFunction),
+            mF(mNumFDimensions),
+            mJ(mNumFDimensions, mNumPDimensions),
+            mJTJ(mNumPDimensions, mNumPDimensions),
+            mNegJTF(mNumPDimensions),
+            mDecomposer(mNumPDimensions),
+            mUseJFunction(false)
+        {
+            LogAssert(mNumPDimensions > 0 && mNumFDimensions > 0, "Invalid dimensions.");
+        }
+
+        // Disallow copy, assignment and move semantics.
+        LevenbergMarquardtMinimizer(LevenbergMarquardtMinimizer const&) = delete;
+        LevenbergMarquardtMinimizer& operator=(LevenbergMarquardtMinimizer const&) = delete;
+        LevenbergMarquardtMinimizer(LevenbergMarquardtMinimizer&&) = delete;
+        LevenbergMarquardtMinimizer& operator=(LevenbergMarquardtMinimizer&&) = delete;
+
+        inline int GetNumPDimensions() const { return mNumPDimensions; }
+        inline int GetNumFDimensions() const { return mNumFDimensions; }
+
+        // The lambda is positive, the multiplier is positive, and the initial
+        // guess for the p-parameter is p0.  Typical choices are lambda =
+        // 0.001 and multiplier = 10.  TODO: Explain lambda in more detail,
+        // Multiview Geometry mentions lambda = 0.001*average(diagonal(JTJ)),
+        // but let's just expose the factor in front of the average.
+
+        struct Result
+        {
+            DVector minLocation;
+            Real minError;
+            Real minErrorDifference;
+            Real minUpdateLength;
+            size_t numIterations;
+            size_t numAdjustments;
+            bool converged;
+        };
+
+        Result operator()(DVector const& p0, size_t maxIterations,
+            Real updateLengthTolerance, Real errorDifferenceTolerance,
+            Real lambdaFactor, Real lambdaAdjust, size_t maxAdjustments)
+        {
+            Result result;
+            result.minLocation = p0;
+            result.minError = std::numeric_limits<Real>::max();
+            result.minErrorDifference = std::numeric_limits<Real>::max();
+            result.minUpdateLength = (Real)0;
+            result.numIterations = 0;
+            result.numAdjustments = 0;
+            result.converged = false;
+
+            // As a simple precaution, ensure that the lambda inputs are
+            // valid.  If invalid, fall back to Gauss-Newton iteration.
+            if (lambdaFactor <= (Real)0 || lambdaAdjust <= (Real)0)
+            {
+                maxAdjustments = 1;
+                lambdaFactor = (Real)0;
+                lambdaAdjust = (Real)1;
+            }
+
+            // As a simple precaution, ensure the tolerances are nonnegative.
+            updateLengthTolerance = std::max(updateLengthTolerance, (Real)0);
+            errorDifferenceTolerance = std::max(errorDifferenceTolerance, (Real)0);
+
+            // Compute the initial error.
+            mFFunction(p0, mF);
+            result.minError = Dot(mF, mF);
+
+            // Do the Levenberg-Marquart iterations.
+            auto pCurrent = p0;
+            for (result.numIterations = 1; result.numIterations <= maxIterations; ++result.numIterations)
+            {
+                std::pair<bool, bool> status;
+                DVector pNext;
+                for (result.numAdjustments = 0; result.numAdjustments < maxAdjustments; ++result.numAdjustments)
+                {
+                    status = DoIteration(pCurrent, lambdaFactor, updateLengthTolerance,
+                        errorDifferenceTolerance, pNext, result);
+                    if (status.first)
+                    {
+                        // Either the Cholesky decomposition failed or the
+                        // iterates converged within tolerance.  TODO: See the
+                        // note in DoIteration about not failing on Cholesky
+                        // decomposition.
+                        return result;
+                    }
+
+                    if (status.second)
+                    {
+                        // The error has been reduced but we have not yet
+                        // converged within tolerance.
+                        break;
+                    }
+
+                    lambdaFactor *= lambdaAdjust;
+                }
+
+                if (result.numAdjustments < maxAdjustments)
+                {
+                    // The current value of lambda led us to an update that
+                    // reduced the error, but the error is not yet small
+                    // enough to conclude we converged.  Reduce lambda for the
+                    // next outer-loop iteration.
+                    lambdaFactor /= lambdaAdjust;
+                }
+                else
+                {
+                    // All lambdas tried during the inner-loop iteration did
+                    // not lead to a reduced error.  If we do nothing here,
+                    // the next inner-loop iteration will continue to multiply
+                    // lambda, risking eventual floating-point overflow.  To
+                    // avoid this, fall back to a Gauss-Newton iterate.
+                    status = DoIteration(pCurrent, lambdaFactor, updateLengthTolerance,
+                        errorDifferenceTolerance, pNext, result);
+                    if (status.first)
+                    {
+                        // Either the Cholesky decomposition failed or the
+                        // iterates converged within tolerance.  TODO: See the
+                        // note in DoIteration about not failing on Cholesky
+                        // decomposition.
+                        return result;
+                    }
+                }
+
+                pCurrent = pNext;
+            }
+
+            return result;
+        }
+
+    private:
+        void ComputeLinearSystemInputs(DVector const& pCurrent, Real lambda)
+        {
+            if (mUseJFunction)
+            {
+                mJFunction(pCurrent, mJ);
+                mJTJ = MultiplyATB(mJ, mJ);
+                mNegJTF = -(mF * mJ);
+            }
+            else
+            {
+                mJPlusFunction(pCurrent, mJTJ, mNegJTF);
+            }
+
+            Real diagonalSum(0);
+            for (int i = 0; i < mNumPDimensions; ++i)
+            {
+                diagonalSum += mJTJ(i, i);
+            }
+
+            Real diagonalAdjust = lambda * diagonalSum / static_cast<Real>(mNumPDimensions);
+            for (int i = 0; i < mNumPDimensions; ++i)
+            {
+                mJTJ(i, i) += diagonalAdjust;
+            }
+        }
+
+        // The returned 'first' is true when the linear system cannot be
+        // solved (result.converged is false in this case) or when the
+        // error is reduced to within the tolerances specified by the caller
+        // (result.converged is true in this case).  When the 'first' value
+        // is true, the 'second' value is true when the error is reduced or
+        // false when it is not.
+        std::pair<bool, bool> DoIteration(DVector const& pCurrent, Real lambdaFactor,
+            Real updateLengthTolerance, Real errorDifferenceTolerance, DVector& pNext,
+            Result& result)
+        {
+            ComputeLinearSystemInputs(pCurrent, lambdaFactor);
+            if (!mDecomposer.Factor(mJTJ))
+            {
+                // TODO: The matrix mJTJ is positive semi-definite, so the
+                // failure can occur when mJTJ has a zero eigenvalue in
+                // which case mJTJ is not invertible.  Generate an iterate
+                // anyway, perhaps using gradient descent?
+                return std::make_pair(true, false);
+            }
+            mDecomposer.SolveLower(mJTJ, mNegJTF);
+            mDecomposer.SolveUpper(mJTJ, mNegJTF);
+
+            pNext = pCurrent + mNegJTF;
+            mFFunction(pNext, mF);
+            Real error = Dot(mF, mF);
+            if (error < result.minError)
+            {
+                result.minErrorDifference = result.minError - error;
+                result.minUpdateLength = Length(mNegJTF);
+                result.minLocation = pNext;
+                result.minError = error;
+                if (result.minErrorDifference <= errorDifferenceTolerance
+                    || result.minUpdateLength <= updateLengthTolerance)
+                {
+                    result.converged = true;
+                    return std::make_pair(true, true);
+                }
+                else
+                {
+                    return std::make_pair(false, true);
+                }
+            }
+            else
+            {
+                return std::make_pair(false, false);
+            }
+        }
+
+        int mNumPDimensions, mNumFDimensions;
+        FFunction mFFunction;
+        JFunction mJFunction;
+        JPlusFunction mJPlusFunction;
+
+        // Storage for J^T(p)*J(p) and -J^T(p)*F(p) during the iterations.
+        RVector mF;
+        JMatrix mJ;
+        JTJMatrix mJTJ;
+        JTFVector mNegJTF;
+
+        CholeskyDecomposition<Real> mDecomposer;
+
+        bool mUseJFunction;
+    };
+}
diff --git a/Engine/Plugins/Experimental/GeometryProcessing/Source/GeometryAlgorithms/Private/GTEngine/Mathematics/GteLine.h b/Engine/Plugins/Experimental/GeometryProcessing/Source/GeometryAlgorithms/Private/GTEngine/Mathematics/GteLine.h
new file mode 100644
--- /dev/null
+++ b/Engine/Plugins/Experimental/GeometryProcessing/Source/GeometryAlgorithms/Private/GTEngine/Mathematics/GteLine.h
@@ -0,0 +1,112 @@
+// David Eberly, Geometric Tools, Redmond WA 98052
+// Copyright (c) 1998-2019
+// Distributed under the Boost Software License, Version 1.0.
+// http://www.boost.org/LICENSE_1_0.txt
+// http://www.geometrictools.com/License/Boost/LICENSE_1_0.txt
+// File Version: 3.0.0 (2016/06/19)
+
+#pragma once
+
+#include <GTEngine/Mathematics/GteVector.h>
+
+// The line is represented by P+t*D, where P is an origin point, D is a
+// unit-length direction vector, and t is any real number.  The user must
+// ensure that D is unit length.
+
+namespace gte
+{
+
+template <int N, typename Real>
+class Line
+{
+public:
+    // Construction and destruction.  The default constructor sets the origin
+    // to (0,...,0) and the line direction to (1,0,...,0).
+    Line();
+    Line(Vector<N, Real> const& inOrigin, Vector<N, Real> const& inDirection);
+
+    // Public member access.  The direction must be unit length.
+    Vector<N, Real> origin, direction;
+
+public:
+    // Comparisons to support sorted containers.
+    bool operator==(Line const& line) const;
+    bool operator!=(Line const& line) const;
+    bool operator< (Line const& line) const;
+    bool operator<=(Line const& line) const;
+    bool operator> (Line const& line) const;
+    bool operator>=(Line const& line) const;
+};
+
+// Template aliases for convenience.
+template <typename Real>
+using Line2 = Line<2, Real>;
+
+template <typename Real>
+using Line3 = Line<3, Real>;
+
+
+template <int N, typename Real>
+Line<N, Real>::Line()
+{
+    origin.MakeZero();
+    direction.MakeUnit(0);
+}
+
+template <int N, typename Real>
+Line<N, Real>::Line(Vector<N, Real> const& inOrigin,
+    Vector<N, Real> const& inDirection)
+    :
+    origin(inOrigin),
+    direction(inDirection)
+{
+}
+
+template <int N, typename Real>
+bool Line<N, Real>::operator==(Line const& line) const
+{
+    return origin == line.origin && direction == line.direction;
+}
+
+template <int N, typename Real>
+bool Line<N, Real>::operator!=(Line const& line) const
+{
+    return !operator==(line);
+}
+
+template <int N, typename Real>
+bool Line<N, Real>::operator<(Line const& line) const
+{
+    if (origin < line.origin)
+    {
+        return true;
+    }
+
+    if (origin > line.origin)
+    {
+        return false;
+    }
+
+    return direction < line.direction;
+}
+
+template <int N, typename Real>
+bool Line<N, Real>::operator<=(Line const& line) const
+{
+    return operator<(line) || operator==(line);
+}
+
+template <int N, typename Real>
+bool Line<N, Real>::operator>(Line const& line) const
+{
+    return !operator<=(line);
+}
+
+template <int N, typename Real>
+bool Line<N, Real>::operator>=(Line const& line) const
+{
+    return !operator<(line);
+}
+
+
+}
diff --git a/Engine/Plugins/Experimental/GeometryProcessing/Source/GeometryAlgorithms/Private/GTEngine/Mathematics/GteLinearSystem.h b/Engine/Plugins/Experimental/GeometryProcessing/Source/GeometryAlgorithms/Private/GTEngine/Mathematics/GteLinearSystem.h
new file mode 100644
--- /dev/null
+++ b/Engine/Plugins/Experimental/GeometryProcessing/Source/GeometryAlgorithms/Private/GTEngine/Mathematics/GteLinearSystem.h
@@ -0,0 +1,401 @@
+// David Eberly, Geometric Tools, Redmond WA 98052
+// Copyright (c) 1998-2019
+// Distributed under the Boost Software License, Version 1.0.
+// http://www.boost.org/LICENSE_1_0.txt
+// http://www.geometrictools.com/License/Boost/LICENSE_1_0.txt
+// File Version: 3.0.1 (2018/10/05)
+
+#pragma once
+
+#include <GTEngine/LowLevel/GteWrapper.h>
+#include <GTEngine/Mathematics/GteMatrix2x2.h>
+#include <GTEngine/Mathematics/GteMatrix3x3.h>
+#include <GTEngine/Mathematics/GteMatrix4x4.h>
+#include <GTEngine/Mathematics/GteGaussianElimination.h>
+#include <map>
+
+// Solve linear systems of equations where the matrix A is NxN.  The return
+// value of a function is 'true' when A is invertible.  In this case the
+// solution X and the solution is valid.  If the return value is 'false', A
+// is not invertible and X and Y are invalid, so do not use them.  When a
+// matrix is passed as Real*, the storage order is assumed to be the one
+// consistent with your choice of GTE_USE_ROW_MAJOR or GTE_USE_COL_MAJOR.
+//
+// The linear solvers that use the conjugate gradient algorithm are based
+// on the discussion in "Matrix Computations, 2nd edition" by G. H. Golub
+// and Charles F. Van Loan, The Johns Hopkins Press, Baltimore MD, Fourth
+// Printing 1993.
+
+namespace gte
+{
+
+template <typename Real>
+class LinearSystem
+{
+public:
+    // Solve 2x2, 3x3, and 4x4 systems by inverting the matrix directly.  This
+    // avoids the overhead of Gaussian elimination in small dimensions.
+    static bool Solve(Matrix2x2<Real> const& A, Vector2<Real> const& B,
+        Vector2<Real>& X);
+
+    static bool Solve(Matrix3x3<Real> const& A, Vector3<Real> const& B,
+        Vector3<Real>& X);
+
+    static bool Solve(Matrix4x4<Real> const& A, Vector4<Real> const& B,
+        Vector4<Real>& X);
+
+    // Solve A*X = B, where B is Nx1 and the solution X is Nx1.
+    static bool Solve(int N, Real const* A, Real const* B, Real* X);
+
+    // Solve A*X = B, where B is NxM and the solution X is NxM.
+    static bool Solve(int N, int M, Real const* A, Real const* B, Real* X);
+
+    // Solve A*X = B, where A is tridiagonal.  The function expects the
+    // subdiagonal, diagonal, and superdiagonal of A.  The diagonal input
+    // must have N elements.  The subdiagonal and superdiagonal inputs must
+    // have N-1 elements.
+    static bool SolveTridiagonal(int N, Real const* subdiagonal,
+        Real const* diagonal, Real const* superdiagonal, Real const* B,
+        Real* X);
+
+    // Solve A*X = B, where A is tridiagonal.  The function expects the
+    // subdiagonal, diagonal, and superdiagonal of A.  Moreover, the
+    // subdiagonal elements are a constant, the diagonal elements are a
+    // constant, and the superdiagonal elements are a constant.
+    static bool SolveConstantTridiagonal(int N, Real subdiagonal,
+        Real diagonal, Real superdiagonal, Real const* B, Real* X);
+
+    // Solve A*X = B using the conjugate gradient method, where A is
+    // symmetric.  You must specify the maximum number of iterations and a
+    // tolerance for terminating the iterations.  Reasonable choices for
+    // tolerance are 1e-06f for 'float' or 1e-08 for 'double'.
+    static unsigned int SolveSymmetricCG(int N, Real const* A, Real const* B,
+        Real* X, unsigned int maxIterations, Real tolerance);
+
+    // Solve A*X = B using the conjugate gradient method, where A is sparse
+    // and symmetric.  The nonzero entries of the symmetrix matrix A are
+    // stored in a map whose keys are pairs (i,j) and whose values are real
+    // numbers.  The pair (i,j) is the location of the value in the array.
+    // Only one of (i,j) and (j,i) should be stored since A is symmetric.
+    // The column vector B is stored as an array of contiguous values.  You
+    // must specify the maximum number of iterations and a tolerance for
+    // terminating the iterations.  Reasonable choices for tolerance are
+    // 1e-06f for 'float' or 1e-08 for 'double'.
+    typedef std::map<std::array<int, 2>, Real> SparseMatrix;
+    static unsigned int SolveSymmetricCG(int N, SparseMatrix const& A,
+        Real const* B, Real* X, unsigned int maxIterations, Real tolerance);
+
+private:
+    // Support for the conjugate gradient method.
+    static Real Dot(int N, Real const* U, Real const* V);
+    static void Mul(int N, Real const* A, Real const* X, Real* P);
+    static void Mul(int N, SparseMatrix const& A, Real const* X, Real* P);
+    static void UpdateX(int N, Real* X, Real alpha, Real const* P);
+    static void UpdateR(int N, Real* R, Real alpha, Real const* W);
+    static void UpdateP(int N, Real* P, Real beta, Real const* R);
+};
+
+
+template <typename Real>
+bool LinearSystem<Real>::Solve(Matrix2x2<Real> const& A,
+    Vector2<Real> const& B, Vector2<Real>& X)
+{
+    bool invertible;
+    Matrix2x2<Real> invA = Inverse(A, &invertible);
+    if (invertible)
+    {
+        X = invA * B;
+    }
+    else
+    {
+        X = Vector2<Real>::Zero();
+    }
+    return invertible;
+}
+
+template <typename Real>
+bool LinearSystem<Real>::Solve(Matrix3x3<Real> const& A,
+    Vector3<Real> const& B, Vector3<Real>& X)
+{
+    bool invertible;
+    Matrix3x3<Real> invA = Inverse(A, &invertible);
+    if (invertible)
+    {
+        X = invA * B;
+    }
+    else
+    {
+        X = Vector3<Real>::Zero();
+    }
+    return invertible;
+}
+
+template <typename Real>
+bool LinearSystem<Real>::Solve(Matrix4x4<Real> const& A,
+    Vector4<Real> const& B, Vector4<Real>& X)
+{
+    bool invertible;
+    Matrix4x4<Real> invA = Inverse(A, &invertible);
+    if (invertible)
+    {
+        X = invA * B;
+    }
+    else
+    {
+        X = Vector4<Real>::Zero();
+    }
+    return invertible;
+}
+
+template <typename Real>
+bool LinearSystem<Real>::Solve(int N, Real const* A, Real const* B, Real* X)
+{
+    Real determinant;
+    return GaussianElimination<Real>()(N, A, nullptr, determinant, B, X,
+        nullptr, 0, nullptr);
+}
+
+template <typename Real>
+bool LinearSystem<Real>::Solve(int N, int M, Real const* A, Real const* B,
+    Real* X)
+{
+    Real determinant;
+    return GaussianElimination<Real>()(N, A, nullptr, determinant, nullptr,
+        nullptr, B, M, X);
+}
+
+template <typename Real>
+bool LinearSystem<Real>::SolveTridiagonal(int N, Real const* subdiagonal,
+    Real const* diagonal, Real const* superdiagonal, Real const* B,
+    Real* X)
+{
+    if (diagonal[0] == (Real)0)
+    {
+        return false;
+    }
+
+    std::vector<Real> tmp(N - 1);
+    Real expr = diagonal[0];
+    Real invExpr = ((Real)1) / expr;
+    X[0] = B[0] * invExpr;
+
+    int i0, i1;
+    for (i0 = 0, i1 = 1; i1 < N; ++i0, ++i1)
+    {
+        tmp[i0] = superdiagonal[i0] * invExpr;
+        expr = diagonal[i1] - subdiagonal[i0] * tmp[i0];
+        if (expr == (Real)0)
+        {
+            return false;
+        }
+        invExpr = ((Real)1) / expr;
+        X[i1] = (B[i1] - subdiagonal[i0] * X[i0]) * invExpr;
+    }
+
+    for (i0 = N - 1, i1 = N - 2; i1 >= 0; --i0, --i1)
+    {
+        X[i1] -= tmp[i1] * X[i0];
+    }
+    return true;
+}
+
+template <typename Real>
+bool LinearSystem<Real>::SolveConstantTridiagonal(int N, Real subdiagonal,
+    Real diagonal, Real superdiagonal, Real const* B, Real* X)
+{
+    if (diagonal == (Real)0)
+    {
+        return false;
+    }
+
+    std::vector<Real> tmp(N - 1);
+    Real expr = diagonal;
+    Real invExpr = ((Real)1) / expr;
+    X[0] = B[0] * invExpr;
+
+    int i0, i1;
+    for (i0 = 0, i1 = 1; i1 < N; ++i0, ++i1)
+    {
+        tmp[i0] = superdiagonal * invExpr;
+        expr = diagonal - subdiagonal * tmp[i0];
+        if (expr == (Real)0)
+        {
+            return false;
+        }
+        invExpr = ((Real)1) / expr;
+        X[i1] = (B[i1] - subdiagonal * X[i0]) * invExpr;
+    }
+
+    for (i0 = N - 1, i1 = N - 2; i1 >= 0; --i0, --i1)
+    {
+        X[i1] -= tmp[i1] * X[i0];
+    }
+    return true;
+}
+
+template <typename Real>
+unsigned int LinearSystem<Real>::SolveSymmetricCG(int N, Real const* A,
+    Real const* B, Real* X, unsigned int maxIterations, Real tolerance)
+{
+    // The first iteration.
+    std::vector<Real> tmpR(N), tmpP(N), tmpW(N);
+    Real* R = &tmpR[0];
+    Real* P = &tmpP[0];
+    Real* W = &tmpW[0];
+    size_t numBytes = N * sizeof(Real);
+    memset(X, 0, numBytes);
+    Memcpy(R, B, numBytes);
+    Real rho0 = Dot(N, R, R);
+    Memcpy(P, R, numBytes);
+    Mul(N, A, P, W);
+    Real alpha = rho0 / Dot(N, P, W);
+    UpdateX(N, X, alpha, P);
+    UpdateR(N, R, alpha, W);
+    Real rho1 = Dot(N, R, R);
+
+    // The remaining iterations.
+    unsigned int iteration;
+    for (iteration = 1; iteration <= maxIterations; ++iteration)
+    {
+        Real root0 = std::sqrt(rho1);
+        Real norm = Dot(N, B, B);
+        Real root1 = std::sqrt(norm);
+        if (root0 <= tolerance*root1)
+        {
+            break;
+        }
+
+        Real beta = rho1 / rho0;
+        UpdateP(N, P, beta, R);
+        Mul(N, A, P, W);
+        alpha = rho1 / Dot(N, P, W);
+        UpdateX(N, X, alpha, P);
+        UpdateR(N, R, alpha, W);
+        rho0 = rho1;
+        rho1 = Dot(N, R, R);
+    }
+    return iteration;
+}
+
+template <typename Real>
+unsigned int LinearSystem<Real>::SolveSymmetricCG(int N,
+    SparseMatrix const& A, Real const* B, Real* X, unsigned int maxIterations,
+    Real tolerance)
+{
+    // The first iteration.
+    std::vector<Real> tmpR(N), tmpP(N), tmpW(N);
+    Real* R = &tmpR[0];
+    Real* P = &tmpP[0];
+    Real* W = &tmpW[0];
+    size_t numBytes = N * sizeof(Real);
+    memset(X, 0, numBytes);
+    Memcpy(R, B, numBytes);
+    Real rho0 = Dot(N, R, R);
+    Memcpy(P, R, numBytes);
+    Mul(N, A, P, W);
+    Real alpha = rho0 / Dot(N, P, W);
+    UpdateX(N, X, alpha, P);
+    UpdateR(N, R, alpha, W);
+    Real rho1 = Dot(N, R, R);
+
+    // The remaining iterations.
+    unsigned int iteration;
+    for (iteration = 1; iteration <= maxIterations; ++iteration)
+    {
+        Real root0 = std::sqrt(rho1);
+        Real norm = Dot(N, B, B);
+        Real root1 = std::sqrt(norm);
+        if (root0 <= tolerance*root1)
+        {
+            break;
+        }
+
+        Real beta = rho1 / rho0;
+        UpdateP(N, P, beta, R);
+        Mul(N, A, P, W);
+        alpha = rho1 / Dot(N, P, W);
+        UpdateX(N, X, alpha, P);
+        UpdateR(N, R, alpha, W);
+        rho0 = rho1;
+        rho1 = Dot(N, R, R);
+    }
+    return iteration;
+}
+
+template <typename Real>
+Real LinearSystem<Real>::Dot(int N, Real const* U, Real const* V)
+{
+    Real dot = (Real)0;
+    for (int i = 0; i < N; ++i)
+    {
+        dot += U[i] * V[i];
+    }
+    return dot;
+}
+
+template <typename Real>
+void LinearSystem<Real>::Mul(int N, Real const* A, Real const* X, Real* P)
+{
+#if defined(GTE_USE_ROW_MAJOR)
+    LexicoArray2<true, Real> matA(N, N, const_cast<Real*>(A));
+#else
+    LexicoArray2<false, Real> matA(N, N, const_cast<Real*>(A));
+#endif
+
+    memset(P, 0, N * sizeof(Real));
+    for (int row = 0; row < N; ++row)
+    {
+        for (int col = 0; col < N; ++col)
+        {
+            P[row] += matA(row, col) * X[col];
+        }
+    }
+}
+
+template <typename Real>
+void LinearSystem<Real>::Mul(int N, SparseMatrix const& A, Real const* X,
+    Real* P)
+{
+    memset(P, 0, N * sizeof(Real));
+    for (auto const& element : A)
+    {
+        int i = element.first[0];
+        int j = element.first[1];
+        Real value = element.second;
+        P[i] += value * X[j];
+        if (i != j)
+        {
+            P[j] += value * X[i];
+        }
+    }
+}
+
+template <typename Real>
+void LinearSystem<Real>::UpdateX(int N, Real* X, Real alpha, Real const* P)
+{
+    for (int i = 0; i < N; ++i)
+    {
+        X[i] += alpha*P[i];
+    }
+}
+
+template <typename Real>
+void LinearSystem<Real>::UpdateR(int N, Real* R, Real alpha, Real const* W)
+{
+    for (int i = 0; i < N; ++i)
+    {
+        R[i] -= alpha*W[i];
+    }
+}
+
+template <typename Real>
+void LinearSystem<Real>::UpdateP(int N, Real* P, Real beta, Real const* R)
+{
+    for (int i = 0; i < N; ++i)
+    {
+        P[i] = R[i] + beta*P[i];
+    }
+}
+
+
+}
diff --git a/Engine/Plugins/Experimental/GeometryProcessing/Source/GeometryAlgorithms/Private/GTEngine/Mathematics/GteLog2Estimate.h b/Engine/Plugins/Experimental/GeometryProcessing/Source/GeometryAlgorithms/Private/GTEngine/Mathematics/GteLog2Estimate.h
new file mode 100644
--- /dev/null
+++ b/Engine/Plugins/Experimental/GeometryProcessing/Source/GeometryAlgorithms/Private/GTEngine/Mathematics/GteLog2Estimate.h
@@ -0,0 +1,174 @@
+// David Eberly, Geometric Tools, Redmond WA 98052
+// Copyright (c) 1998-2019
+// Distributed under the Boost Software License, Version 1.0.
+// http://www.boost.org/LICENSE_1_0.txt
+// http://www.geometrictools.com/License/Boost/LICENSE_1_0.txt
+// File Version: 3.0.1 (2018/10/05)
+
+#pragma once
+
+#include <GTEngine/Mathematics/GteMath.h>
+
+// Minimax polynomial approximations to log2(x).  The polynomial p(x) of
+// degree D minimizes the quantity maximum{|log2(x) - p(x)| : x in [1,2]}
+// over all polynomials of degree D.
+
+namespace gte
+{
+
+template <typename Real>
+class Log2Estimate
+{
+public:
+    // The input constraint is x in [1,2].  For example,
+    //   float x; // in [1,2]
+    //   float result = Log2Estimate<float>::Degree<3>(x);
+    template <int D>
+    inline static Real Degree(Real x);
+
+    // The input constraint is x > 0.  Range reduction is used to generate a
+    // value y in (0,1], call Degree(y), and add the exponent for the power
+    // of two in the binary scientific representation of x.  For example,
+    //   float x;  // x > 0
+    //   float result = Log2Estimate<float>::DegreeRR<3>(x);
+    template <int D>
+    inline static Real DegreeRR(Real x);
+
+private:
+    // Metaprogramming and private implementation to allow specialization of
+    // a template member function.
+    template <int D> struct degree {};
+    inline static Real Evaluate(degree<1>, Real t);
+    inline static Real Evaluate(degree<2>, Real t);
+    inline static Real Evaluate(degree<3>, Real t);
+    inline static Real Evaluate(degree<4>, Real t);
+    inline static Real Evaluate(degree<5>, Real t);
+    inline static Real Evaluate(degree<6>, Real t);
+    inline static Real Evaluate(degree<7>, Real t);
+    inline static Real Evaluate(degree<8>, Real t);
+};
+
+
+template <typename Real>
+template <int D>
+inline Real Log2Estimate<Real>::Degree(Real x)
+{
+    Real t = x - (Real)1;  // t in (0,1]
+    return Evaluate(degree<D>(), t);
+}
+
+template <typename Real>
+template <int D>
+inline Real Log2Estimate<Real>::DegreeRR(Real x)
+{
+    int p;
+    Real y = std::frexp(x, &p);  // y in [1/2,1)
+    y = ((Real)2)*y;  // y in [1,2)
+    --p;
+    Real poly = Degree<D>(y);
+    Real result = poly + (Real)p;
+    return result;
+}
+
+template <typename Real>
+inline Real Log2Estimate<Real>::Evaluate(degree<1>, Real t)
+{
+    Real poly;
+    poly = (Real)GTE_C_LOG2_DEG1_C1;
+    poly = poly * t;
+    return poly;
+}
+
+template <typename Real>
+inline Real Log2Estimate<Real>::Evaluate(degree<2>, Real t)
+{
+    Real poly;
+    poly = (Real)GTE_C_LOG2_DEG2_C2;
+    poly = (Real)GTE_C_LOG2_DEG2_C1 + poly * t;
+    poly = poly * t;
+    return poly;
+}
+
+template <typename Real>
+inline Real Log2Estimate<Real>::Evaluate(degree<3>, Real t)
+{
+    Real poly;
+    poly = (Real)GTE_C_LOG2_DEG3_C3;
+    poly = (Real)GTE_C_LOG2_DEG3_C2 + poly * t;
+    poly = (Real)GTE_C_LOG2_DEG3_C1 + poly * t;
+    poly = poly * t;
+    return poly;
+}
+
+template <typename Real>
+inline Real Log2Estimate<Real>::Evaluate(degree<4>, Real t)
+{
+    Real poly;
+    poly = (Real)GTE_C_LOG2_DEG4_C4;
+    poly = (Real)GTE_C_LOG2_DEG4_C3 + poly * t;
+    poly = (Real)GTE_C_LOG2_DEG4_C2 + poly * t;
+    poly = (Real)GTE_C_LOG2_DEG4_C1 + poly * t;
+    poly = poly * t;
+    return poly;
+}
+
+template <typename Real>
+inline Real Log2Estimate<Real>::Evaluate(degree<5>, Real t)
+{
+    Real poly;
+    poly = (Real)GTE_C_LOG2_DEG5_C5;
+    poly = (Real)GTE_C_LOG2_DEG5_C4 + poly * t;
+    poly = (Real)GTE_C_LOG2_DEG5_C3 + poly * t;
+    poly = (Real)GTE_C_LOG2_DEG5_C2 + poly * t;
+    poly = (Real)GTE_C_LOG2_DEG5_C1 + poly * t;
+    poly = poly * t;
+    return poly;
+}
+
+template <typename Real>
+inline Real Log2Estimate<Real>::Evaluate(degree<6>, Real t)
+{
+    Real poly;
+    poly = (Real)GTE_C_LOG2_DEG6_C6;
+    poly = (Real)GTE_C_LOG2_DEG6_C5 + poly * t;
+    poly = (Real)GTE_C_LOG2_DEG6_C4 + poly * t;
+    poly = (Real)GTE_C_LOG2_DEG6_C3 + poly * t;
+    poly = (Real)GTE_C_LOG2_DEG6_C2 + poly * t;
+    poly = (Real)GTE_C_LOG2_DEG6_C1 + poly * t;
+    poly = poly * t;
+    return poly;
+}
+
+template <typename Real>
+inline Real Log2Estimate<Real>::Evaluate(degree<7>, Real t)
+{
+    Real poly;
+    poly = (Real)GTE_C_LOG2_DEG7_C7;
+    poly = (Real)GTE_C_LOG2_DEG7_C6 + poly * t;
+    poly = (Real)GTE_C_LOG2_DEG7_C5 + poly * t;
+    poly = (Real)GTE_C_LOG2_DEG7_C4 + poly * t;
+    poly = (Real)GTE_C_LOG2_DEG7_C3 + poly * t;
+    poly = (Real)GTE_C_LOG2_DEG7_C2 + poly * t;
+    poly = (Real)GTE_C_LOG2_DEG7_C1 + poly * t;
+    poly = poly * t;
+    return poly;
+}
+
+template <typename Real>
+inline Real Log2Estimate<Real>::Evaluate(degree<8>, Real t)
+{
+    Real poly;
+    poly = (Real)GTE_C_LOG2_DEG8_C8;
+    poly = (Real)GTE_C_LOG2_DEG8_C7 + poly * t;
+    poly = (Real)GTE_C_LOG2_DEG8_C6 + poly * t;
+    poly = (Real)GTE_C_LOG2_DEG8_C5 + poly * t;
+    poly = (Real)GTE_C_LOG2_DEG8_C4 + poly * t;
+    poly = (Real)GTE_C_LOG2_DEG8_C3 + poly * t;
+    poly = (Real)GTE_C_LOG2_DEG8_C2 + poly * t;
+    poly = (Real)GTE_C_LOG2_DEG8_C1 + poly * t;
+    poly = poly * t;
+    return poly;
+}
+
+
+}
diff --git a/Engine/Plugins/Experimental/GeometryProcessing/Source/GeometryAlgorithms/Private/GTEngine/Mathematics/GteLogEstimate.h b/Engine/Plugins/Experimental/GeometryProcessing/Source/GeometryAlgorithms/Private/GTEngine/Mathematics/GteLogEstimate.h
new file mode 100644
--- /dev/null
+++ b/Engine/Plugins/Experimental/GeometryProcessing/Source/GeometryAlgorithms/Private/GTEngine/Mathematics/GteLogEstimate.h
@@ -0,0 +1,55 @@
+// David Eberly, Geometric Tools, Redmond WA 98052
+// Copyright (c) 1998-2019
+// Distributed under the Boost Software License, Version 1.0.
+// http://www.boost.org/LICENSE_1_0.txt
+// http://www.geometrictools.com/License/Boost/LICENSE_1_0.txt
+// File Version: 3.0.0 (2016/06/19)
+
+#pragma once
+
+#include <GTEngine/Mathematics/GteLog2Estimate.h>
+
+// Minimax polynomial approximations to log2(x).  The polynomial p(x) of
+// degree D minimizes the quantity maximum{|log2(x) - p(x)| : x in [1,2]}
+// over all polynomials of degree D.  The natural logarithm is computed
+// using log(x) = log2(x)/log2(e) = log2(x)*log(2).
+
+namespace gte
+{
+
+template <typename Real>
+class LogEstimate
+{
+public:
+    // The input constraint is x in [1,2].  For example,
+    //   float x; // in [1,2]
+    //   float result = LogEstimate<float>::Degree<3>(x);
+    template <int D>
+    inline static Real Degree(Real x);
+
+    // The input constraint is x > 0.  Range reduction is used to generate a
+    // value y in (0,1], call Degree(y), and add the exponent for the power
+    // of two in the binary scientific representation of x.  For example,
+    //   float x;  // x > 0
+    //   float result = LogEstimate<float>::DegreeRR<3>(x);
+    template <int D>
+    inline static Real DegreeRR(Real x);
+};
+
+
+template <typename Real>
+template <int D>
+inline Real LogEstimate<Real>::Degree(Real x)
+{
+    return Log2Estimate<Real>::Degree<D>(x) * (Real)GTE_C_INV_LN_2;
+}
+
+template <typename Real>
+template <int D>
+inline Real LogEstimate<Real>::DegreeRR(Real x)
+{
+    return Log2Estimate<Real>::DegreeRR<D>(x) * (Real)GTE_C_INV_LN_2;
+}
+
+
+}
diff --git a/Engine/Plugins/Experimental/GeometryProcessing/Source/GeometryAlgorithms/Private/GTEngine/Mathematics/GteMath.h b/Engine/Plugins/Experimental/GeometryProcessing/Source/GeometryAlgorithms/Private/GTEngine/Mathematics/GteMath.h
new file mode 100644
--- /dev/null
+++ b/Engine/Plugins/Experimental/GeometryProcessing/Source/GeometryAlgorithms/Private/GTEngine/Mathematics/GteMath.h
@@ -0,0 +1,654 @@
+// David Eberly, Geometric Tools, Redmond WA 98052
+// Copyright (c) 1998-2019
+// Distributed under the Boost Software License, Version 1.0.
+// http://www.boost.org/LICENSE_1_0.txt
+// http://www.geometrictools.com/License/Boost/LICENSE_1_0.txt
+// File Version: 3.17.1 (2019/01/17)
+
+#pragma once
+
+// This file extends the <cmath> support to include convenient constants and
+// functions.  The shared constants for CPU, Intel SSE and GPU lead to
+// correctly rounded approximations of the constants when using 'float' or
+// 'double'.  The file also includes a type trait, is_arbitrary_precision,
+// to support selecting between floating-point arithmetic (float, double,
+//long double) or arbitrary-precision arithmetic (BSNumber<T>, BSRational<T>)
+// in an implementation using std::enable_if.  There is also a type trait,
+// has_division_operator, to support selecting between numeric types that
+// have a division operator (BSRational<T>) and those that do not have a
+// division operator (BSNumber<T>).
+
+#include <GTEngine/GTEngineDEF.h>
+#include <cmath>
+#include <limits>
+#include <type_traits>
+
+// Maximum number of iterations for bisection before a subinterval
+// degenerates to a single point. TODO: Verify these.  I used the formula:
+// 3 + std::numeric_limits<T>::digits - std::numeric_limits<T>::min_exponent.
+//   IEEEBinary16:  digits = 11, min_exponent = -13
+//   float:         digits = 27, min_exponent = -125
+//   double:        digits = 53, min_exponent = -1021
+// BSNumber and BSRational use std::numeric_limits<unsigned int>::max(),
+// but maybe these should be set to a large number or be user configurable.
+// The MAX_BISECTIONS_GENERIC is an arbitrary choice for now and is used
+// in template code where Real is the template parameter.
+#define GTE_C_MAX_BISECTIONS_FLOAT16    27u
+#define GTE_C_MAX_BISECTIONS_FLOAT32    155u
+#define GTE_C_MAX_BISECTIONS_FLOAT64    1077u
+#define GTE_C_MAX_BISECTIONS_BSNUMBER   0xFFFFFFFFu
+#define GTE_C_MAX_BISECTIONS_BSRATIONAL 0xFFFFFFFFu
+#define GTE_C_MAX_BISECTIONS_GENERIC    2048u
+
+// Constants involving pi.
+#define GTE_C_PI 3.1415926535897931
+#define GTE_C_HALF_PI 1.5707963267948966
+#define GTE_C_QUARTER_PI 0.7853981633974483
+#define GTE_C_TWO_PI 6.2831853071795862
+#define GTE_C_INV_PI 0.3183098861837907
+#define GTE_C_INV_TWO_PI 0.1591549430918953
+#define GTE_C_INV_HALF_PI 0.6366197723675813
+
+// Conversions between degrees and radians.
+#define GTE_C_DEG_TO_RAD 0.0174532925199433
+#define GTE_C_RAD_TO_DEG 57.295779513082321
+
+// Common constants.
+#define GTE_C_SQRT_2 1.4142135623730951
+#define GTE_C_INV_SQRT_2 0.7071067811865475
+#define GTE_C_LN_2 0.6931471805599453
+#define GTE_C_INV_LN_2 1.4426950408889634
+#define GTE_C_LN_10 2.3025850929940459
+#define GTE_C_INV_LN_10 0.43429448190325176
+
+// Constants for minimax polynomial approximations to sqrt(x).
+// The algorithm minimizes the maximum absolute error on [1,2].
+#define GTE_C_SQRT_DEG1_C0 +1.0
+#define GTE_C_SQRT_DEG1_C1 +4.1421356237309505e-01
+#define GTE_C_SQRT_DEG1_MAX_ERROR 1.7766952966368793e-2
+
+#define GTE_C_SQRT_DEG2_C0 +1.0
+#define GTE_C_SQRT_DEG2_C1 +4.8563183076125260e-01
+#define GTE_C_SQRT_DEG2_C2 -7.1418268388157458e-02
+#define GTE_C_SQRT_DEG2_MAX_ERROR 1.1795695163108744e-3
+
+#define GTE_C_SQRT_DEG3_C0 +1.0
+#define GTE_C_SQRT_DEG3_C1 +4.9750045320242231e-01
+#define GTE_C_SQRT_DEG3_C2 -1.0787308044477850e-01
+#define GTE_C_SQRT_DEG3_C3 +2.4586189615451115e-02
+#define GTE_C_SQRT_DEG3_MAX_ERROR 1.1309620116468910e-4
+
+#define GTE_C_SQRT_DEG4_C0 +1.0
+#define GTE_C_SQRT_DEG4_C1 +4.9955939832918816e-01
+#define GTE_C_SQRT_DEG4_C2 -1.2024066151943025e-01
+#define GTE_C_SQRT_DEG4_C3 +4.5461507257698486e-02
+#define GTE_C_SQRT_DEG4_C4 -1.0566681694362146e-02
+#define GTE_C_SQRT_DEG4_MAX_ERROR 1.2741170151556180e-5
+
+#define GTE_C_SQRT_DEG5_C0 +1.0
+#define GTE_C_SQRT_DEG5_C1 +4.9992197660031912e-01
+#define GTE_C_SQRT_DEG5_C2 -1.2378506719245053e-01
+#define GTE_C_SQRT_DEG5_C3 +5.6122776972699739e-02
+#define GTE_C_SQRT_DEG5_C4 -2.3128836281145482e-02
+#define GTE_C_SQRT_DEG5_C5 +5.0827122737047148e-03
+#define GTE_C_SQRT_DEG5_MAX_ERROR 1.5725568940708201e-6
+
+#define GTE_C_SQRT_DEG6_C0 +1.0
+#define GTE_C_SQRT_DEG6_C1 +4.9998616695784914e-01
+#define GTE_C_SQRT_DEG6_C2 -1.2470733323278438e-01
+#define GTE_C_SQRT_DEG6_C3 +6.0388587356982271e-02
+#define GTE_C_SQRT_DEG6_C4 -3.1692053551807930e-02
+#define GTE_C_SQRT_DEG6_C5 +1.2856590305148075e-02
+#define GTE_C_SQRT_DEG6_C6 -2.6183954624343642e-03
+#define GTE_C_SQRT_DEG6_MAX_ERROR 2.0584155535630089e-7
+
+#define GTE_C_SQRT_DEG7_C0 +1.0
+#define GTE_C_SQRT_DEG7_C1 +4.9999754817809228e-01
+#define GTE_C_SQRT_DEG7_C2 -1.2493243476353655e-01
+#define GTE_C_SQRT_DEG7_C3 +6.1859954146370910e-02
+#define GTE_C_SQRT_DEG7_C4 -3.6091595023208356e-02
+#define GTE_C_SQRT_DEG7_C5 +1.9483946523450868e-02
+#define GTE_C_SQRT_DEG7_C6 -7.5166134568007692e-03
+#define GTE_C_SQRT_DEG7_C7 +1.4127567687864939e-03
+#define GTE_C_SQRT_DEG7_MAX_ERROR 2.8072302919734948e-8
+
+#define GTE_C_SQRT_DEG8_C0 +1.0
+#define GTE_C_SQRT_DEG8_C1 +4.9999956583056759e-01
+#define GTE_C_SQRT_DEG8_C2 -1.2498490369914350e-01
+#define GTE_C_SQRT_DEG8_C3 +6.2318494667579216e-02
+#define GTE_C_SQRT_DEG8_C4 -3.7982961896432244e-02
+#define GTE_C_SQRT_DEG8_C5 +2.3642612312869460e-02
+#define GTE_C_SQRT_DEG8_C6 -1.2529377587270574e-02
+#define GTE_C_SQRT_DEG8_C7 +4.5382426960713929e-03
+#define GTE_C_SQRT_DEG8_C8 -7.8810995273670414e-04
+#define GTE_C_SQRT_DEG8_MAX_ERROR 3.9460605685825989e-9
+
+// Constants for minimax polynomial approximations to 1/sqrt(x).
+// The algorithm minimizes the maximum absolute error on [1,2].
+#define GTE_C_INVSQRT_DEG1_C0 +1.0
+#define GTE_C_INVSQRT_DEG1_C1 -2.9289321881345254e-01
+#define GTE_C_INVSQRT_DEG1_MAX_ERROR 3.7814314552701983e-2
+
+#define GTE_C_INVSQRT_DEG2_C0 +1.0
+#define GTE_C_INVSQRT_DEG2_C1 -4.4539812104566801e-01
+#define GTE_C_INVSQRT_DEG2_C2 +1.5250490223221547e-01
+#define GTE_C_INVSQRT_DEG2_MAX_ERROR 4.1953446330581234e-3
+
+#define GTE_C_INVSQRT_DEG3_C0 +1.0
+#define GTE_C_INVSQRT_DEG3_C1 -4.8703230993068791e-01
+#define GTE_C_INVSQRT_DEG3_C2 +2.8163710486669835e-01
+#define GTE_C_INVSQRT_DEG3_C3 -8.7498013749463421e-02
+#define GTE_C_INVSQRT_DEG3_MAX_ERROR 5.6307702007266786e-4
+
+#define GTE_C_INVSQRT_DEG4_C0 +1.0
+#define GTE_C_INVSQRT_DEG4_C1 -4.9710061558048779e-01
+#define GTE_C_INVSQRT_DEG4_C2 +3.4266247597676802e-01
+#define GTE_C_INVSQRT_DEG4_C3 -1.9106356536293490e-01
+#define GTE_C_INVSQRT_DEG4_C4 +5.2608486153198797e-02
+#define GTE_C_INVSQRT_DEG4_MAX_ERROR 8.1513919987605266e-5
+
+#define GTE_C_INVSQRT_DEG5_C0 +1.0
+#define GTE_C_INVSQRT_DEG5_C1 -4.9937760586004143e-01
+#define GTE_C_INVSQRT_DEG5_C2 +3.6508741295133973e-01
+#define GTE_C_INVSQRT_DEG5_C3 -2.5884890281853501e-01
+#define GTE_C_INVSQRT_DEG5_C4 +1.3275782221320753e-01
+#define GTE_C_INVSQRT_DEG5_C5 -3.2511945299404488e-02
+#define GTE_C_INVSQRT_DEG5_MAX_ERROR 1.2289367475583346e-5
+
+#define GTE_C_INVSQRT_DEG6_C0 +1.0
+#define GTE_C_INVSQRT_DEG6_C1 -4.9987029229547453e-01
+#define GTE_C_INVSQRT_DEG6_C2 +3.7220923604495226e-01
+#define GTE_C_INVSQRT_DEG6_C3 -2.9193067713256937e-01
+#define GTE_C_INVSQRT_DEG6_C4 +1.9937605991094642e-01
+#define GTE_C_INVSQRT_DEG6_C5 -9.3135712130901993e-02
+#define GTE_C_INVSQRT_DEG6_C6 +2.0458166789566690e-02
+#define GTE_C_INVSQRT_DEG6_MAX_ERROR 1.9001451223750465e-6
+
+#define GTE_C_INVSQRT_DEG7_C0 +1.0
+#define GTE_C_INVSQRT_DEG7_C1 -4.9997357250704977e-01
+#define GTE_C_INVSQRT_DEG7_C2 +3.7426216884998809e-01
+#define GTE_C_INVSQRT_DEG7_C3 -3.0539882498248971e-01
+#define GTE_C_INVSQRT_DEG7_C4 +2.3976005607005391e-01
+#define GTE_C_INVSQRT_DEG7_C5 -1.5410326351684489e-01
+#define GTE_C_INVSQRT_DEG7_C6 +6.5598809723041995e-02
+#define GTE_C_INVSQRT_DEG7_C7 -1.3038592450470787e-02
+#define GTE_C_INVSQRT_DEG7_MAX_ERROR 2.9887724993168940e-7
+
+#define GTE_C_INVSQRT_DEG8_C0 +1.0
+#define GTE_C_INVSQRT_DEG8_C1 -4.9999471066120371e-01
+#define GTE_C_INVSQRT_DEG8_C2 +3.7481415745794067e-01
+#define GTE_C_INVSQRT_DEG8_C3 -3.1023804387422160e-01
+#define GTE_C_INVSQRT_DEG8_C4 +2.5977002682930106e-01
+#define GTE_C_INVSQRT_DEG8_C5 -1.9818790717727097e-01
+#define GTE_C_INVSQRT_DEG8_C6 +1.1882414252613671e-01
+#define GTE_C_INVSQRT_DEG8_C7 -4.6270038088550791e-02
+#define GTE_C_INVSQRT_DEG8_C8 +8.3891541755747312e-03
+#define GTE_C_INVSQRT_DEG8_MAX_ERROR 4.7596926146947771e-8
+
+// Constants for minimax polynomial approximations to sin(x).
+// The algorithm minimizes the maximum absolute error on [-pi/2,pi/2].
+#define GTE_C_SIN_DEG3_C0 +1.0
+#define GTE_C_SIN_DEG3_C1 -1.4727245910375519e-01
+#define GTE_C_SIN_DEG3_MAX_ERROR 1.3481903639145865e-2
+
+#define GTE_C_SIN_DEG5_C0 +1.0
+#define GTE_C_SIN_DEG5_C1 -1.6600599923812209e-01
+#define GTE_C_SIN_DEG5_C2 +7.5924178409012000e-03
+#define GTE_C_SIN_DEG5_MAX_ERROR 1.4001209384639779e-4
+
+#define GTE_C_SIN_DEG7_C0 +1.0
+#define GTE_C_SIN_DEG7_C1 -1.6665578084732124e-01
+#define GTE_C_SIN_DEG7_C2 +8.3109378830028557e-03
+#define GTE_C_SIN_DEG7_C3 -1.8447486103462252e-04
+#define GTE_C_SIN_DEG7_MAX_ERROR 1.0205878936686563e-6
+
+#define GTE_C_SIN_DEG9_C0 +1.0
+#define GTE_C_SIN_DEG9_C1 -1.6666656235308897e-01
+#define GTE_C_SIN_DEG9_C2 +8.3329962509886002e-03
+#define GTE_C_SIN_DEG9_C3 -1.9805100675274190e-04
+#define GTE_C_SIN_DEG9_C4 +2.5967200279475300e-06
+#define GTE_C_SIN_DEG9_MAX_ERROR 5.2010746265374053e-9
+
+#define GTE_C_SIN_DEG11_C0 +1.0
+#define GTE_C_SIN_DEG11_C1 -1.6666666601721269e-01
+#define GTE_C_SIN_DEG11_C2 +8.3333303183525942e-03
+#define GTE_C_SIN_DEG11_C3 -1.9840782426250314e-04
+#define GTE_C_SIN_DEG11_C4 +2.7521557770526783e-06
+#define GTE_C_SIN_DEG11_C5 -2.3828544692960918e-08
+#define GTE_C_SIN_DEG11_MAX_ERROR 1.9295870457014530e-11
+
+// Constants for minimax polynomial approximations to cos(x).
+// The algorithm minimizes the maximum absolute error on [-pi/2,pi/2].
+#define GTE_C_COS_DEG2_C0 +1.0
+#define GTE_C_COS_DEG2_C1 -4.0528473456935105e-01
+#define GTE_C_COS_DEG2_MAX_ERROR 5.4870946878404048e-2
+
+#define GTE_C_COS_DEG4_C0 +1.0
+#define GTE_C_COS_DEG4_C1 -4.9607181958647262e-01
+#define GTE_C_COS_DEG4_C2 +3.6794619653489236e-02
+#define GTE_C_COS_DEG4_MAX_ERROR 9.1879932449712154e-4
+
+#define GTE_C_COS_DEG6_C0 +1.0
+#define GTE_C_COS_DEG6_C1 -4.9992746217057404e-01
+#define GTE_C_COS_DEG6_C2 +4.1493920348353308e-02
+#define GTE_C_COS_DEG6_C3 -1.2712435011987822e-03
+#define GTE_C_COS_DEG6_MAX_ERROR 9.2028470133065365e-6
+
+#define GTE_C_COS_DEG8_C0 +1.0
+#define GTE_C_COS_DEG8_C1 -4.9999925121358291e-01
+#define GTE_C_COS_DEG8_C2 +4.1663780117805693e-02
+#define GTE_C_COS_DEG8_C3 -1.3854239405310942e-03
+#define GTE_C_COS_DEG8_C4 +2.3154171575501259e-05
+#define GTE_C_COS_DEG8_MAX_ERROR 5.9804533020235695e-8
+
+#define GTE_C_COS_DEG10_C0 +1.0
+#define GTE_C_COS_DEG10_C1 -4.9999999508695869e-01
+#define GTE_C_COS_DEG10_C2 +4.1666638865338612e-02
+#define GTE_C_COS_DEG10_C3 -1.3888377661039897e-03
+#define GTE_C_COS_DEG10_C4 +2.4760495088926859e-05
+#define GTE_C_COS_DEG10_C5 -2.6051615464872668e-07
+#define GTE_C_COS_DEG10_MAX_ERROR 2.7006769043325107e-10
+
+// Constants for minimax polynomial approximations to tan(x).
+// The algorithm minimizes the maximum absolute error on [-pi/4,pi/4].
+#define GTE_C_TAN_DEG3_C0 1.0
+#define GTE_C_TAN_DEG3_C1 4.4295926544736286e-01
+#define GTE_C_TAN_DEG3_MAX_ERROR 1.1661892256204731e-2
+
+#define GTE_C_TAN_DEG5_C0 1.0
+#define GTE_C_TAN_DEG5_C1 3.1401320403542421e-01
+#define GTE_C_TAN_DEG5_C2 2.0903948109240345e-01
+#define GTE_C_TAN_DEG5_MAX_ERROR 5.8431854390143118e-4
+
+#define GTE_C_TAN_DEG7_C0 1.0
+#define GTE_C_TAN_DEG7_C1 3.3607213284422555e-01
+#define GTE_C_TAN_DEG7_C2 1.1261037305184907e-01
+#define GTE_C_TAN_DEG7_C3 9.8352099470524479e-02
+#define GTE_C_TAN_DEG7_MAX_ERROR 3.5418688397723108e-5
+
+#define GTE_C_TAN_DEG9_C0 1.0
+#define GTE_C_TAN_DEG9_C1 3.3299232843941784e-01
+#define GTE_C_TAN_DEG9_C2 1.3747843432474838e-01
+#define GTE_C_TAN_DEG9_C3 3.7696344813028304e-02
+#define GTE_C_TAN_DEG9_C4 4.6097377279281204e-02
+#define GTE_C_TAN_DEG9_MAX_ERROR 2.2988173242199927e-6
+
+#define GTE_C_TAN_DEG11_C0 1.0
+#define GTE_C_TAN_DEG11_C1 3.3337224456224224e-01
+#define GTE_C_TAN_DEG11_C2 1.3264516053824593e-01
+#define GTE_C_TAN_DEG11_C3 5.8145237645931047e-02
+#define GTE_C_TAN_DEG11_C4 1.0732193237572574e-02
+#define GTE_C_TAN_DEG11_C5 2.1558456793513869e-02
+#define GTE_C_TAN_DEG11_MAX_ERROR 1.5426257940140409e-7
+
+#define GTE_C_TAN_DEG13_C0 1.0
+#define GTE_C_TAN_DEG13_C1 3.3332916426394554e-01
+#define GTE_C_TAN_DEG13_C2 1.3343404625112498e-01
+#define GTE_C_TAN_DEG13_C3 5.3104565343119248e-02
+#define GTE_C_TAN_DEG13_C4 2.5355038312682154e-02
+#define GTE_C_TAN_DEG13_C5 1.8253255966556026e-03
+#define GTE_C_TAN_DEG13_C6 1.0069407176615641e-02
+#define GTE_C_TAN_DEG13_MAX_ERROR 1.0550264249037378e-8
+
+// Constants for minimax polynomial approximations to acos(x), where the
+// approximation is of the form acos(x) = sqrt(1 - x)*p(x) with p(x) a
+// polynomial.  The algorithm minimizes the maximum error
+// |acos(x)/sqrt(1-x) - p(x)| on [0,1].  At the same time we get an
+// approximation for asin(x) = pi/2 - acos(x).
+#define GTE_C_ACOS_DEG1_C0 +1.5707963267948966
+#define GTE_C_ACOS_DEG1_C1 -1.5658276442180141e-01
+#define GTE_C_ACOS_DEG1_MAX_ERROR 1.1659002803738105e-2
+
+#define GTE_C_ACOS_DEG2_C0 +1.5707963267948966
+#define GTE_C_ACOS_DEG2_C1 -2.0347053865798365e-01
+#define GTE_C_ACOS_DEG2_C2 +4.6887774236182234e-02
+#define GTE_C_ACOS_DEG2_MAX_ERROR 9.0311602490029258e-4
+
+#define GTE_C_ACOS_DEG3_C0 +1.5707963267948966
+#define GTE_C_ACOS_DEG3_C1 -2.1253291899190285e-01
+#define GTE_C_ACOS_DEG3_C2 +7.4773789639484223e-02
+#define GTE_C_ACOS_DEG3_C3 -1.8823635069382449e-02
+#define GTE_C_ACOS_DEG3_MAX_ERROR 9.3066396954288172e-5
+
+#define GTE_C_ACOS_DEG4_C0 +1.5707963267948966
+#define GTE_C_ACOS_DEG4_C1 -2.1422258835275865e-01
+#define GTE_C_ACOS_DEG4_C2 +8.4936675142844198e-02
+#define GTE_C_ACOS_DEG4_C3 -3.5991475120957794e-02
+#define GTE_C_ACOS_DEG4_C4 +8.6946239090712751e-03
+#define GTE_C_ACOS_DEG4_MAX_ERROR 1.0930595804481413e-5
+
+#define GTE_C_ACOS_DEG5_C0 +1.5707963267948966
+#define GTE_C_ACOS_DEG5_C1 -2.1453292139805524e-01
+#define GTE_C_ACOS_DEG5_C2 +8.7973089282889383e-02
+#define GTE_C_ACOS_DEG5_C3 -4.5130266382166440e-02
+#define GTE_C_ACOS_DEG5_C4 +1.9467466687281387e-02
+#define GTE_C_ACOS_DEG5_C5 -4.3601326117634898e-03
+#define GTE_C_ACOS_DEG5_MAX_ERROR 1.3861070257241426-6
+
+#define GTE_C_ACOS_DEG6_C0 +1.5707963267948966
+#define GTE_C_ACOS_DEG6_C1 -2.1458939285677325e-01
+#define GTE_C_ACOS_DEG6_C2 +8.8784960563641491e-02
+#define GTE_C_ACOS_DEG6_C3 -4.8887131453156485e-02
+#define GTE_C_ACOS_DEG6_C4 +2.7011519960012720e-02
+#define GTE_C_ACOS_DEG6_C5 -1.1210537323478320e-02
+#define GTE_C_ACOS_DEG6_C6 +2.3078166879102469e-03
+#define GTE_C_ACOS_DEG6_MAX_ERROR 1.8491291330427484e-7
+
+#define GTE_C_ACOS_DEG7_C0 +1.5707963267948966
+#define GTE_C_ACOS_DEG7_C1 -2.1459960076929829e-01
+#define GTE_C_ACOS_DEG7_C2 +8.8986946573346160e-02
+#define GTE_C_ACOS_DEG7_C3 -5.0207843052845647e-02
+#define GTE_C_ACOS_DEG7_C4 +3.0961594977611639e-02
+#define GTE_C_ACOS_DEG7_C5 -1.7162031184398074e-02
+#define GTE_C_ACOS_DEG7_C6 +6.7072304676685235e-03
+#define GTE_C_ACOS_DEG7_C7 -1.2690614339589956e-03
+#define GTE_C_ACOS_DEG7_MAX_ERROR 2.5574620927948377e-8
+
+#define GTE_C_ACOS_DEG8_C0 +1.5707963267948966
+#define GTE_C_ACOS_DEG8_C1 -2.1460143648688035e-01
+#define GTE_C_ACOS_DEG8_C2 +8.9034700107934128e-02
+#define GTE_C_ACOS_DEG8_C3 -5.0625279962389413e-02
+#define GTE_C_ACOS_DEG8_C4 +3.2683762943179318e-02
+#define GTE_C_ACOS_DEG8_C5 -2.0949278766238422e-02
+#define GTE_C_ACOS_DEG8_C6 +1.1272900916992512e-02
+#define GTE_C_ACOS_DEG8_C7 -4.1160981058965262e-03
+#define GTE_C_ACOS_DEG8_C8 +7.1796493341480527e-04
+#define GTE_C_ACOS_DEG8_MAX_ERROR 3.6340015129032732e-9
+
+// Constants for minimax polynomial approximations to atan(x).
+// The algorithm minimizes the maximum absolute error on [-1,1].
+#define GTE_C_ATAN_DEG3_C0 +1.0
+#define GTE_C_ATAN_DEG3_C1 -2.1460183660255172e-01
+#define GTE_C_ATAN_DEG3_MAX_ERROR 1.5970326392614240e-2
+
+#define GTE_C_ATAN_DEG5_C0 +1.0
+#define GTE_C_ATAN_DEG5_C1 -3.0189478312144946e-01
+#define GTE_C_ATAN_DEG5_C2 +8.7292946518897740e-02
+#define GTE_C_ATAN_DEG5_MAX_ERROR 1.3509832247372636e-3
+
+#define GTE_C_ATAN_DEG7_C0 +1.0
+#define GTE_C_ATAN_DEG7_C1 -3.2570157599356531e-01
+#define GTE_C_ATAN_DEG7_C2 +1.5342994884206673e-01
+#define GTE_C_ATAN_DEG7_C3 -4.2330209451053591e-02
+#define GTE_C_ATAN_DEG7_MAX_ERROR 1.5051227215514412e-4
+
+#define GTE_C_ATAN_DEG9_C0 +1.0
+#define GTE_C_ATAN_DEG9_C1 -3.3157878236439586e-01
+#define GTE_C_ATAN_DEG9_C2 +1.8383034738018011e-01
+#define GTE_C_ATAN_DEG9_C3 -8.9253037587244677e-02
+#define GTE_C_ATAN_DEG9_C4 +2.2399635968909593e-02
+#define GTE_C_ATAN_DEG9_MAX_ERROR 1.8921598624582064e-5
+
+#define GTE_C_ATAN_DEG11_C0 +1.0
+#define GTE_C_ATAN_DEG11_C1 -3.3294527685374087e-01
+#define GTE_C_ATAN_DEG11_C2 +1.9498657165383548e-01
+#define GTE_C_ATAN_DEG11_C3 -1.1921576270475498e-01
+#define GTE_C_ATAN_DEG11_C4 +5.5063351366968050e-02
+#define GTE_C_ATAN_DEG11_C5 -1.2490720064867844e-02
+#define GTE_C_ATAN_DEG11_MAX_ERROR 2.5477724974187765e-6
+
+#define GTE_C_ATAN_DEG13_C0 +1.0
+#define GTE_C_ATAN_DEG13_C1 -3.3324998579202170e-01
+#define GTE_C_ATAN_DEG13_C2 +1.9856563505717162e-01
+#define GTE_C_ATAN_DEG13_C3 -1.3374657325451267e-01
+#define GTE_C_ATAN_DEG13_C4 +8.1675882859940430e-02
+#define GTE_C_ATAN_DEG13_C5 -3.5059680836411644e-02
+#define GTE_C_ATAN_DEG13_C6 +7.2128853633444123e-03
+#define GTE_C_ATAN_DEG13_MAX_ERROR 3.5859104691865484e-7
+
+// Constants for minimax polynomial approximations to exp2(x) = 2^x.
+// The algorithm minimizes the maximum absolute error on [0,1].
+#define GTE_C_EXP2_DEG1_C0 1.0
+#define GTE_C_EXP2_DEG1_C1 1.0
+#define GTE_C_EXP2_DEG1_MAX_ERROR 8.6071332055934313e-2
+
+#define GTE_C_EXP2_DEG2_C0 1.0
+#define GTE_C_EXP2_DEG2_C1 6.5571332605741528e-01
+#define GTE_C_EXP2_DEG2_C2 3.4428667394258472e-01
+#define GTE_C_EXP2_DEG2_MAX_ERROR 3.8132476831060358e-3
+
+#define GTE_C_EXP2_DEG3_C0 1.0
+#define GTE_C_EXP2_DEG3_C1 6.9589012084456225e-01
+#define GTE_C_EXP2_DEG3_C2 2.2486494900110188e-01
+#define GTE_C_EXP2_DEG3_C3 7.9244930154334980e-02
+#define GTE_C_EXP2_DEG3_MAX_ERROR 1.4694877755186408e-4
+
+#define GTE_C_EXP2_DEG4_C0 1.0
+#define GTE_C_EXP2_DEG4_C1 6.9300392358459195e-01
+#define GTE_C_EXP2_DEG4_C2 2.4154981722455560e-01
+#define GTE_C_EXP2_DEG4_C3 5.1744260331489045e-02
+#define GTE_C_EXP2_DEG4_C4 1.3701998859367848e-02
+#define GTE_C_EXP2_DEG4_MAX_ERROR 4.7617792624521371e-6
+
+#define GTE_C_EXP2_DEG5_C0 1.0
+#define GTE_C_EXP2_DEG5_C1 6.9315298010274962e-01
+#define GTE_C_EXP2_DEG5_C2 2.4014712313022102e-01
+#define GTE_C_EXP2_DEG5_C3 5.5855296413199085e-02
+#define GTE_C_EXP2_DEG5_C4 8.9477503096873079e-03
+#define GTE_C_EXP2_DEG5_C5 1.8968500441332026e-03
+#define GTE_C_EXP2_DEG5_MAX_ERROR 1.3162098333463490e-7
+
+#define GTE_C_EXP2_DEG6_C0 1.0
+#define GTE_C_EXP2_DEG6_C1 6.9314698914837525e-01
+#define GTE_C_EXP2_DEG6_C2 2.4023013440952923e-01
+#define GTE_C_EXP2_DEG6_C3 5.5481276898206033e-02
+#define GTE_C_EXP2_DEG6_C4 9.6838443037086108e-03
+#define GTE_C_EXP2_DEG6_C5 1.2388324048515642e-03
+#define GTE_C_EXP2_DEG6_C6 2.1892283501756538e-04
+#define GTE_C_EXP2_DEG6_MAX_ERROR 3.1589168225654163e-9
+
+#define GTE_C_EXP2_DEG7_C0 1.0
+#define GTE_C_EXP2_DEG7_C1 6.9314718588750690e-01
+#define GTE_C_EXP2_DEG7_C2 2.4022637363165700e-01
+#define GTE_C_EXP2_DEG7_C3 5.5505235570535660e-02
+#define GTE_C_EXP2_DEG7_C4 9.6136265387940512e-03
+#define GTE_C_EXP2_DEG7_C5 1.3429234504656051e-03
+#define GTE_C_EXP2_DEG7_C6 1.4299202757683815e-04
+#define GTE_C_EXP2_DEG7_C7 2.1662892777385423e-05
+#define GTE_C_EXP2_DEG7_MAX_ERROR 6.6864513925679603e-11
+
+// Constants for minimax polynomial approximations to log2(x).
+// The algorithm minimizes the maximum absolute error on [1,2].
+// The polynomials all have constant term zero.
+#define GTE_C_LOG2_DEG1_C1 +1.0
+#define GTE_C_LOG2_DEG1_MAX_ERROR 8.6071332055934202e-2
+
+#define GTE_C_LOG2_DEG2_C1 +1.3465553856377803 
+#define GTE_C_LOG2_DEG2_C2 -3.4655538563778032e-01
+#define GTE_C_LOG2_DEG2_MAX_ERROR 7.6362868906658110e-3
+
+#define GTE_C_LOG2_DEG3_C1 +1.4228653756681227
+#define GTE_C_LOG2_DEG3_C2 -5.8208556916449616e-01
+#define GTE_C_LOG2_DEG3_C3 +1.5922019349637218e-01
+#define GTE_C_LOG2_DEG3_MAX_ERROR 8.7902902652883808e-4
+
+#define GTE_C_LOG2_DEG4_C1 +1.4387257478171547
+#define GTE_C_LOG2_DEG4_C2 -6.7778401359918661e-01
+#define GTE_C_LOG2_DEG4_C3 +3.2118898377713379e-01
+#define GTE_C_LOG2_DEG4_C4 -8.2130717995088531e-02
+#define GTE_C_LOG2_DEG4_MAX_ERROR 1.1318551355360418e-4
+
+#define GTE_C_LOG2_DEG5_C1 +1.4419170408633741
+#define GTE_C_LOG2_DEG5_C2 -7.0909645927612530e-01
+#define GTE_C_LOG2_DEG5_C3 +4.1560609399164150e-01
+#define GTE_C_LOG2_DEG5_C4 -1.9357573729558908e-01
+#define GTE_C_LOG2_DEG5_C5 +4.5149061716699634e-02
+#define GTE_C_LOG2_DEG5_MAX_ERROR 1.5521274478735858e-5
+
+#define GTE_C_LOG2_DEG6_C1 +1.4425449435950917
+#define GTE_C_LOG2_DEG6_C2 -7.1814525675038965e-01
+#define GTE_C_LOG2_DEG6_C3 +4.5754919692564044e-01
+#define GTE_C_LOG2_DEG6_C4 -2.7790534462849337e-01
+#define GTE_C_LOG2_DEG6_C5 +1.2179791068763279e-01
+#define GTE_C_LOG2_DEG6_C6 -2.5841449829670182e-02
+#define GTE_C_LOG2_DEG6_MAX_ERROR 2.2162051216689793e-6
+
+#define GTE_C_LOG2_DEG7_C1 +1.4426664401536078
+#define GTE_C_LOG2_DEG7_C2 -7.2055423726162360e-01
+#define GTE_C_LOG2_DEG7_C3 +4.7332419162501083e-01
+#define GTE_C_LOG2_DEG7_C4 -3.2514018752954144e-01
+#define GTE_C_LOG2_DEG7_C5 +1.9302965529095673e-01
+#define GTE_C_LOG2_DEG7_C6 -7.8534970641157997e-02
+#define GTE_C_LOG2_DEG7_C7 +1.5209108363023915e-02
+#define GTE_C_LOG2_DEG7_MAX_ERROR 3.2546531700261561e-7
+
+#define GTE_C_LOG2_DEG8_C1 +1.4426896453621882
+#define GTE_C_LOG2_DEG8_C2 -7.2115893912535967e-01
+#define GTE_C_LOG2_DEG8_C3 +4.7861716616785088e-01
+#define GTE_C_LOG2_DEG8_C4 -3.4699935395019565e-01
+#define GTE_C_LOG2_DEG8_C5 +2.4114048765477492e-01
+#define GTE_C_LOG2_DEG8_C6 -1.3657398692885181e-01
+#define GTE_C_LOG2_DEG8_C7 +5.1421382871922106e-02
+#define GTE_C_LOG2_DEG8_C8 -9.1364020499895560e-03
+#define GTE_C_LOG2_DEG8_MAX_ERROR 4.8796219218050219e-8
+
+// These functions are convenient for some applications.  The classes
+// BSNumber, BSRational and IEEEBinary16 have implementations that
+// (for now) use typecasting to call the 'float' or 'double' versions.
+namespace gte
+{
+    inline float atandivpi(float x)
+    {
+        return std::atan(x) * (float)GTE_C_INV_PI;
+    }
+
+    inline float atan2divpi(float y, float x)
+    {
+        return std::atan2(y, x) * (float)GTE_C_INV_PI;
+    }
+
+    inline float clamp(float x, float xmin, float xmax)
+    {
+        return (x <= xmin ? xmin : (x >= xmax ? xmax : x));
+    }
+
+    inline float cospi(float x)
+    {
+        return std::cos(x * (float)GTE_C_PI);
+    }
+
+    inline float exp10(float x)
+    {
+        return std::exp(x * (float)GTE_C_LN_10);
+    }
+
+    inline float invsqrt(float x)
+    {
+        return 1.0f / std::sqrt(x);
+    }
+
+    inline int isign(float x)
+    {
+        return (x > 0.0f ? 1 : (x < 0.0f ? -1 : 0));
+    }
+
+    inline float saturate(float x)
+    {
+        return (x <= 0.0f ? 0.0f : (x >= 1.0f ? 1.0f : x));
+    }
+
+    inline float sign(float x)
+    {
+        return (x > 0.0f ? 1.0f : (x < 0.0f ? -1.0f : 0.0f));
+    }
+
+    inline float sinpi(float x)
+    {
+        return std::sin(x * (float)GTE_C_PI);
+    }
+
+    inline float sqr(float x)
+    {
+        return x * x;
+    }
+
+
+    inline double atandivpi(double x)
+    {
+        return std::atan(x) * GTE_C_INV_PI;
+    }
+
+    inline double atan2divpi(double y, double x)
+    {
+        return std::atan2(y, x) * GTE_C_INV_PI;
+    }
+
+    inline double clamp(double x, double xmin, double xmax)
+    {
+        return (x <= xmin ? xmin : (x >= xmax ? xmax : x));
+    }
+
+    inline double cospi(double x)
+    {
+        return std::cos(x * GTE_C_PI);
+    }
+
+    inline double exp10(double x)
+    {
+        return std::exp(x * GTE_C_LN_10);
+    }
+
+    inline double invsqrt(double x)
+    {
+        return 1.0 / std::sqrt(x);
+    }
+
+    inline double sign(double x)
+    {
+        return (x > 0.0 ? 1.0 : (x < 0.0 ? -1.0 : 0.0f));
+    }
+
+    inline int isign(double x)
+    {
+        return (x > 0.0 ? 1 : (x < 0.0 ? -1 : 0));
+    }
+
+    inline double saturate(double x)
+    {
+        return (x <= 0.0 ? 0.0 : (x >= 1.0 ? 1.0 : x));
+    }
+
+    inline double sinpi(double x)
+    {
+        return std::sin(x * GTE_C_PI);
+    }
+
+    inline double sqr(double x)
+    {
+        return x * x;
+    }
+}
+
+// Type traits to support std::enable_if conditional compilation for
+// numerical computations.
+namespace gte
+{
+    // The trait is_arbitrary_precision<T> for type T of float, double or
+    // long double generates is_arbitrary_precision<T>::value of false.  The
+    // implementations for arbitrary-precision arithmetic are found in
+    // GteArbitraryPrecision.h.
+    template <typename T>
+    struct is_arbitrary_precision_internal : std::false_type {};
+
+    template <typename T>
+    struct is_arbitrary_precision : is_arbitrary_precision_internal<std::remove_cv_t<T>>::type {};
+
+    // The trait has_division_operator<T> for type T of float, double or
+    // long double generates has_division_operator<T>::value of true.  The
+    // implementations for arbitrary-precision arithmetic are found in
+    // GteArbitraryPrecision.h.
+    template <typename T>
+    struct has_division_operator_internal : std::false_type {};
+
+    template <typename T>
+    struct has_division_operator : has_division_operator_internal<std::remove_cv_t<T>>::type {};
+
+    template <>
+    struct has_division_operator_internal<float> : std::true_type {};
+
+    template <>
+    struct has_division_operator_internal<double> : std::true_type {};
+
+    template <>
+    struct has_division_operator_internal<long double> : std::true_type {};
+}
