move GeometryProcessing plugin from NotForLicensees to Experimental

diff --git a/Engine/Plugins/Experimental/GeometryProcessing/Source/GeometryAlgorithms/Private/GTEngine/Mathematics/GteMatrix.h b/Engine/Plugins/Experimental/GeometryProcessing/Source/GeometryAlgorithms/Private/GTEngine/Mathematics/GteMatrix.h
new file mode 100644
--- /dev/null
+++ b/Engine/Plugins/Experimental/GeometryProcessing/Source/GeometryAlgorithms/Private/GTEngine/Mathematics/GteMatrix.h
@@ -0,0 +1,957 @@
+// David Eberly, Geometric Tools, Redmond WA 98052
+// Copyright (c) 1998-2019
+// Distributed under the Boost Software License, Version 1.0.
+// http://www.boost.org/LICENSE_1_0.txt
+// http://www.geometrictools.com/License/Boost/LICENSE_1_0.txt
+// File Version: 3.0.2 (2018/10/05)
+
+#pragma once
+
+#include <GTEngine/Mathematics/GteVector.h>
+#include <GTEngine/Mathematics/GteGaussianElimination.h>
+
+namespace gte
+{
+
+template <int NumRows, int NumCols, typename Real>
+class Matrix
+{
+public:
+    // The table is initialized to zero.
+    Matrix();
+
+    // The table is fully initialized by the inputs.  The 'values' must be
+    // specified in row-major order, regardless of the active storage scheme
+    // (GTE_USE_ROW_MAJOR or GTE_USE_COL_MAJOR).
+    Matrix(std::array<Real, NumRows*NumCols> const& values);
+
+    // At most NumRows*NumCols are copied from the initializer list, setting
+    // any remaining elements to zero.  The 'values' must be specified in
+    // row-major order, regardless of the active storage scheme
+    // (GTE_USE_ROW_MAJOR or GTE_USE_COL_MAJOR).  Create the zero matrix using
+    // the syntax
+    //   Matrix<NumRows,NumCols,Real> zero{(Real)0};
+    // WARNING: The C++ 11 specification states that
+    //   Matrix<NumRows,NumCols,Real> zero{};
+    // will lead to a call of the default constructor, not the initializer
+    // constructor!
+    Matrix(std::initializer_list<Real> values);
+
+    // For 0 <= r < NumRows and 0 <= c < NumCols, element (r,c) is 1 and all
+    // others are 0.  If either of r or c is invalid, the zero matrix is
+    // created.  This is a convenience for creating the standard Euclidean
+    // basis matrices; see also MakeUnit(int,int) and Unit(int,int).
+    Matrix(int r, int c);
+
+    // The copy constructor, destructor, and assignment operator are generated
+    // by the compiler.
+
+    // Member access for which the storage representation is transparent.  The
+    // matrix entry in row r and column c is A(r,c).  The first operator()
+    // returns a const reference rather than a Real value.  This supports
+    // writing via standard file operations that require a const pointer to
+    // data.
+    inline Real const& operator()(int r, int c) const;
+    inline Real& operator()(int r, int c);
+
+    // Member access by rows or by columns.
+    void SetRow(int r, Vector<NumCols,Real> const& vec);
+    void SetCol(int c, Vector<NumRows,Real> const& vec);
+    Vector<NumCols,Real> GetRow(int r) const;
+    Vector<NumRows,Real> GetCol(int c) const;
+
+    // Member access by 1-dimensional index.  NOTE: These accessors are
+    // useful for the manipulation of matrix entries when it does not
+    // matter whether storage is row-major or column-major.  Do not use
+    // constructs such as M[c+NumCols*r] or M[r+NumRows*c] that expose the
+    // storage convention.
+    inline Real const& operator[](int i) const;
+    inline Real& operator[](int i);
+
+    // Comparisons for sorted containers and geometric ordering.
+    inline bool operator==(Matrix const& mat) const;
+    inline bool operator!=(Matrix const& mat) const;
+    inline bool operator< (Matrix const& mat) const;
+    inline bool operator<=(Matrix const& mat) const;
+    inline bool operator> (Matrix const& mat) const;
+    inline bool operator>=(Matrix const& mat) const;
+
+    // Special matrices.
+    void MakeZero();  // All components are 0.
+    void MakeUnit(int r, int c);  // Component (r,c) is 1, all others zero.
+    void MakeIdentity();  // Diagonal entries 1, others 0, even when nonsquare
+    static Matrix Zero();
+    static Matrix Unit(int r, int c);
+    static Matrix Identity();
+
+protected:
+    class Table
+    {
+    public:
+        // Storage-order-independent element access as 2D array.
+        inline Real const& operator()(int r, int c) const;
+        inline Real& operator()(int r, int c);
+
+        // Element access as 1D array.  Use this internally only when
+        // the 2D storage order is not relevant.
+        inline Real const& operator[](int i) const;
+        inline Real& operator[](int i);
+
+#if defined(GTE_USE_ROW_MAJOR)
+        std::array<std::array<Real,NumCols>,NumRows> mStorage;
+#else
+        std::array<std::array<Real,NumRows>,NumCols> mStorage;
+#endif
+    };
+
+    Table mTable;
+};
+
+// Unary operations.
+template <int NumRows, int NumCols, typename Real>
+Matrix<NumRows,NumCols,Real>
+operator+(Matrix<NumRows,NumCols,Real> const& M);
+
+template <int NumRows, int NumCols, typename Real>
+Matrix<NumRows,NumCols,Real>
+operator-(Matrix<NumRows,NumCols,Real> const& M);
+
+// Linear-algebraic operations.
+template <int NumRows, int NumCols, typename Real>
+Matrix<NumRows,NumCols,Real>
+operator+(
+    Matrix<NumRows,NumCols,Real> const& M0,
+    Matrix<NumRows,NumCols,Real> const& M1);
+
+template <int NumRows, int NumCols, typename Real>
+Matrix<NumRows,NumCols,Real>
+operator-(
+    Matrix<NumRows,NumCols,Real> const& M0,
+    Matrix<NumRows,NumCols,Real> const& M1);
+
+template <int NumRows, int NumCols, typename Real>
+Matrix<NumRows,NumCols,Real>
+operator*(Matrix<NumRows,NumCols,Real> const& M, Real scalar);
+
+template <int NumRows, int NumCols, typename Real>
+Matrix<NumRows,NumCols,Real>
+operator*(Real scalar, Matrix<NumRows,NumCols,Real> const& M);
+
+template <int NumRows, int NumCols, typename Real>
+Matrix<NumRows,NumCols,Real>
+operator/(Matrix<NumRows,NumCols,Real> const& M, Real scalar);
+
+template <int NumRows, int NumCols, typename Real>
+Matrix<NumRows,NumCols,Real>&
+operator+=(
+    Matrix<NumRows,NumCols,Real>& M0,
+    Matrix<NumRows,NumCols,Real> const& M1);
+
+template <int NumRows, int NumCols, typename Real>
+Matrix<NumRows,NumCols,Real>&
+operator-=(
+    Matrix<NumRows,NumCols,Real>& M0,
+    Matrix<NumRows,NumCols,Real> const& M1);
+
+template <int NumRows, int NumCols, typename Real>
+Matrix<NumRows,NumCols,Real>&
+operator*=(Matrix<NumRows,NumCols,Real>& M, Real scalar);
+
+template <int NumRows, int NumCols, typename Real>
+Matrix<NumRows,NumCols,Real>&
+operator/=(Matrix<NumRows,NumCols,Real>& M, Real scalar);
+
+// Geometric operations.
+template <int NumRows, int NumCols, typename Real>
+Real L1Norm(Matrix<NumRows,NumCols,Real> const& M);
+
+template <int NumRows, int NumCols, typename Real>
+Real L2Norm(Matrix<NumRows,NumCols,Real> const& M);
+
+template <int NumRows, int NumCols, typename Real>
+Real LInfinityNorm(Matrix<NumRows,NumCols,Real> const& M);
+
+template <int N, typename Real>
+Matrix<N,N,Real> Inverse(Matrix<N,N,Real> const& M,
+    bool* reportInvertibility = nullptr);
+
+template <int N, typename Real>
+Real Determinant(Matrix<N, N, Real> const& M);
+
+// M^T
+template <int NumRows, int NumCols, typename Real>
+Matrix<NumCols,NumRows,Real>
+Transpose(Matrix<NumRows,NumCols,Real> const& M);
+
+// M*V
+template <int NumRows, int NumCols, typename Real>
+Vector<NumRows,Real>
+operator*(
+    Matrix<NumRows,NumCols,Real> const& M,
+    Vector<NumCols,Real> const& V);
+
+// V^T*M
+template <int NumRows, int NumCols, typename Real>
+Vector<NumCols,Real>
+operator*(
+    Vector<NumRows,Real> const& V,
+    Matrix<NumRows,NumCols,Real> const& M);
+
+// A*B
+template <int NumRows, int NumCols, int NumCommon, typename Real>
+Matrix<NumRows,NumCols,Real>
+operator*(
+    Matrix<NumRows,NumCommon,Real> const& A,
+    Matrix<NumCommon,NumCols,Real> const& B);
+
+template <int NumRows, int NumCols, int NumCommon, typename Real>
+Matrix<NumRows,NumCols,Real>
+MultiplyAB(
+    Matrix<NumRows,NumCommon,Real> const& A,
+    Matrix<NumCommon,NumCols,Real> const& B);
+
+// A*B^T
+template <int NumRows, int NumCols, int NumCommon, typename Real>
+Matrix<NumRows,NumCols,Real>
+MultiplyABT(
+    Matrix<NumRows,NumCommon,Real> const& A,
+    Matrix<NumCols,NumCommon,Real> const& B);
+
+// A^T*B
+template <int NumRows, int NumCols, int NumCommon, typename Real>
+Matrix<NumRows,NumCols,Real>
+MultiplyATB(
+    Matrix<NumCommon,NumRows,Real> const& A,
+    Matrix<NumCommon,NumCols,Real> const& B);
+
+// A^T*B^T
+template <int NumRows, int NumCols, int NumCommon, typename Real>
+Matrix<NumRows,NumCols,Real>
+MultiplyATBT(
+    Matrix<NumCommon,NumRows,Real> const& A,
+    Matrix<NumCols,NumCommon,Real> const& B);
+
+// M*D, D is diagonal NumCols-by-NumCols
+template <int NumRows, int NumCols, typename Real>
+Matrix<NumRows,NumCols,Real>
+MultiplyMD(
+    Matrix<NumRows,NumCols,Real> const& M,
+    Vector<NumCols,Real> const& D);
+
+// D*M, D is diagonal NumRows-by-NumRows
+template <int NumRows, int NumCols, typename Real>
+Matrix<NumRows,NumCols,Real>
+MultiplyDM(
+    Vector<NumRows,Real> const& D,
+    Matrix<NumRows,NumCols,Real> const& M);
+
+// U*V^T, U is NumRows-by-1, V is Num-Cols-by-1, result is NumRows-by-NumCols.
+template <int NumRows, int NumCols, typename Real>
+Matrix<NumRows,NumCols,Real>
+OuterProduct(Vector<NumRows, Real> const& U, Vector<NumCols, Real> const& V);
+
+// Initialization to a diagonal matrix whose diagonal entries are the
+// components of D.
+template <int N, typename Real>
+void MakeDiagonal(Vector<N, Real> const& D, Matrix<N, N, Real>& M);
+
+// Create an (N+1)-by-(N+1) matrix H by setting the upper N-by-N block to the
+// input N-by-N matrix and all other entries to 0 except for the last row
+// and last column entry which is set to 1.
+template <int N, typename Real>
+Matrix<N + 1, N+1, Real> HLift(Matrix<N, N, Real> const& M);
+
+// Extract the upper (N-1)-by-(N-1) block of the input N-by-N matrix.
+template <int N, typename Real>
+Matrix<N - 1, N - 1, Real> HProject(Matrix<N, N, Real> const& M);
+
+
+template <int NumRows, int NumCols, typename Real>
+Matrix<NumRows, NumCols, Real>::Matrix()
+{
+    MakeZero();
+}
+
+template <int NumRows, int NumCols, typename Real>
+Matrix<NumRows, NumCols, Real>::Matrix(
+    std::array<Real, NumRows*NumCols> const& values)
+{
+    for (int r = 0, i = 0; r < NumRows; ++r)
+    {
+        for (int c = 0; c < NumCols; ++c, ++i)
+        {
+            mTable(r, c) = values[i];
+        }
+    }
+}
+
+template <int NumRows, int NumCols, typename Real>
+Matrix<NumRows, NumCols, Real>::Matrix(std::initializer_list<Real> values)
+{
+    int const numValues = static_cast<int>(values.size());
+    auto iter = values.begin();
+    int r, c, i;
+    for (r = 0, i = 0; r < NumRows; ++r)
+    {
+        for (c = 0; c < NumCols; ++c, ++i)
+        {
+            if (i < numValues)
+            {
+                mTable(r, c) = *iter++;
+            }
+            else
+            {
+                break;
+            }
+        }
+
+        if (c < NumCols)
+        {
+            // Fill in the remaining columns of the current row with zeros.
+            for (/**/; c < NumCols; ++c)
+            {
+                mTable(r, c) = (Real)0;
+            }
+            ++r;
+            break;
+        }
+    }
+
+    if (r < NumRows)
+    {
+        // Fill in the remain rows with zeros.
+        for (/**/; r < NumRows; ++r)
+        {
+            for (c = 0; c < NumCols; ++c)
+            {
+                mTable(r, c) = (Real)0;
+            }
+        }
+    }
+}
+
+template <int NumRows, int NumCols, typename Real>
+Matrix<NumRows, NumCols, Real>::Matrix(int r, int c)
+{
+    MakeUnit(r, c);
+}
+
+template <int NumRows, int NumCols, typename Real> inline
+Real const& Matrix<NumRows, NumCols, Real>::operator()(int r, int c) const
+{
+    return mTable(r, c);
+}
+
+template <int NumRows, int NumCols, typename Real> inline
+Real& Matrix<NumRows, NumCols, Real>::operator()(int r, int c)
+{
+    return mTable(r, c);
+}
+
+template <int NumRows, int NumCols, typename Real>
+void Matrix<NumRows, NumCols, Real>::SetRow(int r,
+    Vector<NumCols, Real> const& vec)
+{
+    for (int c = 0; c < NumCols; ++c)
+    {
+        mTable(r, c) = vec[c];
+    }
+}
+
+template <int NumRows, int NumCols, typename Real>
+void Matrix<NumRows, NumCols, Real>::SetCol(int c,
+    Vector<NumRows, Real> const& vec)
+{
+    for (int r = 0; r < NumRows; ++r)
+    {
+        mTable(r, c) = vec[r];
+    }
+}
+
+template <int NumRows, int NumCols, typename Real>
+Vector<NumCols, Real> Matrix<NumRows, NumCols, Real>::GetRow(int r) const
+{
+    Vector<NumCols, Real> vec;
+    for (int c = 0; c < NumCols; ++c)
+    {
+        vec[c] = mTable(r, c);
+    }
+    return vec;
+}
+
+template <int NumRows, int NumCols, typename Real>
+Vector<NumRows, Real> Matrix<NumRows, NumCols, Real>::GetCol(int c) const
+{
+    Vector<NumRows, Real> vec;
+    for (int r = 0; r < NumRows; ++r)
+    {
+        vec[r] = mTable(r, c);
+    }
+    return vec;
+}
+
+template <int NumRows, int NumCols, typename Real> inline
+Real const& Matrix<NumRows, NumCols, Real>::operator[](int i) const
+{
+    return mTable[i];
+}
+
+template <int NumRows, int NumCols, typename Real> inline
+Real& Matrix<NumRows, NumCols, Real>::operator[](int i)
+{
+    return mTable[i];
+}
+
+template <int NumRows, int NumCols, typename Real> inline
+bool Matrix<NumRows, NumCols, Real>::operator==(Matrix const& mat) const
+{
+    return mTable.mStorage == mat.mTable.mStorage;
+}
+
+template <int NumRows, int NumCols, typename Real> inline
+bool Matrix<NumRows, NumCols, Real>::operator!=(Matrix const& mat) const
+{
+    return mTable.mStorage != mat.mTable.mStorage;
+}
+
+template <int NumRows, int NumCols, typename Real> inline
+bool Matrix<NumRows, NumCols, Real>::operator<(Matrix const& mat) const
+{
+    return mTable.mStorage < mat.mTable.mStorage;
+}
+
+template <int NumRows, int NumCols, typename Real> inline
+bool Matrix<NumRows, NumCols, Real>::operator<=(Matrix const& mat) const
+{
+    return mTable.mStorage <= mat.mTable.mStorage;
+}
+
+template <int NumRows, int NumCols, typename Real> inline
+bool Matrix<NumRows, NumCols, Real>::operator>(Matrix const& mat) const
+{
+    return mTable.mStorage > mat.mTable.mStorage;
+}
+
+template <int NumRows, int NumCols, typename Real> inline
+bool Matrix<NumRows, NumCols, Real>::operator>=(Matrix const& mat) const
+{
+    return mTable.mStorage >= mat.mTable.mStorage;
+}
+
+template <int NumRows, int NumCols, typename Real>
+void Matrix<NumRows, NumCols, Real>::MakeZero()
+{
+    Real const zero = (Real)0;
+    for (int i = 0; i < NumRows * NumCols; ++i)
+    {
+        mTable[i] = zero;
+    }
+}
+
+template <int NumRows, int NumCols, typename Real>
+void Matrix<NumRows, NumCols, Real>::MakeUnit(int r, int c)
+{
+    MakeZero();
+    if (0 <= r && r < NumRows && 0 <= c && c < NumCols)
+    {
+        mTable(r, c) = (Real)1;
+    }
+}
+
+template <int NumRows, int NumCols, typename Real>
+void Matrix<NumRows, NumCols, Real>::MakeIdentity()
+{
+    MakeZero();
+    int const numDiagonal = (NumRows <= NumCols ? NumRows : NumCols);
+    for (int i = 0; i < numDiagonal; ++i)
+    {
+        mTable(i, i) = (Real)1;
+    }
+}
+
+template <int NumRows, int NumCols, typename Real>
+Matrix<NumRows, NumCols, Real> Matrix<NumRows, NumCols, Real>::Zero()
+{
+    Matrix M;
+    M.MakeZero();
+    return M;
+}
+
+template <int NumRows, int NumCols, typename Real>
+Matrix<NumRows, NumCols, Real> Matrix<NumRows, NumCols, Real>::Unit(int r,
+    int c)
+{
+    Matrix M;
+    M.MakeUnit(r, c);
+    return M;
+}
+
+template <int NumRows, int NumCols, typename Real>
+Matrix<NumRows, NumCols, Real> Matrix<NumRows, NumCols, Real>::Identity()
+{
+    Matrix M;
+    M.MakeIdentity();
+    return M;
+}
+
+
+
+template <int NumRows, int NumCols, typename Real>
+Matrix<NumRows, NumCols, Real>
+operator+(Matrix<NumRows, NumCols, Real> const& M)
+{
+    return M;
+}
+
+template <int NumRows, int NumCols, typename Real>
+Matrix<NumRows, NumCols, Real>
+operator-(Matrix<NumRows, NumCols, Real> const& M)
+{
+    Matrix<NumRows, NumCols, Real> result;
+    for (int i = 0; i < NumRows*NumCols; ++i)
+    {
+        result[i] = -M[i];
+    }
+    return result;
+}
+
+template <int NumRows, int NumCols, typename Real>
+Matrix<NumRows, NumCols, Real>
+operator+(
+    Matrix<NumRows, NumCols, Real> const& M0,
+    Matrix<NumRows, NumCols, Real> const& M1)
+{
+    Matrix<NumRows, NumCols, Real> result = M0;
+    return result += M1;
+}
+
+template <int NumRows, int NumCols, typename Real>
+Matrix<NumRows, NumCols, Real>
+operator-(
+    Matrix<NumRows, NumCols, Real> const& M0,
+    Matrix<NumRows, NumCols, Real> const& M1)
+{
+    Matrix<NumRows, NumCols, Real> result = M0;
+    return result -= M1;
+}
+
+template <int NumRows, int NumCols, typename Real>
+Matrix<NumRows, NumCols, Real>
+operator*(Matrix<NumRows, NumCols, Real> const& M, Real scalar)
+{
+    Matrix<NumRows, NumCols, Real> result = M;
+    return result *= scalar;
+}
+
+template <int NumRows, int NumCols, typename Real>
+Matrix<NumRows, NumCols, Real>
+operator*(Real scalar, Matrix<NumRows, NumCols, Real> const& M)
+{
+    Matrix<NumRows, NumCols, Real> result = M;
+    return result *= scalar;
+}
+
+template <int NumRows, int NumCols, typename Real>
+Matrix<NumRows, NumCols, Real>
+operator/(Matrix<NumRows, NumCols, Real> const& M, Real scalar)
+{
+    Matrix<NumRows, NumCols, Real> result = M;
+    return result /= scalar;
+}
+
+template <int NumRows, int NumCols, typename Real>
+Matrix<NumRows, NumCols, Real>&
+operator+=(
+    Matrix<NumRows, NumCols, Real>& M0,
+    Matrix<NumRows, NumCols, Real> const& M1)
+{
+    for (int i = 0; i < NumRows*NumCols; ++i)
+    {
+        M0[i] += M1[i];
+    }
+    return M0;
+}
+
+template <int NumRows, int NumCols, typename Real>
+Matrix<NumRows, NumCols, Real>&
+operator-=(
+    Matrix<NumRows, NumCols, Real>& M0,
+    Matrix<NumRows, NumCols, Real> const& M1)
+{
+    for (int i = 0; i < NumRows*NumCols; ++i)
+    {
+        M0[i] -= M1[i];
+    }
+    return M0;
+}
+
+template <int NumRows, int NumCols, typename Real>
+Matrix<NumRows, NumCols, Real>&
+operator*=(Matrix<NumRows, NumCols, Real>& M, Real scalar)
+{
+    for (int i = 0; i < NumRows*NumCols; ++i)
+    {
+        M[i] *= scalar;
+    }
+    return M;
+}
+
+template <int NumRows, int NumCols, typename Real>
+Matrix<NumRows, NumCols, Real>&
+operator/=(Matrix<NumRows, NumCols, Real>& M, Real scalar)
+{
+    if (scalar != (Real)0)
+    {
+        Real invScalar = ((Real)1) / scalar;
+        for (int i = 0; i < NumRows*NumCols; ++i)
+        {
+            M[i] *= invScalar;
+        }
+    }
+    else
+    {
+        for (int i = 0; i < NumRows*NumCols; ++i)
+        {
+            M[i] = (Real)0;
+        }
+    }
+    return M;
+}
+
+template <int NumRows, int NumCols, typename Real>
+Real L1Norm(Matrix<NumRows, NumCols, Real> const& M)
+{
+    Real sum = std::abs(M[0]);
+    for (int i = 1; i < NumRows*NumCols; ++i)
+    {
+        sum += std::abs(M[i]);
+    }
+    return sum;
+}
+
+template <int NumRows, int NumCols, typename Real>
+Real L2Norm(Matrix<NumRows, NumCols, Real> const& M)
+{
+    Real sum = M[0] * M[0];
+    for (int i = 1; i < NumRows*NumCols; ++i)
+    {
+        sum += M[i] * M[i];
+    }
+    return std::sqrt(sum);
+}
+
+template <int NumRows, int NumCols, typename Real>
+Real LInfinityNorm(Matrix<NumRows, NumCols, Real> const& M)
+{
+    Real maxAbsElement = M[0];
+    for (int i = 1; i < NumRows*NumCols; ++i)
+    {
+        Real absElement = std::abs(M[i]);
+        if (absElement > maxAbsElement)
+        {
+            maxAbsElement = absElement;
+        }
+    }
+    return maxAbsElement;
+}
+
+template <int N, typename Real>
+Matrix<N, N, Real> Inverse(Matrix<N, N, Real> const& M,
+    bool* reportInvertibility)
+{
+    Matrix<N, N, Real> invM;
+    Real determinant;
+    bool invertible = GaussianElimination<Real>()(N, &M[0], &invM[0],
+        determinant, nullptr, nullptr, nullptr, 0, nullptr);
+    if (reportInvertibility)
+    {
+        *reportInvertibility = invertible;
+    }
+    return invM;
+}
+
+template <int N, typename Real>
+Real Determinant(Matrix<N, N, Real> const& M)
+{
+    Real determinant;
+    GaussianElimination<Real>()(N, &M[0], nullptr, determinant, nullptr,
+        nullptr, nullptr, 0, nullptr);
+    return determinant;
+}
+
+template <int NumRows, int NumCols, typename Real>
+Matrix<NumCols, NumRows, Real>
+Transpose(Matrix<NumRows, NumCols, Real> const& M)
+{
+    Matrix<NumCols, NumRows, Real> result;
+    for (int r = 0; r < NumRows; ++r)
+    {
+        for (int c = 0; c < NumCols; ++c)
+        {
+            result(c, r) = M(r, c);
+        }
+    }
+    return result;
+}
+
+template <int NumRows, int NumCols, typename Real>
+Vector<NumRows, Real>
+operator*(
+    Matrix<NumRows, NumCols, Real> const& M,
+    Vector<NumCols, Real> const& V)
+{
+    Vector<NumRows, Real> result;
+    for (int r = 0; r < NumRows; ++r)
+    {
+        result[r] = (Real)0;
+        for (int c = 0; c < NumCols; ++c)
+        {
+            result[r] += M(r, c) * V[c];
+        }
+    }
+    return result;
+}
+
+template <int NumRows, int NumCols, typename Real>
+Vector<NumCols, Real> operator*(Vector<NumRows, Real> const& V,
+    Matrix<NumRows, NumCols, Real> const& M)
+{
+    Vector<NumCols, Real> result;
+    for (int c = 0; c < NumCols; ++c)
+    {
+        result[c] = (Real)0;
+        for (int r = 0; r < NumRows; ++r)
+        {
+            result[c] += V[r] * M(r, c);
+        }
+    }
+    return result;
+}
+
+template <int NumRows, int NumCols, int NumCommon, typename Real>
+Matrix<NumRows, NumCols, Real>
+operator*(
+    Matrix<NumRows, NumCommon, Real> const& A,
+    Matrix<NumCommon, NumCols, Real> const& B)
+{
+    return MultiplyAB(A, B);
+}
+
+template <int NumRows, int NumCols, int NumCommon, typename Real>
+Matrix<NumRows, NumCols, Real>
+MultiplyAB(
+    Matrix<NumRows, NumCommon, Real> const& A,
+    Matrix<NumCommon, NumCols, Real> const& B)
+{
+    Matrix<NumRows, NumCols, Real> result;
+    for (int r = 0; r < NumRows; ++r)
+    {
+        for (int c = 0; c < NumCols; ++c)
+        {
+            result(r, c) = (Real)0;
+            for (int i = 0; i < NumCommon; ++i)
+            {
+                result(r, c) += A(r, i) * B(i, c);
+            }
+        }
+    }
+    return result;
+}
+
+template <int NumRows, int NumCols, int NumCommon, typename Real>
+Matrix<NumRows, NumCols, Real>
+MultiplyABT(
+    Matrix<NumRows, NumCommon, Real> const& A,
+    Matrix<NumCols, NumCommon, Real> const& B)
+{
+    Matrix<NumRows, NumCols, Real> result;
+    for (int r = 0; r < NumRows; ++r)
+    {
+        for (int c = 0; c < NumCols; ++c)
+        {
+            result(r, c) = (Real)0;
+            for (int i = 0; i < NumCommon; ++i)
+            {
+                result(r, c) += A(r, i) * B(c, i);
+            }
+        }
+    }
+    return result;
+}
+
+template <int NumRows, int NumCols, int NumCommon, typename Real>
+Matrix<NumRows, NumCols, Real>
+MultiplyATB(
+    Matrix<NumCommon, NumRows, Real> const& A,
+    Matrix<NumCommon, NumCols, Real> const& B)
+{
+    Matrix<NumRows, NumCols, Real> result;
+    for (int r = 0; r < NumRows; ++r)
+    {
+        for (int c = 0; c < NumCols; ++c)
+        {
+            result(r, c) = (Real)0;
+            for (int i = 0; i < NumCommon; ++i)
+            {
+                result(r, c) += A(i, r) * B(i, c);
+            }
+        }
+    }
+    return result;
+}
+
+template <int NumRows, int NumCols, int NumCommon, typename Real>
+Matrix<NumRows, NumCols, Real>
+MultiplyATBT(
+    Matrix<NumCommon, NumRows, Real> const& A,
+    Matrix<NumCols, NumCommon, Real> const& B)
+{
+    Matrix<NumRows, NumCols, Real> result;
+    for (int r = 0; r < NumRows; ++r)
+    {
+        for (int c = 0; c < NumCols; ++c)
+        {
+            result(r, c) = (Real)0;
+            for (int i = 0; i < NumCommon; ++i)
+            {
+                result(r, c) += A(i, r) * B(c, i);
+            }
+        }
+    }
+    return result;
+}
+
+template <int NumRows, int NumCols, typename Real>
+Matrix<NumRows, NumCols, Real>
+MultiplyMD(
+    Matrix<NumRows, NumCols, Real> const& M,
+    Vector<NumCols, Real> const& D)
+{
+    Matrix<NumRows, NumCols, Real> result;
+    for (int r = 0; r < NumRows; ++r)
+    {
+        for (int c = 0; c < NumCols; ++c)
+        {
+            result(r, c) = M(r, c) * D[c];
+        }
+    }
+    return result;
+}
+
+template <int NumRows, int NumCols, typename Real>
+Matrix<NumRows, NumCols, Real>
+MultiplyDM(
+    Vector<NumRows, Real> const& D,
+    Matrix<NumRows, NumCols, Real> const& M)
+{
+    Matrix<NumRows, NumCols, Real> result;
+    for (int r = 0; r < NumRows; ++r)
+    {
+        for (int c = 0; c < NumCols; ++c)
+        {
+            result(r, c) = D[r] * M(r, c);
+        }
+    }
+    return result;
+}
+
+template <int NumRows, int NumCols, typename Real>
+Matrix<NumRows, NumCols, Real>
+OuterProduct(Vector<NumRows, Real> const& U, Vector<NumCols, Real> const& V)
+{
+    Matrix<NumRows, NumCols, Real> result;
+    for (int r = 0; r < NumRows; ++r)
+    {
+        for (int c = 0; c < NumCols; ++c)
+        {
+            result(r, c) = U[r] * V[c];
+        }
+    }
+    return result;
+}
+
+template <int N, typename Real>
+void MakeDiagonal(Vector<N, Real> const& D, Matrix<N, N, Real>& M)
+{
+    for (int i = 0; i < N*N; ++i)
+    {
+        M[i] = (Real)0;
+    }
+
+    for (int i = 0; i < N; ++i)
+    {
+        M(i, i) = D[i];
+    }
+}
+
+template <int N, typename Real>
+Matrix<N + 1, N + 1, Real> HLift(Matrix<N, N, Real> const& M)
+{
+    Matrix<N + 1, N + 1, Real> result;
+    result.MakeIdentity();
+    for (int r = 0; r < N; ++r)
+    {
+        for (int c = 0; c < N; ++c)
+        {
+            result(r, c) = M(r, c);
+        }
+    }
+    return result;
+}
+
+// Extract the upper (N-1)-by-(N-1) block of the input N-by-N matrix.
+template <int N, typename Real>
+Matrix<N - 1, N - 1, Real> HProject(Matrix<N, N, Real> const& M)
+{
+    static_assert(N >= 2, "Invalid matrix dimension.");
+    Matrix<N - 1, N - 1, Real> result;
+    for (int r = 0; r < N - 1; ++r)
+    {
+        for (int c = 0; c < N - 1; ++c)
+        {
+            result(r, c) = M(r, c);
+        }
+    }
+    return result;
+}
+
+
+// Matrix<N,C,Real>::Table
+
+template <int NumRows, int NumCols, typename Real> inline
+Real const& Matrix<NumRows, NumCols, Real>::Table::operator()(int r, int c)
+    const
+{
+#if defined(GTE_USE_ROW_MAJOR)
+    return mStorage[r][c];
+#else
+    return mStorage[c][r];
+#endif
+}
+
+template <int NumRows, int NumCols, typename Real> inline
+Real& Matrix<NumRows, NumCols, Real>::Table::operator()(int r, int c)
+{
+#if defined(GTE_USE_ROW_MAJOR)
+    return mStorage[r][c];
+#else
+    return mStorage[c][r];
+#endif
+}
+
+template <int NumRows, int NumCols, typename Real> inline
+Real const& Matrix<NumRows, NumCols, Real>::Table::operator[](int i) const
+{
+    Real const* elements = &mStorage[0][0];
+    return elements[i];
+}
+
+template <int NumRows, int NumCols, typename Real> inline
+Real& Matrix<NumRows, NumCols, Real>::Table::operator[](int i)
+{
+    Real* elements = &mStorage[0][0];
+    return elements[i];
+}
+
+
+}
diff --git a/Engine/Plugins/Experimental/GeometryProcessing/Source/GeometryAlgorithms/Private/GTEngine/Mathematics/GteMatrix2x2.h b/Engine/Plugins/Experimental/GeometryProcessing/Source/GeometryAlgorithms/Private/GTEngine/Mathematics/GteMatrix2x2.h
new file mode 100644
--- /dev/null
+++ b/Engine/Plugins/Experimental/GeometryProcessing/Source/GeometryAlgorithms/Private/GTEngine/Mathematics/GteMatrix2x2.h
@@ -0,0 +1,132 @@
+// David Eberly, Geometric Tools, Redmond WA 98052
+// Copyright (c) 1998-2019
+// Distributed under the Boost Software License, Version 1.0.
+// http://www.boost.org/LICENSE_1_0.txt
+// http://www.geometrictools.com/License/Boost/LICENSE_1_0.txt
+// File Version: 3.0.1 (2018/10/05)
+
+#pragma once
+
+#include <GTEngine/Mathematics/GteMatrix.h>
+#include <GTEngine/Mathematics/GteVector2.h>
+
+namespace gte
+{
+
+// Template alias for convenience.
+template <typename Real>
+using Matrix2x2 = Matrix<2, 2, Real>;
+
+// Create a rotation matrix from an angle (in radians).  The matrix is
+// [GTE_USE_MAT_VEC]
+//   R(t) = {{c,-s},{s,c}}
+// [GTE_USE_VEC_MAT]
+//   R(t) = {{c,s},{-s,c}}
+// where c = cos(t), s = sin(t), and the inner-brace pairs are rows of the
+// matrix.
+template <typename Real>
+void MakeRotation(Real angle, Matrix2x2<Real>& rotation);
+
+// Get the angle (radians) from a rotation matrix.  The caller is
+// responsible for ensuring the matrix is a rotation.
+template <typename Real>
+Real GetRotationAngle(Matrix2x2<Real> const& rotation);
+
+// Geometric operations.
+template <typename Real>
+Matrix2x2<Real> Inverse(Matrix2x2<Real> const& M,
+    bool* reportInvertibility = nullptr);
+
+template <typename Real>
+Matrix2x2<Real> Adjoint(Matrix2x2<Real> const& M);
+
+template <typename Real>
+Real Determinant(Matrix2x2<Real> const& M);
+
+template <typename Real>
+Real Trace(Matrix2x2<Real> const& M);
+
+
+template <typename Real>
+void MakeRotation(Real angle, Matrix2x2<Real>& rotation)
+{
+    Real cs = std::cos(angle);
+    Real sn = std::sin(angle);
+#if defined(GTE_USE_MAT_VEC)
+    rotation(0, 0) = cs;
+    rotation(0, 1) = -sn;
+    rotation(1, 0) = sn;
+    rotation(1, 1) = cs;
+#else
+    rotation(0, 0) = cs;
+    rotation(0, 1) = sn;
+    rotation(1, 0) = -sn;
+    rotation(1, 1) = cs;
+#endif
+}
+
+template <typename Real>
+Real GetRotationAngle(Matrix2x2<Real> const& rotation)
+{
+#if defined(GTE_USE_MAT_VEC)
+    return std::atan2(rotation(1, 0), rotation(0, 0));
+#else
+    return std::atan2(rotation(0, 1), rotation(0, 0));
+#endif
+}
+
+template <typename Real>
+Matrix2x2<Real> Inverse(Matrix2x2<Real> const& M, bool* reportInvertibility)
+{
+    Matrix2x2<Real> inverse;
+    bool invertible;
+    Real det = M(0, 0)*M(1, 1) - M(0, 1)*M(1, 0);
+    if (det != (Real)0)
+    {
+        Real invDet = ((Real)1) / det;
+        inverse = Matrix2x2<Real>
+        {
+            M(1, 1)*invDet, -M(0, 1)*invDet,
+                -M(1, 0)*invDet, M(0, 0)*invDet
+        };
+        invertible = true;
+    }
+    else
+    {
+        inverse.MakeZero();
+        invertible = false;
+    }
+
+    if (reportInvertibility)
+    {
+        *reportInvertibility = invertible;
+    }
+    return inverse;
+}
+
+template <typename Real>
+Matrix2x2<Real> Adjoint(Matrix2x2<Real> const& M)
+{
+    return Matrix2x2<Real>
+    {
+        M(1, 1), -M(0, 1),
+            -M(1, 0), M(0, 0)
+    };
+}
+
+template <typename Real>
+Real Determinant(Matrix2x2<Real> const& M)
+{
+    Real det = M(0, 0)*M(1, 1) - M(0, 1)*M(1, 0);
+    return det;
+}
+
+template <typename Real>
+Real Trace(Matrix2x2<Real> const& M)
+{
+    Real trace = M(0, 0) + M(1, 1);
+    return trace;
+}
+
+
+}
diff --git a/Engine/Plugins/Experimental/GeometryProcessing/Source/GeometryAlgorithms/Private/GTEngine/Mathematics/GteMatrix3x3.h b/Engine/Plugins/Experimental/GeometryProcessing/Source/GeometryAlgorithms/Private/GTEngine/Mathematics/GteMatrix3x3.h
new file mode 100644
--- /dev/null
+++ b/Engine/Plugins/Experimental/GeometryProcessing/Source/GeometryAlgorithms/Private/GTEngine/Mathematics/GteMatrix3x3.h
@@ -0,0 +1,109 @@
+// David Eberly, Geometric Tools, Redmond WA 98052
+// Copyright (c) 1998-2019
+// Distributed under the Boost Software License, Version 1.0.
+// http://www.boost.org/LICENSE_1_0.txt
+// http://www.geometrictools.com/License/Boost/LICENSE_1_0.txt
+// File Version: 3.0.0 (2016/06/19)
+
+#pragma once
+
+#include <GTEngine/Mathematics/GteMatrix.h>
+#include <GTEngine/Mathematics/GteVector3.h>
+
+namespace gte
+{
+
+// Template alias for convenience.
+template <typename Real>
+using Matrix3x3 = Matrix<3, 3, Real>;
+
+// Geometric operations.
+template <typename Real>
+Matrix3x3<Real> Inverse(Matrix3x3<Real> const& M,
+    bool* reportInvertibility = nullptr);
+
+template <typename Real>
+Matrix3x3<Real> Adjoint(Matrix3x3<Real> const& M);
+
+template <typename Real>
+Real Determinant(Matrix3x3<Real> const& M);
+
+template <typename Real>
+Real Trace(Matrix3x3<Real> const& M);
+
+
+template <typename Real>
+Matrix3x3<Real> Inverse(Matrix3x3<Real> const& M, bool* reportInvertibility)
+{
+    Matrix3x3<Real> inverse;
+    bool invertible;
+    Real c00 = M(1, 1)*M(2, 2) - M(1, 2)*M(2, 1);
+    Real c10 = M(1, 2)*M(2, 0) - M(1, 0)*M(2, 2);
+    Real c20 = M(1, 0)*M(2, 1) - M(1, 1)*M(2, 0);
+    Real det = M(0, 0)*c00 + M(0, 1)*c10 + M(0, 2)*c20;
+    if (det != (Real)0)
+    {
+        Real invDet = ((Real)1) / det;
+        inverse = Matrix3x3<Real>
+        {
+            c00*invDet,
+                (M(0, 2)*M(2, 1) - M(0, 1)*M(2, 2))*invDet,
+                (M(0, 1)*M(1, 2) - M(0, 2)*M(1, 1))*invDet,
+                c10*invDet,
+                (M(0, 0)*M(2, 2) - M(0, 2)*M(2, 0))*invDet,
+                (M(0, 2)*M(1, 0) - M(0, 0)*M(1, 2))*invDet,
+                c20*invDet,
+                (M(0, 1)*M(2, 0) - M(0, 0)*M(2, 1))*invDet,
+                (M(0, 0)*M(1, 1) - M(0, 1)*M(1, 0))*invDet
+        };
+        invertible = true;
+    }
+    else
+    {
+        inverse.MakeZero();
+        invertible = false;
+    }
+
+    if (reportInvertibility)
+    {
+        *reportInvertibility = invertible;
+    }
+    return inverse;
+}
+
+template <typename Real>
+Matrix3x3<Real> Adjoint(Matrix3x3<Real> const& M)
+{
+    return Matrix3x3<Real>
+    {
+        M(1, 1)*M(2, 2) - M(1, 2)*M(2, 1),
+            M(0, 2)*M(2, 1) - M(0, 1)*M(2, 2),
+            M(0, 1)*M(1, 2) - M(0, 2)*M(1, 1),
+            M(1, 2)*M(2, 0) - M(1, 0)*M(2, 2),
+            M(0, 0)*M(2, 2) - M(0, 2)*M(2, 0),
+            M(0, 2)*M(1, 0) - M(0, 0)*M(1, 2),
+            M(1, 0)*M(2, 1) - M(1, 1)*M(2, 0),
+            M(0, 1)*M(2, 0) - M(0, 0)*M(2, 1),
+            M(0, 0)*M(1, 1) - M(0, 1)*M(1, 0)
+    };
+}
+
+template <typename Real>
+Real Determinant(Matrix3x3<Real> const& M)
+{
+    Real c00 = M(1, 1)*M(2, 2) - M(1, 2)*M(2, 1);
+    Real c10 = M(1, 2)*M(2, 0) - M(1, 0)*M(2, 2);
+    Real c20 = M(1, 0)*M(2, 1) - M(1, 1)*M(2, 0);
+    Real det = M(0, 0)*c00 + M(0, 1)*c10 + M(0, 2)*c20;
+    return det;
+}
+
+template <typename Real>
+Real Trace(Matrix3x3<Real> const& M)
+{
+    Real trace = M(0, 0) + M(1, 1) + M(2, 2);
+    return trace;
+}
+
+
+}
diff --git a/Engine/Plugins/Experimental/GeometryProcessing/Source/GeometryAlgorithms/Private/GTEngine/Mathematics/GteMatrix4x4.h b/Engine/Plugins/Experimental/GeometryProcessing/Source/GeometryAlgorithms/Private/GTEngine/Mathematics/GteMatrix4x4.h
new file mode 100644
--- /dev/null
+++ b/Engine/Plugins/Experimental/GeometryProcessing/Source/GeometryAlgorithms/Private/GTEngine/Mathematics/GteMatrix4x4.h
@@ -0,0 +1,349 @@
+// David Eberly, Geometric Tools, Redmond WA 98052
+// Copyright (c) 1998-2019
+// Distributed under the Boost Software License, Version 1.0.
+// http://www.boost.org/LICENSE_1_0.txt
+// http://www.geometrictools.com/License/Boost/LICENSE_1_0.txt
+// File Version: 3.0.0 (2016/06/19)
+
+#pragma once
+
+#include <GTEngine/Mathematics/GteMatrix.h>
+#include <GTEngine/Mathematics/GteVector4.h>
+
+namespace gte
+{
+
+// Template alias for convenience.
+template <typename Real>
+using Matrix4x4 = Matrix<4, 4, Real>;
+
+// Geometric operations.
+template <typename Real>
+Matrix4x4<Real> Inverse(Matrix4x4<Real> const& M,
+    bool* reportInvertibility = nullptr);
+
+template <typename Real>
+Matrix4x4<Real> Adjoint(Matrix4x4<Real> const& M);
+
+template <typename Real>
+Real Determinant(Matrix4x4<Real> const& M);
+
+template <typename Real>
+Real Trace(Matrix4x4<Real> const& M);
+
+// Special matrices.  In the comments, the matrices are shown using the
+// GTE_USE_MAT_VEC multiplication convention.
+
+// The projection plane is Dot(N,X-P) = 0 where N is a 3-by-1 unit-length
+// normal vector and P is a 3-by-1 point on the plane.  The projection is
+// oblique to the plane, in the direction of the 3-by-1 vector D.  Necessarily
+// Dot(N,D) is not zero for this projection to make sense.  Given a 3-by-1
+// point U, compute the intersection of the line U+t*D with the plane to
+// obtain t = -Dot(N,U-P)/Dot(N,D); then
+//
+//   projection(U) = P + [I - D*N^T/Dot(N,D)]*(U-P)
+//
+// A 4-by-4 homogeneous transformation representing the projection is
+//
+//       +-                               -+
+//   M = | D*N^T - Dot(N,D)*I   -Dot(N,P)D |
+//       |          0^T          -Dot(N,D) |
+//       +-                               -+
+//
+// where M applies to [U^T 1]^T by M*[U^T 1]^T.  The matrix is chosen so
+// that M[3][3] > 0 whenever Dot(N,D) < 0; the projection is onto the
+// "positive side" of the plane.
+template <typename Real>
+Matrix4x4<Real> MakeObliqueProjection(Vector4<Real> const& origin,
+    Vector4<Real> const& normal, Vector4<Real> const& direction);
+
+// The perspective projection of a point onto a plane is
+//
+//     +-                                                 -+
+// M = | Dot(N,E-P)*I - E*N^T    -(Dot(N,E-P)*I - E*N^T)*E |
+//     |        -N^t                      Dot(N,E)         |
+//     +-                                                 -+
+//
+// where E is the eye point, P is a point on the plane, and N is a
+// unit-length plane normal.
+template <typename Real>
+Matrix4x4<Real> MakePerspectiveProjection(Vector4<Real> const& origin,
+    Vector4<Real> const& normal, Vector4<Real> const& eye);
+
+// The reflection of a point through a plane is
+//     +-                         -+
+// M = | I-2*N*N^T    2*Dot(N,P)*N |
+//     |     0^T            1      |
+//     +-                         -+
+//
+// where P is a point on the plane and N is a unit-length plane normal.
+
+template <typename Real>
+Matrix4x4<Real> MakeReflection(Vector4<Real> const& origin,
+    Vector4<Real> const& normal);
+
+
+template <typename Real>
+Matrix4x4<Real> Inverse(Matrix4x4<Real> const& M, bool* reportInvertibility)
+{
+    Matrix4x4<Real> inverse;
+    bool invertible;
+    Real a0 = M(0, 0) * M(1, 1) - M(0, 1) * M(1, 0);
+    Real a1 = M(0, 0) * M(1, 2) - M(0, 2) * M(1, 0);
+    Real a2 = M(0, 0) * M(1, 3) - M(0, 3) * M(1, 0);
+    Real a3 = M(0, 1) * M(1, 2) - M(0, 2) * M(1, 1);
+    Real a4 = M(0, 1) * M(1, 3) - M(0, 3) * M(1, 1);
+    Real a5 = M(0, 2) * M(1, 3) - M(0, 3) * M(1, 2);
+    Real b0 = M(2, 0) * M(3, 1) - M(2, 1) * M(3, 0);
+    Real b1 = M(2, 0) * M(3, 2) - M(2, 2) * M(3, 0);
+    Real b2 = M(2, 0) * M(3, 3) - M(2, 3) * M(3, 0);
+    Real b3 = M(2, 1) * M(3, 2) - M(2, 2) * M(3, 1);
+    Real b4 = M(2, 1) * M(3, 3) - M(2, 3) * M(3, 1);
+    Real b5 = M(2, 2) * M(3, 3) - M(2, 3) * M(3, 2);
+    Real det = a0 * b5 - a1 * b4 + a2 * b3 + a3 * b2 - a4 * b1 + a5 * b0;
+    if (det != (Real)0)
+    {
+        Real invDet = ((Real)1) / det;
+        inverse = Matrix4x4<Real>
+        {
+            (+M(1, 1) * b5 - M(1, 2) * b4 + M(1, 3) * b3) * invDet,
+            (-M(0, 1) * b5 + M(0, 2) * b4 - M(0, 3) * b3) * invDet,
+            (+M(3, 1) * a5 - M(3, 2) * a4 + M(3, 3) * a3) * invDet,
+            (-M(2, 1) * a5 + M(2, 2) * a4 - M(2, 3) * a3) * invDet,
+            (-M(1, 0) * b5 + M(1, 2) * b2 - M(1, 3) * b1) * invDet,
+            (+M(0, 0) * b5 - M(0, 2) * b2 + M(0, 3) * b1) * invDet,
+            (-M(3, 0) * a5 + M(3, 2) * a2 - M(3, 3) * a1) * invDet,
+            (+M(2, 0) * a5 - M(2, 2) * a2 + M(2, 3) * a1) * invDet,
+            (+M(1, 0) * b4 - M(1, 1) * b2 + M(1, 3) * b0) * invDet,
+            (-M(0, 0) * b4 + M(0, 1) * b2 - M(0, 3) * b0) * invDet,
+            (+M(3, 0) * a4 - M(3, 1) * a2 + M(3, 3) * a0) * invDet,
+            (-M(2, 0) * a4 + M(2, 1) * a2 - M(2, 3) * a0) * invDet,
+            (-M(1, 0) * b3 + M(1, 1) * b1 - M(1, 2) * b0) * invDet,
+            (+M(0, 0) * b3 - M(0, 1) * b1 + M(0, 2) * b0) * invDet,
+            (-M(3, 0) * a3 + M(3, 1) * a1 - M(3, 2) * a0) * invDet,
+            (+M(2, 0) * a3 - M(2, 1) * a1 + M(2, 2) * a0) * invDet
+        };
+        invertible = true;
+    }
+    else
+    {
+        inverse.MakeZero();
+        invertible = false;
+    }
+
+    if (reportInvertibility)
+    {
+        *reportInvertibility = invertible;
+    }
+    return inverse;
+}
+
+template <typename Real>
+Matrix4x4<Real> Adjoint(Matrix4x4<Real> const& M)
+{
+    Real a0 = M(0, 0) * M(1, 1) - M(0, 1) * M(1, 0);
+    Real a1 = M(0, 0) * M(1, 2) - M(0, 2) * M(1, 0);
+    Real a2 = M(0, 0) * M(1, 3) - M(0, 3) * M(1, 0);
+    Real a3 = M(0, 1) * M(1, 2) - M(0, 2) * M(1, 1);
+    Real a4 = M(0, 1) * M(1, 3) - M(0, 3) * M(1, 1);
+    Real a5 = M(0, 2) * M(1, 3) - M(0, 3) * M(1, 2);
+    Real b0 = M(2, 0) * M(3, 1) - M(2, 1) * M(3, 0);
+    Real b1 = M(2, 0) * M(3, 2) - M(2, 2) * M(3, 0);
+    Real b2 = M(2, 0) * M(3, 3) - M(2, 3) * M(3, 0);
+    Real b3 = M(2, 1) * M(3, 2) - M(2, 2) * M(3, 1);
+    Real b4 = M(2, 1) * M(3, 3) - M(2, 3) * M(3, 1);
+    Real b5 = M(2, 2) * M(3, 3) - M(2, 3) * M(3, 2);
+
+    return Matrix4x4<Real>
+    {
+        +M(1, 1) * b5 - M(1, 2) * b4 + M(1, 3) * b3,
+        -M(0, 1) * b5 + M(0, 2) * b4 - M(0, 3) * b3,
+        +M(3, 1) * a5 - M(3, 2) * a4 + M(3, 3) * a3,
+        -M(2, 1) * a5 + M(2, 2) * a4 - M(2, 3) * a3,
+        -M(1, 0) * b5 + M(1, 2) * b2 - M(1, 3) * b1,
+        +M(0, 0) * b5 - M(0, 2) * b2 + M(0, 3) * b1,
+        -M(3, 0) * a5 + M(3, 2) * a2 - M(3, 3) * a1,
+        +M(2, 0) * a5 - M(2, 2) * a2 + M(2, 3) * a1,
+        +M(1, 0) * b4 - M(1, 1) * b2 + M(1, 3) * b0,
+        -M(0, 0) * b4 + M(0, 1) * b2 - M(0, 3) * b0,
+        +M(3, 0) * a4 - M(3, 1) * a2 + M(3, 3) * a0,
+        -M(2, 0) * a4 + M(2, 1) * a2 - M(2, 3) * a0,
+        -M(1, 0) * b3 + M(1, 1) * b1 - M(1, 2) * b0,
+        +M(0, 0) * b3 - M(0, 1) * b1 + M(0, 2) * b0,
+        -M(3, 0) * a3 + M(3, 1) * a1 - M(3, 2) * a0,
+        +M(2, 0) * a3 - M(2, 1) * a1 + M(2, 2) * a0
+    };
+}
+
+template <typename Real>
+Real Determinant(Matrix4x4<Real> const& M)
+{
+    Real a0 = M(0, 0) * M(1, 1) - M(0, 1) * M(1, 0);
+    Real a1 = M(0, 0) * M(1, 2) - M(0, 2) * M(1, 0);
+    Real a2 = M(0, 0) * M(1, 3) - M(0, 3) * M(1, 0);
+    Real a3 = M(0, 1) * M(1, 2) - M(0, 2) * M(1, 1);
+    Real a4 = M(0, 1) * M(1, 3) - M(0, 3) * M(1, 1);
+    Real a5 = M(0, 2) * M(1, 3) - M(0, 3) * M(1, 2);
+    Real b0 = M(2, 0) * M(3, 1) - M(2, 1) * M(3, 0);
+    Real b1 = M(2, 0) * M(3, 2) - M(2, 2) * M(3, 0);
+    Real b2 = M(2, 0) * M(3, 3) - M(2, 3) * M(3, 0);
+    Real b3 = M(2, 1) * M(3, 2) - M(2, 2) * M(3, 1);
+    Real b4 = M(2, 1) * M(3, 3) - M(2, 3) * M(3, 1);
+    Real b5 = M(2, 2) * M(3, 3) - M(2, 3) * M(3, 2);
+    Real det = a0 * b5 - a1 * b4 + a2 * b3 + a3 * b2 - a4 * b1 + a5 * b0;
+    return det;
+}
+
+template <typename Real>
+Real Trace(Matrix4x4<Real> const& M)
+{
+    Real trace = M(0, 0) + M(1, 1) + M(2, 2) + M(3, 3);
+    return trace;
+}
+
+template <typename Real>
+Matrix4x4<Real> MakeObliqueProjection(Vector4<Real> const& origin,
+    Vector4<Real> const& normal, Vector4<Real> const& direction)
+{
+    Matrix4x4<Real> M;
+
+    Real const zero = (Real)0;
+    Real dotND = Dot(normal, direction);
+    Real dotNO = Dot(origin, normal);
+
+#if defined(GTE_USE_MAT_VEC)
+    M(0, 0) = direction[0] * normal[0] - dotND;
+    M(0, 1) = direction[0] * normal[1];
+    M(0, 2) = direction[0] * normal[2];
+    M(0, 3) = -dotNO * direction[0];
+    M(1, 0) = direction[1] * normal[0];
+    M(1, 1) = direction[1] * normal[1] - dotND;
+    M(1, 2) = direction[1] * normal[2];
+    M(1, 3) = -dotNO * direction[1];
+    M(2, 0) = direction[2] * normal[0];
+    M(2, 1) = direction[2] * normal[1];
+    M(2, 2) = direction[2] * normal[2] - dotND;
+    M(2, 3) = -dotNO * direction[2];
+    M(3, 0) = zero;
+    M(3, 1) = zero;
+    M(3, 2) = zero;
+    M(3, 3) = -dotND;
+#else
+    M(0, 0) = direction[0] * normal[0] - dotND;
+    M(1, 0) = direction[0] * normal[1];
+    M(2, 0) = direction[0] * normal[2];
+    M(3, 0) = -dotNO * direction[0];
+    M(0, 1) = direction[1] * normal[0];
+    M(1, 1) = direction[1] * normal[1] - dotND;
+    M(2, 1) = direction[1] * normal[2];
+    M(3, 1) = -dotNO * direction[1];
+    M(0, 2) = direction[2] * normal[0];
+    M(1, 2) = direction[2] * normal[1];
+    M(2, 2) = direction[2] * normal[2] - dotND;
+    M(3, 2) = -dotNO * direction[2];
+    M(0, 2) = zero;
+    M(1, 3) = zero;
+    M(2, 3) = zero;
+    M(3, 3) = -dotND;
+#endif
+
+    return M;
+}
+
+template <typename Real>
+Matrix4x4<Real> MakePerspectiveProjection(Vector4<Real> const& origin,
+    Vector4<Real> const& normal, Vector4<Real> const& eye)
+{
+    Matrix4x4<Real> M;
+
+    Real dotND = Dot(normal, eye - origin);
+
+#if defined(GTE_USE_MAT_VEC)
+    M(0, 0) = dotND - eye[0] * normal[0];
+    M(0, 1) = -eye[0] * normal[1];
+    M(0, 2) = -eye[0] * normal[2];
+    M(0, 3) = -(M(0, 0) * eye[0] + M(0, 1) * eye[1] + M(0, 2) * eye[2]);
+    M(1, 0) = -eye[1] * normal[0];
+    M(1, 1) = dotND - eye[1] * normal[1];
+    M(1, 2) = -eye[1] * normal[2];
+    M(1, 3) = -(M(1, 0) * eye[0] + M(1, 1) * eye[1] + M(1, 2) * eye[2]);
+    M(2, 0) = -eye[2] * normal[0];
+    M(2, 1) = -eye[2] * normal[1];
+    M(2, 2) = dotND - eye[2] * normal[2];
+    M(2, 3) = -(M(2, 0) * eye[0] + M(2, 1) * eye[1] + M(2, 2) * eye[2]);
+    M(3, 0) = -normal[0];
+    M(3, 1) = -normal[1];
+    M(3, 2) = -normal[2];
+    M(3, 3) = Dot(eye, normal);
+#else
+    M(0, 0) = dotND - eye[0] * normal[0];
+    M(1, 0) = -eye[0] * normal[1];
+    M(2, 0) = -eye[0] * normal[2];
+    M(3, 0) = -(M(0, 0) * eye[0] + M(0, 1) * eye[1] + M(0, 2) * eye[2]);
+    M(0, 1) = -eye[1] * normal[0];
+    M(1, 1) = dotND - eye[1] * normal[1];
+    M(2, 1) = -eye[1] * normal[2];
+    M(3, 1) = -(M(1, 0) * eye[0] + M(1, 1) * eye[1] + M(1, 2) * eye[2]);
+    M(0, 2) = -eye[2] * normal[0];
+    M(1, 2) = -eye[2] * normal[1];
+    M(2, 2) = dotND - eye[2] * normal[2];
+    M(3, 2) = -(M(2, 0) * eye[0] + M(2, 1) * eye[1] + M(2, 2) * eye[2]);
+    M(0, 3) = -normal[0];
+    M(1, 3) = -normal[1];
+    M(2, 3) = -normal[2];
+    M(3, 3) = Dot(eye, normal);
+#endif
+
+    return M;
+}
+
+template <typename Real>
+Matrix4x4<Real> MakeReflection(Vector4<Real> const& origin,
+    Vector4<Real> const& normal)
+{
+    Matrix4x4<Real> M;
+
+    Real const zero = (Real)0, one = (Real)1, two = (Real)2;
+    Real twoDotNO = two * Dot(origin, normal);
+
+#if defined(GTE_USE_MAT_VEC)
+    M(0, 0) = one - two * normal[0] * normal[0];
+    M(0, 1) = -two * normal[0] * normal[1];
+    M(0, 2) = -two * normal[0] * normal[2];
+    M(0, 3) = twoDotNO * normal[0];
+    M(1, 0) = M(0, 1);
+    M(1, 1) = one - two * normal[1] * normal[1];
+    M(1, 2) = -two * normal[1] * normal[2];
+    M(1, 3) = twoDotNO * normal[1];
+    M(2, 0) = M(0, 2);
+    M(2, 1) = M(1, 2);
+    M(2, 2) = one - two * normal[2] * normal[2];
+    M(2, 3) = twoDotNO * normal[2];
+    M(3, 0) = zero;
+    M(3, 1) = zero;
+    M(3, 2) = zero;
+    M(3, 3) = one;
+#else
+    M(0, 0) = one - two * normal[0] * normal[0];
+    M(1, 0) = -two * normal[0] * normal[1];
+    M(2, 0) = -two * normal[0] * normal[2];
+    M(3, 0) = twoDotNO * normal[0];
+    M(0, 1) = M(1, 0);
+    M(1, 1) = one - two * normal[1] * normal[1];
+    M(2, 1) = -two * normal[1] * normal[2];
+    M(3, 1) = twoDotNO * normal[1];
+    M(0, 2) = M(2, 0);
+    M(1, 2) = M(2, 1);
+    M(2, 2) = one - two * normal[2] * normal[2];
+    M(3, 2) = twoDotNO * normal[2];
+    M(0, 3) = zero;
+    M(1, 3) = zero;
+    M(2, 3) = zero;
+    M(3, 3) = one;
+#endif
+
+    return M;
+}
+
+
+}
diff --git a/Engine/Plugins/Experimental/GeometryProcessing/Source/GeometryAlgorithms/Private/GTEngine/Mathematics/GteMesh.h b/Engine/Plugins/Experimental/GeometryProcessing/Source/GeometryAlgorithms/Private/GTEngine/Mathematics/GteMesh.h
new file mode 100644
--- /dev/null
+++ b/Engine/Plugins/Experimental/GeometryProcessing/Source/GeometryAlgorithms/Private/GTEngine/Mathematics/GteMesh.h
@@ -0,0 +1,747 @@
+// David Eberly, Geometric Tools, Redmond WA 98052
+// Copyright (c) 1998-2019
+// Distributed under the Boost Software License, Version 1.0.
+// http://www.boost.org/LICENSE_1_0.txt
+// http://www.geometrictools.com/License/Boost/LICENSE_1_0.txt
+// File Version: 3.3.1 (2018/10/05)
+
+#pragma once
+
+#include <GTEngine/LowLevel/GteLogger.h>
+#include <GTEngine/Mathematics/GteMath.h>
+#include <GTEngine/Mathematics/GteVector2.h>
+#include <GTEngine/Mathematics/GteVector3.h>
+#include <GTEngine/Mathematics/GteMatrix.h>
+#include <GTEngine/Mathematics/GteIndexAttribute.h>
+#include <GTEngine/Mathematics/GteVertexAttribute.h>
+
+// The Mesh class is designed to support triangulations of surfaces of a small
+// number of topologies. See the documents
+//   http://www.geometrictools.com/MeshDifferentialGeometry.pdf
+//   http://www.geometrictools.com/MeshFactory.pdf
+// for details.
+//
+// You must set the vertex attribute sources before calling Update().
+//
+// The semantic "position" is required and its source must be an array
+// of Real with at least 3 channels so that positions are computed as
+// Vector3<Real>.
+//
+// The positions are assumed to be parameterized by texture coordinates
+// (u,v); the position is thought of as a function P(u,v).  If texture
+// coordinates are provided, the semantic must be "tcoord".  If texture
+// coordinates are not provided, default texture coordinates are computed
+// internally as described in the mesh factory document.
+//
+// The frame for the tangent space is optional.  All vectors in the frame
+// must have sources that are arrays of Real with at least 3 channels per
+// attribute.  If normal vectors are provided, the semantic must be
+// "normal".
+//
+// Two options are supported for tangent vectors.  The first option is
+// that the tangents are surface derivatives dP/du and dP/dv, which are
+// not necessarily unit length or orthogonal.  The semantics must be
+// "dpdu" and "dpdv".  The second option is that the tangents are unit
+// length and orthogonal, with the infrequent possibility that a vertex
+// is degenerate in that dP/du and dP/dv are linearly dependent.  The
+// semantics must be "tangent" and "bitangent".
+//
+// For each provided vertex attribute, a derived class can initialize
+// that attribute by overriding one of the Initialize*() functions whose
+// stubs are defined in this class.
+
+namespace gte
+{
+
+enum class MeshTopology
+{
+    ARBITRARY,
+    RECTANGLE,
+    CYLINDER,
+    TORUS,
+    DISK,
+    SPHERE
+};
+
+class MeshDescription
+{
+public:
+    // Constructor for MeshTopology::ARBITRARY.  The members topology,
+    // numVertices, and numTriangles are set in the obvious manner.  The
+    // members numRows and numCols are set to zero. The remaining members
+    // must be set explicitly by the client.
+    inline MeshDescription(uint32_t inNumVertices, uint32_t inNumTriangles);
+
+    // Constructor for topologies other than MeshTopology::ARBITRARY.
+    // Compute the number of vertices and triangles for the mesh based on the
+    // requested number of rows and columns.  If the number of rows or columns
+    // is invalid for the specified topology, they are modified to be valid,
+    // in which case inNumRows/numRows and inNumCols/numCols can differ.  If
+    // the input topology is MeshTopology::ARBITRARY, then inNumRows and
+    // inNumCols are assigned to numVertices and numTriangles, respectively,
+    // and numRows and numCols are set to zero.  The remaining members must be
+    // set explicitly by the client.
+    inline MeshDescription(MeshTopology inTopology, uint32_t inNumRows, uint32_t inNumCols);
+
+    MeshTopology topology;
+    uint32_t numVertices;
+    uint32_t numTriangles;
+    std::vector<VertexAttribute> vertexAttributes;
+    IndexAttribute indexAttribute;
+    bool wantDynamicTangentSpaceUpdate;  // default: false
+    bool wantCCW;  // default: true
+
+    // For internal use only.
+    bool hasTangentSpaceVectors;
+    bool allowUpdateFrame;
+    uint32_t numRows, numCols;
+    uint32_t rMax, cMax, rIncrement;
+
+    // After an attempt to construct a Mesh or Mesh-derived object, examine
+    // this value to determine whether the construction was successful.
+    bool constructed;
+};
+
+template <typename Real>
+class Mesh
+{
+public:
+    // Construction and destruction.  This constructor is for ARBITRARY topology.
+    // The vertices and indices must already be assigned by the client.  Derived
+    // classes use the protected constructor, but assignment of vertices and
+    // indices occurs in the derived-class constructors.
+    Mesh(MeshDescription const& description, std::vector<MeshTopology> const& validTopologies);
+
+    virtual ~Mesh();
+
+    // No copying or assignment is allowed.
+    Mesh(Mesh const&) = delete;
+    Mesh& operator=(Mesh const&) = delete;
+
+    // Member accessors.
+    inline MeshDescription const& GetDescription() const;
+
+    // If the underlying geometric data varies dynamically, call this function
+    // to update whatever vertex attributes are specified by the vertex pool.
+    // A derived class
+    void Update();
+
+protected:
+    // Access the vertex attributes.
+    inline Vector3<Real>& Position(uint32_t i);
+    inline Vector3<Real>& Normal(uint32_t i);
+    inline Vector3<Real>& Tangent(uint32_t i);
+    inline Vector3<Real>& Bitangent(uint32_t i);
+    inline Vector3<Real>& DPDU(uint32_t i);
+    inline Vector3<Real>& DPDV(uint32_t i);
+    inline Vector2<Real>& TCoord(uint32_t i);
+
+    // Compute the indices for non-arbitrary topologies.  This function is
+    // called by derived classes.
+    void ComputeIndices();
+
+    // The Update() function allows derived classes to use algorithms
+    // different from least-squares fitting to compute the normals (when
+    // no tangent-space information is requested) or to compute the frame
+    // (normals and tangent space).  The UpdatePositions() is a stub; the
+    // base-class has no knowledge about how positions should be modified.
+    // A derived class, however, might choose to use dynamic updating
+    // and override UpdatePositions().  The base-class UpdateNormals()
+    // computes vertex normals as averages of area-weighted triangle
+    // normals (nonparametric approach).  The base-class UpdateFrame()
+    // uses a least-squares algorithm for estimating the tangent space
+    // (parametric approach).
+    virtual void UpdatePositions() {}
+    virtual void UpdateNormals();
+    virtual void UpdateFrame();
+
+    // Constructor inputs.
+    // The client requests this via the constructor; however, if it is
+    // requested and the vertex attributes do not contain entries for
+    // "tangent", "bitangent", "dpdu", or "dpdv", then this member is
+    // set to false.
+    MeshDescription mDescription;
+
+    // Copied from mVertexAttributes when available.
+    Vector3<Real>* mPositions;
+    Vector3<Real>* mNormals;
+    Vector3<Real>* mTangents;
+    Vector3<Real>* mBitangents;
+    Vector3<Real>* mDPDUs;
+    Vector3<Real>* mDPDVs;
+    Vector2<Real>* mTCoords;
+    size_t mPositionStride;
+    size_t mNormalStride;
+    size_t mTangentStride;
+    size_t mBitangentStride;
+    size_t mDPDUStride;
+    size_t mDPDVStride;
+    size_t mTCoordStride;
+
+    // When dynamic tangent-space updates are requested, the update algorithm
+    // requires texture coordinates (user-specified or non-local).  It is
+    // possible to create a vertex-adjacent set (with indices into the
+    // vertex array) for each mesh vertex; however, instead we rely on a
+    // triangle iteration and incrementally store the information needed for
+    // the estimation of the tangent space.  Each vertex has associated
+    // matrices D and U, but we need to store only U^T*U and D^T*U.  See the
+    // PDF for details.
+    std::vector<Matrix<2, 2, Real>> mUTU;
+    std::vector<Matrix<3, 2, Real>> mDTU;
+};
+
+
+inline MeshDescription::MeshDescription(uint32_t inNumVertices, uint32_t inNumTriangles)
+    :
+    topology(MeshTopology::ARBITRARY),
+    numVertices(inNumVertices),
+    numTriangles(inNumTriangles),
+    wantDynamicTangentSpaceUpdate(false),
+    wantCCW(true),
+    hasTangentSpaceVectors(false),
+    allowUpdateFrame(false),
+    numRows(0),
+    numCols(0),
+    rMax(0),
+    cMax(0),
+    rIncrement(0),
+    constructed(false)
+{
+    LogAssert(numVertices >= 3, "Invalid number of vertices.");
+    LogAssert(numTriangles >= 1, "Invalid number of triangles.");
+}
+
+inline MeshDescription::MeshDescription(MeshTopology inTopology, uint32_t inNumRows, uint32_t inNumCols)
+    :
+    topology(inTopology),
+    wantDynamicTangentSpaceUpdate(false),
+    wantCCW(true),
+    hasTangentSpaceVectors(false),
+    allowUpdateFrame(false),
+    constructed(false)
+{
+    switch (topology)
+    {
+    case MeshTopology::ARBITRARY:
+        numVertices = inNumRows;
+        numTriangles = inNumCols;
+        numRows = 0;
+        numCols = 0;
+        rMax = 0;
+        cMax = 0;
+        rIncrement = 0;
+        break;
+
+    case MeshTopology::RECTANGLE:
+        numRows = std::max(inNumRows, 2u);
+        numCols = std::max(inNumCols, 2u);
+        rMax = numRows - 1;
+        cMax = numCols - 1;
+        rIncrement = numCols;
+        numVertices = (rMax + 1) * (cMax + 1);
+        numTriangles = 2 * rMax * cMax;
+        break;
+
+    case MeshTopology::CYLINDER:
+        numRows = std::max(inNumRows, 2u);
+        numCols = std::max(inNumCols, 3u);
+        rMax = numRows - 1;
+        cMax = numCols;
+        rIncrement = numCols + 1;
+        numVertices = (rMax + 1) * (cMax + 1);
+        numTriangles = 2 * rMax * cMax;
+        break;
+
+    case MeshTopology::TORUS:
+        numRows = std::max(inNumRows, 2u);
+        numCols = std::max(inNumCols, 3u);
+        rMax = numRows;
+        cMax = numCols;
+        rIncrement = numCols + 1;
+        numVertices = (rMax + 1) * (cMax + 1);
+        numTriangles = 2 * rMax * cMax;
+        break;
+
+    case MeshTopology::DISK:
+        numRows = std::max(inNumRows, 1u);
+        numCols = std::max(inNumCols, 3u);
+        rMax = numRows - 1;
+        cMax = numCols;
+        rIncrement = numCols + 1;
+        numVertices = (rMax + 1) * (cMax + 1) + 1;
+        numTriangles = 2 * rMax * cMax + numCols;
+        break;
+
+    case MeshTopology::SPHERE:
+        numRows = std::max(inNumRows, 1u);
+        numCols = std::max(inNumCols, 3u);
+        rMax = numRows - 1;
+        cMax = numCols;
+        rIncrement = numCols + 1;
+        numVertices = (rMax + 1) * (cMax + 1) + 2;
+        numTriangles = 2 * rMax * cMax + 2 * numCols;
+        break;
+    }
+}
+
+template <typename Real>
+Mesh<Real>::Mesh(MeshDescription const& description, std::vector<MeshTopology> const& validTopologies)
+    :
+    mDescription(description),
+    mPositions(nullptr),
+    mNormals(nullptr),
+    mTangents(nullptr),
+    mBitangents(nullptr),
+    mDPDUs(nullptr),
+    mDPDVs(nullptr),
+    mTCoords(nullptr),
+    mPositionStride(0),
+    mNormalStride(0),
+    mTangentStride(0),
+    mBitangentStride(0),
+    mDPDUStride(0),
+    mDPDVStride(0),
+    mTCoordStride(0)
+{
+    mDescription.constructed = false;
+    for (auto const& topology : validTopologies)
+    {
+        if (mDescription.topology == topology)
+        {
+            mDescription.constructed = true;
+            break;
+        }
+    }
+
+    if (!mDescription.indexAttribute.source)
+    {
+        LogError("The mesh needs triangles/indices.");
+        mDescription.constructed = false;
+        return;
+    }
+
+    // Set sources for the requested vertex attributes.
+    mDescription.hasTangentSpaceVectors = false;
+    mDescription.allowUpdateFrame = mDescription.wantDynamicTangentSpaceUpdate;
+    for (auto const& attribute : mDescription.vertexAttributes)
+    {
+        if (attribute.source != nullptr && attribute.stride > 0)
+        {
+            if (attribute.semantic == "position")
+            {
+                mPositions = reinterpret_cast<Vector3<Real>*>(attribute.source);
+                mPositionStride = attribute.stride;
+                continue;
+            }
+
+            if (attribute.semantic == "normal")
+            {
+                mNormals = reinterpret_cast<Vector3<Real>*>(attribute.source);
+                mNormalStride = attribute.stride;
+                continue;
+            }
+
+            if (attribute.semantic == "tangent")
+            {
+                mTangents = reinterpret_cast<Vector3<Real>*>(attribute.source);
+                mTangentStride = attribute.stride;
+                mDescription.hasTangentSpaceVectors = true;
+                continue;
+            }
+
+            if (attribute.semantic == "bitangent")
+            {
+                mBitangents = reinterpret_cast<Vector3<Real>*>(attribute.source);
+                mBitangentStride = attribute.stride;
+                mDescription.hasTangentSpaceVectors = true;
+                continue;
+            }
+
+            if (attribute.semantic == "dpdu")
+            {
+                mDPDUs = reinterpret_cast<Vector3<Real>*>(attribute.source);
+                mDPDUStride = attribute.stride;
+                mDescription.hasTangentSpaceVectors = true;
+                continue;
+            }
+
+            if (attribute.semantic == "dpdv")
+            {
+                mDPDVs = reinterpret_cast<Vector3<Real>*>(attribute.source);
+                mDPDVStride = attribute.stride;
+                mDescription.hasTangentSpaceVectors = true;
+                continue;
+            }
+
+            if (attribute.semantic == "tcoord")
+            {
+                mTCoords = reinterpret_cast<Vector2<Real>*>(attribute.source);
+                mTCoordStride = attribute.stride;
+                continue;
+            }
+        }
+    }
+
+    if (!mPositions)
+    {
+        LogError("The mesh needs positions.");
+        mPositions = nullptr;
+        mNormals = nullptr;
+        mTangents = nullptr;
+        mBitangents = nullptr;
+        mDPDUs = nullptr;
+        mDPDVs = nullptr;
+        mTCoords = nullptr;
+        mPositionStride = 0;
+        mNormalStride = 0;
+        mTangentStride = 0;
+        mBitangentStride = 0;
+        mDPDUStride = 0;
+        mDPDVStride = 0;
+        mTCoordStride = 0;
+        mDescription.constructed = false;
+        return;
+    }
+
+    // The initial value of allowUpdateFrame is the client request about
+    // wanting dynamic tangent-space updates.  If the vertex attributes do
+    // not include tangent-space vectors, then dynamic updates are not
+    // necessary.  If tangent-space vectors are present, the update algorithm
+    // requires texture coordinates (mTCoords must be nonnull) or must compute
+    // local coordinates (mNormals must be nonnull).
+    if (mDescription.allowUpdateFrame)
+    {
+        if (!mDescription.hasTangentSpaceVectors)
+        {
+            mDescription.allowUpdateFrame = false;
+        }
+
+        if (!mTCoords && !mNormals)
+        {
+            mDescription.allowUpdateFrame = false;
+        }
+    }
+
+    if (mDescription.allowUpdateFrame)
+    {
+        mUTU.resize(mDescription.numVertices);
+        mDTU.resize(mDescription.numVertices);
+    }
+}
+
+template <typename Real>
+Mesh<Real>::~Mesh()
+{
+}
+
+template <typename Real>
+inline MeshDescription const& Mesh<Real>::GetDescription() const
+{
+    return mDescription;
+}
+
+template <typename Real>
+void Mesh<Real>::Update()
+{
+    if (!mDescription.constructed)
+    {
+        LogError("The Mesh object failed the construction.");
+        return;
+    }
+
+    UpdatePositions();
+
+    if (mDescription.allowUpdateFrame)
+    {
+        UpdateFrame();
+    }
+    else if (mNormals)
+    {
+        UpdateNormals();
+    }
+    // else: The mesh has no frame data, so there is nothing to do.
+}
+
+template <typename Real>
+inline Vector3<Real>& Mesh<Real>::Position(uint32_t i)
+{
+    char* positions = reinterpret_cast<char*>(mPositions);
+    return *reinterpret_cast<Vector3<Real>*>(positions + i * mPositionStride);
+}
+
+template <typename Real>
+inline Vector3<Real>& Mesh<Real>::Normal(uint32_t i)
+{
+    char* normals = reinterpret_cast<char*>(mNormals);
+    return *reinterpret_cast<Vector3<Real>*>(normals + i * mNormalStride);
+}
+
+template <typename Real>
+inline Vector3<Real>& Mesh<Real>::Tangent(uint32_t i)
+{
+    char* tangents = reinterpret_cast<char*>(mTangents);
+    return *reinterpret_cast<Vector3<Real>*>(tangents + i * mTangentStride);
+}
+
+template <typename Real>
+inline Vector3<Real>& Mesh<Real>::Bitangent(uint32_t i)
+{
+    char* bitangents = reinterpret_cast<char*>(mBitangents);
+    return *reinterpret_cast<Vector3<Real>*>(bitangents + i * mBitangentStride);
+}
+
+template <typename Real>
+inline Vector3<Real>& Mesh<Real>::DPDU(uint32_t i)
+{
+    char* dpdus = reinterpret_cast<char*>(mDPDUs);
+    return *reinterpret_cast<Vector3<Real>*>(dpdus + i * mDPDUStride);
+}
+
+template <typename Real>
+inline Vector3<Real>& Mesh<Real>::DPDV(uint32_t i)
+{
+    char* dpdvs = reinterpret_cast<char*>(mDPDVs);
+    return *reinterpret_cast<Vector3<Real>*>(dpdvs + i * mDPDVStride);
+}
+
+template <typename Real>
+inline Vector2<Real>& Mesh<Real>::TCoord(uint32_t i)
+{
+    char* tcoords = reinterpret_cast<char*>(mTCoords);
+    return *reinterpret_cast<Vector2<Real>*>(tcoords + i * mTCoordStride);
+}
+
+template <typename Real>
+void Mesh<Real>::ComputeIndices()
+{
+    uint32_t t = 0;
+    for (uint32_t r = 0, i = 0; r < mDescription.rMax; ++r)
+    {
+        uint32_t v0 = i, v1 = v0 + 1;
+        i += mDescription.rIncrement;
+        uint32_t v2 = i, v3 = v2 + 1;
+        for (uint32_t c = 0; c < mDescription.cMax; ++c, ++v0, ++v1, ++v2, ++v3)
+        {
+            if (mDescription.wantCCW)
+            {
+                mDescription.indexAttribute.SetTriangle(t++, v0, v1, v2);
+                mDescription.indexAttribute.SetTriangle(t++, v1, v3, v2);
+            }
+            else
+            {
+                mDescription.indexAttribute.SetTriangle(t++, v0, v2, v1);
+                mDescription.indexAttribute.SetTriangle(t++, v1, v2, v3);
+            }
+        }
+    }
+
+    if (mDescription.topology == MeshTopology::DISK)
+    {
+        uint32_t v0 = 0, v1 = 1, v2 = mDescription.numVertices - 1;
+        for (unsigned int c = 0; c < mDescription.numCols; ++c, ++v0, ++v1)
+        {
+            if (mDescription.wantCCW)
+            {
+                mDescription.indexAttribute.SetTriangle(t++, v0, v2, v1);
+            }
+            else
+            {
+                mDescription.indexAttribute.SetTriangle(t++, v0, v1, v2);
+            }
+        }
+    }
+    else if (mDescription.topology == MeshTopology::SPHERE)
+    {
+        uint32_t v0 = 0, v1 = 1, v2 = mDescription.numVertices - 2;
+        for (uint32_t c = 0; c < mDescription.numCols; ++c, ++v0, ++v1)
+        {
+            if (mDescription.wantCCW)
+            {
+                mDescription.indexAttribute.SetTriangle(t++, v0, v2, v1);
+            }
+            else
+            {
+                mDescription.indexAttribute.SetTriangle(t++, v0, v1, v2);
+            }
+        }
+
+        v0 = (mDescription.numRows - 1) * mDescription.numCols;
+        v1 = v0 + 1;
+        v2 = mDescription.numVertices - 1;
+        for (uint32_t c = 0; c < mDescription.numCols; ++c, ++v0, ++v1)
+        {
+            if (mDescription.wantCCW)
+            {
+                mDescription.indexAttribute.SetTriangle(t++, v0, v2, v1);
+            }
+            else
+            {
+                mDescription.indexAttribute.SetTriangle(t++, v0, v1, v2);
+            }
+        }
+    }
+}
+
+template <typename Real>
+void Mesh<Real>::UpdateNormals()
+{
+    // Compute normal vector as normalized weighted averages of triangle
+    // normal vectors.
+
+    // Set the normals to zero to allow accumulation of triangle normals.
+    Vector3<Real> zero{ (Real)0, (Real)0, (Real)0 };
+    for (uint32_t i = 0; i < mDescription.numVertices; ++i)
+    {
+        Normal(i) = zero;
+    }
+
+    // Accumulate the triangle normals.
+    for (uint32_t t = 0; t < mDescription.numTriangles; ++t)
+    {
+        // Get the positions for the triangle.
+        uint32_t v0, v1, v2;
+        mDescription.indexAttribute.GetTriangle(t, v0, v1, v2);
+        Vector3<Real> P0 = Position(v0);
+        Vector3<Real> P1 = Position(v1);
+        Vector3<Real> P2 = Position(v2);
+
+        // Get the edge vectors.
+        Vector3<Real> E1 = P1 - P0;
+        Vector3<Real> E2 = P2 - P0;
+
+        // Compute a triangle normal show length is twice the area of the
+        // triangle.
+        Vector3<Real> triangleNormal = Cross(E1, E2);
+
+        // Accumulate the triangle normals.
+        Normal(v0) += triangleNormal;
+        Normal(v1) += triangleNormal;
+        Normal(v2) += triangleNormal;
+    }
+
+    // Normalize the normals.
+    for (uint32_t i = 0; i < mDescription.numVertices; ++i)
+    {
+        Normalize(Normal(i), true);
+    }
+}
+
+template <typename Real>
+void Mesh<Real>::UpdateFrame()
+{
+    if (!mTCoords)
+    {
+        // We need to compute vertex normals first in order to compute
+        // local texture coordinates.  The vertex normals are recomputed
+        // later based on estimated tangent vectors.
+        UpdateNormals();
+    }
+
+    // Use the least-squares algorithm to estimate the tangent-space vectors
+    // and, if requested, normal vectors.
+    Matrix<2, 2, Real> zero2x2;  // initialized to zero
+    Matrix<3, 2, Real> zero3x2;  // initialized to zero
+    std::fill(mUTU.begin(), mUTU.end(), zero2x2);
+    std::fill(mDTU.begin(), mDTU.end(), zero3x2);
+    for (uint32_t t = 0; t < mDescription.numTriangles; ++t)
+    {
+        // Get the positions and differences for the triangle.
+        uint32_t v0, v1, v2;
+        mDescription.indexAttribute.GetTriangle(t, v0, v1, v2);
+        Vector3<Real> P0 = Position(v0);
+        Vector3<Real> P1 = Position(v1);
+        Vector3<Real> P2 = Position(v2);
+        Vector3<Real> D10 = P1 - P0;
+        Vector3<Real> D20 = P2 - P0;
+        Vector3<Real> D21 = P2 - P1;
+
+        if (mTCoords)
+        {
+            // Get the texture coordinates and differences for the triangle.
+            Vector2<Real> C0 = TCoord(v0);
+            Vector2<Real> C1 = TCoord(v1);
+            Vector2<Real> C2 = TCoord(v2);
+            Vector2<Real> U10 = C1 - C0;
+            Vector2<Real> U20 = C2 - C0;
+            Vector2<Real> U21 = C2 - C1;
+
+            // Compute the outer products.
+            Matrix<2, 2, Real> outerU10 = OuterProduct(U10, U10);
+            Matrix<2, 2, Real> outerU20 = OuterProduct(U20, U20);
+            Matrix<2, 2, Real> outerU21 = OuterProduct(U21, U21);
+            Matrix<3, 2, Real> outerD10 = OuterProduct(D10, U10);
+            Matrix<3, 2, Real> outerD20 = OuterProduct(D20, U20);
+            Matrix<3, 2, Real> outerD21 = OuterProduct(D21, U21);
+
+            // Keep a running sum of U^T*U and D^T*U.
+            mUTU[v0] += outerU10 + outerU20;
+            mUTU[v1] += outerU10 + outerU21;
+            mUTU[v2] += outerU20 + outerU21;
+            mDTU[v0] += outerD10 + outerD20;
+            mDTU[v1] += outerD10 + outerD21;
+            mDTU[v2] += outerD20 + outerD21;
+        }
+        else
+        {
+            // Compute local coordinates and differences for the triangle.
+            Vector3<Real> basis[3];
+
+            basis[0] = Normal(v0);
+            ComputeOrthogonalComplement(1, basis, true);
+            Vector2<Real> U10{ Dot(basis[1], D10), Dot(basis[2], D10) };
+            Vector2<Real> U20{ Dot(basis[1], D20), Dot(basis[2], D20) };
+            mUTU[v0] += OuterProduct(U10, U10) + OuterProduct(U20, U20);
+            mDTU[v0] += OuterProduct(D10, U10) + OuterProduct(D20, U20);
+
+            basis[0] = Normal(v1);
+            ComputeOrthogonalComplement(1, basis, true);
+            Vector2<Real> U01{ Dot(basis[1], D10), Dot(basis[2], D10) };
+            Vector2<Real> U21{ Dot(basis[1], D21), Dot(basis[2], D21) };
+            mUTU[v1] += OuterProduct(U01, U01) + OuterProduct(U21, U21);
+            mDTU[v1] += OuterProduct(D10, U01) + OuterProduct(D21, U21);
+
+            basis[0] = Normal(v2);
+            ComputeOrthogonalComplement(1, basis, true);
+            Vector2<Real> U02{ Dot(basis[1], D20), Dot(basis[2], D20) };
+            Vector2<Real> U12{ Dot(basis[1], D21), Dot(basis[2], D21) };
+            mUTU[v2] += OuterProduct(U02, U02) + OuterProduct(U12, U12);
+            mDTU[v2] += OuterProduct(D20, U02) + OuterProduct(D21, U12);
+        }
+
+    }
+
+    for (uint32_t i = 0; i < mDescription.numVertices; ++i)
+    {
+        Matrix<3, 2, Real> jacobian = mDTU[i] * Inverse(mUTU[i]);
+
+        Vector3<Real> basis[3];
+        basis[0] = { jacobian(0, 0), jacobian(1, 0), jacobian(2, 0) };
+        basis[1] = { jacobian(0, 1), jacobian(1, 1), jacobian(2, 1) };
+
+        if (mDPDUs)
+        {
+            DPDU(i) = basis[0];
+        }
+        if (mDPDVs)
+        {
+            DPDV(i) = basis[1];
+        }
+
+        ComputeOrthogonalComplement(2, basis, true);
+
+        if (mNormals)
+        {
+            Normal(i) = basis[2];
+        }
+        if (mTangents)
+        {
+            Tangent(i) = basis[0];
+        }
+        if (mBitangents)
+        {
+            Bitangent(i) = basis[1];
+        }
+    }
+}
+
+}
diff --git a/Engine/Plugins/Experimental/GeometryProcessing/Source/GeometryAlgorithms/Private/GTEngine/Mathematics/GteMinimalCycleBasis.h b/Engine/Plugins/Experimental/GeometryProcessing/Source/GeometryAlgorithms/Private/GTEngine/Mathematics/GteMinimalCycleBasis.h
new file mode 100644
--- /dev/null
+++ b/Engine/Plugins/Experimental/GeometryProcessing/Source/GeometryAlgorithms/Private/GTEngine/Mathematics/GteMinimalCycleBasis.h
@@ -0,0 +1,709 @@
+// David Eberly, Geometric Tools, Redmond WA 98052
+// Copyright (c) 1998-2019
+// Distributed under the Boost Software License, Version 1.0.
+// http://www.boost.org/LICENSE_1_0.txt
+// http://www.geometrictools.com/License/Boost/LICENSE_1_0.txt
+// File Version: 3.0.0 (2016/06/19)
+
+#pragma once
+
+#include <GTEngine/LowLevel/GteLogger.h>
+#include <GTEngine/LowLevel/GteMinHeap.h>
+#include <array>
+#include <map>
+#include <memory>
+#include <set>
+#include <stack>
+
+// Extract the minimal cycle basis for a planar graph.  The input vertices and
+// edges must form a graph for which edges intersect only at vertices; that is,
+// no two edges must intersect at an interior point of one of the edges.  The
+// algorithm is described in 
+//   http://www.geometrictools.com/Documentation/MinimalCycleBasis.pdf
+// The graph might have filaments, which are polylines in the graph that are
+// not shared by a cycle.  These are also extracted by the implementation.
+// Because the inputs to the constructor are vertices and edges of the graph,
+// isolated vertices are ignored.
+//
+// The computations that determine which adjacent vertex to visit next during
+// a filament or cycle traversal do not require division, so the exact
+// arithmetic type BSNumber<UIntegerAP32> suffices for ComputeType when you
+// want to ensure a correct output.  (Floating-point rounding errors
+// potentially can lead to an incorrect output.)
+
+namespace gte
+{
+
+template <typename Real>
+class MinimalCycleBasis
+{
+public:
+    struct Tree
+    {
+        std::vector<int> cycle;
+        std::vector<std::shared_ptr<Tree>> children;
+    };
+
+    // The input positions and edges must form a planar graph for which edges
+    // intersect only at vertices; that is, no two edges must intersect at an
+    // interior point of one of the edges.
+    MinimalCycleBasis(
+        std::vector<std::array<Real, 2>> const& positions,
+        std::vector<std::array<int, 2>> const& edges,
+        std::vector<std::shared_ptr<Tree>>& forest);
+
+private:
+    // No copy or assignment allowed.
+    MinimalCycleBasis(MinimalCycleBasis const&) = delete;
+    MinimalCycleBasis& operator=(MinimalCycleBasis const&) = delete;
+
+    struct Vertex
+    {
+        Vertex(int inName, std::array<Real, 2> const* inPosition);
+        bool operator< (Vertex const& vertex) const;
+
+        // The index into the 'positions' input provided to the call to
+        // operator().  The index is used when reporting cycles to the
+        // caller of the constructor for MinimalCycleBasis.
+        int name;
+
+        // Multiple vertices can share a position during processing of
+        // graph components.
+        std::array<Real, 2> const* position;
+
+        // The mVertexStorage member owns the Vertex objects and maintains
+        // the reference counts on those objects.  The adjacent pointers
+        // are considered to be weak pointers; neither object ownership nor
+        // reference counting is required by 'adjacent'.
+        std::set<Vertex*> adjacent;
+
+        // Support for depth-first traversal of a graph.
+        int visited;
+    };
+
+    // The constructor uses GetComponents(...) and DepthFirstSearch(...) to
+    // get the connected components of the graph implied by the input 'edges'.
+    // Recursive processing uses only DepthFirstSearch(...) to collect
+    // vertices of the subgraphs of the original graph.
+    static void DepthFirstSearch(Vertex* vInitial, std::vector<Vertex*>& component);
+
+    // Support for traversing a simply connected component of the graph.
+    std::shared_ptr<Tree> ExtractBasis(std::vector<Vertex*>& component);
+    void RemoveFilaments(std::vector<Vertex*>& component);
+    std::shared_ptr<Tree> ExtractCycleFromComponent(std::vector<Vertex*>& component);
+    std::shared_ptr<Tree> ExtractCycleFromClosedWalk(std::vector<Vertex*>& closedWalk);
+    std::vector<int> ExtractCycle(std::vector<Vertex*>& closedWalk);
+
+    Vertex* GetClockwiseMost(Vertex* vPrev, Vertex* vCurr) const;
+    Vertex* GetCounterclockwiseMost(Vertex* vPrev, Vertex* vCurr) const;
+
+    // Storage for referenced vertices of the original graph and for new
+    // vertices added during graph traversal.
+    std::vector<std::shared_ptr<Vertex>> mVertexStorage;
+};
+
+
+template <typename Real>
+MinimalCycleBasis<Real>::MinimalCycleBasis(std::vector<std::array<Real, 2>> const& positions,
+    std::vector<std::array<int, 2>> const& edges, std::vector<std::shared_ptr<Tree>>& forest)
+{
+    forest.clear();
+    if (positions.size() == 0 || edges.size() == 0)
+    {
+        // The graph is empty, so there are no filaments or cycles.
+        return;
+    }
+
+    // Determine the unique positions referenced by the edges.
+    std::map<int, std::shared_ptr<Vertex>> unique;
+    for (auto const& edge : edges)
+    {
+        for (int i = 0; i < 2; ++i)
+        {
+            int name = edge[i];
+            if (unique.find(name) == unique.end())
+            {
+                std::shared_ptr<Vertex> vertex =
+                    std::make_shared<Vertex>(name, &positions[name]);
+                unique.insert(std::make_pair(name, vertex));
+
+            }
+        }
+    }
+
+    // Assign responsibility for ownership of the Vertex objects.
+    std::vector<Vertex*> vertices;
+    mVertexStorage.reserve(unique.size());
+    vertices.reserve(unique.size());
+    for (auto const& element : unique)
+    {
+        mVertexStorage.push_back(element.second);
+        vertices.push_back(element.second.get());
+    }
+
+    // Determine the adjacencies from the edge information.
+    for (auto const& edge : edges)
+    {
+        auto iter0 = unique.find(edge[0]);
+        auto iter1 = unique.find(edge[1]);
+        iter0->second->adjacent.insert(iter1->second.get());
+        iter1->second->adjacent.insert(iter0->second.get());
+    }
+
+    // Get the connected components of the graph.  The 'visited' flags are
+    // 0 (unvisited), 1 (discovered), 2 (finished).  The Vertex constructor
+    // sets all 'visited' flags to 0.
+    std::vector<std::vector<Vertex*>> components;
+    for (auto vInitial : mVertexStorage)
+    {
+        if (vInitial->visited == 0)
+        {
+            components.push_back(std::vector<Vertex*>());
+            DepthFirstSearch(vInitial.get(), components.back());
+        }
+    }
+
+    // The depth-first search is used later for collecting vertices for
+    // subgraphs that are detached from the main graph, so the 'visited'
+    // flags must be reset to zero after component finding.
+    for (auto vertex : mVertexStorage)
+    {
+        vertex->visited = 0;
+    }
+
+    // Get the primitives for the components.
+    for (auto& component : components)
+    {
+        forest.push_back(ExtractBasis(component));
+    }
+}
+
+template <typename Real>
+void MinimalCycleBasis<Real>::DepthFirstSearch(Vertex* vInitial, std::vector<Vertex*>& component)
+{
+    std::stack<Vertex*> vStack;
+    vStack.push(vInitial);
+    while (vStack.size() > 0)
+    {
+        Vertex* vertex = vStack.top();
+        vertex->visited = 1;
+        size_t i = 0;
+        for (auto adjacent : vertex->adjacent)
+        {
+            if (adjacent && adjacent->visited == 0)
+            {
+                vStack.push(adjacent);
+                break;
+            }
+            ++i;
+        }
+
+        if (i == vertex->adjacent.size())
+        {
+            vertex->visited = 2;
+            component.push_back(vertex);
+            vStack.pop();
+        }
+    }
+}
+
+template <typename Real>
+std::shared_ptr<typename MinimalCycleBasis<Real>::Tree>
+    MinimalCycleBasis<Real>::ExtractBasis(std::vector<Vertex*>& component)
+{
+    // The root will not have its 'cycle' member set.  The children are
+    // the cycle trees extracted from the component.
+    std::shared_ptr<Tree> tree = std::make_shared<Tree>();
+    while (component.size() > 0)
+    {
+        RemoveFilaments(component);
+        if (component.size() > 0)
+        {
+            tree->children.push_back(ExtractCycleFromComponent(component));
+        }
+    }
+
+    if (tree->cycle.size() == 0 && tree->children.size() == 1)
+    {
+        // Replace the parent by the child to avoid having two empty
+        // cycles in parent/child.
+        std::shared_ptr<Tree> child = tree->children.back();
+        tree->cycle = std::move(child->cycle);
+        tree->children = std::move(child->children);
+    }
+    return tree;
+}
+
+template <typename Real>
+void MinimalCycleBasis<Real>::RemoveFilaments(std::vector<Vertex*>& component)
+{
+    // Locate all filament endpoints, which are vertices, each having exactly
+    // one adjacent vertex.
+    std::vector<Vertex*> endpoints;
+    for (auto vertex : component)
+    {
+        if (vertex->adjacent.size() == 1)
+        {
+            endpoints.push_back(vertex);
+        }
+    }
+
+    if (endpoints.size() > 0)
+    {
+        // Remove the filaments from the component.  If a filament has two
+        // endpoints, each having one adjacent vertex, the adjacency set of
+        // the final visited vertex become empty.  We must test for that
+        // condition before starting a new filament removal.
+        for (auto vertex : endpoints)
+        {
+            if (vertex->adjacent.size() == 1)
+            {
+                // Traverse the filament and remove the vertices.
+                while (vertex->adjacent.size() == 1)
+                {
+                    // Break the connection between the two vertices.
+                    Vertex* adjacent = *vertex->adjacent.begin();
+                    adjacent->adjacent.erase(vertex);
+                    vertex->adjacent.erase(adjacent);
+
+                    // Traverse to the adjacent vertex.
+                    vertex = adjacent;
+                }
+            }
+        }
+
+        // At this time the component is either empty (it was a union of
+        // polylines) or it has no filaments and at least one cycle.  Remove
+        // the isolated vertices generated by filament extraction.
+        std::vector<Vertex*> remaining;
+        remaining.reserve(component.size());
+        for (auto vertex : component)
+        {
+            if (vertex->adjacent.size() > 0)
+            {
+                remaining.push_back(vertex);
+            }
+        }
+        component = std::move(remaining);
+    }
+}
+
+template <typename Real>
+std::shared_ptr<typename MinimalCycleBasis<Real>::Tree>
+    MinimalCycleBasis<Real>::ExtractCycleFromComponent(std::vector<Vertex*>& component)
+{
+    // Search for the left-most vertex of the component.  If two or more
+    // vertices attain minimum x-value, select the one that has minimum
+    // y-value.
+    Vertex* minVertex = component[0];
+    for (auto vertex : component)
+    {
+        if (*vertex->position < *minVertex->position)
+        {
+            minVertex = vertex;
+        }
+    }
+
+    // Traverse the closed walk, duplicating the starting vertex as the
+    // last vertex.
+    std::vector<Vertex*> closedWalk;
+    Vertex* vCurr = minVertex;
+    Vertex* vStart = vCurr;
+    closedWalk.push_back(vStart);
+    Vertex* vAdj = GetClockwiseMost(nullptr, vStart);
+    while (vAdj != vStart)
+    {
+        closedWalk.push_back(vAdj);
+        Vertex* vNext = GetCounterclockwiseMost(vCurr, vAdj);
+        vCurr = vAdj;
+        vAdj = vNext;
+    }
+    closedWalk.push_back(vStart);
+
+    // Recursively process the closed walk to extract cycles.
+    std::shared_ptr<Tree> tree = ExtractCycleFromClosedWalk(closedWalk);
+
+    // The isolated vertices generated by cycle removal are also removed from
+    // the component.
+    std::vector<Vertex*> remaining;
+    remaining.reserve(component.size());
+    for (auto vertex : component)
+    {
+        if (vertex->adjacent.size() > 0)
+        {
+            remaining.push_back(vertex);
+        }
+    }
+    component = std::move(remaining);
+
+    return tree;
+}
+
+template <typename Real>
+std::shared_ptr<typename MinimalCycleBasis<Real>::Tree>
+    MinimalCycleBasis<Real>::ExtractCycleFromClosedWalk(std::vector<Vertex*>& closedWalk)
+{
+    std::shared_ptr<Tree> tree = std::make_shared<Tree>();
+
+    std::map<Vertex*, int> duplicates;
+    std::set<int> detachments;
+    int numClosedWalk = static_cast<int>(closedWalk.size());
+    for (int i = 1; i < numClosedWalk - 1; ++i)
+    {
+        auto diter = duplicates.find(closedWalk[i]);
+        if (diter == duplicates.end())
+        {
+            // We have not yet visited this vertex.
+            duplicates.insert(std::make_pair(closedWalk[i], i));
+            continue;
+        }
+
+        // The vertex has been visited previously.  Collapse the closed walk
+        // by removing the subwalk sharing this vertex.  Note that the vertex
+        // is pointed to by closedWalk[diter->second] and closedWalk[i].
+        int iMin = diter->second, iMax = i;
+        detachments.insert(iMin);
+        for (int j = iMin + 1; j < iMax; ++j)
+        {
+            Vertex* vertex = closedWalk[j];
+            duplicates.erase(vertex);
+            detachments.erase(j);
+        }
+        closedWalk.erase(closedWalk.begin() + iMin + 1, closedWalk.begin() + iMax + 1);
+        numClosedWalk = static_cast<int>(closedWalk.size());
+        i = iMin;
+    }
+
+    if (numClosedWalk > 3)
+    {
+        // We do not know whether closedWalk[0] is a detachment point.  To
+        // determine this, we must test for any edges strictly contained by
+        // the wedge formed by the edges <closedWalk[0],closedWalk[N-1]> and
+        // <closedWalk[0],closedWalk[1]>.  However, we must execute this test
+        // even for the known detachment points.  The ensuing logic is designed
+        // to handle this and reduce the amount of code, so we insert
+        // closedWalk[0] into the detachment set and will ignore it later if
+        // it actually is not.
+        detachments.insert(0);
+
+        // Detach subgraphs from the vertices of the cycle.
+        for (auto i : detachments)
+        {
+            Vertex* original = closedWalk[i];
+            Vertex* maxVertex = closedWalk[i + 1];
+            Vertex* minVertex = (i > 0 ? closedWalk[i - 1] : closedWalk[numClosedWalk - 2]);
+
+            std::array<Real, 2> dMin, dMax;
+            for (int j = 0; j < 2; ++j)
+            {
+                dMin[j] = (*minVertex->position)[j] - (*original->position)[j];
+                dMax[j] = (*maxVertex->position)[j] - (*original->position)[j];
+            }
+
+            bool isConvex = (dMax[0] * dMin[1] >= dMax[1] * dMin[0]); (void)isConvex;
+
+            std::set<Vertex*> inWedge;
+            std::set<Vertex*> adjacent = original->adjacent;
+            for (auto vertex : adjacent)
+            {
+                if (vertex->name == minVertex->name || vertex->name == maxVertex->name)
+                {
+                    continue;
+                }
+
+                std::array<Real, 2> dVer;
+                for (int j = 0; j < 2; ++j)
+                {
+                    dVer[j] = (*vertex->position)[j] - (*original->position)[j];
+                }
+
+                bool containsVertex;
+                if (isConvex)
+                {
+                    containsVertex =
+                        dVer[0] * dMin[1] > dVer[1] * dMin[0] &&
+                        dVer[0] * dMax[1] < dVer[1] * dMax[0];
+                }
+                else
+                {
+                    containsVertex =
+                        (dVer[0] * dMin[1] > dVer[1] * dMin[0]) ||
+                        (dVer[0] * dMax[1] < dVer[1] * dMax[0]);
+                }
+                if (containsVertex)
+                {
+                    inWedge.insert(vertex);
+                }
+            }
+
+            if (inWedge.size() > 0)
+            {
+                // The clone will manage the adjacents for 'original' that lie
+                // inside the wedge defined by the first and last edges of the
+                // subgraph rooted at 'original'.  The sorting is in the
+                // clockwise direction.
+                std::shared_ptr<Vertex> clone =
+                    std::make_shared<Vertex>(original->name, original->position);
+                mVertexStorage.push_back(clone);
+
+                // Detach the edges inside the wedge.
+                for (auto vertex : inWedge)
+                {
+                    original->adjacent.erase(vertex);
+                    vertex->adjacent.erase(original);
+                    clone->adjacent.insert(vertex);
+                    vertex->adjacent.insert(clone.get());
+                }
+
+                // Get the subgraph (it is a single connected component).
+                std::vector<Vertex*> component;
+                DepthFirstSearch(clone.get(), component);
+
+                // Extract the cycles of the subgraph.
+                tree->children.push_back(ExtractBasis(component));
+            }
+            // else the candidate was closedWalk[0] and it has no subgraph
+            // to detach.
+        }
+
+        tree->cycle = std::move(ExtractCycle(closedWalk));
+    }
+    else
+    {
+        // Detach the subgraph from vertex closedWalk[0]; the subgraph
+        // is attached via a filament.
+        Vertex* original = closedWalk[0];
+        Vertex* adjacent = closedWalk[1];
+
+        std::shared_ptr<Vertex> clone =
+            std::make_shared<Vertex>(original->name, original->position);
+        mVertexStorage.push_back(clone);
+
+        original->adjacent.erase(adjacent);
+        adjacent->adjacent.erase(original);
+        clone->adjacent.insert(adjacent);
+        adjacent->adjacent.insert(clone.get());
+
+        // Get the subgraph (it is a single connected component).
+        std::vector<Vertex*> component;
+        DepthFirstSearch(clone.get(), component);
+
+        // Extract the cycles of the subgraph.
+        tree->children.push_back(ExtractBasis(component));
+        if (tree->cycle.size() == 0 && tree->children.size() == 1)
+        {
+            // Replace the parent by the child to avoid having two empty
+            // cycles in parent/child.
+            std::shared_ptr<Tree> child = tree->children.back();
+            tree->cycle = std::move(child->cycle);
+            tree->children = std::move(child->children);
+        }
+    }
+
+    return tree;
+}
+
+template <typename Real>
+std::vector<int> MinimalCycleBasis<Real>::ExtractCycle(std::vector<Vertex*>& closedWalk)
+{
+    // TODO:  This logic was designed not to remove filaments after the
+    // cycle deletion is complete.  Modify this to allow filament removal.
+
+    // The closed walk is a cycle.
+    int const numVertices = static_cast<int>(closedWalk.size());
+    std::vector<int> cycle(numVertices);
+    for (int i = 0; i < numVertices; ++i)
+    {
+        cycle[i] = closedWalk[i]->name;
+    }
+
+    // The clockwise-most edge is always removable.
+    Vertex* v0 = closedWalk[0];
+    Vertex* v1 = closedWalk[1];
+    Vertex* vBranch = (v0->adjacent.size() > 2 ? v0 : nullptr);
+    v0->adjacent.erase(v1);
+    v1->adjacent.erase(v0);
+
+    // Remove edges while traversing counterclockwise.
+    while (v1 != vBranch && v1->adjacent.size() == 1)
+    {
+        Vertex* adj = *v1->adjacent.begin();
+        v1->adjacent.erase(adj);
+        adj->adjacent.erase(v1);
+        v1 = adj;
+    }
+
+    if (v1 != v0)
+    {
+        // If v1 had exactly 3 adjacent vertices, removal of the CCW edge
+        // that shared v1 leads to v1 having 2 adjacent vertices.  When
+        // the CW removal occurs and we reach v1, the edge deletion will
+        // lead to v1 having 1 adjacent vertex, making it a filament
+        // endpoints.  We must ensure we do not delete v1 in this case,
+        // allowing the recursive algorithm to handle the filament later.
+        vBranch = v1;
+
+        // Remove edges while traversing clockwise.
+        while (v0 != vBranch && v0->adjacent.size() == 1)
+        {
+            v1 = *v0->adjacent.begin();
+            v0->adjacent.erase(v1);
+            v1->adjacent.erase(v0);
+            v0 = v1;
+        }
+    }
+    // else the cycle is its own connected component.
+
+    return cycle;
+}
+
+template <typename Real>
+typename MinimalCycleBasis<Real>::Vertex*
+    MinimalCycleBasis<Real>::GetClockwiseMost(Vertex* vPrev, Vertex* vCurr) const
+{
+    Vertex* vNext = nullptr;
+    bool vCurrConvex = false;
+    std::array<Real, 2> dCurr, dNext;
+    if (vPrev)
+    {
+        dCurr[0] = (*vCurr->position)[0] - (*vPrev->position)[0];
+        dCurr[1] = (*vCurr->position)[1] - (*vPrev->position)[1];
+    }
+    else
+    {
+        dCurr[0] = static_cast<Real>(0);
+        dCurr[1] = static_cast<Real>(-1);
+    }
+
+    for (auto vAdj : vCurr->adjacent)
+    {
+        // vAdj is a vertex adjacent to vCurr.  No backtracking is allowed.
+        if (vAdj == vPrev)
+        {
+            continue;
+        }
+
+        // Compute the potential direction to move in.
+        std::array<Real, 2> dAdj;
+        dAdj[0] = (*vAdj->position)[0] - (*vCurr->position)[0];
+        dAdj[1] = (*vAdj->position)[1] - (*vCurr->position)[1];
+
+        // Select the first candidate.
+        if (!vNext)
+        {
+            vNext = vAdj;
+            dNext = dAdj;
+            vCurrConvex = (dNext[0] * dCurr[1] <= dNext[1] * dCurr[0]);
+            continue;
+        }
+
+        // Update if the next candidate is clockwise of the current
+        // clockwise-most vertex.
+        if (vCurrConvex)
+        {
+            if (dCurr[0] * dAdj[1] < dCurr[1] * dAdj[0]
+                || dNext[0] * dAdj[1] < dNext[1] * dAdj[0])
+            {
+                vNext = vAdj;
+                dNext = dAdj;
+                vCurrConvex = (dNext[0] * dCurr[1] <= dNext[1] * dCurr[0]);
+            }
+        }
+        else
+        {
+            if (dCurr[0] * dAdj[1] < dCurr[1] * dAdj[0]
+                && dNext[0] * dAdj[1] < dNext[1] * dAdj[0])
+            {
+                vNext = vAdj;
+                dNext = dAdj;
+                vCurrConvex = (dNext[0] * dCurr[1] < dNext[1] * dCurr[0]);
+            }
+        }
+    }
+
+    return vNext;
+}
+
+template <typename Real>
+typename MinimalCycleBasis<Real>::Vertex*
+    MinimalCycleBasis<Real>::GetCounterclockwiseMost(Vertex* vPrev, Vertex* vCurr) const
+{
+    Vertex* vNext = nullptr;
+    bool vCurrConvex = false;
+    std::array<Real, 2> dCurr, dNext;
+    if (vPrev)
+    {
+        dCurr[0] = (*vCurr->position)[0] - (*vPrev->position)[0];
+        dCurr[1] = (*vCurr->position)[1] - (*vPrev->position)[1];
+    }
+    else
+    {
+        dCurr[0] = static_cast<Real>(0);
+        dCurr[1] = static_cast<Real>(-1);
+    }
+
+    for (auto vAdj : vCurr->adjacent)
+    {
+        // vAdj is a vertex adjacent to vCurr.  No backtracking is allowed.
+        if (vAdj == vPrev)
+        {
+            continue;
+        }
+
+        // Compute the potential direction to move in.
+        std::array<Real, 2> dAdj;
+        dAdj[0] = (*vAdj->position)[0] - (*vCurr->position)[0];
+        dAdj[1] = (*vAdj->position)[1] - (*vCurr->position)[1];
+
+        // Select the first candidate.
+        if (!vNext)
+        {
+            vNext = vAdj;
+            dNext = dAdj;
+            vCurrConvex = (dNext[0] * dCurr[1] <= dNext[1] * dCurr[0]);
+            continue;
+        }
+
+        // Select the next candidate if it is counterclockwise of the current
+        // counterclockwise-most vertex.
+        if (vCurrConvex)
+        {
+            if (dCurr[0] * dAdj[1] > dCurr[1] * dAdj[0]
+                && dNext[0] * dAdj[1] > dNext[1] * dAdj[0])
+            {
+                vNext = vAdj;
+                dNext = dAdj;
+                vCurrConvex = (dNext[0] * dCurr[1] <= dNext[1] * dCurr[0]);
+            }
+        }
+        else
+        {
+            if (dCurr[0] * dAdj[1] > dCurr[1] * dAdj[0]
+                || dNext[0] * dAdj[1] > dNext[1] * dAdj[0])
+            {
+                vNext = vAdj;
+                dNext = dAdj;
+                vCurrConvex = (dNext[0] * dCurr[1] <= dNext[1] * dCurr[0]);
+            }
+        }
+    }
+
+    return vNext;
+}
+
+template <typename Real>
+MinimalCycleBasis<Real>::Vertex::Vertex(int inName, std::array<Real, 2> const* inPosition)
+    :
+    name(inName),
+    position(inPosition),
+    visited(0)
+{
+}
+
+template <typename Real>
+bool MinimalCycleBasis<Real>::Vertex::operator< (Vertex const& vertex) const
+{
+    return name < vertex.name;
+}
+
+}
