move GeometryProcessing plugin from NotForLicensees to Experimental

diff --git a/Engine/Plugins/Experimental/GeometryProcessing/Source/GeometryAlgorithms/Private/GTEngine/Mathematics/GteMinimize1.h b/Engine/Plugins/Experimental/GeometryProcessing/Source/GeometryAlgorithms/Private/GTEngine/Mathematics/GteMinimize1.h
new file mode 100644
--- /dev/null
+++ b/Engine/Plugins/Experimental/GeometryProcessing/Source/GeometryAlgorithms/Private/GTEngine/Mathematics/GteMinimize1.h
@@ -0,0 +1,336 @@
+// David Eberly, Geometric Tools, Redmond WA 98052
+// Copyright (c) 1998-2019
+// Distributed under the Boost Software License, Version 1.0.
+// http://www.boost.org/LICENSE_1_0.txt
+// http://www.geometrictools.com/License/Boost/LICENSE_1_0.txt
+// File Version: 3.0.0 (2016/06/19)
+
+#pragma once
+
+#include <GTEngine/LowLevel/GteLogger.h>
+#include <cmath>
+#include <functional>
+
+// The interval [t0,t1] provided to GetMinimum(Real,Real,Real,Real&,Real&)
+// is processed by examining subintervals.  On each subinteral [a,b], the
+// values f0 = F(a), f1 = F((a+b)/2), and f2 = F(b) are examined.  If
+// {f0,f1,f2} is monotonic, then [a,b] is subdivided and processed.  The
+// maximum depth of the recursion is limited by 'maxLevel'.  If {f0,f1,f2}
+// is not monotonic, then two cases arise.  First, if f1 = min{f0,f1,f2},
+// then {f0,f1,f2} is said to "bracket a minimum" and GetBracketedMinimum(*)
+// is called to locate the function minimum.  The process uses a form of
+// bisection called "parabolic interpolation" and the maximum number of
+// bisection steps is 'maxBracket'.  Second, if f1 = max{f0,f1,f2}, then
+// {f0,f1,f2} brackets a maximum.  The minimum search continues recursively
+// as before on [a,(a+b)/2] and [(a+b)/2,b].
+
+namespace gte
+{
+
+template <typename Real>
+class Minimize1
+{
+public:
+    // Construction.
+    Minimize1(std::function<Real(Real)> const& F, int maxLevel,
+        int maxBracket);
+
+    // Search for a minimum of 'function' on the interval [t0,t1] using an
+    // initial guess of 'tInitial'.  The location of the minimum is 'tMin' and
+    // the value of the minimum is 'fMin'.
+    void GetMinimum(Real t0, Real t1, Real tInitial, Real& tMin, Real& fMin);
+
+private:
+    // This is called to start the search on [t0,tInitial] and [tInitial,t1].
+    void GetMinimum(Real t0, Real f0, Real t1, Real f1, int level);
+
+    // This is called recursively to search [a,(a+b)/2] and [(a+b)/2,b].
+    void GetMinimum(Real t0, Real f0, Real tm, Real fm, Real t1, Real f1,
+        int level);
+
+    // This is called when {f0,f1,f2} brackets a minimum.
+    void GetBracketedMinimum(Real t0, Real f0, Real tm, Real fm, Real t1,
+        Real f1, int level);
+
+    std::function<Real(Real)> mFunction;
+    int mMaxLevel;
+    int mMaxBracket;
+    Real mTMin, mFMin;
+};
+
+
+template <typename Real>
+Minimize1<Real>::Minimize1(std::function<Real(Real)> const& F, int maxLevel,
+    int maxBracket)
+    :
+    mFunction(F),
+    mMaxLevel(maxLevel),
+    mMaxBracket(maxBracket)
+{
+}
+
+template <typename Real>
+void Minimize1<Real>::GetMinimum(Real t0, Real t1, Real tInitial,
+    Real& tMin, Real& fMin)
+{
+    LogAssert(t0 <= tInitial && tInitial <= t1, "Invalid initial t value.");
+
+    mTMin = std::numeric_limits<Real>::max();
+    mFMin = std::numeric_limits<Real>::max();
+
+    Real f0 = mFunction(t0);
+    if (f0 < mFMin)
+    {
+        mTMin = t0;
+        mFMin = f0;
+    }
+
+    Real fInitial = mFunction(tInitial);
+    if (fInitial < mFMin)
+    {
+        mTMin = tInitial;
+        mFMin = fInitial;
+    }
+
+    Real f1 = mFunction(t1);
+    if (f1 < mFMin)
+    {
+        mTMin = t1;
+        mFMin = f1;
+    }
+
+    GetMinimum(t0, f0, tInitial, fInitial, t1, f1, mMaxLevel);
+
+    tMin = mTMin;
+    fMin = mFMin;
+}
+
+template <typename Real>
+void Minimize1<Real>::GetMinimum(Real t0, Real f0, Real tm, Real fm, Real t1,
+    Real f1, int level)
+{
+    if (level-- == 0)
+    {
+        return;
+    }
+
+    if ((t1 - tm)*(f0 - fm) > (tm - t0)*(fm - f1))
+    {
+        // The quadratic fit has positive second derivative at the midpoint.
+        if (f1 > f0)
+        {
+            if (fm >= f0)
+            {
+                // Increasing, repeat on [t0,tm].
+                GetMinimum(t0, f0, tm, fm, level);
+            }
+            else
+            {
+                // Not monotonic, have a bracket.
+                GetBracketedMinimum(t0, f0, tm, fm, t1, f1, level);
+            }
+        }
+        else if (f1 < f0)
+        {
+            if (fm >= f1)
+            {
+                // Decreasing, repeat on [tm,t1].
+                GetMinimum(tm, fm, t1, f1, level);
+            }
+            else
+            {
+                // Not monotonic, have a bracket.
+                GetBracketedMinimum(t0, f0, tm, fm, t1, f1, level);
+            }
+        }
+        else
+        {
+            // Constant, repeat on [t0,tm] and [tm,t1].
+            GetMinimum(t0, f0, tm, fm, level);
+            GetMinimum(tm, fm, t1, f1, level);
+        }
+    }
+    else
+    {
+        // The quadratic fit has a nonpositive second derivative at the
+        // midpoint.
+        if (f1 > f0)
+        {
+            // Repeat on [t0,tm].
+            GetMinimum(t0, f0, tm, fm, level);
+        }
+        else if (f1 < f0)
+        {
+            // Repeat on [tm,t1].
+            GetMinimum(tm, fm, t1, f1, level);
+        }
+        else
+        {
+            // Repeat on [t0,tm] and [tm,t1].
+            GetMinimum(t0, f0, tm, fm, level);
+            GetMinimum(tm, fm, t1, f1, level);
+        }
+    }
+}
+
+template <typename Real>
+void Minimize1<Real>::GetMinimum(Real t0, Real f0, Real t1, Real f1,
+    int level)
+{
+    if (level-- == 0)
+    {
+        return;
+    }
+
+    Real tm = ((Real)0.5)*(t0 + t1);
+    Real fm = mFunction(tm);
+    if (fm < mFMin)
+    {
+        mTMin = tm;
+        mFMin = fm;
+    }
+
+    if (f0 - ((Real)2)*fm + f1 >(Real)0)
+    {
+        // The quadratic fit has positive second derivative at the midpoint.
+        if (f1 > f0)
+        {
+            if (fm >= f0)
+            {
+                // Increasing, repeat on [t0,tm].
+                GetMinimum(t0, f0, tm, fm, level);
+            }
+            else
+            {
+                // Not monotonic, have a bracket.
+                GetBracketedMinimum(t0, f0, tm, fm, t1, f1, level);
+            }
+        }
+        else if (f1 < f0)
+        {
+            if (fm >= f1)
+            {
+                // Decreasing, repeat on [tm,t1].
+                GetMinimum(tm, fm, t1, f1, level);
+            }
+            else
+            {
+                // Not monotonic, have a bracket.
+                GetBracketedMinimum(t0, f0, tm, fm, t1, f1, level);
+            }
+        }
+        else
+        {
+            // Constant, repeat on [t0,tm] and [tm,t1].
+            GetMinimum(t0, f0, tm, fm, level);
+            GetMinimum(tm, fm, t1, f1, level);
+        }
+    }
+    else
+    {
+        // The quadratic fit has nonpositive second derivative at the
+        // midpoint.
+        if (f1 > f0)
+        {
+            // Repeat on [t0,tm].
+            GetMinimum(t0, f0, tm, fm, level);
+        }
+        else if (f1 < f0)
+        {
+            // Repeat on [tm,t1].
+            GetMinimum(tm, fm, t1, f1, level);
+        }
+        else
+        {
+            // Repeat on [t0,tm] and [tm,t1].
+            GetMinimum(t0, f0, tm, fm, level);
+            GetMinimum(tm, fm, t1, f1, level);
+        }
+    }
+}
+
+template <typename Real>
+void Minimize1<Real>::GetBracketedMinimum(Real t0, Real f0, Real tm,
+    Real fm, Real t1, Real f1, int level)
+{
+    for (int i = 0; i < mMaxBracket; ++i)
+    {
+        // Update minimum value.
+        if (fm < mFMin)
+        {
+            mTMin = tm;
+            mFMin = fm;
+        }
+
+        // Test for convergence.  TODO: Expose the epsilon and tolerance
+        // parameters to the caller.
+        Real const epsilon = (Real)1e-08;
+        Real const tolerance = (Real)1e-04;
+        if (std::abs(t1 - t0) <= ((Real)2)*tolerance*std::abs(tm) + epsilon)
+        {
+            break;
+        }
+
+        // Compute vertex of interpolating parabola.
+        Real dt0 = t0 - tm;
+        Real dt1 = t1 - tm;
+        Real df0 = f0 - fm;
+        Real df1 = f1 - fm;
+        Real tmp0 = dt0*df1;
+        Real tmp1 = dt1*df0;
+        Real denom = tmp1 - tmp0;
+        if (std::abs(denom) < epsilon)
+        {
+            return;
+        }
+
+        Real tv = tm + ((Real)0.5)*(dt1*tmp1 - dt0*tmp0) / denom;
+        LogAssert(t0 <= tv && tv <= t1, "Vertex not in interval.");
+        Real fv = mFunction(tv);
+        if (fv < mFMin)
+        {
+            mTMin = tv;
+            mFMin = fv;
+        }
+
+        if (tv < tm)
+        {
+            if (fv < fm)
+            {
+                t1 = tm;
+                f1 = fm;
+                tm = tv;
+                fm = fv;
+            }
+            else
+            {
+                t0 = tv;
+                f0 = fv;
+            }
+        }
+        else if (tv > tm)
+        {
+            if (fv < fm)
+            {
+                t0 = tm;
+                f0 = fm;
+                tm = tv;
+                fm = fv;
+            }
+            else
+            {
+                t1 = tv;
+                f1 = fv;
+            }
+        }
+        else
+        {
+            // The vertex of parabola is already at middle sample point.
+            GetMinimum(t0, f0, tm, fm, level);
+            GetMinimum(tm, fm, t1, f1, level);
+        }
+    }
+}
+
+
+}
diff --git a/Engine/Plugins/Experimental/GeometryProcessing/Source/GeometryAlgorithms/Private/GTEngine/Mathematics/GteMinimizeN.h b/Engine/Plugins/Experimental/GeometryProcessing/Source/GeometryAlgorithms/Private/GTEngine/Mathematics/GteMinimizeN.h
new file mode 100644
--- /dev/null
+++ b/Engine/Plugins/Experimental/GeometryProcessing/Source/GeometryAlgorithms/Private/GTEngine/Mathematics/GteMinimizeN.h
@@ -0,0 +1,213 @@
+// David Eberly, Geometric Tools, Redmond WA 98052
+// Copyright (c) 1998-2019
+// Distributed under the Boost Software License, Version 1.0.
+// http://www.boost.org/LICENSE_1_0.txt
+// http://www.geometrictools.com/License/Boost/LICENSE_1_0.txt
+// File Version: 3.0.0 (2016/06/19)
+
+#pragma once
+
+#include <GTEngine/LowLevel/GteWrapper.h>
+#include <GTEngine/Mathematics/GteGVector.h>
+#include <GTEngine/Mathematics/GteMinimize1.h>
+#include <cstring>
+
+// The Cartesian-product domain provided to GetMinimum(*) has minimum values
+// stored in t0[0..d-1] and maximum values stored in t1[0..d-1], where d is
+// 'dimensions'.  The domain is searched along lines through the current
+// estimate of the minimum location.  Each such line is searched for a minimum
+// using a Minimize1<Real> object.  This is called "Powell's Direction Set
+// Method".  The parameters 'maxLevel' and 'maxBracket' are used by
+// Minimize1<Real>, so read the documentation for that class (in its header
+// file) to understand what these mean.  The input 'maxIterations' is the
+// number of iterations for the direction-set method.
+
+namespace gte
+{
+
+template <typename Real>
+class MinimizeN
+{
+public:
+    // Construction.
+    MinimizeN(int dimensions, std::function<Real(Real const*)> const& F,
+        int maxLevel, int maxBracket, int maxIterations);
+
+    // Find the minimum on the Cartesian-product domain whose minimum values
+    // are stored in t0[0..d-1] and whose maximum values are stored in
+    // t1[0..d-1], where d is 'dimensions'.  An initial guess is specified in
+    // tInitial[0..d-1].  The location of the minimum is tMin[0..d-1] and
+    // the value of the minimum is 'fMin'.
+    void GetMinimum(Real const* t0, Real const* t1, Real const* tInitial,
+        Real* tMin, Real& fMin);
+
+private:
+    // The current estimate of the minimum location is mTCurr[0..d-1].  The
+    // direction of the current line to search is mDCurr[0..d-1].  This line
+    // must be clipped against the Cartesian-product domain, a process
+    // implemented in this function.  If the line is mTCurr+s*mDCurr, the
+    // clip result is the s-interval [ell0,ell1].
+    void ComputeDomain(Real const* t0, Real const* t1, Real& ell0,
+        Real& ell1);
+
+    int mDimensions;
+    std::function<Real(Real const*)> mFunction;
+    int mMaxIterations;
+    std::vector<GVector<Real>> mDirections;
+    int mDConjIndex;
+    int mDCurrIndex;
+    GVector<Real> mTCurr;
+    GVector<Real> mTSave;
+    Real mFCurr;
+    Minimize1<Real> mMinimizer;
+};
+
+
+template <typename Real>
+MinimizeN<Real>::MinimizeN(int dimensions,
+    std::function<Real(Real const*)> const& F, int maxLevel, int maxBracket,
+    int maxIterations)
+    :
+    mDimensions(dimensions),
+    mFunction(F),
+    mMaxIterations(maxIterations),
+    mDirections(dimensions + 1),
+    mDConjIndex(dimensions),
+    mDCurrIndex(0),
+    mTCurr(dimensions),
+    mTSave(dimensions),
+    mMinimizer(
+    [this](Real t)
+{
+    return mFunction(&(mTCurr + t*mDirections[mDCurrIndex])[0]);
+},
+maxLevel, maxBracket)
+{
+    for (auto& direction : mDirections)
+    {
+        direction.SetSize(dimensions);
+    }
+}
+
+template <typename Real>
+void MinimizeN<Real>::GetMinimum(Real const* t0, Real const* t1,
+    Real const* tInitial, Real* tMin, Real& fMin)
+{
+    // The initial guess.
+    size_t numBytes = mDimensions*sizeof(Real);
+    mFCurr = mFunction(tInitial);
+    Memcpy(&mTSave[0], tInitial, numBytes);
+    Memcpy(&mTCurr[0], tInitial, numBytes);
+
+    // Initialize the direction set to the standard Euclidean basis.
+    for (int i = 0; i < mDimensions; ++i)
+    {
+        mDirections[i].MakeUnit(i);
+    }
+
+    Real ell0, ell1, ellMin;
+    for (int iter = 0; iter < mMaxIterations; ++iter)
+    {
+        // Find minimum in each direction and update current location.
+        for (int i = 0; i < mDimensions; ++i)
+        {
+            mDCurrIndex = i;
+            ComputeDomain(t0, t1, ell0, ell1);
+            mMinimizer.GetMinimum(ell0, ell1, (Real)0, ellMin, mFCurr);
+            mTCurr += ellMin*mDirections[i];
+        }
+
+        // Estimate a unit-length conjugate direction.  TODO: Expose
+        // epsilon to the caller.
+        mDirections[mDConjIndex] = mTCurr - mTSave;
+        Real length = Length(mDirections[mDConjIndex]);
+        Real const epsilon = (Real)1e-06;
+        if (length < epsilon)
+        {
+            // New position did not change significantly from old one.
+            // Should there be a better convergence criterion here?
+            break;
+        }
+
+        mDirections[mDConjIndex] /= length;
+
+        // Minimize in conjugate direction.
+        mDCurrIndex = mDConjIndex;
+        ComputeDomain(t0, t1, ell0, ell1);
+        mMinimizer.GetMinimum(ell0, ell1, (Real)0, ellMin, mFCurr);
+        mTCurr += ellMin*mDirections[mDCurrIndex];
+
+        // Cycle the directions and add conjugate direction to set.
+        mDConjIndex = 0;
+        for (int i = 0; i < mDimensions; ++i)
+        {
+            mDirections[i] = mDirections[i + 1];
+        }
+
+        // Set parameters for next pass.
+        mTSave = mTCurr;
+    }
+
+    Memcpy(tMin, &mTCurr[0], numBytes);
+    fMin = mFCurr;
+}
+
+template <typename Real>
+void MinimizeN<Real>::ComputeDomain(Real const* t0, Real const* t1,
+    Real& ell0, Real& ell1)
+{
+    ell0 = -std::numeric_limits<Real>::max();
+    ell1 = +std::numeric_limits<Real>::max();
+
+    for (int i = 0; i < mDimensions; ++i)
+    {
+        Real value = mDirections[mDCurrIndex][i];
+        if (value != (Real)0)
+        {
+            Real b0 = t0[i] - mTCurr[i];
+            Real b1 = t1[i] - mTCurr[i];
+            Real inv = ((Real)1) / value;
+            if (value >(Real)0)
+            {
+                // The valid t-interval is [b0,b1].
+                b0 *= inv;
+                if (b0 > ell0)
+                {
+                    ell0 = b0;
+                }
+                b1 *= inv;
+                if (b1 < ell1)
+                {
+                    ell1 = b1;
+                }
+            }
+            else
+            {
+                // The valid t-interval is [b1,b0].
+                b0 *= inv;
+                if (b0 < ell1)
+                {
+                    ell1 = b0;
+                }
+                b1 *= inv;
+                if (b1 > ell0)
+                {
+                    ell0 = b1;
+                }
+            }
+        }
+    }
+
+    // Correction if numerical errors lead to values nearly zero.
+    if (ell0 > (Real)0)
+    {
+        ell0 = (Real)0;
+    }
+    if (ell1 < (Real)0)
+    {
+        ell1 = (Real)0;
+    }
+}
+
+
+}
diff --git a/Engine/Plugins/Experimental/GeometryProcessing/Source/GeometryAlgorithms/Private/GTEngine/Mathematics/GteMinimumAreaBox2.h b/Engine/Plugins/Experimental/GeometryProcessing/Source/GeometryAlgorithms/Private/GTEngine/Mathematics/GteMinimumAreaBox2.h
new file mode 100644
--- /dev/null
+++ b/Engine/Plugins/Experimental/GeometryProcessing/Source/GeometryAlgorithms/Private/GTEngine/Mathematics/GteMinimumAreaBox2.h
@@ -0,0 +1,719 @@
+// David Eberly, Geometric Tools, Redmond WA 98052
+// Copyright (c) 1998-2019
+// Distributed under the Boost Software License, Version 1.0.
+// http://www.boost.org/LICENSE_1_0.txt
+// http://www.geometrictools.com/License/Boost/LICENSE_1_0.txt
+// File Version: 3.0.2 (2018/10/05)
+
+#pragma once
+
+#include <GTEngine/Mathematics/GteOrientedBox.h>
+#include <GTEngine/Mathematics/GteConvexHull2.h>
+#include <type_traits>
+
+// Compute a minimum-area oriented box containing the specified points.  The
+// algorithm uses the rotating calipers method.
+//   http://www-cgrl.cs.mcgill.ca/~godfried/research/calipers.html
+//   http://cgm.cs.mcgill.ca/~orm/rotcal.html
+// The box is supported by the convex hull of the points, so the algorithm
+// is really about computing the minimum-area box containing a convex polygon.
+// The rotating calipers approach is O(n) in time for n polygon edges.
+//
+// A detailed description of the algorithm and implementation is found in
+//   http://www.geometrictools.com/Documentation/MinimumAreaRectangle.pdf
+//
+// NOTE: This algorithm guarantees a correct output only when ComputeType is
+// an exact arithmetic type that supports division.  In GTEngine, one such
+// type is BSRational<UIntegerAP32> (arbitrary precision).  Another such type
+// is BSRational<UIntegerFP32<N>> (fixed precision), where N is chosen large
+// enough for your input data sets.  If you choose ComputeType to be 'float'
+// or 'double', the output is not guaranteed to be correct.
+//
+// See GeometricTools/GTEngine/Samples/Geometrics/MinimumAreaBox2 for an
+// example of how to use the code.
+
+namespace gte
+{
+
+template <typename InputType, typename ComputeType>
+class MinimumAreaBox2
+{
+public:
+    // The class is a functor to support computing the minimum-area box of
+    // multiple data sets using the same class object.
+    MinimumAreaBox2();
+
+    // The points are arbitrary, so we must compute the convex hull from
+    // them in order to compute the minimum-area box.  The input parameters
+    // are necessary for using ConvexHull2.  NOTE:  ConvexHull2 guarantees
+    // that the hull does not have three consecutive collinear points.
+    OrientedBox2<InputType> operator()(int numPoints,
+        Vector2<InputType> const* points,
+        bool useRotatingCalipers =
+            !std::is_floating_point<ComputeType>::value);
+
+    // The points already form a counterclockwise, nondegenerate convex
+    // polygon.  If the points directly are the convex polygon, set
+    // numIndices to 0 and indices to nullptr.  If the polygon vertices
+    // are a subset of the incoming points, that subset is identified by
+    // numIndices >= 3 and indices having numIndices elements.
+    OrientedBox2<InputType> operator()(int numPoints,
+        Vector2<InputType> const* points, int numIndices, int const* indices,
+        bool useRotatingCalipers =
+            !std::is_floating_point<ComputeType>::value);
+
+    // Member access.
+    inline int GetNumPoints() const;
+    inline Vector2<InputType> const* GetPoints() const;
+    inline std::vector<int> const& GetHull() const;
+    inline std::array<int, 4> const& GetSupportIndices() const;
+    inline InputType GetArea() const;
+
+private:
+    // The box axes are U[i] and are usually not unit-length in order to allow
+    // exact arithmetic.  The box is supported by mPoints[index[i]], where i
+    // is one of the enumerations above.  The box axes are not necessarily unit
+    // length, but they have the same length.  They need to be normalized for
+    // conversion back to InputType.
+    struct Box
+    {
+        Vector2<ComputeType> U[2];
+        std::array<int, 4> index;  // order: bottom, right, top, left
+        ComputeType sqrLenU0, area;
+    };
+
+    // The rotating calipers algorithm has a loop invariant that requires
+    // the convex polygon not to have collinear points.  Any such points
+    // must be removed first.  The code is also executed for the O(n^2)
+    // algorithm to reduce the number of process edges.
+    void RemoveCollinearPoints(std::vector<Vector2<ComputeType>>& vertices);
+
+    // This is the slow O(n^2) search.
+    Box ComputeBoxForEdgeOrderNSqr(
+        std::vector<Vector2<ComputeType>> const& vertices);
+
+    // The fast O(n) search.
+    Box ComputeBoxForEdgeOrderN(
+        std::vector<Vector2<ComputeType>> const& vertices);
+
+    // Compute the smallest box for the polygon edge <V[i0],V[i1]>.
+    Box SmallestBox(int i0, int i1,
+        std::vector<Vector2<ComputeType>> const& vertices);
+
+    // Compute (sin(angle))^2 for the polygon edges emanating from the support
+    // vertices of the box.  The return value is 'true' if at least one angle
+    // is in [0,pi/2); otherwise, the return value is 'false' and the original
+    // polygon must be a rectangle.
+    bool ComputeAngles(std::vector<Vector2<ComputeType>> const& vertices,
+        Box const& box, std::array<std::pair<ComputeType, int>, 4>& A,
+        int& numA) const;
+
+    // Sort the angles indirectly.  The sorted indices are returned.  This
+    // avoids swapping elements of A[], which can be expensive when
+    // ComputeType is an exact rational type.
+    std::array<int, 4> SortAngles(
+        std::array<std::pair<ComputeType, int>, 4> const& A, int numA) const;
+
+    bool UpdateSupport(std::array<std::pair<ComputeType, int>, 4> const& A,
+        int numA, std::array<int, 4> const& sort,
+        std::vector<Vector2<ComputeType>> const& vertices,
+        std::vector<bool>& visited, Box& box);
+
+    // Convert the ComputeType box to the InputType box.  When the ComputeType
+    // is an exact rational type, the conversions are performed to avoid
+    // precision loss until necessary at the last step.
+    void ConvertTo(Box const& minBox,
+        std::vector<Vector2<ComputeType>> const& computePoints,
+        OrientedBox2<InputType>& itMinBox);
+
+    // The input points to be bound.
+    int mNumPoints;
+    Vector2<InputType> const* mPoints;
+
+    // The indices into mPoints/mComputePoints for the convex hull vertices.
+    std::vector<int> mHull;
+
+    // The support indices for the minimum-area box.
+    std::array<int, 4> mSupportIndices;
+
+    // The area of the minimum-area box.  The ComputeType value is exact,
+    // so the only rounding errors occur in the conversion from ComputeType
+    // to InputType (default rounding mode is round-to-nearest-ties-to-even).
+    InputType mArea;
+
+    // Convenient values that occur regularly in the code.  When using
+    // rational ComputeType, we construct these numbers only once.
+    ComputeType mZero, mOne, mNegOne, mHalf;
+};
+
+
+template <typename InputType, typename ComputeType>
+MinimumAreaBox2<InputType, ComputeType>::MinimumAreaBox2()
+    :
+    mNumPoints(0),
+    mPoints(nullptr),
+    mArea((InputType)0),
+    mZero(0),
+    mOne(1),
+    mNegOne(-1),
+    mHalf((InputType)0.5)
+{
+    mSupportIndices = { 0, 0, 0, 0 };
+}
+
+template <typename InputType, typename ComputeType>
+OrientedBox2<InputType> MinimumAreaBox2<InputType, ComputeType>::operator()(
+    int numPoints, Vector2<InputType> const* points, bool useRotatingCalipers)
+{
+    mNumPoints = numPoints;
+    mPoints = points;
+    mHull.clear();
+
+    // Get the convex hull of the points.
+    ConvexHull2<InputType, ComputeType> ch2;
+    ch2(mNumPoints, mPoints, (InputType)0);
+    int dimension = ch2.GetDimension();
+
+    OrientedBox2<InputType> minBox;
+
+    if (dimension == 0)
+    {
+        // The points are all effectively the same (using fuzzy epsilon).
+        minBox.center = mPoints[0];
+        minBox.axis[0] = Vector2<InputType>::Unit(0);
+        minBox.axis[1] = Vector2<InputType>::Unit(1);
+        minBox.extent[0] = (InputType)0;
+        minBox.extent[1] = (InputType)0;
+        mHull.resize(1);
+        mHull[0] = 0;
+        return minBox;
+    }
+
+    if (dimension == 1)
+    {
+        // The points effectively lie on a line (using fuzzy epsilon).
+        // Determine the extreme t-values for the points represented as
+        // P = origin + t*direction.  We know that 'origin' is an input
+        // vertex, so we can start both t-extremes at zero.
+        Line2<InputType> const& line = ch2.GetLine();
+        InputType tmin = (InputType)0, tmax = (InputType)0;
+        int imin = 0, imax = 0;
+        for (int i = 0; i < mNumPoints; ++i)
+        {
+            Vector2<InputType> diff = mPoints[i] - line.origin;
+            InputType t = Dot(diff, line.direction);
+            if (t > tmax)
+            {
+                tmax = t;
+                imax = i;
+            }
+            else if (t < tmin)
+            {
+                tmin = t;
+                imin = i;
+            }
+        }
+
+        minBox.center = line.origin +
+            ((InputType)0.5)*(tmin + tmax) * line.direction;
+        minBox.extent[0] = ((InputType)0.5)*(tmax - tmin);
+        minBox.extent[1] = (InputType)0;
+        minBox.axis[0] = line.direction;
+        minBox.axis[1] = -Perp(line.direction);
+        mHull.resize(2);
+        mHull[0] = imin;
+        mHull[1] = imax;
+        return minBox;
+    }
+
+    mHull = ch2.GetHull();
+    Vector2<ComputeType> const* queryPoints = ch2.GetQuery().GetVertices();
+    std::vector<Vector2<ComputeType>> computePoints(mHull.size());
+    for (size_t i = 0; i < mHull.size(); ++i)
+    {
+        computePoints[i] = queryPoints[mHull[i]];
+    }
+
+    RemoveCollinearPoints(computePoints);
+
+    Box box;
+    if (useRotatingCalipers)
+    {
+        box = ComputeBoxForEdgeOrderN(computePoints);
+    }
+    else
+    {
+        box = ComputeBoxForEdgeOrderNSqr(computePoints);
+    }
+
+    ConvertTo(box, computePoints, minBox);
+    return minBox;
+}
+
+template <typename InputType, typename ComputeType>
+OrientedBox2<InputType> MinimumAreaBox2<InputType, ComputeType>::operator()(
+    int numPoints, Vector2<InputType> const* points, int numIndices,
+    int const* indices, bool useRotatingCalipers)
+{
+    mHull.clear();
+
+    OrientedBox2<InputType> minBox;
+
+    if (numPoints < 3 || !points || (indices && numIndices < 3))
+    {
+        minBox.center = Vector2<InputType>::Zero();
+        minBox.axis[0] = Vector2<InputType>::Unit(0);
+        minBox.axis[1] = Vector2<InputType>::Unit(1);
+        minBox.extent = Vector2<InputType>::Zero();
+        return minBox;
+    }
+
+    if (indices)
+    {
+        mHull.resize(numIndices);
+        std::copy(indices, indices + numIndices, mHull.begin());
+    }
+    else
+    {
+        numIndices = numPoints;
+        mHull.resize(numIndices);
+        for (int i = 0; i < numIndices; ++i)
+        {
+            mHull[i] = i;
+        }
+    }
+
+    std::vector<Vector2<ComputeType>> computePoints(numIndices);
+    for (int i = 0; i < numIndices; ++i)
+    {
+        int h = mHull[i];
+        computePoints[i][0] = (ComputeType)points[h][0];
+        computePoints[i][1] = (ComputeType)points[h][1];
+    }
+
+    RemoveCollinearPoints(computePoints);
+
+    Box box;
+    if (useRotatingCalipers)
+    {
+        box = ComputeBoxForEdgeOrderN(computePoints);
+    }
+    else
+    {
+        box = ComputeBoxForEdgeOrderNSqr(computePoints);
+    }
+
+    ConvertTo(box, computePoints, minBox);
+    return minBox;
+}
+
+template <typename InputType, typename ComputeType> inline
+int MinimumAreaBox2<InputType, ComputeType>::GetNumPoints() const
+{
+    return mNumPoints;
+}
+
+template <typename InputType, typename ComputeType> inline
+Vector2<InputType> const* MinimumAreaBox2<InputType, ComputeType>::GetPoints()
+    const
+{
+    return mPoints;
+}
+
+template <typename InputType, typename ComputeType> inline
+std::vector<int> const&
+MinimumAreaBox2<InputType, ComputeType>::GetHull() const
+{
+    return mHull;
+}
+
+template <typename InputType, typename ComputeType> inline
+std::array<int, 4> const&
+MinimumAreaBox2<InputType, ComputeType>::GetSupportIndices() const
+{
+    return mSupportIndices;
+}
+
+template <typename InputType, typename ComputeType> inline
+InputType MinimumAreaBox2<InputType, ComputeType>::GetArea() const
+{
+    return mArea;
+}
+
+template <typename InputType, typename ComputeType>
+void MinimumAreaBox2<InputType, ComputeType>::RemoveCollinearPoints(
+    std::vector<Vector2<ComputeType>>& vertices)
+{
+    std::vector<Vector2<ComputeType>> tmpVertices = vertices;
+
+    int const numVertices = static_cast<int>(vertices.size());
+    int numNoncollinear = 0;
+    Vector2<ComputeType> ePrev = tmpVertices[0] - tmpVertices.back();
+    for (int i0 = 0, i1 = 1; i0 < numVertices; ++i0)
+    {
+        Vector2<ComputeType> eNext = tmpVertices[i1] - tmpVertices[i0];
+
+        ComputeType dp = DotPerp(ePrev, eNext);
+        if (dp != mZero)
+        {
+            vertices[numNoncollinear++] = tmpVertices[i0];
+        }
+
+        ePrev = eNext;
+        if (++i1 == numVertices)
+        {
+            i1 = 0;
+        }
+    }
+
+    vertices.resize(numNoncollinear);
+}
+
+template <typename InputType, typename ComputeType>
+typename MinimumAreaBox2<InputType, ComputeType>::Box
+MinimumAreaBox2<InputType, ComputeType>::ComputeBoxForEdgeOrderNSqr(
+    std::vector<Vector2<ComputeType>> const& vertices)
+{
+    Box minBox;
+    minBox.area = mNegOne;
+    int const numIndices = static_cast<int>(vertices.size());
+    for (int i0 = numIndices - 1, i1 = 0; i1 < numIndices; i0 = i1++)
+    {
+        Box box = SmallestBox(i0, i1, vertices);
+        if (minBox.area == mNegOne || box.area < minBox.area)
+        {
+            minBox = box;
+        }
+    }
+    return minBox;
+}
+
+template <typename InputType, typename ComputeType>
+typename MinimumAreaBox2<InputType, ComputeType>::Box
+MinimumAreaBox2<InputType, ComputeType>::ComputeBoxForEdgeOrderN(
+    std::vector<Vector2<ComputeType>> const& vertices)
+{
+    // The inputs are assumed to be the vertices of a convex polygon that
+    // is counterclockwise ordered.  The input points must not contain three
+    // consecutive collinear points.
+
+    // When the bounding box corresponding to a polygon edge is computed,
+    // we mark the edge as visited.  If the edge is encountered later, the
+    // algorithm terminates.
+    std::vector<bool> visited(vertices.size());
+    std::fill(visited.begin(), visited.end(), false);
+
+    // Start the minimum-area rectangle search with the edge from the last
+    // polygon vertex to the first.  When updating the extremes, we want the
+    // bottom-most point on the left edge, the top-most point on the right
+    // edge, the left-most point on the top edge, and the right-most point
+    // on the bottom edge.  The polygon edges starting at these points are
+    // then guaranteed not to coincide with a box edge except when an extreme
+    // point is shared by two box edges (at a corner).
+    Box minBox = SmallestBox((int)vertices.size() - 1, 0, vertices);
+    visited[minBox.index[0]] = true;
+
+    // Execute the rotating calipers algorithm.
+    Box box = minBox;
+    for (size_t i = 0; i < vertices.size(); ++i)
+    {
+        std::array<std::pair<ComputeType, int>, 4> A;
+        int numA;
+        if (!ComputeAngles(vertices, box, A, numA))
+        {
+            // The polygon is a rectangle, so the search is over.
+            break;
+        }
+
+        // Indirectly sort the A-array.
+        std::array<int, 4> sort = SortAngles(A, numA);
+
+        // Update the supporting indices (box.index[]) and the box axis
+        // directions (box.U[]).
+        if (!UpdateSupport(A, numA, sort, vertices, visited, box))
+        {
+            // We have already processed the box polygon edge, so the search
+            // is over.
+            break;
+        }
+
+        if (box.area < minBox.area)
+        {
+            minBox = box;
+        }
+    }
+
+    return minBox;
+}
+
+template <typename InputType, typename ComputeType>
+typename MinimumAreaBox2<InputType, ComputeType>::Box
+MinimumAreaBox2<InputType, ComputeType>::SmallestBox(int i0, int i1,
+    std::vector<Vector2<ComputeType>> const& vertices)
+{
+    Box box;
+    box.U[0] = vertices[i1] - vertices[i0];
+    box.U[1] = -Perp(box.U[0]);
+    box.index = { i1, i1, i1, i1 };
+    box.sqrLenU0 = Dot(box.U[0], box.U[0]);
+
+    Vector2<ComputeType> const& origin = vertices[i1];
+    Vector2<ComputeType> support[4];
+    for (int j = 0; j < 4; ++j)
+    {
+        support[j] = { mZero, mZero };
+    }
+
+    int i = 0;
+    for (auto const& vertex : vertices)
+    {
+        Vector2<ComputeType> diff = vertex - origin;
+        Vector2<ComputeType> v = { Dot(box.U[0], diff), Dot(box.U[1], diff) };
+
+        // The right-most vertex of the bottom edge is vertices[i1].  The
+        // assumption of no triple of collinear vertices guarantees that
+        // box.index[0] is i1, which is the initial value assigned at the
+        // beginning of this function.  Therefore, there is no need to test
+        // for other vertices farther to the right than vertices[i1].
+
+        if (v[0] > support[1][0] ||
+            (v[0] == support[1][0] && v[1] > support[1][1]))
+        {
+            // New right maximum OR same right maximum but closer to top.
+            box.index[1] = i;
+            support[1] = v;
+        }
+
+        if (v[1] > support[2][1] ||
+            (v[1] == support[2][1] && v[0] < support[2][0]))
+        {
+            // New top maximum OR same top maximum but closer to left.
+            box.index[2] = i;
+            support[2] = v;
+        }
+
+        if (v[0] < support[3][0] ||
+            (v[0] == support[3][0] && v[1] < support[3][1]))
+        {
+            // New left minimum OR same left minimum but closer to bottom.
+            box.index[3] = i;
+            support[3] = v;
+        }
+
+        ++i;
+    }
+
+    // The comment in the loop has the implication that support[0] = { 0, 0 },
+    // so the scaled height (support[2][1] - support[0][1]) is simply
+    // support[2][1].
+    ComputeType scaledWidth = support[1][0] - support[3][0];
+    ComputeType scaledHeight = support[2][1];
+    box.area = scaledWidth * scaledHeight / box.sqrLenU0;
+    return box;
+}
+
+template <typename InputType, typename ComputeType>
+bool MinimumAreaBox2<InputType, ComputeType>::ComputeAngles(
+    std::vector<Vector2<ComputeType>> const& vertices, Box const& box,
+    std::array<std::pair<ComputeType, int>, 4>& A, int& numA) const
+{
+    int const numVertices = static_cast<int>(vertices.size());
+    numA = 0;
+    for (int k0 = 3, k1 = 0; k1 < 4; k0 = k1++)
+    {
+        if (box.index[k0] != box.index[k1])
+        {
+            // The box edges are ordered in k1 as U[0], U[1], -U[0], -U[1].
+            Vector2<ComputeType> D =
+                ((k0 & 2) ? -box.U[k0 & 1] : box.U[k0 & 1]);
+            int j0 = box.index[k0], j1 = j0 + 1;
+            if (j1 == numVertices)
+            {
+                j1 = 0;
+            }
+            Vector2<ComputeType> E = vertices[j1] - vertices[j0];
+            ComputeType dp = DotPerp(D, E);
+            ComputeType esqrlen = Dot(E, E);
+            ComputeType sinThetaSqr = (dp * dp) / esqrlen;
+            A[numA++] = std::make_pair(sinThetaSqr, k0);
+        }
+    }
+    return numA > 0;
+}
+
+template <typename InputType, typename ComputeType>
+std::array<int, 4> MinimumAreaBox2<InputType, ComputeType>::SortAngles(
+    std::array<std::pair<ComputeType, int>, 4> const& A, int numA) const
+{
+    std::array<int, 4> sort = { 0, 1, 2, 3 };
+    if (numA > 1)
+    {
+        if (numA == 2)
+        {
+            if (A[sort[0]].first > A[sort[1]].first)
+            {
+                std::swap(sort[0], sort[1]);
+            }
+        }
+        else if (numA == 3)
+        {
+            if (A[sort[0]].first > A[sort[1]].first)
+            {
+                std::swap(sort[0], sort[1]);
+            }
+            if (A[sort[0]].first > A[sort[2]].first)
+            {
+                std::swap(sort[0], sort[2]);
+            }
+            if (A[sort[1]].first > A[sort[2]].first)
+            {
+                std::swap(sort[1], sort[2]);
+            }
+        }
+        else  // numA == 4
+        {
+            if (A[sort[0]].first > A[sort[1]].first)
+            {
+                std::swap(sort[0], sort[1]);
+            }
+            if (A[sort[2]].first > A[sort[3]].first)
+            {
+                std::swap(sort[2], sort[3]);
+            }
+            if (A[sort[0]].first > A[sort[2]].first)
+            {
+                std::swap(sort[0], sort[2]);
+            }
+            if (A[sort[1]].first > A[sort[3]].first)
+            {
+                std::swap(sort[1], sort[3]);
+            }
+            if (A[sort[1]].first > A[sort[2]].first)
+            {
+                std::swap(sort[1], sort[2]);
+            }
+        }
+    }
+    return sort;
+}
+
+template <typename InputType, typename ComputeType>
+bool MinimumAreaBox2<InputType, ComputeType>::UpdateSupport(
+    std::array<std::pair<ComputeType, int>, 4> const& A, int numA,
+    std::array<int, 4> const& sort,
+    std::vector<Vector2<ComputeType>> const& vertices,
+    std::vector<bool>& visited, Box& box)
+{
+    // Replace the support vertices of those edges attaining minimum angle
+    // with the other endpoints of the edges.
+    int const numVertices = static_cast<int>(vertices.size());
+    auto const& amin = A[sort[0]];
+    for (int k = 0; k < numA; ++k)
+    {
+        auto const& a = A[sort[k]];
+        if (a.first == amin.first)
+        {
+            if (++box.index[a.second] == numVertices)
+            {
+                box.index[a.second] = 0;
+            }
+        }
+        else
+        {
+            break;
+        }
+    }
+
+    int bottom = box.index[amin.second];
+    if (visited[bottom])
+    {
+        // We have already processed this polygon edge.
+        return false;
+    }
+    visited[bottom] = true;
+
+    // Cycle the vertices so that the bottom support occurs first.
+    std::array<int, 4> nextIndex;
+    for (int k = 0; k < 4; ++k)
+    {
+        nextIndex[k] = box.index[(amin.second + k) % 4];
+    }
+    box.index = nextIndex;
+
+    // Compute the box axis directions.
+    int j1 = box.index[0], j0 = j1 - 1;
+    if (j0 < 0)
+    {
+        j0 = numVertices - 1;
+    }
+    box.U[0] = vertices[j1] - vertices[j0];
+    box.U[1] = -Perp(box.U[0]);
+    box.sqrLenU0 = Dot(box.U[0], box.U[0]);
+
+    // Compute the box area.
+    Vector2<ComputeType> diff[2] =
+    {
+        vertices[box.index[1]] - vertices[box.index[3]],
+        vertices[box.index[2]] - vertices[box.index[0]]
+    };
+    box.area = Dot(box.U[0], diff[0]) * Dot(box.U[1], diff[1]) / box.sqrLenU0;
+    return true;
+}
+
+template <typename InputType, typename ComputeType>
+void MinimumAreaBox2<InputType, ComputeType>::ConvertTo(Box const& minBox,
+    std::vector<Vector2<ComputeType>> const& computePoints,
+    OrientedBox2<InputType>& itMinBox)
+{
+    // The sum, difference, and center are all computed exactly.
+    Vector2<ComputeType> sum[2] =
+    {
+        computePoints[minBox.index[1]] + computePoints[minBox.index[3]],
+        computePoints[minBox.index[2]] + computePoints[minBox.index[0]]
+    };
+
+    Vector2<ComputeType> difference[2] =
+    {
+        computePoints[minBox.index[1]] - computePoints[minBox.index[3]],
+        computePoints[minBox.index[2]] - computePoints[minBox.index[0]]
+    };
+
+    Vector2<ComputeType> center = mHalf * (
+        Dot(minBox.U[0], sum[0]) * minBox.U[0] +
+        Dot(minBox.U[1], sum[1]) * minBox.U[1]) / minBox.sqrLenU0;
+
+    // Calculate the squared extent using ComputeType to avoid loss of
+    // precision before computing a squared root.
+    Vector2<ComputeType> sqrExtent;
+    for (int i = 0; i < 2; ++i)
+    {
+        sqrExtent[i] = mHalf * Dot(minBox.U[i], difference[i]);
+        sqrExtent[i] *= sqrExtent[i];
+        sqrExtent[i] /= minBox.sqrLenU0;
+    }
+
+    for (int i = 0; i < 2; ++i)
+    {
+        itMinBox.center[i] = (InputType)center[i];
+        itMinBox.extent[i] = std::sqrt((InputType)sqrExtent[i]);
+
+        // Before converting to floating-point, factor out the maximum
+        // component using ComputeType to generate rational numbers in a
+        // range that avoids loss of precision during the conversion and
+        // normalization.
+        Vector2<ComputeType> const& axis = minBox.U[i];
+        ComputeType cmax = std::max(std::abs(axis[0]), std::abs(axis[1]));
+        ComputeType invCMax = mOne / cmax;
+        for (int j = 0; j < 2; ++j)
+        {
+            itMinBox.axis[i][j] = (InputType)(axis[j] * invCMax);
+        }
+        Normalize(itMinBox.axis[i]);
+    }
+
+    mSupportIndices = minBox.index;
+    mArea = (InputType)minBox.area;
+}
+
+
+}
diff --git a/Engine/Plugins/Experimental/GeometryProcessing/Source/GeometryAlgorithms/Private/GTEngine/Mathematics/GteMinimumAreaCircle2.h b/Engine/Plugins/Experimental/GeometryProcessing/Source/GeometryAlgorithms/Private/GTEngine/Mathematics/GteMinimumAreaCircle2.h
new file mode 100644
--- /dev/null
+++ b/Engine/Plugins/Experimental/GeometryProcessing/Source/GeometryAlgorithms/Private/GTEngine/Mathematics/GteMinimumAreaCircle2.h
@@ -0,0 +1,445 @@
+// David Eberly, Geometric Tools, Redmond WA 98052
+// Copyright (c) 1998-2019
+// Distributed under the Boost Software License, Version 1.0.
+// http://www.boost.org/LICENSE_1_0.txt
+// http://www.geometrictools.com/License/Boost/LICENSE_1_0.txt
+// File Version: 3.0.2 (2018/10/28)
+
+#pragma once
+
+#include <GTEngine/LowLevel/GteLogger.h>
+#include <GTEngine/Mathematics/GteContCircle2.h>
+#include <GTEngine/Mathematics/GteHypersphere.h>
+#include <GTEngine/Mathematics/GteLinearSystem.h>
+#include <algorithm>
+#include <random>
+#include <utility>
+
+// Compute the minimum area circle containing the input set of points.  The
+// algorithm randomly permutes the input points so that the construction
+// occurs in 'expected' O(N) time.  All internal minimal circle calculations
+// store the squared radius in the radius member of Circle2.  Only at the
+// end is a sqrt computed.
+//
+// The most robust choice for ComputeType is BSRational<T> for exact rational
+// arithmetic.  As long as this code is a correct implementation of the theory
+// (which I hope it is), you will obtain the minimum-area circle containing
+// the points.
+//
+// Instead, if you choose ComputeType to be float or double, floating-point
+// rounding errors can cause the UpdateSupport{2,3} functions to fail.
+// The failure is trapped in those functions and a simple bounding circle is
+// computed using GetContainer in file GteContCircle2.h.  This circle is
+// generally not the minimum-area circle containing the points.  The
+// minimum-area algorithm is terminated immediately.  The circle is returned
+// as well as a bool value of 'true' when the circle is minimum area or
+// 'false' when the failure is trapped.  When 'false' is returned, you can
+// try another call to the operator()(...) function.  The random shuffle
+// that occurs is highly likely to be different from the previous shuffle,
+// and there is a chance that the algorithm can succeed just because of the
+// different ordering of points.
+
+namespace gte
+{
+    template <typename InputType, typename ComputeType>
+    class MinimumAreaCircle2
+    {
+    public:
+        bool operator()(int numPoints, Vector2<InputType> const* points, Circle2<InputType>& minimal)
+        {
+            if (numPoints >= 1 && points)
+            {
+                // Function array to avoid switch statement in the main loop.
+                std::function<UpdateResult(int)> update[4];
+                update[1] = [this](int i) { return UpdateSupport1(i); };
+                update[2] = [this](int i) { return UpdateSupport2(i); };
+                update[3] = [this](int i) { return UpdateSupport3(i); };
+
+                // Process only the unique points.
+                std::vector<int> permuted(numPoints);
+                for (int i = 0; i < numPoints; ++i)
+                {
+                    permuted[i] = i;
+                }
+                std::sort(permuted.begin(), permuted.end(),
+                    [points](int i0, int i1) { return points[i0] < points[i1]; });
+                auto end = std::unique(permuted.begin(), permuted.end(),
+                    [points](int i0, int i1) { return points[i0] == points[i1]; });
+                permuted.erase(end, permuted.end());
+                numPoints = static_cast<int>(permuted.size());
+
+                // Create a random permutation of the points.
+                std::shuffle(permuted.begin(), permuted.end(), mDRE);
+
+                // Convert to the compute type, which is a simple copy when
+                // ComputeType is the same as InputType.
+                mComputePoints.resize(numPoints);
+                for (int i = 0; i < numPoints; ++i)
+                {
+                    for (int j = 0; j < 2; ++j)
+                    {
+                        mComputePoints[i][j] = points[permuted[i]][j];
+                    }
+                }
+
+                // Start with the first point.
+                Circle2<ComputeType> ctMinimal = ExactCircle1(0);
+                mNumSupport = 1;
+                mSupport[0] = 0;
+
+                // The loop restarts from the beginning of the point list each
+                // time the circle needs updating.  Linus Kallberg (Computer
+                // Science at Malardalen University in Sweden) discovered that
+                // performance is better when the remaining points in the
+                // array are processed before restarting.  The points
+                // processed before the point that caused the update are
+                // likely to be enclosed by the new circle (or near the circle
+                // boundary) because they were enclosed by the previous
+                // circle.  The chances are better that points after the
+                // current one will cause growth of the bounding circle.
+                for (int i = 1 % numPoints, n = 0; i != n; i = (i + 1) % numPoints)
+                {
+                    if (!SupportContains(i))
+                    {
+                        if (!Contains(i, ctMinimal))
+                        {
+                            auto result = update[mNumSupport](i);
+                            if (result.second == true)
+                            {
+                                if (result.first.radius > ctMinimal.radius)
+                                {
+                                    ctMinimal = result.first;
+                                    n = i;
+                                }
+                            }
+                            else
+                            {
+                                // This case can happen when ComputeType is
+                                // float or double.  See the comments at the
+                                // beginning of this file.
+                                LogWarning("ComputeType is not exact and failure occurred. Returning non-minimal circle.");
+                                GetContainer(numPoints, points, minimal);
+                                mNumSupport = 0;
+                                mSupport.fill(0);
+                                return false;
+                            }
+                        }
+                    }
+                }
+
+                for (int j = 0; j < 2; ++j)
+                {
+                    minimal.center[j] = static_cast<InputType>(ctMinimal.center[j]);
+                }
+                minimal.radius = static_cast<InputType>(ctMinimal.radius);
+                minimal.radius = std::sqrt(minimal.radius);
+
+                for (int i = 0; i < mNumSupport; ++i)
+                {
+                    mSupport[i] = permuted[mSupport[i]];
+                }
+                return true;
+            }
+            else
+            {
+                LogError("Input must contain points.");
+                minimal.center = Vector2<InputType>::Zero();
+                minimal.radius = std::numeric_limits<InputType>::max();
+                return false;
+            }
+        }
+
+        // Member access.
+        inline int GetNumSupport() const
+        {
+            return mNumSupport;
+        }
+
+        inline std::array<int, 3> const& GetSupport() const
+        {
+            return mSupport;
+        }
+
+    private:
+        // Test whether point P is inside circle C using squared distance and
+        // squared radius.
+        bool Contains(int i, Circle2<ComputeType> const& circle) const
+        {
+            // NOTE: In this algorithm, circle.radius is the *squared radius*
+            // until the function returns at which time a square root is
+            // applied.
+            Vector2<ComputeType> diff = mComputePoints[i] - circle.center;
+            return Dot(diff, diff) <= circle.radius;
+        }
+
+        Circle2<ComputeType> ExactCircle1(int i0) const
+        {
+            Circle2<ComputeType> minimal;
+            minimal.center = mComputePoints[i0];
+            minimal.radius = (ComputeType)0;
+            return minimal;
+        }
+
+        Circle2<ComputeType> ExactCircle2(int i0, int i1) const
+        {
+            Vector2<ComputeType> const& P0 = mComputePoints[i0];
+            Vector2<ComputeType> const& P1 = mComputePoints[i1];
+            Vector2<ComputeType> diff = P1 - P0;
+            Circle2<ComputeType> minimal;
+            minimal.center = ((ComputeType)0.5)*(P0 + P1);
+            minimal.radius = ((ComputeType)0.25)*Dot(diff, diff);
+            return minimal;
+        }
+
+        Circle2<ComputeType> ExactCircle3(int i0, int i1, int i2) const
+        {
+            // Compute the 2D circle containing P0, P1, and P2.  The center in
+            // barycentric coordinates is C = x0*P0 + x1*P1 + x2*P2, where
+            // x0 + x1 + x2 = 1.  The center is equidistant from the three
+            // points, so |C - P0| = |C - P1| = |C - P2| = R, where R is the
+            // radius of the circle.  From these conditions,
+            //   C - P0 = x0*E0 + x1*E1 - E0
+            //   C - P1 = x0*E0 + x1*E1 - E1
+            //   C - P2 = x0*E0 + x1*E1
+            // where E0 = P0 - P2 and E1 = P1 - P2, which leads to
+            //   r^2 = |x0*E0 + x1*E1|^2 - 2*Dot(E0, x0*E0 + x1*E1) + |E0|^2
+            //   r^2 = |x0*E0 + x1*E1|^2 - 2*Dot(E1, x0*E0 + x1*E1) + |E1|^2
+            //   r^2 = |x0*E0 + x1*E1|^2
+            // Subtracting the last equation from the first two and writing
+            // the equations as a linear system,
+            //
+            // +-                     -++   -+       +-          -+
+            // | Dot(E0,E0) Dot(E0,E1) || x0 | = 0.5 | Dot(E0,E0) |
+            // | Dot(E1,E0) Dot(E1,E1) || x1 |       | Dot(E1,E1) |
+            // +-                     -++   -+       +-          -+
+            //
+            // The following code solves this system for x0 and x1 and then
+            // evaluates the third equation in r^2 to obtain r.
+
+            Vector2<ComputeType> const& P0 = mComputePoints[i0];
+            Vector2<ComputeType> const& P1 = mComputePoints[i1];
+            Vector2<ComputeType> const& P2 = mComputePoints[i2];
+
+            Vector2<ComputeType> E0 = P0 - P2;
+            Vector2<ComputeType> E1 = P1 - P2;
+
+            Matrix2x2<ComputeType> A;
+            A(0, 0) = Dot(E0, E0);
+            A(0, 1) = Dot(E0, E1);
+            A(1, 0) = A(0, 1);
+            A(1, 1) = Dot(E1, E1);
+
+            ComputeType const half = (ComputeType)0.5;
+            Vector2<ComputeType> B{ half * A(0, 0), half* A(1, 1) };
+
+            Circle2<ComputeType> minimal;
+            Vector2<ComputeType> X;
+            if (LinearSystem<ComputeType>::Solve(A, B, X))
+            {
+                ComputeType x2 = (ComputeType)1 - X[0] - X[1];
+                minimal.center = X[0] * P0 + X[1] * P1 + x2 * P2;
+                Vector2<ComputeType> tmp = X[0] * E0 + X[1] * E1;
+                minimal.radius = Dot(tmp, tmp);
+            }
+            else
+            {
+                minimal.center = Vector2<ComputeType>::Zero();
+                minimal.radius = (ComputeType)std::numeric_limits<InputType>::max();
+            }
+
+            return minimal;
+        }
+
+        typedef std::pair<Circle2<ComputeType>, bool> UpdateResult;
+
+        UpdateResult UpdateSupport1(int i)
+        {
+            Circle2<ComputeType> minimal = ExactCircle2(mSupport[0], i);
+            mNumSupport = 2;
+            mSupport[1] = i;
+            return std::make_pair(minimal, true);
+        }
+
+        UpdateResult UpdateSupport2(int i)
+        {
+            // Permutations of type 2, used for calling ExactCircle2(...).
+            int const numType2 = 2;
+            int const type2[numType2][2] =
+            {
+                { 0, /*2*/ 1 },
+                { 1, /*2*/ 0 }
+            };
+
+            // Permutations of type 3, used for calling ExactCircle3(...).
+            int const numType3 = 1;  // {0, 1, 2}
+
+            Circle2<ComputeType> circle[numType2 + numType3];
+            ComputeType minRSqr = (ComputeType)std::numeric_limits<InputType>::max();
+            int iCircle = 0, iMinRSqr = -1;
+            int k0, k1;
+
+            // Permutations of type 2.
+            for (int j = 0; j < numType2; ++j, ++iCircle)
+            {
+                k0 = mSupport[type2[j][0]];
+                circle[iCircle] = ExactCircle2(k0, i);
+                if (circle[iCircle].radius < minRSqr)
+                {
+                    k1 = mSupport[type2[j][1]];
+                    if (Contains(k1, circle[iCircle]))
+                    {
+                        minRSqr = circle[iCircle].radius;
+                        iMinRSqr = iCircle;
+                    }
+                }
+            }
+
+            // Permutations of type 3.
+            k0 = mSupport[0];
+            k1 = mSupport[1];
+            circle[iCircle] = ExactCircle3(k0, k1, i);
+            if (circle[iCircle].radius < minRSqr)
+            {
+                minRSqr = circle[iCircle].radius;
+                iMinRSqr = iCircle;
+            }
+
+            switch (iMinRSqr)
+            {
+            case 0:
+                mSupport[1] = i;
+                break;
+            case 1:
+                mSupport[0] = i;
+                break;
+            case 2:
+                mNumSupport = 3;
+                mSupport[2] = i;
+                break;
+            case -1:
+                // For exact arithmetic, iMinRSqr >= 0, but for floating-point
+                // arithmetic, round-off errors can lead to iMinRSqr == -1.
+                // When this happens, use a simple bounding circle for the
+                // result and terminate the minimum-area algorithm.
+                return std::make_pair(Circle2<ComputeType>(), false);
+            }
+
+            return std::make_pair(circle[iMinRSqr], true);
+        }
+
+        UpdateResult UpdateSupport3(int i)
+        {
+            // Permutations of type 2, used for calling ExactCircle2(...).
+            int const numType2 = 3;
+            int const type2[numType2][3] =
+            {
+                { 0, /*3*/ 1, 2 },
+                { 1, /*3*/ 0, 2 },
+                { 2, /*3*/ 0, 1 }
+            };
+
+            // Permutations of type 2, used for calling ExactCircle3(...).
+            int const numType3 = 3;
+            int const type3[numType3][3] =
+            {
+                { 0, 1, /*3*/ 2 },
+                { 0, 2, /*3*/ 1 },
+                { 1, 2, /*3*/ 0 }
+            };
+
+            Circle2<ComputeType> circle[numType2 + numType3];
+            ComputeType minRSqr = (ComputeType)std::numeric_limits<InputType>::max();
+            int iCircle = 0, iMinRSqr = -1;
+            int k0, k1, k2;
+
+            // Permutations of type 2.
+            for (int j = 0; j < numType2; ++j, ++iCircle)
+            {
+                k0 = mSupport[type2[j][0]];
+                circle[iCircle] = ExactCircle2(k0, i);
+                if (circle[iCircle].radius < minRSqr)
+                {
+                    k1 = mSupport[type2[j][1]];
+                    k2 = mSupport[type2[j][2]];
+                    if (Contains(k1, circle[iCircle]) && Contains(k2, circle[iCircle]))
+                    {
+                        minRSqr = circle[iCircle].radius;
+                        iMinRSqr = iCircle;
+                    }
+                }
+            }
+
+            // Permutations of type 3.
+            for (int j = 0; j < numType3; ++j, ++iCircle)
+            {
+                k0 = mSupport[type3[j][0]];
+                k1 = mSupport[type3[j][1]];
+                circle[iCircle] = ExactCircle3(k0, k1, i);
+                if (circle[iCircle].radius < minRSqr)
+                {
+                    k2 = mSupport[type3[j][2]];
+                    if (Contains(k2, circle[iCircle]))
+                    {
+                        minRSqr = circle[iCircle].radius;
+                        iMinRSqr = iCircle;
+                    }
+                }
+            }
+
+            switch (iMinRSqr)
+            {
+            case 0:
+                mNumSupport = 2;
+                mSupport[1] = i;
+                break;
+            case 1:
+                mNumSupport = 2;
+                mSupport[0] = i;
+                break;
+            case 2:
+                mNumSupport = 2;
+                mSupport[0] = mSupport[2];
+                mSupport[1] = i;
+                break;
+            case 3:
+                mSupport[2] = i;
+                break;
+            case 4:
+                mSupport[1] = i;
+                break;
+            case 5:
+                mSupport[0] = i;
+                break;
+            case -1:
+                // For exact arithmetic, iMinRSqr >= 0, but for floating-point
+                // arithmetic, round-off errors can lead to iMinRSqr == -1.
+                // When this happens, use a simple bounding circle for the
+                // result and terminate the minimum-area algorithm.
+                return std::make_pair(Circle2<ComputeType>(), false);
+            }
+
+            return std::make_pair(circle[iMinRSqr], true);
+        }
+
+        // Indices of points that support the current minimum area circle.
+        bool SupportContains(int j) const
+        {
+            for (int i = 0; i < mNumSupport; ++i)
+            {
+                if (j == mSupport[i])
+                {
+                    return true;
+                }
+            }
+            return false;
+        }
+
+        int mNumSupport;
+        std::array<int, 3> mSupport;
+
+        // Random permutation of the unique input points to produce expected
+        // linear time for the algorithm.
+        std::default_random_engine mDRE;
+        std::vector<Vector2<ComputeType>> mComputePoints;
+    };
+}
