move GeometryProcessing plugin from NotForLicensees to Experimental

diff --git a/Engine/Plugins/Experimental/GeometryProcessing/Source/GeometryAlgorithms/Private/GTEngine/Mathematics/GteMinimumVolumeBox3.h b/Engine/Plugins/Experimental/GeometryProcessing/Source/GeometryAlgorithms/Private/GTEngine/Mathematics/GteMinimumVolumeBox3.h
new file mode 100644
--- /dev/null
+++ b/Engine/Plugins/Experimental/GeometryProcessing/Source/GeometryAlgorithms/Private/GTEngine/Mathematics/GteMinimumVolumeBox3.h
@@ -0,0 +1,1252 @@
+// David Eberly, Geometric Tools, Redmond WA 98052
+// Copyright (c) 1998-2019
+// Distributed under the Boost Software License, Version 1.0.
+// http://www.boost.org/LICENSE_1_0.txt
+// http://www.geometrictools.com/License/Boost/LICENSE_1_0.txt
+// File Version: 3.0.1 (2018/10/05)
+
+#pragma once
+
+#include <GTEngine/Mathematics/GteConvexHull3.h>
+#include <GTEngine/Mathematics/GteEdgeKey.h>
+#include <GTEngine/Mathematics/GteMinimumAreaBox2.h>
+#include <GTEngine/Mathematics/GteOrientedBox.h>
+#include <thread>
+#include <type_traits>
+
+// Compute a minimum-volume oriented box containing the specified points.  The
+// algorithm is really about computing the minimum-volume box containing the
+// convex hull of the points, so we must compute the convex hull or you must
+// pass an already built hull to the code.
+//
+// The minimum-volume oriented box has a face coincident with a hull face
+// or has three mutually orthogonal edges coincident with three hull edges
+// that (of course) are mutually orthogonal.
+//    J.O'Rourke, "Finding minimal enclosing boxes",
+//    Internat. J. Comput. Inform. Sci., 14:183-199, 1985.
+//
+// A detailed description of the algorithm and implementation is found in
+// the documents
+//   http://www.geometrictools.com/Documentation/MinimumVolumeBox.pdf
+//   http://www.geometrictools.com/Documentation/MinimumAreaRectangle.pdf
+//
+// NOTE: This algorithm guarantees a correct output only when ComputeType is
+// an exact arithmetic type that supports division.  In GTEngine, one such
+// type is BSRational<UIntegerAP32> (arbitrary precision).  Another such type
+// is BSRational<UIntegerFP32<N>> (fixed precision), where N is chosen large
+// enough for your input data sets.  If you choose ComputeType to be 'float'
+// or 'double', the output is not guaranteed to be correct.
+//
+// See GeometricTools/GTEngine/Samples/Geometrics/MinimumVolumeBox3 for an
+// example of how to use the code.
+
+namespace gte
+{
+
+template <typename InputType, typename ComputeType>
+class MinimumVolumeBox3
+{
+public:
+    // The class is a functor to support computing the minimum-volume box of
+    // multiple data sets using the same class object.  For multithreading
+    // in ProcessFaces, choose 'numThreads' subject to the constraints
+    //     1 <= numThreads <= std::thread::hardware_concurrency()
+    // To execute ProcessEdges in a thread separate from the main thrad,
+    // choose 'threadProcessEdges' to 'true'.
+    MinimumVolumeBox3(unsigned int numThreads = 1, bool threadProcessEdges = false);
+
+    // The points are arbitrary, so we must compute the convex hull from
+    // them in order to compute the minimum-area box.  The input parameters
+    // are necessary for using ConvexHull3.
+    OrientedBox3<InputType> operator()(int numPoints, Vector3<InputType> const* points,
+        bool useRotatingCalipers = !std::is_floating_point<ComputeType>::value);
+
+    // The points form a nondegenerate convex polyhedron.  The indices input
+    // must be nonnull and specify the triangle faces.
+    OrientedBox3<InputType> operator()(int numPoints, Vector3<InputType> const* points,
+        int numIndices, int const* indices,
+        bool useRotatingCalipers = !std::is_floating_point<ComputeType>::value);
+
+    // Member access.
+    inline int GetNumPoints() const;
+    inline Vector3<InputType> const* GetPoints() const;
+    inline std::vector<int> const& GetHull() const;
+    inline InputType GetVolume() const;
+
+private:
+    struct Box
+    {
+        Vector3<ComputeType> P, U[3];
+        ComputeType sqrLenU[3], range[3][2], volume;
+    };
+
+    struct ExtrudeRectangle
+    {
+        Vector3<ComputeType> U[2];
+        std::array<int, 4> index;
+        ComputeType sqrLenU[2], area;
+    };
+
+    // Compute the minimum-volume box relative to each hull face.
+    void ProcessFaces(ETManifoldMesh const& mesh, Box& minBox);
+
+    // Compute the minimum-volume box for each triple of orthgonal hull edges.
+    void ProcessEdges(ETManifoldMesh const& mesh, Box& minBox);
+
+    // Compute the minimum-volume box relative to a single hull face.
+    typedef ETManifoldMesh::Triangle Triangle;
+    void ProcessFace(std::shared_ptr<Triangle> const& supportTri,
+        std::vector<Vector3<ComputeType>> const& normal,
+        std::map<std::shared_ptr<Triangle>, int> const& triNormalMap,
+        ETManifoldMesh::EMap const& emap, Box& localMinBox);
+
+    // The rotating calipers algorithm has a loop invariant that requires
+    // the convex polygon not to have collinear points.  Any such points
+    // must be removed first.  The code is also executed for the O(n^2)
+    // algorithm to reduce the number of process edges.
+    void RemoveCollinearPoints(Vector3<ComputeType> const& N, std::vector<int>& polyline);
+
+    // This is the slow order O(n^2) search.
+    void ComputeBoxForFaceOrderNSqr(Vector3<ComputeType> const& N, std::vector<int> const& polyline, Box& box);
+
+    // This is the rotating calipers version, which is O(n).
+    void ComputeBoxForFaceOrderN(Vector3<ComputeType> const& N, std::vector<int> const& polyline, Box& box);
+
+    // Compute the smallest rectangle for the polyline edge <V[i0],V[i1]>.
+    ExtrudeRectangle SmallestRectangle(int i0, int i1, Vector3<ComputeType> const& N, std::vector<int> const& polyline);
+
+    // Compute (sin(angle))^2 for the polyline edges emanating from the
+    // support vertices of the rectangle.  The return value is 'true' if at
+    // least one angle is in [0,pi/2); otherwise, the return value is 'false'
+    // and the original polyline must project to a rectangle.
+    bool ComputeAngles(Vector3<ComputeType> const& N,
+        std::vector<int> const& polyline, ExtrudeRectangle const& rct,
+        std::array<std::pair<ComputeType, int>, 4>& A, int& numA) const;
+
+    // Sort the angles indirectly.  The sorted indices are returned.  This
+    // avoids swapping elements of A[], which can be expensive when
+    // ComputeType is an exact rational type.
+    std::array<int, 4> SortAngles(std::array<std::pair<ComputeType, int>, 4> const& A, int numA) const;
+
+    bool UpdateSupport(std::array<std::pair<ComputeType, int>, 4> const& A, int numA,
+        std::array<int, 4> const& sort, Vector3<ComputeType> const& N, std::vector<int> const& polyline,
+        std::vector<bool>& visited, ExtrudeRectangle& rct);
+
+    // Convert the extruded box to the minimum-volume box of InputType.  When
+    // the ComputeType is an exact rational type, the conversions are
+    // performed to avoid precision loss until necessary at the last step.
+    void ConvertTo(Box const& minBox, OrientedBox3<InputType>& itMinBox);
+
+    // The code is multithreaded, both for convex hull computation and
+    // computing minimum-volume extruded boxes for the hull faces.  The
+    // default value is 1, which implies a single-threaded computation (on
+    // the main thread).
+    unsigned int mNumThreads;
+    bool mThreadProcessEdges;
+
+    // The input points to be bound.
+    int mNumPoints;
+    Vector3<InputType> const* mPoints;
+
+    // The ComputeType conversions of the input points.  Only points of the
+    // convex hull (vertices of a convex polyhedron) are converted for
+    // performance when ComputeType is rational.
+    Vector3<ComputeType> const* mComputePoints;
+
+    // The indices into mPoints/mComputePoints for the convex hull vertices.
+    std::vector<int> mHull;
+
+    // The unique indices in mHull.  This set allows us to compute only for
+    // the hull vertices and avoids redundant computations if the indices
+    // were to have repeated indices into mPoints/mComputePoints.  This is
+    // a performance improvement for rational ComputeType.
+    std::set<int> mUniqueIndices;
+
+    // The caller can specify whether to use rotating calipers or the slower
+    // all-edge processing for computing an extruded bounding box.
+    bool mUseRotatingCalipers;
+
+    // The volume of the minimum-volume box.  The ComputeType value is exact,
+    // so the only rounding errors occur in the conversion from ComputeType
+    // to InputType (default rounding mode is round-to-nearest-ties-to-even).
+    InputType mVolume;
+
+    // Convenient values that occur regularly in the code.  When using
+    // rational ComputeType, we construct these numbers only once.
+    ComputeType mZero, mOne, mNegOne, mHalf;
+};
+
+
+template <typename InputType, typename ComputeType>
+MinimumVolumeBox3<InputType, ComputeType>::MinimumVolumeBox3(unsigned int numThreads, bool threadProcessEdges)
+    :
+    mNumThreads(numThreads),
+    mThreadProcessEdges(threadProcessEdges),
+    mNumPoints(0),
+    mPoints(nullptr),
+    mComputePoints(nullptr),
+    mUseRotatingCalipers(true),
+    mVolume((InputType)0),
+    mZero(0),
+    mOne(1),
+    mNegOne(-1),
+    mHalf((InputType)0.5)
+{
+}
+
+template <typename InputType, typename ComputeType>
+OrientedBox3<InputType> MinimumVolumeBox3<InputType, ComputeType>::operator()(
+    int numPoints, Vector3<InputType> const* points, bool useRotatingCalipers)
+{
+    mNumPoints = numPoints;
+    mPoints = points;
+    mUseRotatingCalipers = useRotatingCalipers;
+    mHull.clear();
+    mUniqueIndices.clear();
+
+    // Get the convex hull of the points.
+    ConvexHull3<InputType, ComputeType> ch3;
+    ch3(mNumPoints, mPoints, (InputType)0);
+    int dimension = ch3.GetDimension();
+
+    OrientedBox3<InputType> itMinBox;
+
+    if (dimension == 0)
+    {
+        // The points are all effectively the same (using fuzzy epsilon).
+        itMinBox.center = mPoints[0];
+        itMinBox.axis[0] = Vector3<InputType>::Unit(0);
+        itMinBox.axis[1] = Vector3<InputType>::Unit(1);
+        itMinBox.axis[2] = Vector3<InputType>::Unit(2);
+        itMinBox.extent[0] = (InputType)0;
+        itMinBox.extent[1] = (InputType)0;
+        itMinBox.extent[2] = (InputType)0;
+        mHull.resize(1);
+        mHull[0] = 0;
+        return itMinBox;
+    }
+
+    if (dimension == 1)
+    {
+        // The points effectively lie on a line (using fuzzy epsilon).
+        // Determine the extreme t-values for the points represented as
+        // P = origin + t*direction.  We know that 'origin' is an input
+        // vertex, so we can start both t-extremes at zero.
+        Line3<InputType> const& line = ch3.GetLine();
+        InputType tmin = (InputType)0, tmax = (InputType)0;
+        int imin = 0, imax = 0;
+        for (int i = 0; i < mNumPoints; ++i)
+        {
+            Vector3<InputType> diff = mPoints[i] - line.origin;
+            InputType t = Dot(diff, line.direction);
+            if (t > tmax)
+            {
+                tmax = t;
+                imax = i;
+            }
+            else if (t < tmin)
+            {
+                tmin = t;
+                imin = i;
+            }
+        }
+
+        itMinBox.center = line.origin + ((InputType)0.5)*(tmin + tmax) * line.direction;
+        itMinBox.extent[0] = ((InputType)0.5)*(tmax - tmin);
+        itMinBox.extent[1] = (InputType)0;
+        itMinBox.extent[2] = (InputType)0;
+        itMinBox.axis[0] = line.direction;
+        ComputeOrthogonalComplement(1, &itMinBox.axis[0]);
+        mHull.resize(2);
+        mHull[0] = imin;
+        mHull[1] = imax;
+        return itMinBox;
+    }
+
+    if (dimension == 2)
+    {
+        // The points effectively line on a plane (using fuzzy epsilon).
+        // Project the points onto the plane and compute the minimum-area
+        // bounding box of them.
+        Plane3<InputType> const& plane = ch3.GetPlane();
+
+        // Get a coordinate system relative to the plane of the points.
+        // Choose the origin to be any of the input points.
+        Vector3<InputType> origin = mPoints[0];
+        Vector3<InputType> basis[3];
+        basis[0] = plane.normal;
+        ComputeOrthogonalComplement(1, basis);
+
+        // Project the input points onto the plane.
+        std::vector<Vector2<InputType>> projection(mNumPoints);
+        for (int i = 0; i < mNumPoints; ++i)
+        {
+            Vector3<InputType> diff = mPoints[i] - origin;
+            projection[i][0] = Dot(basis[1], diff);
+            projection[i][1] = Dot(basis[2], diff);
+        }
+
+        // Compute the minimum area box in 2D.
+        MinimumAreaBox2<InputType, ComputeType> mab2;
+        OrientedBox2<InputType> rectangle = mab2(mNumPoints, &projection[0]);
+
+        // Lift the values into 3D.
+        itMinBox.center = origin + rectangle.center[0] * basis[1] + rectangle.center[1] * basis[2];
+        itMinBox.axis[0] = rectangle.axis[0][0] * basis[1] + rectangle.axis[0][1] * basis[2];
+        itMinBox.axis[1] = rectangle.axis[1][0] * basis[1] + rectangle.axis[1][1] * basis[2];
+        itMinBox.axis[2] = basis[0];
+        itMinBox.extent[0] = rectangle.extent[0];
+        itMinBox.extent[1] = rectangle.extent[1];
+        itMinBox.extent[2] = (InputType)0;
+        mHull = mab2.GetHull();
+        return itMinBox;
+    }
+
+    // Get the set of unique indices of the hull.  This is used to project
+    // hull vertices onto lines.
+    ETManifoldMesh const& mesh = ch3.GetHullMesh();
+    mHull.resize(3 * mesh.GetTriangles().size());
+    int h = 0;
+    for (auto const& element : mesh.GetTriangles())
+    {
+        for (int i = 0; i < 3; ++i, ++h)
+        {
+            int index = element.first.V[i];
+            mHull[h] = index;
+            mUniqueIndices.insert(index);
+        }
+    }
+
+    mComputePoints = ch3.GetQuery().GetVertices();
+
+    Box minBox, minBoxEdges;
+    minBox.volume = mNegOne;
+    minBoxEdges.volume = mNegOne;
+
+    if (mThreadProcessEdges)
+    {
+        std::thread doEdges([this, &mesh, &minBoxEdges]()
+        {
+            ProcessEdges(mesh, minBoxEdges);
+        });
+        ProcessFaces(mesh, minBox);
+        doEdges.join();
+    }
+    else
+    {
+        ProcessEdges(mesh, minBoxEdges);
+        ProcessFaces(mesh, minBox);
+    }
+
+    if (minBoxEdges.volume != mNegOne
+        && minBoxEdges.volume < minBox.volume)
+    {
+        minBox = minBoxEdges;
+    }
+
+    ConvertTo(minBox, itMinBox);
+    mComputePoints = nullptr;
+    return itMinBox;
+}
+
+template <typename InputType, typename ComputeType>
+OrientedBox3<InputType> MinimumVolumeBox3<InputType, ComputeType>::operator()(
+    int numPoints, Vector3<InputType> const* points, int numIndices,
+    int const* indices, bool useRotatingCalipers)
+{
+    mNumPoints = numPoints;
+    mPoints = points;
+    mUseRotatingCalipers = useRotatingCalipers;
+    mUniqueIndices.clear();
+
+    // Build the mesh from the indices.  The box construction uses the edge
+    // map of the mesh.
+    ETManifoldMesh mesh;
+    int numTriangles = numIndices / 3;
+    for (int t = 0; t < numTriangles; ++t)
+    {
+        int v0 = *indices++;
+        int v1 = *indices++;
+        int v2 = *indices++;
+        mesh.Insert(v0, v1, v2);
+    }
+
+    // Get the set of unique indices of the hull.  This is used to project
+    // hull vertices onto lines.
+    mHull.resize(3 * mesh.GetTriangles().size());
+    int h = 0;
+    for (auto const& element : mesh.GetTriangles())
+    {
+        for (int i = 0; i < 3; ++i, ++h)
+        {
+            int index = element.first.V[i];
+            mHull[h] = index;
+            mUniqueIndices.insert(index);
+        }
+    }
+
+    // Create the ComputeType points to be used downstream.
+    std::vector<Vector3<ComputeType>> computePoints(mNumPoints);
+    for (auto i : mUniqueIndices)
+    {
+        for (int j = 0; j < 3; ++j)
+        {
+            computePoints[i][j] = (ComputeType)mPoints[i][j];
+        }
+    }
+
+    OrientedBox3<InputType> itMinBox;
+    mComputePoints = &computePoints[0];
+
+    Box minBox, minBoxEdges;
+    minBox.volume = mNegOne;
+    minBoxEdges.volume = mNegOne;
+
+    if (mThreadProcessEdges)
+    {
+        std::thread doEdges([this, &mesh, &minBoxEdges]()
+        {
+            ProcessEdges(mesh, minBoxEdges);
+        });
+        ProcessFaces(mesh, minBox);
+        doEdges.join();
+    }
+    else
+    {
+        ProcessEdges(mesh, minBoxEdges);
+        ProcessFaces(mesh, minBox);
+    }
+
+    if (minBoxEdges.volume != mNegOne && minBoxEdges.volume < minBox.volume)
+    {
+        minBox = minBoxEdges;
+    }
+
+    ConvertTo(minBox, itMinBox);
+    mComputePoints = nullptr;
+    return itMinBox;
+}
+
+template <typename InputType, typename ComputeType> inline
+int MinimumVolumeBox3<InputType, ComputeType>::GetNumPoints() const
+{
+    return mNumPoints;
+}
+
+template <typename InputType, typename ComputeType> inline
+Vector3<InputType> const*
+MinimumVolumeBox3<InputType, ComputeType>::GetPoints() const
+{
+    return mPoints;
+}
+
+template <typename InputType, typename ComputeType> inline
+std::vector<int> const& MinimumVolumeBox3<InputType, ComputeType>::GetHull()
+    const
+{
+    return mHull;
+}
+
+template <typename InputType, typename ComputeType> inline
+InputType MinimumVolumeBox3<InputType, ComputeType>::GetVolume() const
+{
+    return mVolume;
+}
+
+template <typename InputType, typename ComputeType>
+void MinimumVolumeBox3<InputType, ComputeType>::ProcessFaces(ETManifoldMesh const& mesh, Box& minBox)
+{
+    // Get the mesh data structures.
+    auto const& tmap = mesh.GetTriangles();
+    auto const& emap = mesh.GetEdges();
+
+    // Compute inner-pointing face normals for searching boxes supported by
+    // a face and an extreme vertex.  The indirection in triNormalMap, using
+    // an integer index instead of the normal/sqrlength pair itself, avoids
+    // expensive copies when using exact arithmetic.
+    std::vector<Vector3<ComputeType>> normal(tmap.size());
+    std::map<std::shared_ptr<Triangle>, int> triNormalMap;
+    int index = 0;
+    for (auto const& element : tmap)
+    {
+        auto tri = element.second;
+        Vector3<ComputeType> const& v0 = mComputePoints[tri->V[0]];
+        Vector3<ComputeType> const& v1 = mComputePoints[tri->V[1]];
+        Vector3<ComputeType> const& v2 = mComputePoints[tri->V[2]];
+        Vector3<ComputeType> edge1 = v1 - v0;
+        Vector3<ComputeType> edge2 = v2 - v0;
+        normal[index] = Cross(edge2, edge1);  // inner-pointing normal
+        triNormalMap[tri] = index++;
+    }
+
+    // Process the triangle faces.  For each face, compute the polyline of
+    // edges that supports the bounding box with a face coincident to the
+    // triangle face.  The projection of the polyline onto the plane of the
+    // triangle face is a convex polygon, so we can use the method of rotating
+    // calipers to compute its minimum-area box efficiently.
+    unsigned int numFaces = static_cast<unsigned int>(tmap.size());
+    if (mNumThreads > 1 && numFaces >= mNumThreads)
+    {
+        // Repackage the triangle pointers to support the partitioning of
+        // faces for multithreaded face processing.
+        std::vector<std::shared_ptr<Triangle>> triangles;
+        triangles.reserve(numFaces);
+        for (auto const& element : tmap)
+        {
+            triangles.push_back(element.second);
+        }
+
+        // Partition the data for multiple threads.
+        unsigned int numFacesPerThread = numFaces / mNumThreads;
+        std::vector<unsigned int> imin(mNumThreads), imax(mNumThreads);
+        std::vector<Box> localMinBox(mNumThreads);
+        for (unsigned int t = 0; t < mNumThreads; ++t)
+        {
+            imin[t] = t * numFacesPerThread;
+            imax[t] = imin[t] + numFacesPerThread - 1;
+            localMinBox[t].volume = mNegOne;
+        }
+        imax[mNumThreads - 1] = numFaces - 1;
+
+        // Execute the face processing in multiple threads.
+        std::vector<std::thread> process(mNumThreads);
+        for (unsigned int t = 0; t < mNumThreads; ++t)
+        {
+            process[t] = std::thread([this, t, &imin, &imax, &triangles,
+                &normal, &triNormalMap, &emap, &localMinBox]()
+            {
+                for (unsigned int i = imin[t]; i <= imax[t]; ++i)
+                {
+                    auto const& supportTri = triangles[i];
+                    ProcessFace(supportTri, normal, triNormalMap, emap, localMinBox[t]);
+                }
+            });
+        }
+
+        // Wait for all threads to finish.
+        for (unsigned int t = 0; t < mNumThreads; ++t)
+        {
+            process[t].join();
+
+            // Update the minimum-volume box candidate.
+            if (minBox.volume == mNegOne || localMinBox[t].volume < minBox.volume)
+            {
+                minBox = localMinBox[t];
+            }
+        }
+    }
+    else
+    {
+        for (auto const& element : tmap)
+        {
+            auto const& supportTri = element.second;
+            ProcessFace(supportTri, normal, triNormalMap, emap, minBox);
+        }
+    }
+}
+
+template <typename InputType, typename ComputeType>
+void MinimumVolumeBox3<InputType, ComputeType>::ProcessEdges(ETManifoldMesh const& mesh, Box& minBox)
+{
+    // The minimum-volume box can also be supported by three mutually
+    // orthogonal edges of the convex hull.  For each triple of orthogonal
+    // edges, compute the minimum-volume box for that coordinate frame by
+    // projecting the points onto the axes of the frame.  Use a hull vertex
+    // as the origin.
+    int index = mesh.GetTriangles().begin()->first.V[0];
+    Vector3<ComputeType> const& origin = mComputePoints[index];
+    Vector3<ComputeType> U[3];
+    std::array<ComputeType, 3> sqrLenU;
+
+    auto const& emap = mesh.GetEdges();
+    auto e2 = emap.begin(), end = emap.end();
+    for (/**/; e2 != end; ++e2)
+    {
+        U[2] = mComputePoints[e2->first.V[1]] - mComputePoints[e2->first.V[0]];
+        auto e1 = e2;
+        for (++e1; e1 != end; ++e1)
+        {
+            U[1] = mComputePoints[e1->first.V[1]] - mComputePoints[e1->first.V[0]];
+            if (Dot(U[1], U[2]) != mZero)
+            {
+                continue;
+            }
+            sqrLenU[1] = Dot(U[1], U[1]);
+
+            auto e0 = e1;
+            for (++e0; e0 != end; ++e0)
+            {
+                U[0] = mComputePoints[e0->first.V[1]] - mComputePoints[e0->first.V[0]];
+                sqrLenU[0] = Dot(U[0], U[0]);
+                if (Dot(U[0], U[1]) != mZero || Dot(U[0], U[2]) != mZero)
+                {
+                    continue;
+                }
+
+                // The three edges are mutually orthogonal.  To support exact
+                // rational arithmetic for volume computation, we replace U[2]
+                // by a parallel vector.
+                U[2] = Cross(U[0], U[1]);
+                sqrLenU[2] = sqrLenU[0] * sqrLenU[1];
+
+                // Project the vertices onto the lines containing the edges.
+                // Use vertex 0 as the origin.
+                std::array<ComputeType, 3> umin, umax;
+                for (int j = 0; j < 3; ++j)
+                {
+                    umin[j] = mZero;
+                    umax[j] = mZero;
+                }
+
+                for (auto i : mUniqueIndices)
+                {
+                    Vector3<ComputeType> diff = mComputePoints[i] - origin;
+                    for (int j = 0; j < 3; ++j)
+                    {
+                        ComputeType dot = Dot(diff, U[j]);
+                        if (dot < umin[j])
+                        {
+                            umin[j] = dot;
+                        }
+                        else if (dot > umax[j])
+                        {
+                            umax[j] = dot;
+                        }
+                    }
+                }
+
+                ComputeType volume = (umax[0] - umin[0]) / sqrLenU[0];
+                volume *= (umax[1] - umin[1]) / sqrLenU[1];
+                volume *= (umax[2] - umin[2]);
+
+                // Update current minimum-volume box (if necessary).
+                if (minBox.volume == mOne || volume < minBox.volume)
+                {
+                    // The edge keys have unordered vertices, so it is
+                    // possible that {U[0],U[1],U[2]} is a left-handed set.
+                    // We need a right-handed set.
+                    if (DotCross(U[0], U[1], U[2]) < mZero)
+                    {
+                        U[2] = -U[2];
+                    }
+
+                    minBox.P = origin;
+                    for (int j = 0; j < 3; ++j)
+                    {
+                        minBox.U[j] = U[j];
+                        minBox.sqrLenU[j] = sqrLenU[j];
+                        for (int k = 0; k < 3; ++k)
+                        {
+                            minBox.range[k][0] = umin[k];
+                            minBox.range[k][1] = umax[k];
+                        }
+                    }
+                    minBox.volume = volume;
+
+                }
+            }
+        }
+    }
+}
+
+template <typename InputType, typename ComputeType>
+void MinimumVolumeBox3<InputType, ComputeType>::ProcessFace(std::shared_ptr<Triangle> const& supportTri,
+    std::vector<Vector3<ComputeType>> const& normal, std::map<std::shared_ptr<Triangle>, int> const& triNormalMap,
+    ETManifoldMesh::EMap const& emap, Box& localMinBox)
+{
+    // Get the supporting triangle information.
+    Vector3<ComputeType> const& supportNormal = normal[triNormalMap.find(supportTri)->second];
+
+    // Build the polyline of supporting edges.  The pair (v,polyline[v])
+    // represents an edge directed appropriately (see next set of
+    // comments).
+    std::vector<int> polyline(mNumPoints);
+    int polylineStart = -1;
+    for (auto const& edgeElement : emap)
+    {
+        auto const& edge = *edgeElement.second;
+        auto const& tri0 = edge.T[0].lock();
+        auto const& tri1 = edge.T[1].lock();
+        auto const& normal0 = normal[triNormalMap.find(tri0)->second];
+        auto const& normal1 = normal[triNormalMap.find(tri1)->second];
+        ComputeType dot0 = Dot(supportNormal, normal0);
+        ComputeType dot1 = Dot(supportNormal, normal1);
+
+        std::shared_ptr<Triangle> tri;
+        if (dot0 < mZero && dot1 >= mZero)
+        {
+            tri = tri0;
+        }
+        else if (dot1 < mZero && dot0 >= mZero)
+        {
+            tri = tri1;
+        }
+
+        if (tri)
+        {
+            // The edge supports the bounding box.  Insert the edge in the
+            // list using clockwise order relative to tri.  This will lead
+            // to a polyline whose projection onto the plane of the hull
+            // face is a convex polygon that is counterclockwise oriented.
+            for (int j0 = 2, j1 = 0; j1 < 3; j0 = j1++)
+            {
+                if (tri->V[j1] == edge.V[0])
+                {
+                    if (tri->V[j0] == edge.V[1])
+                    {
+                        polyline[edge.V[1]] = edge.V[0];
+                    }
+                    else
+                    {
+                        polyline[edge.V[0]] = edge.V[1];
+                    }
+                    polylineStart = edge.V[0];
+                    break;
+                }
+            }
+        }
+    }
+
+    // Rearrange the edges to form a closed polyline.  For M vertices, each
+    // ComputeBoxFor*() function starts with the edge from closedPolyline[M-1]
+    // to closedPolyline[0].
+    std::vector<int> closedPolyline(mNumPoints);
+    int numClosedPolyline = 0;
+    int v = polylineStart;
+    for (auto& cp : closedPolyline)
+    {
+        cp = v;
+        ++numClosedPolyline;
+        v = polyline[v];
+        if (v == polylineStart)
+        {
+            break;
+        }
+    }
+    closedPolyline.resize(numClosedPolyline);
+
+    // This avoids redundant face testing in the O(n^2) or O(n) algorithms
+    // and it simplifies the O(n) implementation.
+    RemoveCollinearPoints(supportNormal, closedPolyline);
+
+    // Compute the box coincident to the hull triangle that has minimum
+    // area on the face coincident with the triangle.
+    Box faceBox;
+    if (mUseRotatingCalipers)
+    {
+        ComputeBoxForFaceOrderN(supportNormal, closedPolyline, faceBox);
+    }
+    else
+    {
+        ComputeBoxForFaceOrderNSqr(supportNormal, closedPolyline, faceBox);
+    }
+
+    // Update the minimum-volume box candidate.
+    if (localMinBox.volume == mNegOne || faceBox.volume < localMinBox.volume)
+    {
+        localMinBox = faceBox;
+    }
+}
+
+template <typename InputType, typename ComputeType>
+void MinimumVolumeBox3<InputType, ComputeType>::RemoveCollinearPoints(
+    Vector3<ComputeType> const& N, std::vector<int>& polyline)
+{
+    std::vector<int> tmpPolyline = polyline;
+
+    int const numPolyline = static_cast<int>(polyline.size());
+    int numNoncollinear = 0;
+    Vector3<ComputeType> ePrev =
+        mComputePoints[tmpPolyline[0]] - mComputePoints[tmpPolyline.back()];
+
+    for (int i0 = 0, i1 = 1; i0 < numPolyline; ++i0)
+    {
+        Vector3<ComputeType> eNext =
+            mComputePoints[tmpPolyline[i1]] - mComputePoints[tmpPolyline[i0]];
+
+        ComputeType tsp = DotCross(ePrev, eNext, N);
+        if (tsp != mZero)
+        {
+            polyline[numNoncollinear++] = tmpPolyline[i0];
+        }
+
+        ePrev = eNext;
+        if (++i1 == numPolyline)
+        {
+            i1 = 0;
+        }
+    }
+
+    polyline.resize(numNoncollinear);
+}
+
+template <typename InputType, typename ComputeType>
+void MinimumVolumeBox3<InputType, ComputeType>::ComputeBoxForFaceOrderNSqr(
+    Vector3<ComputeType> const& N, std::vector<int> const& polyline,
+    Box& box)
+{
+    // This code processes the polyline terminator associated with a convex
+    // hull face of inner-pointing normal N.  The polyline is usually not
+    // contained by a plane, and projecting the polyline to a convex polygon
+    // in a plane perpendicular to N introduces floating-point rounding
+    // errors.  The minimum-area box for the projected polyline is computed
+    // indirectly to support exact rational arithmetic.
+
+    box.P = mComputePoints[polyline[0]];
+    box.U[2] = N;
+    box.sqrLenU[2] = Dot(N, N);
+    box.range[1][0] = mZero;
+    box.volume = mNegOne;
+    int const numPolyline = static_cast<int>(polyline.size());
+    for (int i0 = numPolyline - 1, i1 = 0; i1 < numPolyline; i0 = i1++)
+    {
+        // Create a coordinate system for the plane perpendicular to the face
+        // normal and containing a polyline vertex.
+        Vector3<ComputeType> const& P = mComputePoints[polyline[i0]];
+        Vector3<ComputeType> E =
+            mComputePoints[polyline[i1]] - mComputePoints[polyline[i0]];
+
+        Vector3<ComputeType> U1 = Cross(N, E);
+        Vector3<ComputeType> U0 = Cross(U1, N);
+
+        // Compute the smallest rectangle containing the projected polyline.
+        ComputeType min0 = mZero, max0 = mZero, max1 = mZero;
+        for (int j = 0; j < numPolyline; ++j)
+        {
+            Vector3<ComputeType> diff = mComputePoints[polyline[j]] - P;
+            ComputeType dot = Dot(U0, diff);
+            if (dot < min0)
+            {
+                min0 = dot;
+            }
+            else if (dot > max0)
+            {
+                max0 = dot;
+            }
+
+            dot = Dot(U1, diff);
+            if (dot > max1)
+            {
+                max1 = dot;
+            }
+        }
+
+        // The true area is Area(rectangle)*Length(N).  After the smallest
+        // scaled-area rectangle is computed and returned, the box.volume is
+        // updated to be the actual squared volume of the box.
+        ComputeType sqrLenU1 = Dot(U1, U1);
+        ComputeType volume = (max0 - min0) * max1 / sqrLenU1;
+        if (box.volume == mNegOne || volume < box.volume)
+        {
+            box.P = P;
+            box.U[0] = U0;
+            box.U[1] = U1;
+            box.sqrLenU[0] = sqrLenU1 * box.sqrLenU[2];
+            box.sqrLenU[1] = sqrLenU1;
+            box.range[0][0] = min0;
+            box.range[0][1] = max0;
+            box.range[1][1] = max1;
+            box.volume = volume;
+        }
+    }
+
+    // Compute the range of points in the support-normal direction.
+    box.range[2][0] = mZero;
+    box.range[2][1] = mZero;
+    for (auto i : mUniqueIndices)
+    {
+        Vector3<ComputeType> diff = mComputePoints[i] - box.P;
+        ComputeType height = Dot(box.U[2], diff);
+        if (height < box.range[2][0])
+        {
+            box.range[2][0] = height;
+        }
+        else if (height > box.range[2][1])
+        {
+            box.range[2][1] = height;
+        }
+    }
+
+    // Compute the actual volume.
+    box.volume *= (box.range[2][1] - box.range[2][0]) / box.sqrLenU[2];
+}
+
+template <typename InputType, typename ComputeType>
+void MinimumVolumeBox3<InputType, ComputeType>::ComputeBoxForFaceOrderN(
+    Vector3<ComputeType> const& N, std::vector<int> const& polyline, Box& box)
+{
+    // This code processes the polyline terminator associated with a convex
+    // hull face of inner-pointing normal N.  The polyline is usually not
+    // contained by a plane, and projecting the polyline to a convex polygon
+    // in a plane perpendicular to N introduces floating-point rounding
+    // errors.  The minimum-area box for the projected polyline is computed
+    // indirectly to support exact rational arithmetic.
+
+    // When the bounding box corresponding to a polyline edge is computed,
+    // we mark the edge as visited.  If the edge is encountered later, the
+    // algorithm terminates.
+    std::vector<bool> visited(polyline.size());
+    std::fill(visited.begin(), visited.end(), false);
+
+    // Start the minimum-area rectangle search with the edge from the last
+    // polyline vertex to the first.  When updating the extremes, we want the
+    // bottom-most point on the left edge, the top-most point on the right
+    // edge, the left-most point on the top edge, and the right-most point
+    // on the bottom edge.  The polygon edges starting at these points are
+    // then guaranteed not to coincide with a box edge except when an extreme
+    // point is shared by two box edges (at a corner).
+    ExtrudeRectangle minRct = SmallestRectangle((int)polyline.size() - 1, 0,
+        N, polyline);
+    visited[minRct.index[0]] = true;
+
+    // Execute the rotating calipers algorithm.
+    ExtrudeRectangle rct = minRct;
+    for (size_t i = 0; i < polyline.size(); ++i)
+    {
+        std::array<std::pair<ComputeType, int>, 4> A;
+        int numA;
+        if (!ComputeAngles(N, polyline, rct, A, numA))
+        {
+            // The polyline projects to a rectangle, so the search is over.
+            break;
+        }
+
+        // Indirectly sort the A-array.
+        std::array<int, 4> sort = SortAngles(A, numA);
+
+        // Update the supporting indices (rct.index[]) and the rectangle axis
+        // directions (rct.U[]).
+        if (!UpdateSupport(A, numA, sort, N, polyline, visited, rct))
+        {
+            // We have already processed the rectangle polygon edge, so the
+            // search is over.
+            break;
+        }
+
+        if (rct.area < minRct.area)
+        {
+            minRct = rct;
+        }
+    }
+
+    // Store relevant box information for computing volume and converting to
+    // an InputType bounding box.
+    box.P = mComputePoints[polyline[minRct.index[0]]];
+    box.U[0] = minRct.U[0];
+    box.U[1] = minRct.U[1];
+    box.U[2] = N;
+    box.sqrLenU[0] = minRct.sqrLenU[0];
+    box.sqrLenU[1] = minRct.sqrLenU[1];
+    box.sqrLenU[2] = Dot(box.U[2], box.U[2]);
+
+    // Compute the range of points in the plane perpendicular to the support
+    // normal.
+    box.range[0][0] = Dot(box.U[0],
+        mComputePoints[polyline[minRct.index[3]]] - box.P);
+    box.range[0][1] = Dot(box.U[0],
+        mComputePoints[polyline[minRct.index[1]]] - box.P);
+    box.range[1][0] = mZero;
+    box.range[1][1] = Dot(box.U[1],
+        mComputePoints[polyline[minRct.index[2]]] - box.P);
+
+    // Compute the range of points in the support-normal direction.
+    box.range[2][0] = mZero;
+    box.range[2][1] = mZero;
+    for (auto i : mUniqueIndices)
+    {
+        Vector3<ComputeType> diff = mComputePoints[i] - box.P;
+        ComputeType height = Dot(box.U[2], diff);
+        if (height < box.range[2][0])
+        {
+            box.range[2][0] = height;
+        }
+        else if (height > box.range[2][1])
+        {
+            box.range[2][1] = height;
+        }
+    }
+
+    // Compute the actual volume.
+    box.volume =
+        (box.range[0][1] - box.range[0][0]) *
+        ((box.range[1][1] - box.range[1][0]) / box.sqrLenU[1]) *
+        ((box.range[2][1] - box.range[2][0]) / box.sqrLenU[2]);
+}
+
+template <typename InputType, typename ComputeType>
+typename MinimumVolumeBox3<InputType, ComputeType>::ExtrudeRectangle
+MinimumVolumeBox3<InputType, ComputeType>::SmallestRectangle(int i0, int i1,
+    Vector3<ComputeType> const& N, std::vector<int> const& polyline)
+{
+    ExtrudeRectangle rct;
+    Vector3<ComputeType> E =
+        mComputePoints[polyline[i1]] - mComputePoints[polyline[i0]];
+    rct.U[1] = Cross(N, E);
+    rct.U[0] = Cross(rct.U[1], N);
+    rct.index = { i1, i1, i1, i1 };
+    rct.sqrLenU[0] = Dot(rct.U[0], rct.U[0]);
+    rct.sqrLenU[1] = Dot(rct.U[1], rct.U[1]);
+
+    Vector3<ComputeType> const& origin = mComputePoints[polyline[i1]];
+    Vector2<ComputeType> support[4];
+    for (int j = 0; j < 4; ++j)
+    {
+        support[j] = { mZero, mZero };
+    }
+
+    int i = 0;
+    for (auto p : polyline)
+    {
+        Vector3<ComputeType> diff = mComputePoints[p] - origin;
+        Vector2<ComputeType> v = { Dot(rct.U[0], diff), Dot(rct.U[1], diff) };
+
+        // The right-most vertex of the bottom edge is vertices[i1].  The
+        // assumption of no triple of collinear vertices guarantees that
+        // box.index[0] is i1, which is the initial value assigned at the
+        // beginning of this function.  Therefore, there is no need to test
+        // for other vertices farther to the right than vertices[i1].
+
+        if (v[0] > support[1][0] ||
+            (v[0] == support[1][0] && v[1] > support[1][1]))
+        {
+            // New right maximum OR same right maximum but closer to top.
+            rct.index[1] = i;
+            support[1] = v;
+        }
+
+        if (v[1] > support[2][1] ||
+            (v[1] == support[2][1] && v[0] < support[2][0]))
+        {
+            // New top maximum OR same top maximum but closer to left.
+            rct.index[2] = i;
+            support[2] = v;
+        }
+
+        if (v[0] < support[3][0] ||
+            (v[0] == support[3][0] && v[1] < support[3][1]))
+        {
+            // New left minimum OR same left minimum but closer to bottom.
+            rct.index[3] = i;
+            support[3] = v;
+        }
+
+        ++i;
+    }
+
+    // The comment in the loop has the implication that support[0] = { 0, 0 },
+    // so the scaled height (support[2][1] - support[0][1]) is simply
+    // support[2][1].
+    ComputeType scaledWidth = support[1][0] - support[3][0];
+    ComputeType scaledHeight = support[2][1];
+    rct.area = scaledWidth * scaledHeight / rct.sqrLenU[1];
+    return rct;
+}
+
+template <typename InputType, typename ComputeType>
+bool MinimumVolumeBox3<InputType, ComputeType>::ComputeAngles(
+    Vector3<ComputeType> const& N, std::vector<int> const& polyline,
+    ExtrudeRectangle const& rct,
+    std::array<std::pair<ComputeType, int>, 4>& A, int& numA) const
+{
+    int const numPolyline = static_cast<int>(polyline.size());
+    numA = 0;
+    for (int k0 = 3, k1 = 0; k1 < 4; k0 = k1++)
+    {
+        if (rct.index[k0] != rct.index[k1])
+        {
+            // The rct edges are ordered in k1 as U[0], U[1], -U[0], -U[1].
+            int lookup = (k0 & 1);
+            Vector3<ComputeType> D =
+                ((k0 & 2) ? -rct.U[lookup] : rct.U[lookup]);
+            int j0 = rct.index[k0], j1 = j0 + 1;
+            if (j1 == numPolyline)
+            {
+                j1 = 0;
+            }
+            Vector3<ComputeType> E = 
+                mComputePoints[polyline[j1]] - mComputePoints[polyline[j0]];
+            Vector3<ComputeType> Eperp = Cross(N, E);
+            E = Cross(Eperp, N);
+            Vector3<ComputeType> DxE = Cross(D, E);
+            ComputeType csqrlen = Dot(DxE, DxE);
+            ComputeType dsqrlen = rct.sqrLenU[lookup];
+            ComputeType esqrlen = Dot(E, E);
+            ComputeType sinThetaSqr = csqrlen / (dsqrlen * esqrlen);
+            A[numA++] = std::make_pair(sinThetaSqr, k0);
+        }
+    }
+    return numA > 0;
+}
+
+template <typename InputType, typename ComputeType>
+std::array<int, 4> MinimumVolumeBox3<InputType, ComputeType>::SortAngles(
+    std::array<std::pair<ComputeType, int>, 4> const& A, int numA) const
+{
+    std::array<int, 4> sort = { 0, 1, 2, 3 };
+    if (numA > 1)
+    {
+        if (numA == 2)
+        {
+            if (A[sort[0]].first > A[sort[1]].first)
+            {
+                std::swap(sort[0], sort[1]);
+            }
+        }
+        else if (numA == 3)
+        {
+            if (A[sort[0]].first > A[sort[1]].first)
+            {
+                std::swap(sort[0], sort[1]);
+            }
+            if (A[sort[0]].first > A[sort[2]].first)
+            {
+                std::swap(sort[0], sort[2]);
+            }
+            if (A[sort[1]].first > A[sort[2]].first)
+            {
+                std::swap(sort[1], sort[2]);
+            }
+        }
+        else  // numA == 4
+        {
+            if (A[sort[0]].first > A[sort[1]].first)
+            {
+                std::swap(sort[0], sort[1]);
+            }
+            if (A[sort[2]].first > A[sort[3]].first)
+            {
+                std::swap(sort[2], sort[3]);
+            }
+            if (A[sort[0]].first > A[sort[2]].first)
+            {
+                std::swap(sort[0], sort[2]);
+            }
+            if (A[sort[1]].first > A[sort[3]].first)
+            {
+                std::swap(sort[1], sort[3]);
+            }
+            if (A[sort[1]].first > A[sort[2]].first)
+            {
+                std::swap(sort[1], sort[2]);
+            }
+        }
+    }
+    return sort;
+}
+
+template <typename InputType, typename ComputeType>
+bool MinimumVolumeBox3<InputType, ComputeType>::UpdateSupport(
+    std::array<std::pair<ComputeType, int>, 4> const& A, int numA,
+    std::array<int, 4> const& sort, Vector3<ComputeType> const& N,
+    std::vector<int> const& polyline, std::vector<bool>& visited,
+    ExtrudeRectangle& rct)
+{
+    // Replace the support vertices of those edges attaining minimum angle
+    // with the other endpoints of the edges.
+    int const numPolyline = static_cast<int>(polyline.size());
+    auto const& amin = A[sort[0]];
+    for (int k = 0; k < numA; ++k)
+    {
+        auto const& a = A[sort[k]];
+        if (a.first == amin.first)
+        {
+            if (++rct.index[a.second] == numPolyline)
+            {
+                rct.index[a.second] = 0;
+            }
+        }
+        else
+        {
+            break;
+        }
+    }
+
+    int bottom = rct.index[amin.second];
+    if (visited[bottom])
+    {
+        // We have already processed this polyline edge.
+        return false;
+    }
+    visited[bottom] = true;
+
+    // Cycle the vertices so that the bottom support occurs first.
+    std::array<int, 4> nextIndex;
+    for (int k = 0; k < 4; ++k)
+    {
+        nextIndex[k] = rct.index[(amin.second + k) % 4];
+    }
+    rct.index = nextIndex;
+
+    // Compute the rectangle axis directions.
+    int j1 = rct.index[0], j0 = j1 - 1;
+    if (j1 < 0)
+    {
+        j1 = numPolyline - 1;
+    }
+    Vector3<ComputeType> E =
+        mComputePoints[polyline[j1]] - mComputePoints[polyline[j0]];
+    rct.U[1] = Cross(N, E);
+    rct.U[0] = Cross(rct.U[1], N);
+    rct.sqrLenU[0] = Dot(rct.U[0], rct.U[0]);
+    rct.sqrLenU[1] = Dot(rct.U[1], rct.U[1]);
+
+    // Compute the rectangle area.
+    Vector3<ComputeType> diff[2] =
+    {
+        mComputePoints[polyline[rct.index[1]]]
+            - mComputePoints[polyline[rct.index[3]]],
+        mComputePoints[polyline[rct.index[2]]]
+            - mComputePoints[polyline[rct.index[0]]]
+    };
+    ComputeType scaledWidth = Dot(rct.U[0], diff[0]);
+    ComputeType scaledHeight = Dot(rct.U[1], diff[1]);
+    rct.area = scaledWidth * scaledHeight / rct.sqrLenU[1];
+    return true;
+}
+
+template <typename InputType, typename ComputeType>
+void MinimumVolumeBox3<InputType, ComputeType>::ConvertTo(Box const& minBox,
+    OrientedBox3<InputType>& itMinBox)
+{
+    Vector3<ComputeType> center = minBox.P;
+    for (int i = 0; i < 3; ++i)
+    {
+        ComputeType average =
+            mHalf * (minBox.range[i][0] + minBox.range[i][1]);
+        center += (average / minBox.sqrLenU[i]) * minBox.U[i];
+    }
+
+    // Calculate the squared extent using ComputeType to avoid loss of
+    // precision before computing a squared root.
+    Vector3<ComputeType> sqrExtent;
+    for (int i = 0; i < 3; ++i)
+    {
+        sqrExtent[i] = mHalf * (minBox.range[i][1] - minBox.range[i][0]);
+        sqrExtent[i] *= sqrExtent[i];
+        sqrExtent[i] /= minBox.sqrLenU[i];
+    }
+
+    for (int i = 0; i < 3; ++i)
+    {
+        itMinBox.center[i] = (InputType)center[i];
+        itMinBox.extent[i] = std::sqrt((InputType)sqrExtent[i]);
+
+        // Before converting to floating-point, factor out the maximum
+        // component using ComputeType to generate rational numbers in a
+        // range that avoids loss of precision during the conversion and
+        // normalization.
+        Vector3<ComputeType> const& axis = minBox.U[i];
+        ComputeType cmax = std::max(std::abs(axis[0]), std::abs(axis[1]));
+        cmax = std::max(cmax, std::abs(axis[2]));
+        ComputeType invCMax = mOne / cmax;
+        for (int j = 0; j < 3; ++j)
+        {
+            itMinBox.axis[i][j] = (InputType)(axis[j] * invCMax);
+        }
+        Normalize(itMinBox.axis[i]);
+    }
+
+    mVolume = (InputType)minBox.volume;
+}
+
+}
diff --git a/Engine/Plugins/Experimental/GeometryProcessing/Source/GeometryAlgorithms/Private/GTEngine/Mathematics/GteMinimumVolumeSphere3.h b/Engine/Plugins/Experimental/GeometryProcessing/Source/GeometryAlgorithms/Private/GTEngine/Mathematics/GteMinimumVolumeSphere3.h
new file mode 100644
--- /dev/null
+++ b/Engine/Plugins/Experimental/GeometryProcessing/Source/GeometryAlgorithms/Private/GTEngine/Mathematics/GteMinimumVolumeSphere3.h
@@ -0,0 +1,695 @@
+// David Eberly, Geometric Tools, Redmond WA 98052
+// Copyright (c) 1998-2019
+// Distributed under the Boost Software License, Version 1.0.
+// http://www.boost.org/LICENSE_1_0.txt
+// http://www.geometrictools.com/License/Boost/LICENSE_1_0.txt
+// File Version: 3.0.2 (2018/10/28)
+
+#pragma once
+
+#include <GTEngine/LowLevel/GteLogger.h>
+#include <GTEngine/Mathematics/GteContSphere3.h>
+#include <GTEngine/Mathematics/GteHypersphere.h>
+#include <GTEngine/Mathematics/GteLinearSystem.h>
+#include <algorithm>
+#include <random>
+
+// Compute the minimum volume sphere containing the input set of points.  The
+// algorithm randomly permutes the input points so that the construction
+// occurs in 'expected' O(N) time.  All internal minimal sphere calculations
+// store the squared radius in the radius member of Sphere3.  Only at
+// the end is a sqrt computed.
+//
+// The most robust choice for ComputeType is BSRational<T> for exact rational
+// arithmetic.  As long as this code is a correct implementation of the theory
+// (which I hope it is), you will obtain the minimum-volume sphere
+// containing the points.
+//
+// Instead, if you choose ComputeType to be float or double, floating-point
+// rounding errors can cause the UpdateSupport{2,3,4} functions to fail.
+// The failure is trapped in those functions and a simple bounding sphere is
+// computed using GetContainer in file GteContSphere3.h.  This sphere is
+// generally not the minimum-volume sphere containing the points.  The
+// minimum-volume algorithm is terminated immediately.  The sphere is
+// returned as well as a bool value of 'true' when the sphere is minimum
+// volume or 'false' when the failure is trapped.  When 'false' is returned,
+// you can try another call to the operator()(...) function.  The random
+// shuffle that occurs is highly likely to be different from the previous
+// shuffle, and there is a chance that the algorithm can succeed just because
+// of the different ordering of points.
+
+namespace gte
+{
+    template <typename InputType, typename ComputeType>
+    class MinimumVolumeSphere3
+    {
+    public:
+        bool operator()(int numPoints, Vector3<InputType> const* points, Sphere3<InputType>& minimal)
+        {
+            if (numPoints >= 1 && points)
+            {
+                // Function array to avoid switch statement in the main loop.
+                std::function<UpdateResult(int)> update[5];
+                update[1] = [this](int i) { return UpdateSupport1(i); };
+                update[2] = [this](int i) { return UpdateSupport2(i); };
+                update[3] = [this](int i) { return UpdateSupport3(i); };
+                update[4] = [this](int i) { return UpdateSupport4(i); };
+
+                // Process only the unique points.
+                std::vector<int> permuted(numPoints);
+                for (int i = 0; i < numPoints; ++i)
+                {
+                    permuted[i] = i;
+                }
+                std::sort(permuted.begin(), permuted.end(),
+                    [points](int i0, int i1) { return points[i0] < points[i1]; });
+                auto end = std::unique(permuted.begin(), permuted.end(),
+                    [points](int i0, int i1) { return points[i0] == points[i1]; });
+                permuted.erase(end, permuted.end());
+                numPoints = static_cast<int>(permuted.size());
+
+                // Create a random permutation of the points.
+                std::shuffle(permuted.begin(), permuted.end(), mDRE);
+
+                // Convert to the compute type, which is a simple copy when
+                // ComputeType is the same as InputType.
+                mComputePoints.resize(numPoints);
+                for (int i = 0; i < numPoints; ++i)
+                {
+                    for (int j = 0; j < 3; ++j)
+                    {
+                        mComputePoints[i][j] = points[permuted[i]][j];
+                    }
+                }
+
+                // Start with the first point.
+                Sphere3<ComputeType> ctMinimal = ExactSphere1(0);
+                mNumSupport = 1;
+                mSupport[0] = 0;
+
+                // The loop restarts from the beginning of the point list each
+                // time the sphere needs updating.  Linus Kallberg (Computer
+                // Science at Malardalen University in Sweden) discovered that
+                // performance is/ better when the remaining points in the
+                // array are processed before restarting.  The points
+                // processed before the point that caused the/ update are
+                // likely to be enclosed by the new sphere (or near the sphere
+                // boundary) because they were enclosed by the previous
+                // sphere.  The chances are better that points after the
+                // current one will cause growth of the bounding sphere.
+                for (int i = 1 % numPoints, n = 0; i != n; i = (i + 1) % numPoints)
+                {
+                    if (!SupportContains(i))
+                    {
+                        if (!Contains(i, ctMinimal))
+                        {
+                            auto result = update[mNumSupport](i);
+                            if (result.second == true)
+                            {
+                                if (result.first.radius > ctMinimal.radius)
+                                {
+                                    ctMinimal = result.first;
+                                    n = i;
+                                }
+                            }
+                            else
+                            {
+                                // This case can happen when ComputeType is
+                                // float or double.  See the comments at the
+                                // beginning of this file.
+                                LogWarning("ComputeType is not exact and failure occurred. Returning non-minimal sphere.");
+                                GetContainer(numPoints, points, minimal);
+                                mNumSupport = 0;
+                                mSupport.fill(0);
+                                return false;
+                            }
+                        }
+                    }
+                }
+
+                for (int j = 0; j < 3; ++j)
+                {
+                    minimal.center[j] = static_cast<InputType>(ctMinimal.center[j]);
+                }
+                minimal.radius = static_cast<InputType>(ctMinimal.radius);
+                minimal.radius = std::sqrt(minimal.radius);
+
+                for (int i = 0; i < mNumSupport; ++i)
+                {
+                    mSupport[i] = permuted[mSupport[i]];
+                }
+                return true;
+            }
+            else
+            {
+                LogError("Input must contain points.");
+                minimal.center = Vector3<InputType>::Zero();
+                minimal.radius = std::numeric_limits<InputType>::max();
+                return false;
+            }
+        }
+
+        // Member access.
+        inline int GetNumSupport() const
+        {
+            return mNumSupport;
+        }
+
+        inline std::array<int, 4> const& GetSupport() const
+        {
+            return mSupport;
+        }
+
+    private:
+        // Test whether point P is inside sphere S using squared distance and
+        // squared radius.
+        bool Contains(int i, Sphere3<ComputeType> const& sphere) const
+        {
+            // NOTE: In this algorithm, sphere.radius is the *squared radius*
+            // until the function returns at which time a square root is
+            // applied.
+            Vector3<ComputeType> diff = mComputePoints[i] - sphere.center;
+            return Dot(diff, diff) <= sphere.radius;
+        }
+
+        Sphere3<ComputeType> ExactSphere1(int i0) const
+        {
+            Sphere3<ComputeType> minimal;
+            minimal.center = mComputePoints[i0];
+            minimal.radius = (ComputeType)0;
+            return minimal;
+        }
+
+        Sphere3<ComputeType> ExactSphere2(int i0, int i1) const
+        {
+            Vector3<ComputeType> const& P0 = mComputePoints[i0];
+            Vector3<ComputeType> const& P1 = mComputePoints[i1];
+            Sphere3<ComputeType> minimal;
+            minimal.center = ((ComputeType)0.5)*(P0 + P1);
+            Vector3<ComputeType> diff = P1 - P0;
+            minimal.radius = ((ComputeType)0.25)*Dot(diff, diff);
+            return minimal;
+        }
+
+        Sphere3<ComputeType> ExactSphere3(int i0, int i1, int i2) const
+        {
+            // Compute the 2D circle containing P0, P1, and P2.  The center in
+            // barycentric coordinates is C = x0*P0 + x1*P1 + x2*P2, where
+            // x0 + x1 + x2 = 1.  The center is equidistant from the three
+            // points, so |C - P0| = |C - P1| = |C - P2| = R, where R is the
+            // radius of the circle.  From these conditions,
+            //   C - P0 = x0*E0 + x1*E1 - E0
+            //   C - P1 = x0*E0 + x1*E1 - E1
+            //   C - P2 = x0*E0 + x1*E1
+            // where E0 = P0 - P2 and E1 = P1 - P2, which leads to
+            //   r^2 = |x0*E0 + x1*E1|^2 - 2*Dot(E0, x0*E0 + x1*E1) + |E0|^2
+            //   r^2 = |x0*E0 + x1*E1|^2 - 2*Dot(E1, x0*E0 + x1*E1) + |E1|^2
+            //   r^2 = |x0*E0 + x1*E1|^2
+            // Subtracting the last equation from the first two and writing
+            // the equations as a linear system,
+            //
+            // +-                     -++   -+       +-          -+
+            // | Dot(E0,E0) Dot(E0,E1) || x0 | = 0.5 | Dot(E0,E0) |
+            // | Dot(E1,E0) Dot(E1,E1) || x1 |       | Dot(E1,E1) |
+            // +-                     -++   -+       +-          -+
+            //
+            // The following code solves this system for x0 and x1 and then
+            // evaluates the third equation in r^2 to obtain r.
+
+            Vector3<ComputeType> const& P0 = mComputePoints[i0];
+            Vector3<ComputeType> const& P1 = mComputePoints[i1];
+            Vector3<ComputeType> const& P2 = mComputePoints[i2];
+
+            Vector3<ComputeType> E0 = P0 - P2;
+            Vector3<ComputeType> E1 = P1 - P2;
+
+            Matrix2x2<ComputeType> A;
+            A(0, 0) = Dot(E0, E0);
+            A(0, 1) = Dot(E0, E1);
+            A(1, 0) = A(0, 1);
+            A(1, 1) = Dot(E1, E1);
+
+            ComputeType const half = (ComputeType)0.5;
+            Vector2<ComputeType> B{ half*A(0, 0), half*A(1, 1) };
+
+            Sphere3<ComputeType> minimal;
+            Vector2<ComputeType> X;
+            if (LinearSystem<ComputeType>::Solve(A, B, X))
+            {
+                ComputeType x2 = (ComputeType)1 - X[0] - X[1];
+                minimal.center = X[0] * P0 + X[1] * P1 + x2 * P2;
+                Vector3<ComputeType> tmp = X[0] * E0 + X[1] * E1;
+                minimal.radius = Dot(tmp, tmp);
+            }
+            else
+            {
+                minimal.center = Vector3<ComputeType>::Zero();
+                minimal.radius = (ComputeType)std::numeric_limits<InputType>::max();
+            }
+            return minimal;
+        }
+
+        Sphere3<ComputeType> ExactSphere4(int i0, int i1, int i2, int i3) const
+        {
+            // Compute the sphere containing P0, P1, P2, and P3.  The center
+            // in barycentric coordinates is C = x0*P0 + x1*P1 + x2*P2 + x3*P3,
+            // where x0 + x1 + x2 + x3 = 1.  The center is equidistant from
+            // the three points, so |C - P0| = |C - P1| = |C - P2| = |C - P3|
+            // = R, where R is the radius of the sphere.  From these
+            // conditions,
+            //   C - P0 = x0*E0 + x1*E1 + x2*E2 - E0
+            //   C - P1 = x0*E0 + x1*E1 + x2*E2 - E1
+            //   C - P2 = x0*E0 + x1*E1 + x2*E2 - E2
+            //   C - P3 = x0*E0 + x1*E1 + x2*E2
+            // where E0 = P0 - P3, E1 = P1 - P3, and E2 = P2 - P3, which
+            // leads to
+            //   r^2 = |x0*E0+x1*E1+x2*E2|^2 - 2*Dot(E0,x0*E0+x1*E1+x2*E2) + |E0|^2
+            //   r^2 = |x0*E0+x1*E1+x2*E2|^2 - 2*Dot(E1,x0*E0+x1*E1+x2*E2) + |E1|^2
+            //   r^2 = |x0*E0+x1*E1+x2*E2|^2 - 2*Dot(E2,x0*E0+x1*E1+x2*E2) + |E2|^2
+            //   r^2 = |x0*E0+x1*E1+x2*E2|^2
+            // Subtracting the last equation from the first three and writing
+            // the equations as a linear system,
+            //
+            // +-                                -++   -+       +-          -+
+            // | Dot(E0,E0) Dot(E0,E1) Dot(E0,E2) || x0 | = 0.5 | Dot(E0,E0) |
+            // | Dot(E1,E0) Dot(E1,E1) Dot(E1,E2) || x1 |       | Dot(E1,E1) |
+            // | Dot(E2,E0) Dot(E2,E1) Dot(E2,E2) || x2 |       | Dot(E2,E2) |
+            // +-                                -++   -+       +-          -+
+            //
+            // The following code solves this system for x0, x1, and x2 and
+            // then evaluates the fourth equation in r^2 to obtain r.
+
+            Vector3<ComputeType> const& P0 = mComputePoints[i0];
+            Vector3<ComputeType> const& P1 = mComputePoints[i1];
+            Vector3<ComputeType> const& P2 = mComputePoints[i2];
+            Vector3<ComputeType> const& P3 = mComputePoints[i3];
+
+            Vector3<ComputeType> E0 = P0 - P3;
+            Vector3<ComputeType> E1 = P1 - P3;
+            Vector3<ComputeType> E2 = P2 - P3;
+
+            Matrix3x3<ComputeType> A;
+            A(0, 0) = Dot(E0, E0);
+            A(0, 1) = Dot(E0, E1);
+            A(0, 2) = Dot(E0, E2);
+            A(1, 0) = A(0, 1);
+            A(1, 1) = Dot(E1, E1);
+            A(1, 2) = Dot(E1, E2);
+            A(2, 0) = A(0, 2);
+            A(2, 1) = A(1, 2);
+            A(2, 2) = Dot(E2, E2);
+
+            ComputeType const half = (ComputeType)0.5;
+            Vector3<ComputeType> B{ half*A(0, 0), half*A(1, 1), half*A(2, 2) };
+
+            Sphere3<ComputeType> minimal;
+            Vector3<ComputeType> X;
+            if (LinearSystem<ComputeType>::Solve(A, B, X))
+            {
+                ComputeType x3 = (ComputeType)1 - X[0] - X[1] - X[2];
+                minimal.center = X[0] * P0 + X[1] * P1 + X[2] * P2 + x3 * P3;
+                Vector3<ComputeType> tmp = X[0] * E0 + X[1] * E1 + X[2] * E2;
+                minimal.radius = Dot(tmp, tmp);
+            }
+            else
+            {
+                minimal.center = Vector3<ComputeType>::Zero();
+                minimal.radius = (ComputeType)std::numeric_limits<InputType>::max();
+            }
+            return minimal;
+        }
+
+        typedef std::pair<Sphere3<ComputeType>, bool> UpdateResult;
+
+        UpdateResult UpdateSupport1(int i)
+        {
+            Sphere3<ComputeType> minimal = ExactSphere2(mSupport[0], i);
+            mNumSupport = 2;
+            mSupport[1] = i;
+            return std::make_pair(minimal, true);
+        }
+
+        UpdateResult UpdateSupport2(int i)
+        {
+            // Permutations of type 2, used for calling ExactSphere2(...).
+            int const numType2 = 2;
+            int const type2[numType2][2] =
+            {
+                { 0, /*2*/ 1 },
+                { 1, /*2*/ 0 }
+            };
+
+            // Permutations of type 3, used for calling ExactSphere3(...).
+            int const numType3 = 1;  // {0, 1, 2}
+
+            Sphere3<ComputeType> sphere[numType2 + numType3];
+            ComputeType minRSqr = (ComputeType)std::numeric_limits<InputType>::max();
+            int iSphere = 0, iMinRSqr = -1;
+            int k0, k1;
+
+            // Permutations of type 2.
+            for (int j = 0; j < numType2; ++j, ++iSphere)
+            {
+                k0 = mSupport[type2[j][0]];
+                sphere[iSphere] = ExactSphere2(k0, i);
+                if (sphere[iSphere].radius < minRSqr)
+                {
+                    k1 = mSupport[type2[j][1]];
+                    if (Contains(k1, sphere[iSphere]))
+                    {
+                        minRSqr = sphere[iSphere].radius;
+                        iMinRSqr = iSphere;
+                    }
+                }
+            }
+
+            // Permutations of type 3.
+            k0 = mSupport[0];
+            k1 = mSupport[1];
+            sphere[iSphere] = ExactSphere3(k0, k1, i);
+            if (sphere[iSphere].radius < minRSqr)
+            {
+                minRSqr = sphere[iSphere].radius;
+                iMinRSqr = iSphere;
+            }
+
+            switch (iMinRSqr)
+            {
+            case 0:
+                mSupport[1] = i;
+                break;
+            case 1:
+                mSupport[0] = i;
+                break;
+            case 2:
+                mNumSupport = 3;
+                mSupport[2] = i;
+                break;
+            case -1:
+                // For exact arithmetic, iMinRSqr >= 0, but for floating-point
+                // arithmetic, round-off errors can lead to iMinRSqr == -1.
+                // When this happens, use a simple bounding sphere for the
+                // result and terminate the minimum-volume algorithm.
+                return std::make_pair(Sphere3<ComputeType>(), false);
+            }
+
+            return std::make_pair(sphere[iMinRSqr], true);
+        }
+
+        UpdateResult UpdateSupport3(int i)
+        {
+            // Permutations of type 2, used for calling ExactSphere2(...).
+            int const numType2 = 3;
+            int const type2[numType2][3] =
+            {
+                { 0, /*3*/ 1, 2 },
+                { 1, /*3*/ 0, 2 },
+                { 2, /*3*/ 0, 1 }
+            };
+
+            // Permutations of type 3, used for calling ExactSphere3(...).
+            int const numType3 = 3;
+            int const type3[numType3][3] =
+            {
+                { 0, 1, /*3*/ 2 },
+                { 0, 2, /*3*/ 1 },
+                { 1, 2, /*3*/ 0 }
+            };
+
+            // Permutations of type 4, used for calling ExactSphere4(...).
+            int const numType4 = 1;  // {0, 1, 2, 3}
+
+            Sphere3<ComputeType> sphere[numType2 + numType3 + numType4];
+            ComputeType minRSqr = (ComputeType)std::numeric_limits<InputType>::max();
+            int iSphere = 0, iMinRSqr = -1;
+            int k0, k1, k2;
+
+            // Permutations of type 2.
+            for (int j = 0; j < numType2; ++j, ++iSphere)
+            {
+                k0 = mSupport[type2[j][0]];
+                sphere[iSphere] = ExactSphere2(k0, i);
+                if (sphere[iSphere].radius < minRSqr)
+                {
+                    k1 = mSupport[type2[j][1]];
+                    k2 = mSupport[type2[j][2]];
+                    if (Contains(k1, sphere[iSphere]) && Contains(k2, sphere[iSphere]))
+                    {
+                        minRSqr = sphere[iSphere].radius;
+                        iMinRSqr = iSphere;
+                    }
+                }
+            }
+
+            // Permutations of type 3.
+            for (int j = 0; j < numType3; ++j, ++iSphere)
+            {
+                k0 = mSupport[type3[j][0]];
+                k1 = mSupport[type3[j][1]];
+                sphere[iSphere] = ExactSphere3(k0, k1, i);
+                if (sphere[iSphere].radius < minRSqr)
+                {
+                    k2 = mSupport[type3[j][2]];
+                    if (Contains(k2, sphere[iSphere]))
+                    {
+                        minRSqr = sphere[iSphere].radius;
+                        iMinRSqr = iSphere;
+                    }
+                }
+            }
+
+            // Permutations of type 4.
+            k0 = mSupport[0];
+            k1 = mSupport[1];
+            k2 = mSupport[2];
+            sphere[iSphere] = ExactSphere4(k0, k1, k2, i);
+            if (sphere[iSphere].radius < minRSqr)
+            {
+                minRSqr = sphere[iSphere].radius;
+                iMinRSqr = iSphere;
+            }
+
+            switch (iMinRSqr)
+            {
+            case 0:
+                mNumSupport = 2;
+                mSupport[1] = i;
+                break;
+            case 1:
+                mNumSupport = 2;
+                mSupport[0] = i;
+                break;
+            case 2:
+                mNumSupport = 2;
+                mSupport[0] = mSupport[2];
+                mSupport[1] = i;
+                break;
+            case 3:
+                mSupport[2] = i;
+                break;
+            case 4:
+                mSupport[1] = i;
+                break;
+            case 5:
+                mSupport[0] = i;
+                break;
+            case 6:
+                mNumSupport = 4;
+                mSupport[3] = i;
+                break;
+            case -1:
+                // For exact arithmetic, iMinRSqr >= 0, but for floating-point
+                // arithmetic, round-off errors can lead to iMinRSqr == -1.
+                // When this happens, use a simple bounding sphere for the
+                // result and terminate the minimum-area algorithm.
+                return std::make_pair(Sphere3<ComputeType>(), false);
+            }
+
+            return std::make_pair(sphere[iMinRSqr], true);
+        }
+
+        UpdateResult UpdateSupport4(int i)
+        {
+            // Permutations of type 2, used for calling ExactSphere2(...).
+            int const numType2 = 4;
+            int const type2[numType2][4] =
+            {
+                { 0, /*4*/ 1, 2, 3 },
+                { 1, /*4*/ 0, 2, 3 },
+                { 2, /*4*/ 0, 1, 3 },
+                { 3, /*4*/ 0, 1, 2 }
+            };
+
+            // Permutations of type 3, used for calling ExactSphere3(...).
+            int const numType3 = 6;
+            int const type3[numType3][4] =
+            {
+                { 0, 1, /*4*/ 2, 3 },
+                { 0, 2, /*4*/ 1, 3 },
+                { 0, 3, /*4*/ 1, 2 },
+                { 1, 2, /*4*/ 0, 3 },
+                { 1, 3, /*4*/ 0, 2 },
+                { 2, 3, /*4*/ 0, 1 }
+            };
+
+            // Permutations of type 4, used for calling ExactSphere4(...).
+            int const numType4 = 4;
+            int const type4[numType4][4] =
+            {
+                { 0, 1, 2, /*4*/ 3 },
+                { 0, 1, 3, /*4*/ 2 },
+                { 0, 2, 3, /*4*/ 1 },
+                { 1, 2, 3, /*4*/ 0 }
+            };
+
+            Sphere3<ComputeType> sphere[numType2 + numType3 + numType4];
+            ComputeType minRSqr = (ComputeType)std::numeric_limits<InputType>::max();
+            int iSphere = 0, iMinRSqr = -1;
+            int k0, k1, k2, k3;
+
+            // Permutations of type 2.
+            for (int j = 0; j < numType2; ++j, ++iSphere)
+            {
+                k0 = mSupport[type2[j][0]];
+                sphere[iSphere] = ExactSphere2(k0, i);
+                if (sphere[iSphere].radius < minRSqr)
+                {
+                    k1 = mSupport[type2[j][1]];
+                    k2 = mSupport[type2[j][2]];
+                    k3 = mSupport[type2[j][3]];
+                    if (Contains(k1, sphere[iSphere]) && Contains(k2, sphere[iSphere]) && Contains(k3, sphere[iSphere]))
+                    {
+                        minRSqr = sphere[iSphere].radius;
+                        iMinRSqr = iSphere;
+                    }
+                }
+            }
+
+            // Permutations of type 3.
+            for (int j = 0; j < numType3; ++j, ++iSphere)
+            {
+                k0 = mSupport[type3[j][0]];
+                k1 = mSupport[type3[j][1]];
+                sphere[iSphere] = ExactSphere3(k0, k1, i);
+                if (sphere[iSphere].radius < minRSqr)
+                {
+                    k2 = mSupport[type3[j][2]];
+                    k3 = mSupport[type3[j][3]];
+                    if (Contains(k2, sphere[iSphere]) && Contains(k3, sphere[iSphere]))
+                    {
+                        minRSqr = sphere[iSphere].radius;
+                        iMinRSqr = iSphere;
+                    }
+                }
+            }
+
+            // Permutations of type 4.
+            for (int j = 0; j < numType4; ++j, ++iSphere)
+            {
+                k0 = mSupport[type4[j][0]];
+                k1 = mSupport[type4[j][1]];
+                k2 = mSupport[type4[j][2]];
+                sphere[iSphere] = ExactSphere4(k0, k1, k2, i);
+                if (sphere[iSphere].radius < minRSqr)
+                {
+                    k3 = mSupport[type4[j][3]];
+                    if (Contains(k3, sphere[iSphere]))
+                    {
+                        minRSqr = sphere[iSphere].radius;
+                        iMinRSqr = iSphere;
+                    }
+                }
+            }
+
+            switch (iMinRSqr)
+            {
+            case 0:
+                mNumSupport = 2;
+                mSupport[1] = i;
+                break;
+            case 1:
+                mNumSupport = 2;
+                mSupport[0] = i;
+                break;
+            case 2:
+                mNumSupport = 2;
+                mSupport[0] = mSupport[2];
+                mSupport[1] = i;
+                break;
+            case 3:
+                mNumSupport = 2;
+                mSupport[0] = mSupport[3];
+                mSupport[1] = i;
+                break;
+            case 4:
+                mNumSupport = 3;
+                mSupport[2] = i;
+                break;
+            case 5:
+                mNumSupport = 3;
+                mSupport[1] = i;
+                break;
+            case 6:
+                mNumSupport = 3;
+                mSupport[1] = mSupport[3];
+                mSupport[2] = i;
+                break;
+            case 7:
+                mNumSupport = 3;
+                mSupport[0] = i;
+                break;
+            case 8:
+                mNumSupport = 3;
+                mSupport[0] = mSupport[3];
+                mSupport[2] = i;
+                break;
+            case 9:
+                mNumSupport = 3;
+                mSupport[0] = mSupport[3];
+                mSupport[1] = i;
+                break;
+            case 10:
+                mSupport[3] = i;
+                break;
+            case 11:
+                mSupport[2] = i;
+                break;
+            case 12:
+                mSupport[1] = i;
+                break;
+            case 13:
+                mSupport[0] = i;
+                break;
+            case -1:
+                // For exact arithmetic, iMinRSqr >= 0, but for floating-point
+                // arithmetic, round-off errors can lead to iMinRSqr == -1.
+                // When this happens, use a simple bounding sphere for the
+                // result and terminate the minimum-area algorithm.
+                return std::make_pair(Sphere3<ComputeType>(), false);
+            }
+
+            return std::make_pair(sphere[iMinRSqr], true);
+        }
+
+        // Indices of points that support current minimum volume sphere.
+        bool SupportContains(int j) const
+        {
+            for (int i = 0; i < mNumSupport; ++i)
+            {
+                if (j == mSupport[i])
+                {
+                    return true;
+                }
+            }
+            return false;
+        }
+
+        int mNumSupport;
+        std::array<int, 4> mSupport;
+
+        // Random permutation of the unique input points to produce expected
+        // linear time for the algorithm.
+        std::default_random_engine mDRE;
+        std::vector<Vector3<ComputeType>> mComputePoints;
+    };
+}
diff --git a/Engine/Plugins/Experimental/GeometryProcessing/Source/GeometryAlgorithms/Private/GTEngine/Mathematics/GteNURBSCircle.h b/Engine/Plugins/Experimental/GeometryProcessing/Source/GeometryAlgorithms/Private/GTEngine/Mathematics/GteNURBSCircle.h
new file mode 100644
--- /dev/null
+++ b/Engine/Plugins/Experimental/GeometryProcessing/Source/GeometryAlgorithms/Private/GTEngine/Mathematics/GteNURBSCircle.h
@@ -0,0 +1,143 @@
+// David Eberly, Geometric Tools, Redmond WA 98052
+// Copyright (c) 1998-2019
+// Distributed under the Boost Software License, Version 1.0.
+// http://www.boost.org/LICENSE_1_0.txt
+// http://www.geometrictools.com/License/Boost/LICENSE_1_0.txt
+// File Version: 3.18.0 (2018/10/28)
+
+#pragma once
+
+#include <GTEngine/Mathematics/GteNURBSCurve.h>
+
+// The algorithm for representing a circle as a NURBS curve or a sphere as a
+// NURBS surface is described in
+//   http://www.geometrictools.com/Documentation/NURBSCircleSphere.pdf
+// The implementations are related to the documents as shown next.
+//   NURBSQuarterCircleDegree2 implements equation (9)
+//   NURBSQuarterCircleDegree4 implements equation (10)
+//   NURBSHalfCircleDegree3 implements equation (12)
+//   NURBSFullCircleDegree3 implements Section 2.3
+
+namespace gte
+{
+    template <typename Real>
+    class NURBSQuarterCircleDegree2 : public NURBSCurve<2, Real>
+    {
+    public:
+        // Construction.  The quarter circle is x^2 + y^2 = 1 for x >= 0
+        // and y >= 0.  The direction of traversal is counterclockwise as
+        // u increase from 0 to 1.
+        NURBSQuarterCircleDegree2()
+            :
+            NURBSCurve<2, Real>(BasisFunctionInput<Real>(3, 2), nullptr, nullptr)
+        {
+            Real const sqrt2 = std::sqrt((Real)2);
+            this->mWeights[0] = sqrt2;
+            this->mWeights[1] = (Real)1;
+            this->mWeights[2] = sqrt2;
+
+            this->mControls[0] = { (Real)1, (Real)0 };
+            this->mControls[1] = { (Real)1, (Real)1 };
+            this->mControls[2] = { (Real)0, (Real)1 };
+        }
+    };
+
+    template <typename Real>
+    class NURBSQuarterCircleDegree4 : public NURBSCurve<2, Real>
+    {
+    public:
+        // Construction.  The quarter circle is x^2 + y^2 = 1 for x >= 0
+        // and y >= 0.  The direction of traversal is counterclockwise as
+        // u increases from 0 to 1.
+        NURBSQuarterCircleDegree4()
+            :
+            NURBSCurve<2, Real>(BasisFunctionInput<Real>(5, 4), nullptr, nullptr)
+        {
+            Real const sqrt2 = std::sqrt((Real)2);
+            this->mWeights[0] = (Real)1;
+            this->mWeights[1] = (Real)1;
+            this->mWeights[2] = (Real)2 * sqrt2 / (Real)3;
+            this->mWeights[3] = (Real)1;
+            this->mWeights[4] = (Real)1;
+
+            Real const x1 = (Real)1;
+            Real const y1 = (Real)0.5 / sqrt2;
+            Real const x2 = (Real)1 - sqrt2 / (Real)8;
+            this->mControls[0] = { (Real)1, (Real)0 };
+            this->mControls[1] = { x1, y1 };
+            this->mControls[2] = { x2, x2 };
+            this->mControls[3] = { y1, x1 };
+            this->mControls[4] = { (Real)0, (Real)1 };
+        }
+    };
+
+    template <typename Real>
+    class NURBSHalfCircleDegree3 : public NURBSCurve<2, Real>
+    {
+    public:
+        // Construction.  The half circle is x^2 + y^2 = 1 for x >= 0.  The
+        // direction of traversal is counterclockwise as u increases from
+        // 0 to 1.
+        NURBSHalfCircleDegree3()
+            :
+            NURBSCurve<2, Real>(BasisFunctionInput<Real>(4, 3), nullptr, nullptr)
+        {
+            Real const oneThird = (Real)1 / (Real)3;
+            this->mWeights[0] = (Real)1;
+            this->mWeights[1] = oneThird;
+            this->mWeights[2] = oneThird;
+            this->mWeights[3] = (Real)1;
+
+            this->mControls[0] = { (Real)1, (Real)0 };
+            this->mControls[1] = { (Real)1, (Real)2 };
+            this->mControls[2] = { (Real)-1, (Real)2 };
+            this->mControls[3] = { (Real)-1, (Real)0 };
+        }
+    };
+
+    template <typename Real>
+    class NURBSFullCircleDegree3 : public NURBSCurve<2, Real>
+    {
+    public:
+        // Construction.  The full circle is x^2 + y^2 = 1.  The direction of
+        // traversal is counterclockwise as u increases from 0 to 1.
+        NURBSFullCircleDegree3()
+            :
+            NURBSCurve<2, Real>(CreateBasisFunctionInput(), nullptr, nullptr)
+        {
+            Real const oneThird = (Real)1 / (Real)3;
+            this->mWeights[0] = (Real)1;
+            this->mWeights[1] = oneThird;
+            this->mWeights[2] = oneThird;
+            this->mWeights[3] = (Real)1;
+            this->mWeights[4] = oneThird;
+            this->mWeights[5] = oneThird;
+            this->mWeights[6] = (Real)1;
+
+            this->mControls[0] = { (Real)1, (Real)0 };
+            this->mControls[1] = { (Real)1, (Real)2 };
+            this->mControls[2] = { (Real)-1, (Real)2 };
+            this->mControls[3] = { (Real)-1, (Real)0 };
+            this->mControls[4] = { (Real)-1, (Real)-2 };
+            this->mControls[5] = { (Real)1, (Real)-2 };
+            this->mControls[6] = { (Real)1, (Real)0 };
+        }
+
+    private:
+        static BasisFunctionInput<Real> CreateBasisFunctionInput()
+        {
+            // We need knots (0,0,0,0,1/2,1/2,1/2,1,1,1,1).
+            BasisFunctionInput<Real> input;
+            input.numControls = 7;
+            input.degree = 3;
+            input.uniform = true;
+            input.periodic = false;
+            input.numUniqueKnots = 3;
+            input.uniqueKnots.resize(input.numUniqueKnots);
+            input.uniqueKnots[0] = { (Real)0, 4 };
+            input.uniqueKnots[1] = { (Real)0.5, 3 };
+            input.uniqueKnots[2] = { (Real)1, 4 };
+            return input;
+        }
+    };
+}
diff --git a/Engine/Plugins/Experimental/GeometryProcessing/Source/GeometryAlgorithms/Private/GTEngine/Mathematics/GteNURBSCurve.h b/Engine/Plugins/Experimental/GeometryProcessing/Source/GeometryAlgorithms/Private/GTEngine/Mathematics/GteNURBSCurve.h
new file mode 100644
--- /dev/null
+++ b/Engine/Plugins/Experimental/GeometryProcessing/Source/GeometryAlgorithms/Private/GTEngine/Mathematics/GteNURBSCurve.h
@@ -0,0 +1,229 @@
+// David Eberly, Geometric Tools, Redmond WA 98052
+// Copyright (c) 1998-2019
+// Distributed under the Boost Software License, Version 1.0.
+// http://www.boost.org/LICENSE_1_0.txt
+// http://www.geometrictools.com/License/Boost/LICENSE_1_0.txt
+// File Version: 3.0.3 (2018/10/28)
+
+#pragma once
+
+#include <GTEngine/Mathematics/GteBasisFunction.h>
+#include <GTEngine/Mathematics/GteParametricCurve.h>
+
+namespace gte
+{
+    template <int N, typename Real>
+    class NURBSCurve : public ParametricCurve<N, Real>
+    {
+    public:
+        // Construction.  If the input controls is non-null, a copy is made of
+        // the controls.  To defer setting the control points or weights, pass
+        // null pointers and later access the control points or weights via
+        // GetControls(), GetWeights(), SetControl(), or SetWeight() member
+        // functions.  The domain is t in [t[d],t[n]], where t[d] and t[n] are
+        // knots with d the degree and n the number of control points.  To
+        // validate construction, create an object as shown:
+        //     NURBSCurve<N, Real> curve(parameters);
+        //     if (!curve) { <constructor failed, handle accordingly>; }
+        NURBSCurve(BasisFunctionInput<Real> const& input,
+            Vector<N, Real> const* controls, Real const* weights)
+            :
+            ParametricCurve<N, Real>((Real)0, (Real)1),
+            mBasisFunction(input)
+        {
+            if (!mBasisFunction)
+            {
+                // Errors were already generated during construction of the
+                // basis function.
+                return;
+            }
+
+            // The mBasisFunction stores the domain but so does
+            // ParametricCurve.
+            this->mTime.front() = mBasisFunction.GetMinDomain();
+            this->mTime.back() = mBasisFunction.GetMaxDomain();
+
+            // The replication of control points for periodic splines is
+            // avoided by wrapping the i-loop index in Evaluate.
+            mControls.resize(input.numControls);
+            mWeights.resize(input.numControls);
+            if (controls)
+            {
+                std::copy(controls, controls + input.numControls, mControls.begin());
+            }
+            else
+            {
+                Vector<N, Real> zero{ (Real)0 };
+                std::fill(mControls.begin(), mControls.end(), zero);
+            }
+            if (weights)
+            {
+                std::copy(weights, weights + input.numControls, mWeights.begin());
+            }
+            else
+            {
+                std::fill(mWeights.begin(), mWeights.end(), (Real)0);
+            }
+            this->mConstructed = true;
+        }
+
+        // Member access.
+        inline BasisFunction<Real> const& GetBasisFunction() const
+        {
+            return mBasisFunction;
+        }
+
+        inline int GetNumControls() const
+        {
+            return static_cast<int>(mControls.size());
+        }
+
+        inline Vector<N, Real> const* GetControls() const
+        {
+            return mControls.data();
+        }
+
+        inline Vector<N, Real>* GetControls()
+        {
+            return mControls.data();
+        }
+
+        inline Real const* GetWeights() const
+        {
+            return mWeights.data();
+        }
+
+        inline Real* GetWeights()
+        {
+            return mWeights.data();
+        }
+
+        void SetControl(int i, Vector<N, Real> const& control)
+        {
+            if (0 <= i && i < GetNumControls())
+            {
+                mControls[i] = control;
+            }
+        }
+
+        Vector<N, Real> const& GetControl(int i) const
+        {
+            if (0 <= i && i < GetNumControls())
+            {
+                return mControls[i];
+            }
+            else
+            {
+                // Invalid index, return something.
+                return mControls[0];
+            }
+        }
+
+        void SetWeight(int i, Real weight)
+        {
+            if (0 <= i && i < GetNumControls())
+            {
+                mWeights[i] = weight;
+            }
+        }
+
+        Real const& GetWeight(int i) const
+        {
+            if (0 <= i && i < GetNumControls())
+            {
+                return mWeights[i];
+            }
+            else
+            {
+                // Invalid index, return something.
+                return mWeights[0];
+            }
+        }
+
+        // Evaluation of the curve.  The function supports derivative
+        // calculation through order 3; that is, maxOrder <= 3 is required.
+        // If you want only the position, pass in maxOrder of 0.  If you
+        // want the position and first derivative, pass in maxOrder of 1,
+        // and so on.  The output 'values' are ordered as: position, first
+        // derivative, second derivative, third derivative.
+        virtual void Evaluate(Real t, unsigned int maxOrder, Vector<N, Real> values[4]) const
+        {
+            if (!this->mConstructed)
+            {
+                // Errors were already generated during construction.
+                for (unsigned int order = 0; order < 4; ++order)
+                {
+                    values[order].MakeZero();
+                }
+                return;
+            }
+
+            int imin, imax;
+            mBasisFunction.Evaluate(t, maxOrder, imin, imax);
+
+            // Compute position.
+            Vector<N, Real> X;
+            Real w;
+            Compute(0, imin, imax, X, w);
+            Real invW = ((Real)1) / w;
+            values[0] = invW * X;
+
+            if (maxOrder >= 1)
+            {
+                // Compute first derivative.
+                Vector<N, Real> XDer1;
+                Real wDer1;
+                Compute(1, imin, imax, XDer1, wDer1);
+                values[1] = invW * (XDer1 - wDer1 * values[0]);
+
+                if (maxOrder >= 2)
+                {
+                    // Compute second derivative.
+                    Vector<N, Real> XDer2;
+                    Real wDer2;
+                    Compute(2, imin, imax, XDer2, wDer2);
+                    values[2] = invW * (XDer2 - ((Real)2) * wDer1 * values[1] -
+                        wDer2 * values[0]);
+
+                    if (maxOrder == 3)
+                    {
+                        // Compute third derivative.
+                        Vector<N, Real> XDer3;
+                        Real wDer3;
+                        Compute(3, imin, imax, XDer3, wDer3);
+                        values[3] = invW * (XDer3 - ((Real)3) * wDer1 * values[2] -
+                            ((Real)3) * wDer2 * values[1] - wDer3 * values[0]);
+                    }
+                    else
+                    {
+                        values[3].MakeZero();
+                    }
+                }
+            }
+        }
+
+    protected:
+        // Support for Evaluate(...).
+        void Compute(unsigned int order, int imin, int imax, Vector<N, Real>& X, Real& w) const
+        {
+            // The j-index introduces a tiny amount of overhead in order to
+            // handle both aperiodic and periodic splines.  For aperiodic
+            // splines, j = i always.
+
+            int numControls = GetNumControls();
+            X.MakeZero();
+            w = (Real)0;
+            for (int i = imin; i <= imax; ++i)
+            {
+                int j = (i >= numControls ? i - numControls : i);
+                Real tmp = mBasisFunction.GetValue(order, i) * mWeights[j];
+                X += tmp * mControls[j];
+                w += tmp;
+            }
+        }
+
+        BasisFunction<Real> mBasisFunction;
+        std::vector<Vector<N, Real>> mControls;
+        std::vector<Real> mWeights;
+    };
+}
