move GeometryProcessing plugin from NotForLicensees to Experimental

diff --git a/Engine/Plugins/Experimental/GeometryProcessing/Source/GeometryAlgorithms/Private/GTEngine/Mathematics/GteNURBSSphere.h b/Engine/Plugins/Experimental/GeometryProcessing/Source/GeometryAlgorithms/Private/GTEngine/Mathematics/GteNURBSSphere.h
new file mode 100644
--- /dev/null
+++ b/Engine/Plugins/Experimental/GeometryProcessing/Source/GeometryAlgorithms/Private/GTEngine/Mathematics/GteNURBSSphere.h
@@ -0,0 +1,459 @@
+// David Eberly, Geometric Tools, Redmond WA 98052
+// Copyright (c) 1998-2019
+// Distributed under the Boost Software License, Version 1.0.
+// http://www.boost.org/LICENSE_1_0.txt
+// http://www.geometrictools.com/License/Boost/LICENSE_1_0.txt
+// File Version: 3.18.0 (2018/10/28)
+
+#pragma once
+
+#include <GTEngine/Mathematics/GteNURBSSurface.h>
+#include <GTEngine/Mathematics/GteVector3.h>
+#include <functional>
+
+// The algorithm for representing a circle as a NURBS curve or a sphere as a
+// NURBS surface is described in
+//   http://www.geometrictools.com/Documentation/NURBSCircleSphere.pdf
+// The implementations are related to the documents as shown next.
+//   NURBSEighthSphereDegree4 implements Section 3.1.2 (triangular domain)
+//   NURBSHalfSphereDegree3 implements Section 3.2 (rectangular domain)
+//   NURBSFullSphereDegree3 implements Section 2.3 (rectangular domain)
+// TODO: The class NURBSSurface currently assumes a rectangular domain.
+// Once support is added for triangular domains, make that new class a
+// base class of the sphere-representing NURBS.  This will allow sharing
+// of the NURBS basis functions and evaluation framework.
+
+namespace gte
+{
+    template <typename Real>
+    class NURBSEighthSphereDegree4
+    {
+    public:
+        // Construction.  The eigth sphere is x^2 + y^2 + z^2 = 1 for x >= 0,
+        // y >= 0 and z >= 0.
+        NURBSEighthSphereDegree4()
+        {
+            Real const sqrt2 = std::sqrt((Real)2);
+            Real const sqrt3 = std::sqrt((Real)3);
+            Real const a0 = (sqrt3 - (Real)1) / sqrt3;
+            Real const a1 = (sqrt3 + (Real)1) / ((Real)2 * sqrt3);
+            Real const a2 = (Real)1 - ((Real)5 - sqrt2) * ((Real)7 - sqrt3) / (Real)46;
+            Real const b0 = (Real)4 * sqrt3 * (sqrt3 - (Real)1);
+            Real const b1 = (Real)3 * sqrt2;
+            Real const b2 = (Real)4;
+            Real const b3 = sqrt2 * ((Real)3 + (Real)2 * sqrt2 - sqrt3) / sqrt3;
+
+            mControls[0][0] = { (Real)0, (Real)0, (Real)1 };   // P004
+            mControls[0][1] = { (Real)0, a0,   (Real)1 };   // P013
+            mControls[0][2] = { (Real)0, a1,   a1 };     // P022
+            mControls[0][3] = { (Real)0, (Real)1, a0 };     // P031
+            mControls[0][4] = { (Real)0, (Real)1, (Real)0 };   // P040
+
+            mControls[1][0] = { a0,   (Real)0, (Real)1 };   // P103
+            mControls[1][1] = { a2,   a2,   (Real)1 };   // P112
+            mControls[1][2] = { a2,   (Real)1, a2 };     // P121
+            mControls[1][3] = { a0,   (Real)1, (Real)0 };   // P130
+            mControls[1][4] = { (Real)0, (Real)0, (Real)0 };   // unused
+
+            mControls[2][0] = { a1,   (Real)0, a1 };     // P202
+            mControls[2][1] = { (Real)1, a2,   a2 };     // P211
+            mControls[2][2] = { a1,   a1,   (Real)0 };   // P220
+            mControls[2][3] = { (Real)0, (Real)0, (Real)0 };   // unused
+            mControls[2][4] = { (Real)0, (Real)0, (Real)0 };   // unused
+
+            mControls[3][0] = { (Real)1, (Real)0, a0 };     // P301
+            mControls[3][1] = { (Real)1, a0,   (Real)0 };   // P310
+            mControls[3][2] = { (Real)0, (Real)0, (Real)0 };   // unused
+            mControls[3][3] = { (Real)0, (Real)0, (Real)0 };   // unused
+            mControls[3][4] = { (Real)0, (Real)0, (Real)0 };   // unused
+
+            mControls[4][0] = { (Real)1, (Real)0, (Real)0 };   // P400
+            mControls[4][1] = { (Real)0, (Real)0, (Real)0 };   // unused
+            mControls[4][2] = { (Real)0, (Real)0, (Real)0 };   // unused
+            mControls[4][3] = { (Real)0, (Real)0, (Real)0 };   // unused
+            mControls[4][4] = { (Real)0, (Real)0, (Real)0 };   // unused
+
+            mWeights[0][0] = b0;    // w004
+            mWeights[0][1] = b1;    // w013
+            mWeights[0][2] = b2;    // w022
+            mWeights[0][3] = b1;    // w031
+            mWeights[0][4] = b0;    // w040
+
+            mWeights[1][0] = b1;    // w103
+            mWeights[1][1] = b3;    // w112
+            mWeights[1][2] = b3;    // w121
+            mWeights[1][3] = b1;    // w130
+            mWeights[1][4] = (Real)0;  // unused
+
+            mWeights[2][0] = b2;    // w202
+            mWeights[2][1] = b3;    // w211
+            mWeights[2][2] = b2;    // w220
+            mWeights[2][3] = (Real)0;  // unused
+            mWeights[2][4] = (Real)0;  // unused
+
+            mWeights[3][0] = b1;    // w301
+            mWeights[3][1] = b1;    // w310
+            mWeights[3][2] = (Real)0;  // unused
+            mWeights[3][3] = (Real)0;  // unused
+            mWeights[3][4] = (Real)0;  // unused
+
+            mWeights[4][0] = b0;    // w400
+            mWeights[4][1] = (Real)0;  // unused
+            mWeights[4][2] = (Real)0;  // unused
+            mWeights[4][3] = (Real)0;  // unused
+            mWeights[4][4] = (Real)0;  // unused
+        }
+
+        // Evaluation of the surface.  The function supports derivative
+        // calculation through order 2; that is, maxOrder <= 2 is required.
+        // If you want only the position, pass in maxOrder of 0.  If you want
+        // the position and first-order derivatives, pass in maxOrder of 1,
+        // and so on.  The output 'values' are ordered as: position X;
+        // first-order derivatives dX/du, dX/dv; second-order derivatives
+        // d2X/du2, d2X/dudv, d2X/dv2.
+        void Evaluate(Real u, Real v, unsigned int maxOrder, Vector<3, Real> values[6]) const
+        {
+            // TODO: Some of the polynomials are used in other polynomials.
+            // Optimize the code by eliminating the redundant computations.
+
+            Real w = (Real)1 - u - v;
+            Real uu = u * u, uv = u * v, uw = u * w, vv = v * v, vw = v * w, ww = w * w;
+
+            // Compute the order-0 polynomials.  Only the elements to be used
+            // are filled in.  The other terms are uninitialized but never
+            // used.
+            Real B[5][5];
+            B[0][0] = ww * ww;
+            B[0][1] = (Real)4 * vw * ww;
+            B[0][2] = (Real)6 * vv * ww;
+            B[0][3] = (Real)4 * vv * vw;
+            B[0][4] = vv * vv;
+            B[1][0] = (Real)4 * uw * ww;
+            B[1][1] = (Real)12 * uv * ww;
+            B[1][2] = (Real)12 * uv * vw;
+            B[1][3] = (Real)4 * uv * vv;
+            B[2][0] = (Real)6 * uu * ww;
+            B[2][1] = (Real)12 * uu * vw;
+            B[2][2] = (Real)6 * uu * vv;
+            B[3][0] = (Real)4 * uu * uw;
+            B[3][1] = (Real)4 * uu * uv;
+            B[4][0] = uu * uu;
+
+            // Compute the NURBS position.
+            Vector<3, Real> N{ (Real)0, (Real)0, (Real)0 };
+            Real D(0);
+            for (int j1 = 0; j1 <= 4; ++j1)
+            {
+                for (int j0 = 0; j0 <= 4 - j1; ++j0)
+                {
+                    Real product = mWeights[j1][j0] * B[j1][j0];
+                    N += product * mControls[j1][j0];
+                    D += product;
+                }
+            }
+            values[0] = N / D;
+
+            if (maxOrder >= 1)
+            {
+                // Compute the order-1 polynomials.  Only the elements to be
+                // used are filled in.  The other terms are uninitialized but
+                // never used.
+                Real WmU = w - u;
+                Real WmTwoU = WmU - u;
+                Real WmThreeU = WmTwoU - u;
+                Real TwoWmU = w + WmU;
+                Real ThreeWmU = w + TwoWmU;
+                Real WmV = w - v;
+                Real WmTwoV = WmV - v;
+                Real WmThreeV = WmTwoV - v;
+                Real TwoWmV = w + WmV;
+                Real ThreeWmV = w + TwoWmV;
+                Real Dsqr = D * D;
+
+                Real Bu[5][5];
+                Bu[0][0] = (Real)-4 * ww * w;
+                Bu[0][1] = (Real)-12 * v * ww;
+                Bu[0][2] = (Real)-12 * vv * w;
+                Bu[0][3] = (Real)-4 * v * vv;
+                Bu[0][4] = (Real)0;
+                Bu[1][0] = (Real)4 * ww * WmThreeU;
+                Bu[1][1] = (Real)12 * vw * WmTwoU;
+                Bu[1][2] = (Real)12 * vv * WmU;
+                Bu[1][3] = (Real)4 * vv;
+                Bu[2][0] = (Real)12 * uw * WmU;
+                Bu[2][1] = (Real)12 * uv * TwoWmU;
+                Bu[2][2] = (Real)12 * u * vv;
+                Bu[3][0] = (Real)4 * uu * ThreeWmU;
+                Bu[3][1] = (Real)12 * uu * v;
+                Bu[4][0] = (Real)4 * uu * u;
+
+                Real Bv[5][5];
+                Bv[0][0] = (Real)-4 * ww * w;
+                Bv[0][1] = (Real)4 * ww * WmThreeV;
+                Bv[0][2] = (Real)12 * vw * WmV;
+                Bv[0][3] = (Real)4 * vv * ThreeWmV;
+                Bv[0][4] = (Real)4 * vv * v;
+                Bv[1][0] = (Real)-12 * u * ww;
+                Bv[1][1] = (Real)12 * uw * WmTwoV;
+                Bv[1][2] = (Real)12 * uv * TwoWmV;
+                Bv[1][3] = (Real)12 * u * vv;
+                Bv[2][0] = (Real)-12 * uu * w;
+                Bv[2][1] = (Real)12 * uu * WmV;
+                Bv[2][2] = (Real)12 * uu * v;
+                Bv[3][0] = (Real)-4 * uu * u;
+                Bv[3][1] = (Real)4 * uu * u;
+                Bv[4][0] = (Real)0;
+
+                Vector<3, Real> Nu{ (Real)0, (Real)0, (Real)0 };
+                Vector<3, Real> Nv{ (Real)0, (Real)0, (Real)0 };
+                Real Du(0), Dv(0);
+                for (int j1 = 0; j1 <= 4; ++j1)
+                {
+                    for (int j0 = 0; j0 <= 4 - j1; ++j0)
+                    {
+                        Real product = mWeights[j1][j0] * Bu[j1][j0];
+                        Nu += product * mControls[j1][j0];
+                        Du += product;
+                        product = mWeights[j1][j0] * Bv[j1][j0];
+                        Nv += product * mControls[j1][j0];
+                        Dv += product;
+                    }
+                }
+                Vector<3, Real> numerDU = D * Nu - Du * N;
+                Vector<3, Real> numerDV = D * Nv - Dv * N;
+                values[1] = numerDU / Dsqr;
+                values[2] = numerDV / Dsqr;
+
+                if (maxOrder >= 2)
+                {
+                    // Compute the order-2 polynomials.  Only the elements to
+                    // be used are filled in.  The other terms are
+                    // uninitialized but never used.
+                    Real Dcub = Dsqr * D;
+
+                    Real Buu[5][5];
+                    Buu[0][0] = (Real)12 * ww;
+                    Buu[0][1] = (Real)24 * vw;
+                    Buu[0][2] = (Real)12 * vv;
+                    Buu[0][3] = (Real)0;
+                    Buu[0][4] = (Real)0;
+                    Buu[1][0] = (Real)-24 * w * WmU;
+                    Buu[1][1] = (Real)-24 * v * TwoWmU;
+                    Buu[1][2] = (Real)-24 * vv;
+                    Buu[1][3] = (Real)0;
+                    Buu[2][0] = (Real)12 * (ww - (Real)4 * uw + uu);
+                    Buu[2][1] = (Real)24 * v * WmTwoU;
+                    Buu[2][2] = (Real)12 * vv;
+                    Buu[3][0] = (Real)24 * u * WmU;
+                    Buu[3][1] = (Real)24 * uv;
+                    Buu[4][0] = (Real)12 * uu;
+
+                    Real Buv[5][5];
+                    Buv[0][0] = (Real)12 * ww;
+                    Buv[0][1] = (Real)-12 * w * WmTwoV;
+                    Buv[0][2] = (Real)-12 * v * TwoWmV;
+                    Buv[0][3] = (Real)-12 * vv;
+                    Buv[0][4] = (Real)0;
+                    Buv[1][0] = (Real)-12 * w * WmTwoU;
+                    Buv[1][1] = (Real)12 * (ww + (Real)2 * (uv - uw - vw));
+                    Buv[1][2] = (Real)12 * v * ((Real)2 * WmU - v);
+                    Buv[1][3] = (Real)12 * vv;
+                    Buv[2][0] = (Real)-12 * u * TwoWmU;
+                    Buv[2][1] = (Real)12 * u * ((Real)2 * WmV - u);
+                    Buv[2][2] = (Real)24 * uv;
+                    Buv[3][0] = (Real)-12 * uu;
+                    Buv[3][1] = (Real)12 * uu;
+                    Buv[4][0] = (Real)0;
+
+                    Real Bvv[5][5];
+                    Bvv[0][0] = (Real)12 * ww;
+                    Bvv[0][1] = (Real)-24 * w * WmV;
+                    Bvv[0][2] = (Real)12 * (ww - (Real)4 * vw + vv);
+                    Bvv[0][3] = (Real)24 * v * WmV;
+                    Bvv[0][4] = (Real)12 * vv;
+                    Bvv[1][0] = (Real)24 * uw;
+                    Bvv[1][1] = (Real)-24 * u * TwoWmV;
+                    Bvv[1][2] = (Real)24 * u * WmTwoV;
+                    Bvv[1][3] = (Real)24 * uv;
+                    Bvv[2][0] = (Real)12 * uu;
+                    Bvv[2][1] = (Real)-24 * uu;
+                    Bvv[2][2] = (Real)12 * uu;
+                    Bvv[3][0] = (Real)0;
+                    Bvv[3][1] = (Real)0;
+                    Bvv[4][0] = (Real)0;
+
+                    Vector<3, Real> Nuu{ (Real)0, (Real)0, (Real)0 };
+                    Vector<3, Real> Nuv{ (Real)0, (Real)0, (Real)0 };
+                    Vector<3, Real> Nvv{ (Real)0, (Real)0, (Real)0 };
+                    Real Duu(0), Duv(0), Dvv(0);
+                    for (int j1 = 0; j1 <= 4; ++j1)
+                    {
+                        for (int j0 = 0; j0 <= 4 - j1; ++j0)
+                        {
+                            Real product = mWeights[j1][j0] * Buu[j1][j0];
+                            Nuu += product * mControls[j1][j0];
+                            Duu += product;
+                            product = mWeights[j1][j0] * Buv[j1][j0];
+                            Nuv += product * mControls[j1][j0];
+                            Duv += product;
+                            product = mWeights[j1][j0] * Bvv[j1][j0];
+                            Nvv += product * mControls[j1][j0];
+                            Dvv += product;
+                        }
+                    }
+                    Vector<3, Real> termDuu = D * (D * Nuu - Duu * N);
+                    Vector<3, Real> termDuv = D * (D * Nuv - Duv * N - Du * Nv - Dv * Nu);
+                    Vector<3, Real> termDvv = D * (D * Nvv - Dvv * N);
+                    values[3] = (D * termDuu - (Real)2 * Du * numerDU) / Dcub;
+                    values[4] = (D * termDuv + (Real)2 * Du * Dv * N) / Dcub;
+                    values[5] = (D * termDvv - (Real)2 * Dv * numerDV) / Dcub;
+                }
+            }
+        }
+
+    private:
+        // For simplicity of the implementation, 2-dimensional arrays
+        // of size 5-by-5 are used.  Only array[r][c] is used where
+        // 0 <= r <= 4 and 0 <= c < 4 - r.
+        Vector3<Real> mControls[5][5];
+        Real mWeights[5][5];
+    };
+
+    template <typename Real>
+    class NURBSHalfSphereDegree3 : public NURBSSurface<3, Real>
+    {
+    public:
+        NURBSHalfSphereDegree3()
+            :
+            NURBSSurface<3, Real>(BasisFunctionInput<Real>(4, 3),
+                BasisFunctionInput<Real>(4, 3), nullptr, nullptr)
+        {
+            // weight[j][i] is mWeights[i + 4 * j], 0 <= i < 4, 0 <= j < 4
+            Real const oneThird = (Real)1 / (Real)3;
+            Real const oneNinth = (Real)1 / (Real)9;
+            this->mWeights[0] = (Real)1;
+            this->mWeights[1] = oneThird;
+            this->mWeights[2] = oneThird;
+            this->mWeights[3] = (Real)1;
+            this->mWeights[4] = oneThird;
+            this->mWeights[5] = oneNinth;
+            this->mWeights[6] = oneNinth;
+            this->mWeights[7] = oneThird;
+            this->mWeights[8] = oneThird;
+            this->mWeights[9] = oneNinth;
+            this->mWeights[10] = oneNinth;
+            this->mWeights[11] = oneThird;
+            this->mWeights[12] = (Real)1;
+            this->mWeights[13] = oneThird;
+            this->mWeights[14] = oneThird;
+            this->mWeights[15] = (Real)1;
+
+            // control[j][i] is mControls[i + 4 * j], 0 <= i < 4, 0 <= j < 4
+            this->mControls[0] = { (Real)0, (Real)0, (Real)1 };
+            this->mControls[1] = { (Real)0, (Real)0, (Real)1 };
+            this->mControls[2] = { (Real)0, (Real)0, (Real)1 };
+            this->mControls[3] = { (Real)0, (Real)0, (Real)1 };
+            this->mControls[4] = { (Real)2, (Real)0, (Real)1 };
+            this->mControls[5] = { (Real)2, (Real)4, (Real)1 };
+            this->mControls[6] = { (Real)-2, (Real)4, (Real)1 };
+            this->mControls[7] = { (Real)-2, (Real)0, (Real)1 };
+            this->mControls[8] = { (Real)2, (Real)0, (Real)-1 };
+            this->mControls[9] = { (Real)2, (Real)4, (Real)-1 };
+            this->mControls[10] = { (Real)-2, (Real)4, (Real)-1 };
+            this->mControls[11] = { (Real)-2, (Real)0, (Real)-1 };
+            this->mControls[12] = { (Real)0, (Real)0, (Real)-1 };
+            this->mControls[13] = { (Real)0, (Real)0, (Real)-1 };
+            this->mControls[14] = { (Real)0, (Real)0, (Real)-1 };
+            this->mControls[15] = { (Real)0, (Real)0, (Real)-1 };
+        }
+    };
+
+    template <typename Real>
+    class NURBSFullSphereDegree3 : public NURBSSurface<3, Real>
+    {
+    public:
+        NURBSFullSphereDegree3()
+            :
+            NURBSSurface<3, Real>(BasisFunctionInput<Real>(4, 3),
+                CreateBasisFunctionInputV(), nullptr, nullptr)
+        {
+            // weight[j][i] is mWeights[i + 4 * j], 0 <= i < 4, 0 <= j < 7
+            Real const oneThird = (Real)1 / (Real)3;
+            Real const oneNinth = (Real)1 / (Real)9;
+            this->mWeights[0] = (Real)1;
+            this->mWeights[4] = oneThird;
+            this->mWeights[8] = oneThird;
+            this->mWeights[12] = (Real)1;
+            this->mWeights[16] = oneThird;
+            this->mWeights[20] = oneThird;
+            this->mWeights[24] = (Real)1;
+            this->mWeights[1] = oneThird;
+            this->mWeights[5] = oneNinth;
+            this->mWeights[9] = oneNinth;
+            this->mWeights[13] = oneThird;
+            this->mWeights[17] = oneNinth;
+            this->mWeights[21] = oneNinth;
+            this->mWeights[25] = oneThird;
+            this->mWeights[2] = oneThird;
+            this->mWeights[6] = oneNinth;
+            this->mWeights[10] = oneNinth;
+            this->mWeights[14] = oneThird;
+            this->mWeights[18] = oneNinth;
+            this->mWeights[22] = oneNinth;
+            this->mWeights[26] = oneThird;
+            this->mWeights[3] = (Real)1;
+            this->mWeights[7] = oneThird;
+            this->mWeights[11] = oneThird;
+            this->mWeights[15] = (Real)1;
+            this->mWeights[19] = oneThird;
+            this->mWeights[23] = oneThird;
+            this->mWeights[27] = (Real)1;
+
+            // control[j][i] is mControls[i + 4 * j], 0 <= i < 4, 0 <= j < 7
+            this->mControls[0] = { (Real)0, (Real)0, (Real)1 };
+            this->mControls[4] = { (Real)0, (Real)0, (Real)1 };
+            this->mControls[8] = { (Real)0, (Real)0, (Real)1 };
+            this->mControls[12] = { (Real)0, (Real)0, (Real)1 };
+            this->mControls[16] = { (Real)0, (Real)0, (Real)1 };
+            this->mControls[20] = { (Real)0, (Real)0, (Real)1 };
+            this->mControls[24] = { (Real)0, (Real)0, (Real)1 };
+            this->mControls[1] = { (Real)2, (Real)0, (Real)1 };
+            this->mControls[5] = { (Real)2, (Real)4, (Real)1 };
+            this->mControls[9] = { (Real)-2, (Real)4, (Real)1 };
+            this->mControls[13] = { (Real)-2, (Real)0, (Real)1 };
+            this->mControls[17] = { (Real)-2, (Real)-4, (Real)1 };
+            this->mControls[21] = { (Real)2, (Real)-4, (Real)1 };
+            this->mControls[25] = { (Real)2, (Real)0, (Real)1 };
+            this->mControls[2] = { (Real)2, (Real)0, (Real)-1 };
+            this->mControls[6] = { (Real)2, (Real)4, (Real)-1 };
+            this->mControls[10] = { (Real)-2, (Real)4, (Real)-1 };
+            this->mControls[14] = { (Real)-2, (Real)0, (Real)-1 };
+            this->mControls[18] = { (Real)-2, (Real)-4, (Real)-1 };
+            this->mControls[22] = { (Real)2, (Real)-4, (Real)-1 };
+            this->mControls[26] = { (Real)2, (Real)0, (Real)-1 };
+            this->mControls[3] = { (Real)0, (Real)0, (Real)-1 };
+            this->mControls[7] = { (Real)0, (Real)0, (Real)-1 };
+            this->mControls[11] = { (Real)0, (Real)0, (Real)-1 };
+            this->mControls[15] = { (Real)0, (Real)0, (Real)-1 };
+            this->mControls[19] = { (Real)0, (Real)0, (Real)-1 };
+            this->mControls[23] = { (Real)0, (Real)0, (Real)-1 };
+            this->mControls[27] = { (Real)0, (Real)0, (Real)-1 };
+        }
+
+    private:
+        static BasisFunctionInput<Real> CreateBasisFunctionInputV()
+        {
+            BasisFunctionInput<Real> input;
+            input.numControls = 7;
+            input.degree = 3;
+            input.uniform = true;
+            input.periodic = false;
+            input.numUniqueKnots = 3;
+            input.uniqueKnots.resize(input.numUniqueKnots);
+            input.uniqueKnots[0] = { (Real)0, 4 };
+            input.uniqueKnots[1] = { (Real)0.5, 3 };
+            input.uniqueKnots[2] = { (Real)1, 4 };
+            return input;
+        }
+    };
+}
diff --git a/Engine/Plugins/Experimental/GeometryProcessing/Source/GeometryAlgorithms/Private/GTEngine/Mathematics/GteNURBSSurface.h b/Engine/Plugins/Experimental/GeometryProcessing/Source/GeometryAlgorithms/Private/GTEngine/Mathematics/GteNURBSSurface.h
new file mode 100644
--- /dev/null
+++ b/Engine/Plugins/Experimental/GeometryProcessing/Source/GeometryAlgorithms/Private/GTEngine/Mathematics/GteNURBSSurface.h
@@ -0,0 +1,253 @@
+// David Eberly, Geometric Tools, Redmond WA 98052
+// Copyright (c) 1998-2019
+// Distributed under the Boost Software License, Version 1.0.
+// http://www.boost.org/LICENSE_1_0.txt
+// http://www.geometrictools.com/License/Boost/LICENSE_1_0.txt
+// File Version: 3.0.3 (2018/10/28)
+
+#pragma once
+
+#include <GTEngine/Mathematics/GteVector.h>
+#include <GTEngine/Mathematics/GteBasisFunction.h>
+#include <GTEngine/Mathematics/GteParametricSurface.h>
+
+namespace gte
+{
+    template <int N, typename Real>
+    class NURBSSurface : public ParametricSurface<N, Real>
+    {
+    public:
+        // Construction.  If the input controls is non-null, a copy is made of
+        // the controls.  To defer setting the control points or weights, pass
+        // null pointers and later access the control points or weights via
+        // GetControls(), GetWeights(), SetControl(), or SetWeight() member
+        // functions.  The 'controls' and 'weights' must be stored in
+        // row-major order, attribute[i0 + numControls0*i1].  As a 2D array,
+        // this corresponds to attribute2D[i1][i0].  To validate construction,
+        // create an object as shown:
+        //     NURBSSurface<N, Real> surface(parameters);
+        //     if (!surface) { <constructor failed, handle accordingly>; }
+        NURBSSurface(BasisFunctionInput<Real> const& input0,
+            BasisFunctionInput<Real> const& input1,
+            Vector<N, Real> const* controls, Real const* weights)
+            :
+            ParametricSurface<N, Real>((Real)0, (Real)1, (Real)0, (Real)1, true)
+        {
+            BasisFunctionInput<Real> const* input[2] = { &input0, &input1 };
+            for (int i = 0; i < 2; ++i)
+            {
+                mNumControls[i] = input[i]->numControls;
+                mBasisFunction[i].Create(*input[i]);
+                if (!mBasisFunction[i])
+                {
+                    // Errors were already generated during construction of
+                    // the basis functions.
+                    return;
+                }
+            }
+
+            // The mBasisFunction stores the domain but so does
+            // ParametricSurface.
+            this->mUMin = mBasisFunction[0].GetMinDomain();
+            this->mUMax = mBasisFunction[0].GetMaxDomain();
+            this->mVMin = mBasisFunction[1].GetMinDomain();
+            this->mVMax = mBasisFunction[1].GetMaxDomain();
+
+            // The replication of control points for periodic splines is
+            // avoided by wrapping the i-loop index in Evaluate.
+            int numControls = mNumControls[0] * mNumControls[1];
+            mControls.resize(numControls);
+            mWeights.resize(numControls);
+            if (controls)
+            {
+                std::copy(controls, controls + numControls, mControls.begin());
+            }
+            else
+            {
+                memset(mControls.data(), 0, mControls.size() * sizeof(mControls[0]));
+            }
+            if (weights)
+            {
+                std::copy(weights, weights + numControls, mWeights.begin());
+            }
+            else
+            {
+                memset(mWeights.data(), 0, mWeights.size() * sizeof(mWeights[0]));
+            }
+            this->mConstructed = true;
+        }
+
+        // Member access.  The index 'dim' must be in {0,1}.
+        inline BasisFunction<Real> const& GetBasisFunction(int dim) const
+        {
+            return mBasisFunction[dim];
+        }
+
+        inline int GetNumControls(int dim) const
+        {
+            return mNumControls[dim];
+        }
+
+        inline Vector<N, Real> const* GetControls() const
+        {
+            return mControls.data();
+        }
+
+        inline Vector<N, Real>* GetControls()
+        {
+            return mControls.data();
+        }
+
+        inline Real const* GetWeights() const
+        {
+            return mWeights.data();
+        }
+
+        inline Real* GetWeights()
+        {
+            return mWeights.data();
+        }
+
+        void SetControl(int i0, int i1, Vector<N, Real> const& control)
+        {
+            if (0 <= i0 && i0 < GetNumControls(0) && 0 <= i1 && i1 < GetNumControls(1))
+            {
+                mControls[i0 + mNumControls[0] * i1] = control;
+            }
+        }
+
+        Vector<N, Real> const& GetControl(int i0, int i1) const
+        {
+            if (0 <= i0 && i0 < GetNumControls(0) && 0 <= i1 && i1 < GetNumControls(1))
+            {
+                return mControls[i0 + mNumControls[0] * i1];
+            }
+            else
+            {
+                return mControls[0];
+            }
+        }
+
+        void SetWeight(int i0, int i1, Real weight)
+        {
+            if (0 <= i0 && i0 < GetNumControls(0) && 0 <= i1 && i1 < GetNumControls(1))
+            {
+                mWeights[i0 + mNumControls[0] * i1] = weight;
+            }
+        }
+
+        Real const& GetWeight(int i0, int i1) const
+        {
+            if (0 <= i0 && i0 < GetNumControls(0) && 0 <= i1 && i1 < GetNumControls(1))
+            {
+                return mWeights[i0 + mNumControls[0] * i1];
+            }
+            else
+            {
+                return mWeights[0];
+            }
+        }
+
+        // Evaluation of the surface.  The function supports derivative
+        // calculation through order 2; that is, maxOrder <= 2 is required.
+        // If you want only the position, pass in maxOrder of 0.  If you want
+        // the position and first-order derivatives, pass in maxOrder of 1,
+        // and so on.  The output 'values' are ordered as: position X;
+        // first-order derivatives dX/du, dX/dv; second-order derivatives
+        // d2X/du2, d2X/dudv, d2X/dv2.
+        virtual void Evaluate(Real u, Real v, unsigned int maxOrder, Vector<N, Real> values[6]) const
+        {
+            if (!this->mConstructed)
+            {
+                // Errors were already generated during construction.
+                for (int i = 0; i < 6; ++i)
+                {
+                    values[i].MakeZero();
+                }
+                return;
+            }
+
+            int iumin, iumax, ivmin, ivmax;
+            mBasisFunction[0].Evaluate(u, maxOrder, iumin, iumax);
+            mBasisFunction[1].Evaluate(v, maxOrder, ivmin, ivmax);
+
+            // Compute position.
+            Vector<N, Real> X;
+            Real w;
+            Compute(0, 0, iumin, iumax, ivmin, ivmax, X, w);
+            Real invW = ((Real)1) / w;
+            values[0] = invW * X;
+
+            if (maxOrder >= 1)
+            {
+                // Compute first-order derivatives.
+                Vector<N, Real> XDerU;
+                Real wDerU;
+                Compute(1, 0, iumin, iumax, ivmin, ivmax, XDerU, wDerU);
+                values[1] = invW * (XDerU - wDerU * values[0]);
+
+                Vector<N, Real> XDerV;
+                Real wDerV;
+                Compute(0, 1, iumin, iumax, ivmin, ivmax, XDerV, wDerV);
+                values[2] = invW * (XDerV - wDerV * values[0]);
+
+                if (maxOrder >= 2)
+                {
+                    // Compute second-order derivatives.
+                    Vector<N, Real> XDerUU;
+                    Real wDerUU;
+                    Compute(2, 0, iumin, iumax, ivmin, ivmax, XDerUU, wDerUU);
+                    values[3] = invW * (XDerUU - ((Real)2) * wDerU * values[1] -
+                        wDerUU * values[0]);
+
+                    Vector<N, Real> XDerUV;
+                    Real wDerUV;
+                    Compute(1, 1, iumin, iumax, ivmin, ivmax, XDerUV, wDerUV);
+                    values[4] = invW * (XDerUV - wDerU * values[2] - wDerV * values[1]
+                        - wDerUV * values[0]);
+
+                    Vector<N, Real> XDerVV;
+                    Real wDerVV;
+                    Compute(0, 2, iumin, iumax, ivmin, ivmax, XDerVV, wDerVV);
+                    values[5] = invW * (XDerVV - ((Real)2) * wDerV * values[2] -
+                        wDerVV * values[0]);
+                }
+            }
+        }
+
+    protected:
+        // Support for Evaluate(...).
+        void Compute(unsigned int uOrder, unsigned int vOrder, int iumin,
+            int iumax, int ivmin, int ivmax, Vector<N, Real>& X, Real& w) const
+        {
+            // The j*-indices introduce a tiny amount of overhead in order to handle
+            // both aperiodic and periodic splines.  For aperiodic splines, j* = i*
+            // always.
+
+            int const numControls0 = mNumControls[0];
+            int const numControls1 = mNumControls[1];
+            X.MakeZero();
+            w = (Real)0;
+            for (int iv = ivmin; iv <= ivmax; ++iv)
+            {
+                Real tmpv = mBasisFunction[1].GetValue(vOrder, iv);
+                int jv = (iv >= numControls1 ? iv - numControls1 : iv);
+                for (int iu = iumin; iu <= iumax; ++iu)
+                {
+                    Real tmpu = mBasisFunction[0].GetValue(uOrder, iu);
+                    int ju = (iu >= numControls0 ? iu - numControls0 : iu);
+                    int index = ju + numControls0 * jv;
+                    Real tmp = tmpu * tmpv * mWeights[index];
+                    X += tmp * mControls[index];
+                    w += tmp;
+                }
+            }
+        }
+
+        std::array<BasisFunction<Real>, 2> mBasisFunction;
+        std::array<int, 2> mNumControls;
+        std::vector<Vector<N, Real>> mControls;
+        std::vector<Real> mWeights;
+    };
+
+}
diff --git a/Engine/Plugins/Experimental/GeometryProcessing/Source/GeometryAlgorithms/Private/GTEngine/Mathematics/GteNURBSVolume.h b/Engine/Plugins/Experimental/GeometryProcessing/Source/GeometryAlgorithms/Private/GTEngine/Mathematics/GteNURBSVolume.h
new file mode 100644
--- /dev/null
+++ b/Engine/Plugins/Experimental/GeometryProcessing/Source/GeometryAlgorithms/Private/GTEngine/Mathematics/GteNURBSVolume.h
@@ -0,0 +1,264 @@
+// David Eberly, Geometric Tools, Redmond WA 98052
+// Copyright (c) 1998-2019
+// Distributed under the Boost Software License, Version 1.0.
+// http://www.boost.org/LICENSE_1_0.txt
+// http://www.geometrictools.com/License/Boost/LICENSE_1_0.txt
+// File Version: 3.0.3 (2018/10/22)
+
+#pragma once
+
+#include <GTEngine/Mathematics/GteVector.h>
+#include <GTEngine/Mathematics/GteBasisFunction.h>
+
+namespace gte
+{
+    template <int N, typename Real>
+    class NURBSVolume
+    {
+    public:
+        // Construction.  If the input controls is non-null, a copy is made of
+        // the controls.  To defer setting the control points or weights, pass
+        // null pointers and later access the control points or weights via
+        // GetControls(), GetWeights(), SetControl(), or SetWeight() member
+        // functions.  The 'controls' and 'weights' must be stored in
+        // lexicographical order,
+        //   attribute[i0 + numControls0 * (i1 + numControls1 * i2)]
+        // As a 3D array, this corresponds to attribute3D[i2][i1][i0].
+        NURBSVolume(BasisFunctionInput<Real> const& input0,
+            BasisFunctionInput<Real> const& input1,
+            BasisFunctionInput<Real> const& input2,
+            Vector<N, Real> const* controls, Real const* weights)
+            :
+            mConstructed(false)
+        {
+            BasisFunctionInput<Real> const* input[3] = { &input0, &input1, &input2 };
+            for (int i = 0; i < 3; ++i)
+            {
+                mNumControls[i] = input[i]->numControls;
+                mBasisFunction[i].Create(*input[i]);
+                if (!mBasisFunction[i])
+                {
+                    // Errors were already generated during construction of
+                    // the basis functions.
+                    return;
+                }
+            }
+
+            // The replication of control points for periodic splines is
+            // avoided by wrapping the i-loop index in Evaluate.
+            int numControls = mNumControls[0] * mNumControls[1] * mNumControls[2];
+            mControls.resize(numControls);
+            mWeights.resize(numControls);
+            if (controls)
+            {
+                std::copy(controls, controls + numControls, mControls.begin());
+            }
+            else
+            {
+                memset(mControls.data(), 0, mControls.size() * sizeof(mControls[0]));
+            }
+            if (weights)
+            {
+                std::copy(weights, weights + numControls, mWeights.begin());
+            }
+            else
+            {
+                memset(mWeights.data(), 0, mWeights.size() * sizeof(mWeights[0]));
+            }
+            mConstructed = true;
+        }
+
+        // To validate construction, create an object as shown:
+        //     NURBSVolume<N, Real> volume(parameters);
+        //     if (!volume) { <constructor failed, handle accordingly>; }
+        inline operator bool() const
+        {
+            return mConstructed;
+        }
+
+        // Member access.  The index 'dim' must be in {0,1,2}.
+        inline BasisFunction<Real> const& GetBasisFunction(int dim) const
+        {
+            return mBasisFunction[dim];
+        }
+
+        inline Real GetMinDomain(int dim) const
+        {
+            return mBasisFunction[dim].GetMinDomain();
+        }
+
+        inline Real GetMaxDomain(int dim) const
+        {
+            return mBasisFunction[dim].GetMaxDomain();
+        }
+
+        inline int GetNumControls(int dim) const
+        {
+            return mNumControls[dim];
+        }
+
+        inline Vector<N, Real> const* GetControls() const
+        {
+            return mControls.data();
+        }
+
+        inline Vector<N, Real>* GetControls()
+        {
+            return mControls.data();
+        }
+
+        inline Real const* GetWeights() const
+        {
+            return mWeights.data();
+        }
+
+        inline Real* GetWeights()
+        {
+            return mWeights.data();
+        }
+
+        // Evaluation of the volume.  The function supports derivative
+        // calculation through order 2; that is, maxOrder <= 2 is required.
+        // If you want only the position, pass in maxOrder of 0.  If you want
+        // the position and first-order derivatives, pass in maxOrder of 1,
+        // and so on.  The output 'values' are ordered as: position X;
+        // first-order derivatives dX/du, dX/dv, dX/dw; second-order
+        // derivatives d2X/du2, d2X/dv2, d2X/dw2, d2X/dudv, d2X/dudw,
+        // d2X/dvdw.
+        void Evaluate(Real u, Real v, Real w, unsigned int maxOrder, Vector<N, Real> values[10]) const
+        {
+            if (!mConstructed)
+            {
+                // Errors were already generated during construction.
+                for (int i = 0; i < 10; ++i)
+                {
+                    values[i].MakeZero();
+                }
+                return;
+            }
+
+            int iumin, iumax, ivmin, ivmax, iwmin, iwmax;
+            mBasisFunction[0].Evaluate(u, maxOrder, iumin, iumax);
+            mBasisFunction[1].Evaluate(v, maxOrder, ivmin, ivmax);
+            mBasisFunction[2].Evaluate(w, maxOrder, iwmin, iwmax);
+
+            // Compute position.
+            Vector<N, Real> X;
+            Real h;
+            Compute(0, 0, 0, iumin, iumax, ivmin, ivmax, iwmin, iwmax, X, h);
+            Real invH = ((Real)1) / h;
+            values[0] = invH * X;
+
+            if (maxOrder >= 1)
+            {
+                // Compute first-order derivatives.
+                Vector<N, Real> XDerU;
+                Real hDerU;
+                Compute(1, 0, 0, iumin, iumax, ivmin, ivmax, iwmin, iwmax,
+                    XDerU, hDerU);
+                values[1] = invH * (XDerU - hDerU * values[0]);
+
+                Vector<N, Real> XDerV;
+                Real hDerV;
+                Compute(0, 1, 0, iumin, iumax, ivmin, ivmax, iwmin, iwmax,
+                    XDerV, hDerV);
+                values[2] = invH * (XDerV - hDerV * values[0]);
+
+                Vector<N, Real> XDerW;
+                Real hDerW;
+                Compute(0, 0, 1, iumin, iumax, ivmin, ivmax, iwmin, iwmax,
+                    XDerW, hDerW);
+                values[3] = invH * (XDerW - hDerW * values[0]);
+
+                if (maxOrder >= 2)
+                {
+                    // Compute second-order derivatives.
+                    Vector<N, Real> XDerUU;
+                    Real hDerUU;
+                    Compute(2, 0, 0, iumin, iumax, ivmin, ivmax, iwmin, iwmax,
+                        XDerUU, hDerUU);
+                    values[4] = invH * (XDerUU - ((Real)2) * hDerU * values[1] -
+                        hDerUU * values[0]);
+
+                    Vector<N, Real> XDerVV;
+                    Real hDerVV;
+                    Compute(0, 2, 0, iumin, iumax, ivmin, ivmax, iwmin, iwmax,
+                        XDerVV, hDerVV);
+                    values[5] = invH * (XDerVV - ((Real)2) * hDerV * values[2] -
+                        hDerVV * values[0]);
+
+                    Vector<N, Real> XDerWW;
+                    Real hDerWW;
+                    Compute(0, 0, 2, iumin, iumax, ivmin, ivmax, iwmin, iwmax,
+                        XDerWW, hDerWW);
+                    values[6] = invH * (XDerWW - ((Real)2) * hDerW * values[3] -
+                        hDerWW * values[0]);
+
+                    Vector<N, Real> XDerUV;
+                    Real hDerUV;
+                    Compute(1, 1, 0, iumin, iumax, ivmin, ivmax, iwmin, iwmax,
+                        XDerUV, hDerUV);
+                    values[7] = invH * (XDerUV - hDerU * values[2]
+                        - hDerV * values[1] - hDerUV * values[0]);
+
+                    Vector<N, Real> XDerUW;
+                    Real hDerUW;
+                    Compute(1, 0, 1, iumin, iumax, ivmin, ivmax, iwmin, iwmax,
+                        XDerUW, hDerUW);
+                    values[8] = invH * (XDerUW - hDerU * values[3]
+                        - hDerW * values[1] - hDerUW * values[0]);
+
+                    Vector<N, Real> XDerVW;
+                    Real hDerVW;
+                    Compute(0, 1, 1, iumin, iumax, ivmin, ivmax, iwmin, iwmax,
+                        XDerVW, hDerVW);
+                    values[9] = invH * (XDerVW - hDerV * values[3]
+                        - hDerW * values[2] - hDerVW * values[0]);
+                }
+            }
+        }
+
+    private:
+        // Support for Evaluate(...).
+        void Compute(unsigned int uOrder, unsigned int vOrder,
+            unsigned int wOrder, int iumin, int iumax, int ivmin, int ivmax,
+            int iwmin, int iwmax, Vector<N, Real>& X, Real& h) const
+        {
+            // The j*-indices introduce a tiny amount of overhead in order to
+            // handle both aperiodic and periodic splines.  For aperiodic
+            // splines, j* = i* always.
+
+            int const numControls0 = mNumControls[0];
+            int const numControls1 = mNumControls[1];
+            int const numControls2 = mNumControls[2];
+            X.MakeZero();
+            h = (Real)0;
+            for (int iw = iwmin; iw <= iwmax; ++iw)
+            {
+                Real tmpw = mBasisFunction[2].GetValue(wOrder, iw);
+                int jw = (iw >= numControls2 ? iw - numControls2 : iw);
+                for (int iv = ivmin; iv <= ivmax; ++iv)
+                {
+                    Real tmpv = mBasisFunction[1].GetValue(vOrder, iv);
+                    Real tmpvw = tmpv * tmpw;
+                    int jv = (iv >= numControls1 ? iv - numControls1 : iv);
+                    for (int iu = iumin; iu <= iumax; ++iu)
+                    {
+                        Real tmpu = mBasisFunction[0].GetValue(uOrder, iu);
+                        int ju = (iu >= numControls0 ? iu - numControls0 : iu);
+                        int index = ju + numControls0 * (jv + numControls1 * jw);
+                        Real tmp = (tmpu * tmpvw) * mWeights[index];
+                        X += tmp * mControls[index];
+                        h += tmp;
+                    }
+                }
+            }
+        }
+
+        std::array<BasisFunction<Real>, 3> mBasisFunction;
+        std::array<int, 3> mNumControls;
+        std::vector<Vector<N, Real>> mControls;
+        std::vector<Real> mWeights;
+        bool mConstructed;
+    };
+}
diff --git a/Engine/Plugins/Experimental/GeometryProcessing/Source/GeometryAlgorithms/Private/GTEngine/Mathematics/GteNaturalSplineCurve.h b/Engine/Plugins/Experimental/GeometryProcessing/Source/GeometryAlgorithms/Private/GTEngine/Mathematics/GteNaturalSplineCurve.h
new file mode 100644
--- /dev/null
+++ b/Engine/Plugins/Experimental/GeometryProcessing/Source/GeometryAlgorithms/Private/GTEngine/Mathematics/GteNaturalSplineCurve.h
@@ -0,0 +1,404 @@
+// David Eberly, Geometric Tools, Redmond WA 98052
+// Copyright (c) 1998-2019
+// Distributed under the Boost Software License, Version 1.0.
+// http://www.boost.org/LICENSE_1_0.txt
+// http://www.geometrictools.com/License/Boost/LICENSE_1_0.txt
+// File Version: 3.0.1 (2018/02/17)
+
+#pragma once
+
+#include <GTEngine/LowLevel/GteLogger.h>
+#include <GTEngine/Mathematics/GteGMatrix.h>
+#include <GTEngine/Mathematics/GteParametricCurve.h>
+
+namespace gte
+{
+
+template <int N, typename Real>
+class NaturalSplineCurve : public ParametricCurve<N, Real>
+{
+public:
+    // Construction and destruction.  The object copies the input arrays.  The
+    // number of points M must be at least 2.  The first constructor is for a
+    // spline with second derivatives zero at the endpoints (isFree = true)
+    // or a spline that is closed (isFree = false).  The second constructor is
+    // for clamped splines, where you specify the first derivatives at the
+    // endpoints.  Usually, derivative0 = points[1] - points[0] at the first
+    // point and derivative1 = points[M-1] - points[M-2].  To validate
+    // construction, create an object as shown:
+    //     NaturalSplineCurve<N, Real> curve(parameters);
+    //     if (!curve) { <constructor failed, handle accordingly>; }
+    virtual ~NaturalSplineCurve();
+    NaturalSplineCurve(bool isFree, int numPoints,
+        Vector<N, Real> const* points, Real const* times);
+    NaturalSplineCurve(int numPoints, Vector<N, Real> const* points,
+        Real const* times, Vector<N, Real> const& derivative0,
+        Vector<N, Real> const& derivative1);
+
+    // Member access.
+    inline int GetNumPoints() const;
+    inline Vector<N, Real> const* GetPoints() const;
+
+    // Evaluation of the curve.  The function supports derivative calculation
+    // through order 3; that is, maxOrder <= 3 is required.  If you want
+    // only the position, pass in maxOrder of 0.  If you want the position and
+    // first derivative, pass in maxOrder of 1, and so on.  The output
+    // 'values' are ordered as: position, first derivative, second derivative,
+    // third derivative.
+    virtual void Evaluate(Real t, unsigned int maxOrder,
+        Vector<N, Real> values[4]) const;
+
+protected:
+    // Support for construction.
+    void CreateFree();
+    void CreateClosed();
+    void CreateClamped(Vector<N, Real> const& derivative0,
+        Vector<N, Real> const& derivative1);
+
+    // Determine the index i for which times[i] <= t < times[i+1].
+    void GetKeyInfo(Real t, int& key, Real& dt) const;
+
+    // Polynomial coefficients.  mA are the points (constant coefficients of
+    // polynomials.  mB are the degree 1 coefficients, mC are the degree 2
+    // coefficients, and mD are the degree 3 coefficients.
+    std::vector<Vector<N, Real>> mA, mB, mC, mD;
+};
+
+
+template <int N, typename Real>
+NaturalSplineCurve<N, Real>::~NaturalSplineCurve()
+{
+}
+
+template <int N, typename Real>
+NaturalSplineCurve<N, Real>::NaturalSplineCurve(bool isFree, int numPoints,
+    Vector<N, Real> const* points, Real const* times)
+    :
+    ParametricCurve<N, Real>(numPoints - 1, times)
+{
+    if (numPoints < 2 || !points)
+    {
+        LogError("Invalid input.");
+        return;
+    }
+
+    mA.resize(numPoints);
+    std::copy(points, points + numPoints, mA.begin());
+
+    if (isFree)
+    {
+        CreateFree();
+    }
+    else
+    {
+        CreateClosed();
+    }
+
+    this->mConstructed = true;
+}
+
+template <int N, typename Real>
+NaturalSplineCurve<N, Real>::NaturalSplineCurve(int numPoints,
+    Vector<N, Real> const* points, Real const* times,
+    Vector<N, Real> const& derivative0, Vector<N, Real> const& derivative1)
+    :
+    ParametricCurve<N, Real>(numPoints - 1, times)
+{
+    if (numPoints < 2 || !points)
+    {
+        LogError("Invalid input.");
+        return;
+    }
+
+    mA.resize(numPoints);
+    std::copy(points, points + numPoints, mA.begin());
+
+    CreateClamped(derivative0, derivative1);
+    this->mConstructed = true;
+}
+
+template <int N, typename Real> inline
+int NaturalSplineCurve<N, Real>::GetNumPoints() const
+{
+    return static_cast<int>(mA.size());
+}
+
+template <int N, typename Real> inline
+Vector<N, Real> const* NaturalSplineCurve<N, Real>::GetPoints() const
+{
+    return &mA[0];
+}
+
+template <int N, typename Real>
+void NaturalSplineCurve<N, Real>::Evaluate(Real t, unsigned int maxOrder,
+    Vector<N, Real> values[4]) const
+{
+    if (!this->mConstructed)
+    {
+        // Errors were already generated during construction.
+        for (unsigned int order = 0; order < 4; ++order)
+        {
+            values[order].MakeZero();
+        }
+        return;
+    }
+
+    int key = 0;
+    Real dt = (Real)0;
+    GetKeyInfo(t, key, dt);
+
+    // Compute position.
+    values[0] = mA[key] + dt * (mB[key] + dt * (mC[key] + dt * mD[key]));
+    if (maxOrder >= 1)
+    {
+        // Compute first derivative.
+        values[1] = mB[key] + dt * (((Real)2) * mC[key] +
+            ((Real)3) * dt * mD[key]);
+        if (maxOrder >= 2)
+        {
+            // Compute second derivative.
+            values[2] = ((Real)2) * mC[key] + ((Real)6) * dt * mD[key];
+            if (maxOrder == 3)
+            {
+                values[3] = ((Real)6) * mD[key];
+            }
+            else
+            {
+                values[3].MakeZero();
+            }
+        }
+    }
+}
+
+template <int N, typename Real>
+void NaturalSplineCurve<N, Real>::CreateFree()
+{
+    int numSegments = GetNumPoints() - 1;
+    std::vector<Real> dt(numSegments);
+    for (int i = 0; i < numSegments; ++i)
+    {
+        dt[i] = this->mTime[i + 1] - this->mTime[i];
+    }
+
+    std::vector<Real> d2t(numSegments);
+    for (int i = 1; i < numSegments; ++i)
+    {
+        d2t[i] = this->mTime[i + 1] - this->mTime[i - 1];
+    }
+
+    std::vector<Vector<N, Real>> alpha(numSegments);
+    for (int i = 1; i < numSegments; ++i)
+    {
+        Vector<N, Real> numer = ((Real)3)*(dt[i - 1] * mA[i + 1] -
+            d2t[i] * mA[i] + dt[i] * mA[i - 1]);
+        Real invDenom = ((Real)1) / (dt[i - 1] * dt[i]);
+        alpha[i] = invDenom * numer;
+    }
+
+    std::vector<Real> ell(numSegments + 1);
+    std::vector<Real> mu(numSegments);
+    std::vector<Vector<N, Real>> z(numSegments + 1);
+    Real inv;
+
+    ell[0] = (Real)1;
+    mu[0] = (Real)0;
+    z[0].MakeZero();
+    for (int i = 1; i < numSegments; ++i)
+    {
+        ell[i] = ((Real)2) * d2t[i] - dt[i - 1] * mu[i - 1];
+        inv = ((Real)1) / ell[i];
+        mu[i] = inv * dt[i];
+        z[i] = inv * (alpha[i] - dt[i - 1] * z[i - 1]);
+    }
+    ell[numSegments] = (Real)1;
+    z[numSegments].MakeZero();
+
+    mB.resize(numSegments);
+    mC.resize(numSegments + 1);
+    mD.resize(numSegments);
+
+    Real oneThird = ((Real)1) / (Real)3;
+    mC[numSegments].MakeZero();
+    for (int i = numSegments - 1; i >= 0; --i)
+    {
+        mC[i] = z[i] - mu[i] * mC[i + 1];
+        inv = ((Real)1) / dt[i];
+        mB[i] = inv * (mA[i + 1] - mA[i]) - oneThird * dt[i] * (mC[i + 1] +
+            ((Real)2) * mC[i]);
+        mD[i] = oneThird * inv * (mC[i + 1] - mC[i]);
+    }
+}
+
+template <int N, typename Real>
+void NaturalSplineCurve<N, Real>::CreateClosed()
+{
+    // TODO:  A general linear system solver is used here.  The matrix
+    // corresponding to this case is actually "cyclic banded", so a faster
+    // linear solver can be used.  The current linear system code does not
+    // have such a solver.
+
+    int numSegments = GetNumPoints() - 1;
+    std::vector<Real> dt(numSegments);
+    for (int i = 0; i < numSegments; ++i)
+    {
+        dt[i] = this->mTime[i + 1] - this->mTime[i];
+    }
+
+    // Construct matrix of system.
+    GMatrix<Real> mat(numSegments + 1, numSegments + 1);
+    mat(0, 0) = (Real)1;
+    mat(0, numSegments) = (Real)-1;
+    for (int i = 1; i <= numSegments - 1; ++i)
+    {
+        mat(i, i - 1) = dt[i - 1];
+        mat(i, i) = ((Real)2) * (dt[i - 1] + dt[i]);
+        mat(i, i + 1) = dt[i];
+    }
+    mat(numSegments, numSegments - 1) = dt[numSegments - 1];
+    mat(numSegments, 0) = ((Real)2) * (dt[numSegments - 1] + dt[0]);
+    mat(numSegments, 1) = dt[0];
+
+    // Construct right-hand side of system.
+    mC.resize(numSegments + 1);
+    mC[0].MakeZero();
+    Real inv0, inv1;
+    for (int i = 1; i <= numSegments - 1; ++i)
+    {
+        inv0 = ((Real)1) / dt[i];
+        inv1 = ((Real)1) / dt[i - 1];
+        mC[i] = ((Real)3) * (inv0 * (mA[i + 1] - mA[i]) -
+            inv1*(mA[i] - mA[i - 1]));
+    }
+    inv0 = ((Real)1) / dt[0];
+    inv1 = ((Real)1) / dt[numSegments - 1];
+    mC[numSegments] = ((Real)3) * (inv0 * (mA[1] - mA[0]) -
+        inv1 * (mA[0] - mA[numSegments - 1]));
+
+    // Solve the linear systems.
+    GMatrix<Real> invMat = Inverse(mat);
+    GVector<Real> input(numSegments + 1);
+    GVector<Real> output(numSegments + 1);
+    for (int j = 0; j < N; ++j)
+    {
+        for (int i = 0; i <= numSegments; ++i)
+        {
+            input[i] = mC[i][j];
+        }
+        output = invMat * input;
+        for (int i = 0; i <= numSegments; ++i)
+        {
+            mC[i][j] = output[i];
+        }
+    }
+
+    Real oneThird = ((Real)1) / (Real)3;
+    mB.resize(numSegments);
+    mD.resize(numSegments);
+    for (int i = 0; i < numSegments; ++i)
+    {
+        inv0 = ((Real)1) / dt[i];
+        mB[i] = inv0 * (mA[i + 1] - mA[i]) - oneThird * (mC[i + 1] +
+            ((Real)2) * mC[i]) * dt[i];
+        mD[i] = oneThird * inv0 * (mC[i + 1] - mC[i]);
+    }
+}
+
+template <int N, typename Real>
+void NaturalSplineCurve<N, Real>::CreateClamped(
+    Vector<N, Real> const& derivative0, Vector<N, Real> const& derivative1)
+{
+    int numSegments = GetNumPoints() - 1;
+    std::vector<Real> dt(numSegments);
+    for (int i = 0; i < numSegments; ++i)
+    {
+        dt[i] = this->mTime[i + 1] - this->mTime[i];
+    }
+
+    std::vector<Real> d2t(numSegments);
+    for (int i = 1; i < numSegments; ++i)
+    {
+        d2t[i] = this->mTime[i + 1] - this->mTime[i - 1];
+    }
+
+    std::vector<Vector<N, Real>> alpha(numSegments + 1);
+    Real inv = ((Real)1) / dt[0];
+    alpha[0] = ((Real)3) * (inv * (mA[1] - mA[0]) - derivative0);
+    inv = ((Real)1) / dt[numSegments - 1];
+    alpha[numSegments] = ((Real)3)*(derivative1 -
+        inv * (mA[numSegments] - mA[numSegments - 1]));
+    for (int i = 1; i < numSegments; ++i)
+    {
+        Vector<N, Real> numer = ((Real)3)*(dt[i - 1] * mA[i + 1] -
+            d2t[i] * mA[i] + dt[i] * mA[i - 1]);
+        Real invDenom = ((Real)1) / (dt[i - 1] * dt[i]);
+        alpha[i] = invDenom*numer;
+    }
+
+    std::vector<Real> ell(numSegments + 1);
+    std::vector<Real> mu(numSegments);
+    std::vector<Vector<N, Real>> z(numSegments + 1);
+
+    ell[0] = ((Real)2) * dt[0];
+    mu[0] = (Real)0.5;
+    inv = ((Real)1) / ell[0];
+    z[0] = inv * alpha[0];
+
+    for (int i = 1; i < numSegments; ++i)
+    {
+        ell[i] = ((Real)2) * d2t[i] - dt[i - 1] * mu[i - 1];
+        inv = ((Real)1) / ell[i];
+        mu[i] = inv * dt[i];
+        z[i] = inv * (alpha[i] - dt[i - 1] * z[i - 1]);
+    }
+    ell[numSegments] = dt[numSegments - 1] *
+        (((Real)2) - mu[numSegments - 1]);
+    inv = ((Real)1) / ell[numSegments];
+    z[numSegments] = inv * (alpha[numSegments] - dt[numSegments - 1] *
+        z[numSegments - 1]);
+
+    mB.resize(numSegments);
+    mC.resize(numSegments + 1);
+    mD.resize(numSegments);
+
+    Real oneThird = ((Real)1) / (Real)3;
+    mC[numSegments] = z[numSegments];
+    for (int i = numSegments - 1; i >= 0; --i)
+    {
+        mC[i] = z[i] - mu[i] * mC[i + 1];
+        inv = ((Real)1) / dt[i];
+        mB[i] = inv * (mA[i + 1] - mA[i]) - oneThird*dt[i] * (mC[i + 1] +
+            ((Real)2) * mC[i]);
+        mD[i] = oneThird * inv * (mC[i + 1] - mC[i]);
+    }
+}
+
+template <int N, typename Real>
+void NaturalSplineCurve<N, Real>::GetKeyInfo(Real t, int& key, Real& dt) const
+{
+    int numSegments = GetNumPoints() - 1;
+    if (t <= this->mTime[0])
+    {
+        key = 0;
+        dt = (Real)0;
+    }
+    else if (t >= this->mTime[numSegments])
+    {
+        key = numSegments - 1;
+        dt = this->mTime[numSegments] - this->mTime[numSegments - 1];
+    }
+    else
+    {
+        for (int i = 0; i < numSegments; ++i)
+        {
+            if (t < this->mTime[i + 1])
+            {
+                key = i;
+                dt = t - this->mTime[i];
+                break;
+            }
+        }
+    }
+}
+
+
+}
diff --git a/Engine/Plugins/Experimental/GeometryProcessing/Source/GeometryAlgorithms/Private/GTEngine/Mathematics/GteNearestNeighborQuery.h b/Engine/Plugins/Experimental/GeometryProcessing/Source/GeometryAlgorithms/Private/GTEngine/Mathematics/GteNearestNeighborQuery.h
new file mode 100644
--- /dev/null
+++ b/Engine/Plugins/Experimental/GeometryProcessing/Source/GeometryAlgorithms/Private/GTEngine/Mathematics/GteNearestNeighborQuery.h
@@ -0,0 +1,238 @@
+// David Eberly, Geometric Tools, Redmond WA 98052
+// Copyright (c) 1998-2019
+// Distributed under the Boost Software License, Version 1.0.
+// http://www.boost.org/LICENSE_1_0.txt
+// http://www.geometrictools.com/License/Boost/LICENSE_1_0.txt
+// File Version: 3.0.1 (2017/07/26)
+
+#pragma once
+
+#include <GTEngine/LowLevel/GteLogger.h>
+#include <GTEngine/Mathematics/GteVector.h>
+#include <algorithm>
+#include <limits>
+#include <vector>
+
+namespace gte
+{
+
+// Use a kd-tree for sorting used in a query for finding nearest neighbors
+// of a point in a space of the specified dimension N.  The split order is
+// always 0,1,2,...,N-1.  The number of sites at a leaf node is controlled
+// by 'maxLeafSize' and the maximum level of the tree is controlled by
+// 'maxLevels'.  The points are of type Vector<N,Real>.  The 'Site' is a
+// structure of information that minimally implements the function
+// 'Vector<N,Real> GetPosition () const'.  The Site template parameter
+// allows the query to be applied even when it has more local information
+// than just point location.
+template <int N, typename Real, typename Site>
+class NearestNeighborQuery
+{
+public:
+    // Construction.
+    NearestNeighborQuery(std::vector<Site> const& sites, int maxLeafSize, int maxLevel);
+
+    // Member access.
+    inline int GetMaxLeafSize () const;
+    inline int GetMaxLevel () const;
+
+    // Compute up to MaxNeighbors nearest neighbors within the specified
+    // radius of the point.  The returned integer is the number of neighbors
+    // found, possibly zero.  The neighbors array stores indices into the
+    // array passed to the constructor.
+    template <int MaxNeighbors>
+    int FindNeighbors(Vector<N,Real> const& point, Real radius,
+        std::array<int, MaxNeighbors>& neighbors) const;
+
+private:
+    typedef std::pair<Vector<N, Real>, int> SortedPoint;
+
+    struct Node
+    {
+        Real split;
+        int axis;
+        int numSites;
+        int siteOffset;
+        int left;
+        int right;
+    };
+
+    // Populate the node so that it contains the points split along the
+    // coordinate axes.
+    void Build(int numSites, int siteOffset, int nodeIndex, int level);
+
+    int mMaxLeafSize;
+    int mMaxLevel;
+    std::vector<SortedPoint> mSortedPoints;
+    std::vector<Node> mNodes;
+};
+
+
+template <int N, typename Real, typename Site>
+NearestNeighborQuery<N, Real, Site>::NearestNeighborQuery(
+    std::vector<Site> const& sites, int maxLeafSize, int maxLevel)
+    :
+    mMaxLeafSize(maxLeafSize),
+    mMaxLevel(maxLevel),
+    mSortedPoints(sites.size())
+{
+    int const numSites = static_cast<int>(sites.size());
+    for (int i = 0; i < numSites; ++i)
+    {
+        mSortedPoints[i] = std::make_pair(sites[i].GetPosition(), i);
+    }
+
+    mNodes.push_back(Node());
+    Build(numSites, 0, 0, 0);
+}
+
+template <int N, typename Real, typename Site>
+inline int NearestNeighborQuery<N, Real, Site>::GetMaxLeafSize() const
+{
+    return mMaxLeafSize;
+}
+
+template <int N, typename Real, typename Site>
+inline int NearestNeighborQuery<N, Real, Site>::GetMaxLevel() const
+{
+    return mMaxLevel;
+}
+
+template <int N, typename Real, typename Site>
+template <int MaxNeighbors>
+int NearestNeighborQuery<N, Real, Site>::FindNeighbors(Vector<N, Real> const& point,
+    Real radius, std::array<int, MaxNeighbors>& neighbors) const
+{
+    Real sqrRadius = radius * radius;
+    int numNeighbors = 0;
+    std::array<int, MaxNeighbors + 1> localNeighbors;
+    std::array<Real, MaxNeighbors + 1> neighborSqrLength;
+    for (int i = 0; i <= MaxNeighbors; ++i)
+    {
+        localNeighbors[i] = -1;
+        neighborSqrLength[i] = std::numeric_limits<Real>::max();
+    }
+
+    // The kd-tree construction is recursive, simulated here by using a stack.
+    // The maximum depth is limited to 32, because the number of sites is
+    // limited to 2^{32} (the number of 32-bit integer indices).
+    std::array<int, 32> stack;
+    int top = 0;
+    stack[0] = 0;
+
+    while (top >= 0)
+    {
+        Node node = mNodes[stack[top--]];
+
+        if (node.siteOffset != -1)
+        {
+            for (int i = 0, j = node.siteOffset; i < node.numSites; ++i, ++j)
+            {
+                Vector<N, Real> diff = mSortedPoints[j].first - point;
+                Real sqrLength = Dot(diff, diff);
+                if (sqrLength <= sqrRadius)
+                {
+                    // Maintain the nearest neighbors.
+                    int k;
+                    for (k = 0; k < numNeighbors; ++k)
+                    {
+                        if (sqrLength <= neighborSqrLength[k])
+                        {
+                            for (int n = numNeighbors; n > k; --n)
+                            {
+                                localNeighbors[n] = localNeighbors[n - 1];
+                                neighborSqrLength[n] = neighborSqrLength[n - 1];
+                            }
+                            break;
+                        }
+                    }
+                    if (k < MaxNeighbors)
+                    {
+                        localNeighbors[k] = mSortedPoints[j].second;
+                        neighborSqrLength[k] = sqrLength;
+                    }
+                    if (numNeighbors < MaxNeighbors)
+                    {
+                        ++numNeighbors;
+                    }
+                }
+            }
+        }
+
+        if (node.left != -1 && point[node.axis] - radius <= node.split)
+        {
+            stack[++top] = node.left;
+        }
+
+        if (node.right != -1 && point[node.axis] + radius >= node.split)
+        {
+            stack[++top] = node.right;
+        }
+    }
+
+    for (int i = 0; i < numNeighbors; ++i)
+    {
+        neighbors[i] = localNeighbors[i];
+    }
+
+    return numNeighbors;
+}
+
+template <int N, typename Real, typename Site>
+void NearestNeighborQuery<N, Real, Site>::Build(int numSites, int siteOffset,
+    int nodeIndex, int level)
+{
+    LogAssert(siteOffset != -1, "Invalid site offset.");
+    LogAssert(nodeIndex != -1, "Invalid node index.");
+    LogAssert(numSites > 0, "Empty point list.");
+
+    Node& node = mNodes[nodeIndex];
+    node.numSites = numSites;
+
+    if (numSites > mMaxLeafSize && level <= mMaxLevel)
+    {
+        int halfNumSites = numSites / 2;
+
+        // The point set is too large for a leaf node, so split it at the
+        // median.  The O(m log m) sort is not needed; rather, we locate the
+        // median using an order statistic construction that is expected
+        // time O(m).
+        int const axis = level % N;
+        auto sorter = [axis](SortedPoint const& p0, SortedPoint const& p1)
+        {
+            return p0.first[axis] < p1.first[axis];
+        };
+
+        auto begin = mSortedPoints.begin() + siteOffset;
+        auto mid = mSortedPoints.begin() + siteOffset + halfNumSites;
+        auto end = mSortedPoints.begin() + siteOffset + numSites;
+        std::nth_element(begin, mid, end, sorter);
+
+        // Get the median position.
+        node.split = mSortedPoints[siteOffset + halfNumSites].first[axis];
+        node.axis = axis;
+        node.siteOffset = -1;
+
+        // Apply a divide-and-conquer step.
+        int left = (int)mNodes.size(), right = left + 1;
+        node.left = left;
+        node.right = right;
+        mNodes.push_back(Node());
+        mNodes.push_back(Node());
+
+        int nextLevel = level + 1;
+        Build(halfNumSites, siteOffset, left, nextLevel);
+        Build(numSites - halfNumSites, siteOffset + halfNumSites, right, nextLevel);
+    }
+    else
+    {
+        // The number of points is small enough, so make this node a leaf.
+        node.split = std::numeric_limits<Real>::max();
+        node.axis = -1;
+        node.siteOffset = siteOffset;
+        node.left = -1;
+        node.right = -1;
+    }
+}
+
+}
diff --git a/Engine/Plugins/Experimental/GeometryProcessing/Source/GeometryAlgorithms/Private/GTEngine/Mathematics/GteOBBTreeOfPoints.h b/Engine/Plugins/Experimental/GeometryProcessing/Source/GeometryAlgorithms/Private/GTEngine/Mathematics/GteOBBTreeOfPoints.h
new file mode 100644
--- /dev/null
+++ b/Engine/Plugins/Experimental/GeometryProcessing/Source/GeometryAlgorithms/Private/GTEngine/Mathematics/GteOBBTreeOfPoints.h
@@ -0,0 +1,359 @@
+// David Eberly, Geometric Tools, Redmond WA 98052
+// Copyright (c) 1998-2019
+// Distributed under the Boost Software License, Version 1.0.
+// http://www.boost.org/LICENSE_1_0.txt
+// http://www.geometrictools.com/License/Boost/LICENSE_1_0.txt
+// File Version: 3.10.1 (2017/09/16)
+
+#pragma once
+
+#include <GTEngine/Mathematics/GteBitHacks.h>
+#include <GTEngine/Mathematics/GteContOrientedBox3.h>
+
+// The depth of a node in a (nonempty) tree is the distance from the node to
+// the root of the tree.  The height is the maximum depth.  A tree with a
+// single node has height 0.  The set of nodes of a tree with the same depth
+// is refered to as a level of a tree (corresponding to that depth).  A
+// complete binary tree of height H has 2^{H+1}-1 nodes.  The level
+// corresponding to depth D has 2^D nodes, in which case the number of
+// leaf nodes (depth H) is 2^H.
+
+namespace gte
+{
+
+template <typename Real>
+class OBBTreeForPoints
+{
+public:
+    struct Node
+    {
+        Node();
+        OrientedBox3<Real> box;
+        uint32_t depth;
+        uint32_t minIndex, maxIndex;
+        uint32_t leftChild, rightChild;
+    };
+
+
+    // The 'points' array is a collection of vertices, each occupying a
+    // chunk of memory with 'stride' bytes.  A vertex must start at the
+    // first byte of this chunk but does not necessarily fill it.  The
+    // 'height' specifies the height of the tree and must be no larger
+    // than 31.  If it is set to std::numeric_limits<uint32_t>::max(),
+    // then the entire tree is built and the actual height is computed
+    // from 'numPoints'.
+    OBBTreeForPoints(uint32_t numPoints, char const* points, size_t stride,
+        uint32_t height = std::numeric_limits<uint32_t>::max());
+
+    // Member access.
+    inline uint32_t GetNumPoints() const;
+    inline char const* GetPoints() const;
+    inline size_t GetStride() const;
+    inline std::vector<Node> const& GetTree() const;
+    inline uint32_t GetHeight() const;
+    inline std::vector<uint32_t> const& GetPartition() const;
+
+private:
+    inline Vector3<Real> GetPosition(uint32_t index) const;
+    void BuildTree(Node& node, uint32_t i0, uint32_t i1);
+    void ComputeOBB(uint32_t i0, uint32_t i1, OrientedBox3<Real>& box);
+    void SplitPoints(uint32_t i0, uint32_t i1, uint32_t& j0, uint32_t& j1,
+        Vector3<Real> const& origin, Vector3<Real> const& direction);
+
+    struct ProjectionInfo
+    {
+        ProjectionInfo();
+        bool operator< (ProjectionInfo const& info) const;
+        uint32_t pointIndex;
+        Real projection;
+    };
+
+    uint32_t mNumPoints;
+    char const* mPoints;
+    size_t mStride;
+    uint32_t mHeight;
+    std::vector<Node> mTree;
+    std::vector<uint32_t> mPartition;
+};
+
+template <typename Real>
+OBBTreeForPoints<Real>::OBBTreeForPoints(uint32_t numPoints, char const* points, size_t stride, uint32_t height)
+    :
+    mNumPoints(numPoints),
+    mPoints(points),
+    mStride(stride),
+    mHeight(height),
+    mPartition(numPoints)
+{
+    // The tree nodes are indexed by 32-bit unsigned integers, so
+    // the number of nodes can be at most 2^{32} - 1.  This limits
+    // the height to 31.
+
+    uint32_t numNodes;
+    if (mHeight == std::numeric_limits<uint32_t>::max())
+    {
+        uint32_t minPowerOfTwo = (uint32_t)RoundUpToPowerOfTwo(mNumPoints);
+        mHeight = Log2OfPowerOfTwo(minPowerOfTwo);
+        numNodes = 2 * mNumPoints - 1;
+    }
+    else
+    {
+        // The maximum level cannot exceed 30 because we are storing the
+        // indices into the node array as 32-bit unsigned integers.
+        if (mHeight < 32)
+        {
+            numNodes = (uint32_t)(1ULL << (mHeight + 1)) - 1;
+        }
+        else
+        {
+            // When the precondition is not met, return a tree of
+            // height 0 (a single node).
+            mHeight = 0;
+            numNodes = 1;
+        }
+    }
+
+    // The tree is built recursively.  A reference to a Node is passed
+    // to BuildTree and nodes are appended to a std::vector.  Because the
+    // references are on the stack, we must guarantee no reallocations
+    // to avoid invalidating the references.  TODO: This design can be
+    // modified to pass indices to the nodes so that reallocation is not
+    // a problem.
+    mTree.reserve(numNodes);
+
+    // Build the tree recursively.  The array mPartition stores the
+    // indices into the 'points' array so that at a node, the points
+    // represented by the node are those indexed by the range
+    // [node.minIndex, node.maxIndex].
+    for (uint32_t i = 0; i < mNumPoints; ++i)
+    {
+        mPartition[i] = i;
+    }
+    mTree.push_back(Node());
+    BuildTree(mTree.back(), 0, mNumPoints - 1);
+}
+
+template <typename Real>
+inline uint32_t OBBTreeForPoints<Real>::GetNumPoints() const
+{
+    return mNumPoints;
+}
+
+template <typename Real>
+inline char const* OBBTreeForPoints<Real>::GetPoints() const
+{
+    return mPoints;
+}
+
+template <typename Real>
+inline size_t OBBTreeForPoints<Real>::GetStride() const
+{
+    return mStride;
+}
+
+template <typename Real>
+inline std::vector<typename OBBTreeForPoints<Real>::Node> const& OBBTreeForPoints<Real>::GetTree() const
+{
+    return mTree;
+}
+
+template <typename Real>
+inline uint32_t OBBTreeForPoints<Real>::GetHeight() const
+{
+    return mHeight;
+}
+
+template <typename Real>
+inline std::vector<uint32_t> const& OBBTreeForPoints<Real>::GetPartition() const
+{
+    return mPartition;
+}
+
+template <typename Real>
+inline Vector3<Real> OBBTreeForPoints<Real>::GetPosition(uint32_t index) const
+{
+    return *reinterpret_cast<Vector3<Real> const*>(mPoints + index * mStride);
+}
+
+template <typename Real>
+void OBBTreeForPoints<Real>::BuildTree(Node& node, uint32_t i0, uint32_t i1)
+{
+    node.minIndex = i0;
+    node.maxIndex = i1;
+
+    if (i0 == i1)
+    {
+        // We are at a leaf node.  The left and right child indices were
+        // set to std::numeric_limits<uint32_t>::max() during construction.
+
+        // Create a degenerate box whose center is the point.
+        node.box.center = GetPosition(mPartition[i0]);
+        node.box.axis[0] = Vector3<Real>{ (Real)1, (Real)0, (Real)0 };
+        node.box.axis[1] = Vector3<Real>{ (Real)0, (Real)1, (Real)0 };
+        node.box.axis[2] = Vector3<Real>{ (Real)0, (Real)0, (Real)1 };
+        node.box.extent = Vector3<Real>{ (Real)0, (Real)0, (Real)0 };
+    }
+    else  // i0 < i1
+    {
+        // We are at an interior node.  Compute an oriented bounding box.
+        ComputeOBB(i0, i1, node.box);
+
+        if (node.depth == mHeight)
+        {
+            return;
+        }
+
+        // Use the box axis corresponding to largest extent for the
+        // splitting axis.  Partition the points into two subsets, one for
+        // the left child and one for the right child. The subsets have
+        // numbers of elements that differ by at most 1, so the tree is
+        // balanced.
+        Vector3<Real> axis2 = node.box.axis[2];
+        uint32_t j0, j1;
+        SplitPoints(i0, i1, j0, j1, node.box.center, axis2);
+
+        node.leftChild = static_cast<uint32_t>(mTree.size());
+        node.rightChild = node.leftChild + 1;
+        mTree.push_back(Node());
+        Node& leftTree = mTree.back();
+        mTree.push_back(Node());
+        Node& rightTree = mTree.back();
+        leftTree.depth = node.depth + 1;
+        rightTree.depth = node.depth + 1;
+        BuildTree(leftTree, i0, j0);
+        BuildTree(rightTree, j1, i1);
+    }
+}
+
+template <typename Real>
+void OBBTreeForPoints<Real>::ComputeOBB(uint32_t i0, uint32_t i1, OrientedBox3<Real>& box)
+{
+    // Compute the mean of the points.
+    Vector3<Real> zero{ (Real)0, (Real)0, (Real)0 };
+    box.center = zero;
+    for (uint32_t i = i0; i <= i1; ++i)
+    {
+        box.center += GetPosition(mPartition[i]);
+    }
+    Real invSize = ((Real)1) / (Real)(i1 - i0 + 1);
+    box.center *= invSize;
+
+    // Compute the covariance matrix of the points.
+    Real covar00 = (Real)0, covar01 = (Real)0, covar02 = (Real)0;
+    Real covar11 = (Real)0, covar12 = (Real)0, covar22 = (Real)0;
+    for (uint32_t i = i0; i <= i1; ++i)
+    {
+        Vector3<Real> diff = GetPosition(mPartition[i]) - box.center;
+        covar00 += diff[0] * diff[0];
+        covar01 += diff[0] * diff[1];
+        covar02 += diff[0] * diff[2];
+        covar11 += diff[1] * diff[1];
+        covar12 += diff[1] * diff[2];
+        covar22 += diff[2] * diff[2];
+    }
+    covar00 *= invSize;
+    covar01 *= invSize;
+    covar02 *= invSize;
+    covar11 *= invSize;
+    covar12 *= invSize;
+    covar22 *= invSize;
+
+    // Solve the eigensystem and use the eigenvectors for the box axes.
+    SymmetricEigensolver3x3<Real> es;
+    std::array<Real, 3> eval;
+    std::array<std::array<Real, 3>, 3> evec;
+    es(covar00, covar01, covar02, covar11, covar12, covar22, false, +1, eval, evec);
+    for (int i = 0; i < 3; ++i)
+    {
+        box.axis[i] = evec[i];
+    }
+    box.extent = eval;
+
+    // Let C be the box center and let U0, U1, and U2 be the box axes.
+    // Each input point is of the form X = C + y0*U0 + y1*U1 + y2*U2.
+    // The following code computes min(y0), max(y0), min(y1), max(y1),
+    // min(y2), and max(y2).  The box center is then adjusted to be
+    //   C' = C + 0.5*(min(y0)+max(y0))*U0 + 0.5*(min(y1)+max(y1))*U1 +
+    //        0.5*(min(y2)+max(y2))*U2
+    Vector3<Real> pmin = zero, pmax = zero;
+    for (uint32_t i = i0; i <= i1; ++i)
+    {
+        Vector3<Real> diff = GetPosition(mPartition[i]) - box.center;
+        for (int j = 0; j < 3; ++j)
+        {
+            Real dot = Dot(diff, box.axis[j]);
+            if (dot < pmin[j])
+            {
+                pmin[j] = dot;
+            }
+            else if (dot > pmax[j])
+            {
+                pmax[j] = dot;
+            }
+        }
+    }
+
+    Real const half(0.5);
+    for (int j = 0; j < 3; ++j)
+    {
+        box.center += (half * (pmin[j] + pmax[j])) * box.axis[j];
+        box.extent[j] = half * (pmax[j] - pmin[j]);
+    }
+}
+
+template <typename Real>
+void OBBTreeForPoints<Real>::SplitPoints(uint32_t i0, uint32_t i1, uint32_t& j0, uint32_t& j1,
+    Vector3<Real> const& origin, Vector3<Real> const& direction)
+{
+    // Project the points onto the splitting axis.
+    uint32_t numProjections = i1 - i0 + 1;
+    std::vector<ProjectionInfo> info(numProjections);
+    uint32_t i, k;
+    for (i = i0, k = 0; i <= i1; ++i, ++k)
+    {
+        Vector3<Real> diff = GetPosition(mPartition[i]) - origin;
+        info[k].pointIndex = mPartition[i];
+        info[k].projection = Dot(direction, diff);
+    }
+
+    // Partition the projections by the median.
+    uint32_t medianIndex = (numProjections - 1) / 2;
+    std::nth_element(info.begin(), info.begin() + medianIndex, info.end());
+
+    // Partition the points by the median.
+    for (k = 0, j0 = i0 - 1; k <= medianIndex; ++k)
+    {
+        mPartition[++j0] = info[k].pointIndex;
+    }
+    for (j1 = i1 + 1; k < numProjections; ++k)
+    {
+        mPartition[--j1] = info[k].pointIndex;
+    }
+}
+
+template <typename Real>
+OBBTreeForPoints<Real>::Node::Node()
+    :
+    depth(0),
+    minIndex(std::numeric_limits<uint32_t>::max()),
+    maxIndex(std::numeric_limits<uint32_t>::max()),
+    leftChild(std::numeric_limits<uint32_t>::max()),
+    rightChild(std::numeric_limits<uint32_t>::max())
+{
+}
+
+template <typename Real>
+OBBTreeForPoints<Real>::ProjectionInfo::ProjectionInfo()
+    :
+    pointIndex(0),
+    projection((Real)0)
+{
+}
+
+template <typename Real>
+bool OBBTreeForPoints<Real>::ProjectionInfo::operator< (ProjectionInfo const& info) const
+{
+    return projection < info.projection;
+}
+
+}
diff --git a/Engine/Plugins/Experimental/GeometryProcessing/Source/GeometryAlgorithms/Private/GTEngine/Mathematics/GteOdeEuler.h b/Engine/Plugins/Experimental/GeometryProcessing/Source/GeometryAlgorithms/Private/GTEngine/Mathematics/GteOdeEuler.h
new file mode 100644
--- /dev/null
+++ b/Engine/Plugins/Experimental/GeometryProcessing/Source/GeometryAlgorithms/Private/GTEngine/Mathematics/GteOdeEuler.h
@@ -0,0 +1,60 @@
+// David Eberly, Geometric Tools, Redmond WA 98052
+// Copyright (c) 1998-2019
+// Distributed under the Boost Software License, Version 1.0.
+// http://www.boost.org/LICENSE_1_0.txt
+// http://www.geometrictools.com/License/Boost/LICENSE_1_0.txt
+// File Version: 3.0.0 (2016/06/19)
+
+#pragma once
+
+#include <GTEngine/Mathematics/GteOdeSolver.h>
+
+// The TVector template parameter allows you to create solvers with
+// Vector<N,Real> when the dimension N is known at compile time or
+// GVector<Real> when the dimension N is known at run time.  Both classes
+// have 'int GetSize() const' that allow OdeSolver-derived classes to query
+// for the dimension.
+
+namespace gte
+{
+
+template <typename Real, typename TVector>
+class OdeEuler : public OdeSolver<Real,TVector>
+{
+public:
+    // Construction and destruction.
+    virtual ~OdeEuler();
+    OdeEuler(Real tDelta,
+        std::function<TVector(Real, TVector const&)> const& F);
+
+    // Estimate x(t + tDelta) from x(t) using dx/dt = F(t,x).  You may allow
+    // xIn and xOut to be the same object.
+    virtual void Update(Real tIn, TVector const& xIn, Real& tOut,
+        TVector& xOut);
+};
+
+
+template <typename Real, typename TVector>
+OdeEuler<Real, TVector>::~OdeEuler()
+{
+}
+
+template <typename Real, typename TVector>
+OdeEuler<Real, TVector>::OdeEuler(Real tDelta,
+    std::function<TVector(Real, TVector const&)> const& F)
+    :
+    OdeSolver<Real, TVector>(tDelta, F)
+{
+}
+
+template <typename Real, typename TVector>
+void OdeEuler<Real, TVector>::Update(Real tIn, TVector const& xIn,
+    Real& tOut, TVector& xOut)
+{
+    TVector fVector = this->mFunction(tIn, xIn);
+    tOut = tIn + this->mTDelta;
+    xOut = xIn + this->mTDelta * fVector;
+}
+
+
+}
diff --git a/Engine/Plugins/Experimental/GeometryProcessing/Source/GeometryAlgorithms/Private/GTEngine/Mathematics/GteOdeImplicitEuler.h b/Engine/Plugins/Experimental/GeometryProcessing/Source/GeometryAlgorithms/Private/GTEngine/Mathematics/GteOdeImplicitEuler.h
new file mode 100644
--- /dev/null
+++ b/Engine/Plugins/Experimental/GeometryProcessing/Source/GeometryAlgorithms/Private/GTEngine/Mathematics/GteOdeImplicitEuler.h
@@ -0,0 +1,77 @@
+// David Eberly, Geometric Tools, Redmond WA 98052
+// Copyright (c) 1998-2019
+// Distributed under the Boost Software License, Version 1.0.
+// http://www.boost.org/LICENSE_1_0.txt
+// http://www.geometrictools.com/License/Boost/LICENSE_1_0.txt
+// File Version: 3.0.0 (2016/06/19)
+
+#pragma once
+
+#include <GTEngine/Mathematics/GteGMatrix.h>
+#include <GTEngine/Mathematics/GteOdeSolver.h>
+
+// The TVector template parameter allows you to create solvers with
+// Vector<N,Real> when the dimension N is known at compile time or
+// GVector<Real> when the dimension N is known at run time.  Both classes
+// have 'int GetSize() const' that allow OdeSolver-derived classes to query
+// for the dimension.  The TMatrix parameter must be either Matrix<N,N,Real>
+// or GMatrix<Real> accordingly.
+//
+// The function F(t,x) has input t, a scalar, and input x, an N-vector.
+// The first derivative matrix with respect to x is DF(t,x), an
+// N-by-N matrix.  Entry DF(r,c) is the derivative of F[r] with
+// respect to x[c].
+
+namespace gte
+{
+
+template <typename Real, typename TVector, typename TMatrix>
+class OdeImplicitEuler : public OdeSolver<Real,TVector>
+{
+public:
+    // Construction and destruction.
+    virtual ~OdeImplicitEuler();
+    OdeImplicitEuler(Real tDelta,
+        std::function<TVector(Real, TVector const&)> const& F,
+        std::function<TMatrix(Real, TVector const&)> const& DF);
+
+    // Estimate x(t + tDelta) from x(t) using dx/dt = F(t,x).  You may allow
+    // xIn and xOut to be the same object.
+    virtual void Update(Real tIn, TVector const& xIn, Real& tOut,
+        TVector& xOut);
+
+private:
+    std::function<TMatrix(Real, TVector const&)> mDerivativeFunction;
+};
+
+
+template <typename Real, typename TVector, typename TMatrix>
+OdeImplicitEuler<Real, TVector, TMatrix>::~OdeImplicitEuler()
+{
+}
+
+template <typename Real, typename TVector, typename TMatrix>
+OdeImplicitEuler<Real, TVector, TMatrix>::OdeImplicitEuler(Real tDelta,
+    std::function<TVector(Real, TVector const&)> const& F,
+    std::function<TMatrix(Real, TVector const&)> const& DF)
+    :
+    OdeSolver<Real, TVector>(tDelta, F),
+    mDerivativeFunction(DF)
+{
+}
+
+template <typename Real, typename TVector, typename TMatrix>
+void OdeImplicitEuler<Real, TVector, TMatrix>::Update(Real tIn,
+    TVector const& xIn, Real& tOut, TVector& xOut)
+{
+    TVector fVector = this->mFunction(tIn, xIn);
+    TMatrix dfMatrix = mDerivativeFunction(tIn, xIn);
+    TMatrix dgMatrix = TMatrix::Identity() - this->mTDelta * dfMatrix;
+    TMatrix dgInverse = Inverse(dgMatrix);
+    fVector = dgInverse * fVector;
+    tOut = tIn + this->mTDelta;
+    xOut = xIn + this->mTDelta * fVector;
+}
+
+
+}
diff --git a/Engine/Plugins/Experimental/GeometryProcessing/Source/GeometryAlgorithms/Private/GTEngine/Mathematics/GteOdeMidpoint.h b/Engine/Plugins/Experimental/GeometryProcessing/Source/GeometryAlgorithms/Private/GTEngine/Mathematics/GteOdeMidpoint.h
new file mode 100644
--- /dev/null
+++ b/Engine/Plugins/Experimental/GeometryProcessing/Source/GeometryAlgorithms/Private/GTEngine/Mathematics/GteOdeMidpoint.h
@@ -0,0 +1,67 @@
+// David Eberly, Geometric Tools, Redmond WA 98052
+// Copyright (c) 1998-2019
+// Distributed under the Boost Software License, Version 1.0.
+// http://www.boost.org/LICENSE_1_0.txt
+// http://www.geometrictools.com/License/Boost/LICENSE_1_0.txt
+// File Version: 3.0.0 (2016/06/19)
+
+#pragma once
+
+#include <GTEngine/Mathematics/GteOdeSolver.h>
+
+// The TVector template parameter allows you to create solvers with
+// Vector<N,Real> when the dimension N is known at compile time or
+// GVector<Real> when the dimension N is known at run time.  Both classes
+// have 'int GetSize() const' that allow OdeSolver-derived classes to query
+// for the dimension.
+
+namespace gte
+{
+
+template <typename Real, typename TVector>
+class OdeMidpoint : public OdeSolver<Real,TVector>
+{
+public:
+    // Construction and destruction.
+    virtual ~OdeMidpoint();
+    OdeMidpoint(Real tDelta,
+        std::function<TVector(Real, TVector const&)> const& F);
+
+    // Estimate x(t + tDelta) from x(t) using dx/dt = F(t,x).  You may allow
+    // xIn and xOut to be the same object.
+    virtual void Update(Real tIn, TVector const& xIn, Real& tOut,
+        TVector& xOut);
+};
+
+
+template <typename Real, typename TVector>
+OdeMidpoint<Real, TVector>::~OdeMidpoint()
+{
+}
+
+template <typename Real, typename TVector>
+OdeMidpoint<Real, TVector>::OdeMidpoint(Real tDelta,
+    std::function<TVector(Real, TVector const&)> const& F)
+    :
+    OdeSolver<Real, TVector>(tDelta, F)
+{
+}
+
+template <typename Real, typename TVector>
+void OdeMidpoint<Real, TVector>::Update(Real tIn, TVector const& xIn,
+    Real& tOut, TVector& xOut)
+{
+    // Compute the first step.
+    Real halfTDelta = ((Real)0.5) * this->mTDelta;
+    TVector fVector = this->mFunction(tIn, xIn);
+    TVector xTemp = xIn + halfTDelta * fVector;
+
+    // Compute the second step.
+    Real halfT = tIn + halfTDelta;
+    fVector = this->mFunction(halfT, xTemp);
+    tOut = tIn + this->mTDelta;
+    xOut = xIn + this->mTDelta * fVector;
+}
+
+
+}
diff --git a/Engine/Plugins/Experimental/GeometryProcessing/Source/GeometryAlgorithms/Private/GTEngine/Mathematics/GteOdeRungeKutta4.h b/Engine/Plugins/Experimental/GeometryProcessing/Source/GeometryAlgorithms/Private/GTEngine/Mathematics/GteOdeRungeKutta4.h
new file mode 100644
--- /dev/null
+++ b/Engine/Plugins/Experimental/GeometryProcessing/Source/GeometryAlgorithms/Private/GTEngine/Mathematics/GteOdeRungeKutta4.h
@@ -0,0 +1,77 @@
+// David Eberly, Geometric Tools, Redmond WA 98052
+// Copyright (c) 1998-2019
+// Distributed under the Boost Software License, Version 1.0.
+// http://www.boost.org/LICENSE_1_0.txt
+// http://www.geometrictools.com/License/Boost/LICENSE_1_0.txt
+// File Version: 3.0.0 (2016/06/19)
+
+#pragma once
+
+#include <GTEngine/Mathematics/GteOdeSolver.h>
+
+// The TVector template parameter allows you to create solvers with
+// Vector<N,Real> when the dimension N is known at compile time or
+// GVector<Real> when the dimension N is known at run time.  Both classes
+// have 'int GetSize() const' that allow OdeSolver-derived classes to query
+// for the dimension.
+
+namespace gte
+{
+
+template <typename Real, typename TVector>
+class OdeRungeKutta4 : public OdeSolver<Real,TVector>
+{
+public:
+    // Construction and destruction.
+    virtual ~OdeRungeKutta4();
+    OdeRungeKutta4(Real tDelta,
+        std::function<TVector(Real, TVector const&)> const& F);
+
+    // Estimate x(t + tDelta) from x(t) using dx/dt = F(t,x).  You may allow
+    // xIn and xOut to be the same object.
+    virtual void Update(Real tIn, TVector const& xIn, Real& tOut,
+        TVector& xOut);
+};
+
+
+template <typename Real, typename TVector>
+OdeRungeKutta4<Real, TVector>::~OdeRungeKutta4()
+{
+}
+
+template <typename Real, typename TVector>
+OdeRungeKutta4<Real, TVector>::OdeRungeKutta4(Real tDelta,
+    std::function<TVector(Real, TVector const&)> const& F)
+    :
+    OdeSolver<Real, TVector>(tDelta, F)
+{
+}
+
+template <typename Real, typename TVector>
+void OdeRungeKutta4<Real, TVector>::Update(Real tIn, TVector const& xIn,
+    Real& tOut, TVector& xOut)
+{
+    // Compute the first step.
+    Real halfTDelta = ((Real)0.5) * this->mTDelta;
+    TVector fTemp1 = this->mFunction(tIn, xIn);
+    TVector xTemp = xIn + halfTDelta * fTemp1;
+
+    // Compute the second step.
+    Real halfT = tIn + halfTDelta;
+    TVector fTemp2 = this->mFunction(halfT, xTemp);
+    xTemp = xIn + halfTDelta * fTemp2;
+
+    // Compute the third step.
+    TVector fTemp3 = this->mFunction(halfT, xTemp);
+    xTemp = xIn + this->mTDelta * fTemp3;
+
+    // Compute the fourth step.
+    Real sixthTDelta = this->mTDelta / (Real)6;
+    tOut = tIn + this->mTDelta;
+    TVector fTemp4 = this->mFunction(tOut, xTemp);
+    xOut = xIn + sixthTDelta * (
+        fTemp1 + ((Real)2)*(fTemp2 + fTemp3) + fTemp4);
+}
+
+
+}
diff --git a/Engine/Plugins/Experimental/GeometryProcessing/Source/GeometryAlgorithms/Private/GTEngine/Mathematics/GteOdeSolver.h b/Engine/Plugins/Experimental/GeometryProcessing/Source/GeometryAlgorithms/Private/GTEngine/Mathematics/GteOdeSolver.h
new file mode 100644
--- /dev/null
+++ b/Engine/Plugins/Experimental/GeometryProcessing/Source/GeometryAlgorithms/Private/GTEngine/Mathematics/GteOdeSolver.h
@@ -0,0 +1,77 @@
+// David Eberly, Geometric Tools, Redmond WA 98052
+// Copyright (c) 1998-2019
+// Distributed under the Boost Software License, Version 1.0.
+// http://www.boost.org/LICENSE_1_0.txt
+// http://www.geometrictools.com/License/Boost/LICENSE_1_0.txt
+// File Version: 3.0.0 (2016/06/19)
+
+#pragma once
+
+#include <GTEngine/Mathematics/GteGVector.h>
+#include <functional>
+
+// The differential equation is dx/dt = F(t,x).  The TVector template
+// parameter allows you to create solvers with Vector<N,Real> when the
+// dimension N is known at compile time or GVector<Real> when the dimension
+// N is known at run time.  Both classes have 'int GetSize() const' that
+// allow OdeSolver-derived classes to query for the dimension.
+
+namespace gte
+{
+
+template <typename Real, typename TVector>
+class OdeSolver
+{
+public:
+    // Abstract base class.
+public:
+    virtual ~OdeSolver();
+protected:
+    OdeSolver(Real tDelta,
+        std::function<TVector(Real, TVector const&)> const& F);
+
+public:
+    // Member access.
+    inline void SetTDelta(Real tDelta);
+    inline Real GetTDelta() const;
+
+    // Estimate x(t + tDelta) from x(t) using dx/dt = F(t,x).  The
+    // derived classes implement this so that it is possible for xIn and
+    // xOut to be the same object.
+    virtual void Update(Real tIn, TVector const& xIn, Real& tOut,
+        TVector& xOut) = 0;
+
+protected:
+    Real mTDelta;
+    std::function<TVector(Real, TVector const&)> mFunction;
+};
+
+
+template <typename Real, typename TVector>
+OdeSolver<Real, TVector>::~OdeSolver()
+{
+}
+
+template <typename Real, typename TVector>
+OdeSolver<Real, TVector>::OdeSolver(Real tDelta,
+    std::function<TVector(Real, TVector const&)> const& F)
+    :
+    mTDelta(tDelta),
+    mFunction(F)
+{
+}
+
+template <typename Real, typename TVector> inline
+void OdeSolver<Real, TVector>::SetTDelta(Real tDelta)
+{
+    mTDelta = tDelta;
+}
+
+template <typename Real, typename TVector> inline
+Real OdeSolver<Real, TVector>::GetTDelta() const
+{
+    return mTDelta;
+}
+
+
+}
diff --git a/Engine/Plugins/Experimental/GeometryProcessing/Source/GeometryAlgorithms/Private/GTEngine/Mathematics/GteOrientedBox.h b/Engine/Plugins/Experimental/GeometryProcessing/Source/GeometryAlgorithms/Private/GTEngine/Mathematics/GteOrientedBox.h
new file mode 100644
--- /dev/null
+++ b/Engine/Plugins/Experimental/GeometryProcessing/Source/GeometryAlgorithms/Private/GTEngine/Mathematics/GteOrientedBox.h
@@ -0,0 +1,165 @@
+// David Eberly, Geometric Tools, Redmond WA 98052
+// Copyright (c) 1998-2019
+// Distributed under the Boost Software License, Version 1.0.
+// http://www.boost.org/LICENSE_1_0.txt
+// http://www.geometrictools.com/License/Boost/LICENSE_1_0.txt
+// File Version: 3.0.0 (2016/06/19)
+
+#pragma once
+
+#include <GTEngine/Mathematics/GteVector.h>
+
+// A box has center C, axis directions U[i], and extents e[i].  The set
+// {U[0],...,U[N-1]} is orthonormal, which means the vectors are
+// unit-length and mutually perpendicular.  The extents are nonnegative;
+// zero is allowed, meaning the box is degenerate in the corresponding
+// direction.  A point X is represented in box coordinates by
+// X = C + y[0]*U[0] + y[1]*U[1].  This point is inside or on the
+// box whenever |y[i]| <= e[i] for all i.
+
+namespace gte
+{
+
+template <int N, typename Real>
+class OrientedBox
+{
+public:
+    // Construction and destruction.  The default constructor sets the center
+    // to (0,...,0), axis d to Vector<N,Real>::Unit(d), and extent d to +1.
+    OrientedBox();
+    OrientedBox(Vector<N, Real> const& inCenter,
+        std::array<Vector<N, Real>, N> const& inAxis,
+        Vector<N, Real> const& inExtent);
+
+    // Compute the vertices of the box.  If index i has the bit pattern
+    // i = b[N-1]...b[0], then
+    //    vertex[i] = center + sum_{d=0}^{N-1} sign[d] * extent[d] * axis[d]
+    // where sign[d] = 2*b[d] - 1.
+    void GetVertices(std::array<Vector<N, Real>, (1 << N)>& vertex) const;
+
+    // Public member access.  It is required that extent[i] >= 0.
+    Vector<N, Real> center;
+    std::array<Vector<N, Real>, N> axis;
+    Vector<N, Real> extent;
+
+public:
+    // Comparisons to support sorted containers.
+    bool operator==(OrientedBox const& box) const;
+    bool operator!=(OrientedBox const& box) const;
+    bool operator< (OrientedBox const& box) const;
+    bool operator<=(OrientedBox const& box) const;
+    bool operator> (OrientedBox const& box) const;
+    bool operator>=(OrientedBox const& box) const;
+};
+
+// Template aliases for convenience.
+template <typename Real>
+using OrientedBox2 = OrientedBox<2, Real>;
+
+template <typename Real>
+using OrientedBox3 = OrientedBox<3, Real>;
+
+
+template <int N, typename Real>
+OrientedBox<N, Real>::OrientedBox()
+{
+    center.MakeZero();
+    for (int i = 0; i < N; ++i)
+    {
+        axis[i].MakeUnit(i);
+        extent[i] = (Real)1;
+    }
+}
+
+template <int N, typename Real>
+OrientedBox<N, Real>::OrientedBox(Vector<N, Real> const& inCenter,
+    std::array<Vector<N, Real>, N> const& inAxis,
+    Vector<N, Real> const& inExtent)
+    :
+    center(inCenter),
+    axis(inAxis),
+    extent(inExtent)
+{
+}
+
+template <int N, typename Real>
+void OrientedBox<N, Real>::GetVertices(
+    std::array<Vector<N, Real>, (1 << N)>& vertex) const
+{
+    unsigned int const dsup = static_cast<unsigned int>(N);
+    std::array<Vector<N, Real>, N> product;
+    for (unsigned int d = 0; d < dsup; ++d)
+    {
+        product[d] = extent[d] * axis[d];
+    }
+
+    int const imax = (1 << N);
+    for (int i = 0; i < imax; ++i)
+    {
+        vertex[i] = center;
+        for (unsigned int d = 0, mask = 1; d < dsup; ++d, mask <<= 1)
+        {
+            Real sign = (i & mask ? (Real)1 : (Real)-1);
+            vertex[i] += sign * product[d];
+        }
+    }
+}
+
+template <int N, typename Real>
+bool OrientedBox<N, Real>::operator==(OrientedBox const& box) const
+{
+    return center == box.center && axis == box.axis && extent == box.extent;
+}
+
+template <int N, typename Real>
+bool OrientedBox<N, Real>::operator!=(OrientedBox const& box) const
+{
+    return !operator==(box);
+}
+
+template <int N, typename Real>
+bool OrientedBox<N, Real>::operator<(OrientedBox const& box) const
+{
+    if (center < box.center)
+    {
+        return true;
+    }
+
+    if (center > box.center)
+    {
+        return false;
+    }
+
+    if (axis < box.axis)
+    {
+        return true;
+    }
+
+    if (axis > box.axis)
+    {
+        return false;
+    }
+
+    return extent < box.extent;
+}
+
+template <int N, typename Real>
+bool OrientedBox<N, Real>::operator<=(OrientedBox const& box) const
+{
+    return operator<(box) || operator==(box);
+}
+
+template <int N, typename Real>
+bool OrientedBox<N, Real>::operator>(OrientedBox const& box) const
+{
+    return !operator<=(box);
+}
+
+template <int N, typename Real>
+bool OrientedBox<N, Real>::operator>=(OrientedBox const& box) const
+{
+    return !operator<(box);
+}
+
+
+}
