move GeometryProcessing plugin from NotForLicensees to Experimental

diff --git a/Engine/Plugins/Experimental/GeometryProcessing/Source/GeometryAlgorithms/Private/GTEngine/Mathematics/GteParametricCurve.h b/Engine/Plugins/Experimental/GeometryProcessing/Source/GeometryAlgorithms/Private/GTEngine/Mathematics/GteParametricCurve.h
new file mode 100644
--- /dev/null
+++ b/Engine/Plugins/Experimental/GeometryProcessing/Source/GeometryAlgorithms/Private/GTEngine/Mathematics/GteParametricCurve.h
@@ -0,0 +1,355 @@
+// David Eberly, Geometric Tools, Redmond WA 98052
+// Copyright (c) 1998-2019
+// Distributed under the Boost Software License, Version 1.0.
+// http://www.boost.org/LICENSE_1_0.txt
+// http://www.geometrictools.com/License/Boost/LICENSE_1_0.txt
+// File Version: 3.0.0 (2016/06/19)
+
+#pragma once
+
+#include <GTEngine/Mathematics/GteVector.h>
+#include <GTEngine/Mathematics/GteIntegration.h>
+#include <GTEngine/Mathematics/GteRootsBisection.h>
+#include <algorithm>
+#include <vector>
+
+namespace gte
+{
+
+template <int N, typename Real>
+class ParametricCurve
+{
+protected:
+    // Abstract base class for a parameterized curve X(t), where t is the
+    // parameter in [tmin,tmax] and X is an N-tuple position.  The first
+    // constructor is for single-segment curves.  The second constructor is
+    // for multiple-segment curves. The times must be strictly increasing.
+    ParametricCurve(Real tmin, Real tmax);
+    ParametricCurve(int numSegments, Real const* times);
+public:
+    virtual ~ParametricCurve();
+
+    // To validate construction, create an object as shown:
+    //     DerivedClassCurve<N, Real> curve(parameters);
+    //     if (!curve) { <constructor failed, handle accordingly>; }
+    inline operator bool() const;
+
+    // Member access.
+    inline Real GetTMin() const;
+    inline Real GetTMax() const;
+    inline int GetNumSegments() const;
+    Real const* GetTimes() const;
+
+    // This function applies only when the first constructor is used (two
+    // times rather than a sequence of three or more times).
+    void SetTimeInterval(Real tmin, Real tmax);
+
+    // Parameters used in GetLength(...), GetTotalLength(), and GetTime(...).
+    void SetRombergOrder(int order);  // default = 8
+    void SetMaxBisections(unsigned int maxBisections);  // default = 1024
+
+    // Evaluation of the curve.  The function supports derivative calculation
+    // through order 3; that is, maxOrder <= 3 is required.  If you want
+    // only the position, pass in maxOrder of 0.  If you want the position and
+    // first derivative, pass in maxOrder of 1, and so on.  The output
+    // 'values' are ordered as: position, first derivative, second derivative,
+    // third derivative.
+    virtual void Evaluate(Real t, unsigned int maxOrder,
+        Vector<N, Real> values[4]) const = 0;
+
+    // Differential geometric quantities.
+    Vector<N, Real> GetPosition(Real t) const;
+    Vector<N, Real> GetTangent(Real t) const;
+    Real GetSpeed(Real t) const;
+    Real GetLength(Real t0, Real t1) const;
+    Real GetTotalLength() const;
+
+    // Inverse mapping of s = Length(t) given by t = Length^{-1}(s).  The
+    // inverse length function is generally a function that cannot be written
+    // in closed form, in which case it is not directly computable.  Instead,
+    // we can specify s and estimate the root t for F(t) = Length(t) - s.  The
+    // derivative is F'(t) = Speed(t) >= 0, so F(t) is nondecreasing.  To be
+    // robust, we use bisection to locate the root, although it is possible to
+    // use a hybrid of Newton's method and bisection.
+    Real GetTime(Real length) const;
+
+    // Compute a subset of curve points according to the specified attribute.
+    // The input 'numPoints' must be two or larger.
+    void SubdivideByTime(int numPoints, Vector<N, Real>* points) const;
+    void SubdivideByLength(int numPoints, Vector<N, Real>* points) const;
+
+protected:
+    enum
+    {
+        DEFAULT_ROMBERG_ORDER = 8,
+        DEFAULT_MAX_BISECTIONS = 1024
+    };
+
+    std::vector<Real> mTime;
+    mutable std::vector<Real> mSegmentLength;
+    mutable std::vector<Real> mAccumulatedLength;
+    int mRombergOrder;
+    unsigned int mMaxBisections;
+
+    bool mConstructed;
+};
+
+
+template <int N, typename Real>
+ParametricCurve<N, Real>::ParametricCurve(Real tmin, Real tmax)
+    :
+    mTime(2),
+    mSegmentLength(1),
+    mAccumulatedLength(1),
+    mRombergOrder(DEFAULT_ROMBERG_ORDER),
+    mMaxBisections(DEFAULT_MAX_BISECTIONS),
+    mConstructed(false)
+{
+    mTime[0] = tmin;
+    mTime[1] = tmax;
+    mSegmentLength[0] = (Real)0;
+    mAccumulatedLength[0] = (Real)0;
+}
+
+template <int N, typename Real>
+ParametricCurve<N, Real>::ParametricCurve(int numSegments, Real const* times)
+    :
+    mTime(numSegments + 1),
+    mSegmentLength(numSegments),
+    mAccumulatedLength(numSegments),
+    mRombergOrder(DEFAULT_ROMBERG_ORDER),
+    mMaxBisections(DEFAULT_MAX_BISECTIONS),
+    mConstructed(false)
+{
+    std::copy(times, times + numSegments + 1, mTime.begin());
+    mSegmentLength[0] = (Real)0;
+    mAccumulatedLength[0] = (Real)0;
+}
+
+template <int N, typename Real>
+ParametricCurve<N, Real>::~ParametricCurve()
+{
+}
+
+template <int N, typename Real> inline
+ParametricCurve<N, Real>::operator bool() const
+{
+    return mConstructed;
+}
+
+template <int N, typename Real> inline
+Real ParametricCurve<N, Real>::GetTMin() const
+{
+    return mTime.front();
+}
+
+template <int N, typename Real> inline
+Real ParametricCurve<N, Real>::GetTMax() const
+{
+    return mTime.back();
+}
+
+template <int N, typename Real> inline
+int ParametricCurve<N, Real>::GetNumSegments() const
+{
+    return static_cast<int>(mSegmentLength.size());
+}
+
+template <int N, typename Real> inline
+Real const* ParametricCurve<N, Real>::GetTimes() const
+{
+    return &mTime[0];
+}
+
+template <int N, typename Real>
+void ParametricCurve<N, Real>::SetTimeInterval(Real tmin, Real tmax)
+{
+    if (mTime.size() == 2)
+    {
+        mTime[0] = tmin;
+        mTime[1] = tmax;
+    }
+}
+
+template <int N, typename Real>
+void ParametricCurve<N, Real>::SetRombergOrder(int order)
+{
+    mRombergOrder = std::max(order, 1);
+}
+
+template <int N, typename Real>
+void ParametricCurve<N, Real>::SetMaxBisections(unsigned int maxBisections)
+{
+    mMaxBisections = std::max(maxBisections, 1u);
+}
+
+template <int N, typename Real>
+Vector<N, Real> ParametricCurve<N, Real>::GetPosition(Real t) const
+{
+    Vector<N, Real> values[4];
+    Evaluate(t, 0, values);
+    return values[0];
+}
+
+template <int N, typename Real>
+Vector<N, Real> ParametricCurve<N, Real>::GetTangent(Real t) const
+{
+    Vector<N, Real> values[4];
+    Evaluate(t, 1, values);
+    Normalize(values[1]);
+    return values[1];
+}
+
+template <int N, typename Real>
+Real ParametricCurve<N, Real>::GetSpeed(Real t) const
+{
+    Vector<N, Real> values[4];
+    Evaluate(t, 1, values);
+    return Length(values[1]);
+}
+
+template <int N, typename Real>
+Real ParametricCurve<N, Real>::GetLength(Real t0, Real t1) const
+{
+    std::function<Real(Real)> speed = [this](Real t)
+    {
+        return GetSpeed(t);
+    };
+
+    if (mSegmentLength[0] == (Real)0)
+    {
+        // Lazy initialization of lengths of segments.
+        int const numSegments = static_cast<int>(mSegmentLength.size());
+        Real accumulated = (Real)0;
+        for (int i = 0; i < numSegments; ++i)
+        {
+            mSegmentLength[i] = Integration<Real>::Romberg(mRombergOrder,
+                mTime[i], mTime[i + 1], speed);
+            accumulated += mSegmentLength[i];
+            mAccumulatedLength[i] = accumulated;
+        }
+    }
+
+    t0 = std::max(t0, GetTMin());
+    t1 = std::min(t1, GetTMax());
+    auto iter0 = std::lower_bound(mTime.begin(), mTime.end(), t0);
+    int index0 = static_cast<int>(iter0 - mTime.begin());
+    auto iter1 = std::lower_bound(mTime.begin(), mTime.end(), t1);
+    int index1 = static_cast<int>(iter1 - mTime.begin());
+
+    Real length;
+    if (index0 < index1)
+    {
+        length = (Real)0;
+        if (t0 < *iter0)
+        {
+            length += Integration<Real>::Romberg(mRombergOrder, t0,
+                mTime[index0], speed);
+        }
+
+        int isup;
+        if (t1 < *iter1)
+        {
+            length += Integration<Real>::Romberg(mRombergOrder,
+                mTime[index1 - 1], t1, speed);
+            isup = index1 - 1;
+        }
+        else
+        {
+            isup = index1;
+        }
+        for (int i = index0; i < isup; ++i)
+        {
+            length += mSegmentLength[i];
+        }
+    }
+    else
+    {
+        length = Integration<Real>::Romberg(mRombergOrder, t0, t1, speed);
+    }
+    return length;
+}
+
+template <int N, typename Real>
+Real ParametricCurve<N, Real>::GetTotalLength() const
+{
+    if (mAccumulatedLength.back() == (Real)0)
+    {
+        // Lazy evaluation of the accumulated length array.
+        return GetLength(mTime.front(), mTime.back());
+    }
+
+    return mAccumulatedLength.back();
+}
+
+template <int N, typename Real>
+Real ParametricCurve<N, Real>::GetTime(Real length) const
+{
+    if (length > (Real)0)
+    {
+        if (length < GetTotalLength())
+        {
+            std::function<Real(Real)> F = [this, &length](Real t)
+            {
+                return Integration<Real>::Romberg(mRombergOrder,
+                    mTime.front(), t, [this](Real z){ return GetSpeed(z); })
+                    - length;
+            };
+
+            // We know that F(tmin) < 0 and F(tmax) > 0, which allows us to
+            // use bisection.  Rather than bisect the entire interval, let's
+            // narrow it down with a reasonable initial guess.
+            Real ratio = length / GetTotalLength();
+            Real omratio = (Real)1 - ratio;
+            Real tmid = omratio * mTime.front() + ratio * mTime.back();
+            Real fmid = F(tmid);
+            if (fmid > (Real)0)
+            {
+                RootsBisection<Real>::Find(F, mTime.front(), tmid, (Real)-1,
+                    (Real)1, mMaxBisections, tmid);
+            }
+            else if (fmid < (Real)0)
+            {
+                RootsBisection<Real>::Find(F, tmid, mTime.back(), (Real)-1,
+                    (Real)1, mMaxBisections, tmid);
+            }
+            return tmid;
+        }
+        else
+        {
+            return mTime.back();
+        }
+    }
+    else
+    {
+        return mTime.front();
+    }
+}
+
+template <int N, typename Real>
+void ParametricCurve<N, Real>::SubdivideByTime(int numPoints,
+    Vector<N, Real>* points) const
+{
+    Real delta = (mTime.back() - mTime.front()) / (Real)(numPoints - 1);
+    for (int i = 0; i < numPoints; ++i)
+    {
+        Real t = mTime.front() + delta * i;
+        points[i] = GetPosition(t);
+    }
+}
+
+template <int N, typename Real>
+void ParametricCurve<N, Real>::SubdivideByLength(int numPoints,
+    Vector<N, Real>* points) const
+{
+    Real delta = GetTotalLength() / (Real)(numPoints - 1);
+    for (int i = 0; i < numPoints; ++i)
+    {
+        Real length = delta * i;
+        Real t = GetTime(length);
+        points[i] = GetPosition(t);
+    }
+}
+
+
+}
diff --git a/Engine/Plugins/Experimental/GeometryProcessing/Source/GeometryAlgorithms/Private/GTEngine/Mathematics/GteParametricSurface.h b/Engine/Plugins/Experimental/GeometryProcessing/Source/GeometryAlgorithms/Private/GTEngine/Mathematics/GteParametricSurface.h
new file mode 100644
--- /dev/null
+++ b/Engine/Plugins/Experimental/GeometryProcessing/Source/GeometryAlgorithms/Private/GTEngine/Mathematics/GteParametricSurface.h
@@ -0,0 +1,144 @@
+// David Eberly, Geometric Tools, Redmond WA 98052
+// Copyright (c) 1998-2019
+// Distributed under the Boost Software License, Version 1.0.
+// http://www.boost.org/LICENSE_1_0.txt
+// http://www.geometrictools.com/License/Boost/LICENSE_1_0.txt
+// File Version: 3.0.0 (2016/06/19)
+
+#pragma once
+
+#include <GTEngine/Mathematics/GteVector.h>
+
+namespace gte
+{
+
+template <int N, typename Real>
+class ParametricSurface
+{
+protected:
+    // Abstract base class for a parameterized surface X(u,v).  The parametric
+    // domain is either rectangular or triangular.  Valid (u,v) values for a
+    // rectangular domain satisfy
+    //   umin <= u <= umax,  vmin <= v <= vmax
+    // and valid (u,v) values for a triangular domain satisfy
+    //   umin <= u <= umax,  vmin <= v <= vmax,
+    //   (vmax-vmin)*(u-umin)+(umax-umin)*(v-vmax) <= 0
+    ParametricSurface(Real umin, Real umax, Real vmin, Real vmax,
+        bool rectangular);
+public:
+    virtual ~ParametricSurface();
+
+    // To validate construction, create an object as shown:
+    //     DerivedClassSurface<Real> surface(parameters);
+    //     if (!surface) { <constructor failed, handle accordingly>; }
+    inline operator bool() const;
+
+    // Member access.
+    inline Real GetUMin() const;
+    inline Real GetUMax() const;
+    inline Real GetVMin() const;
+    inline Real GetVMax() const;
+    inline bool IsRectangular() const;
+
+    // Evaluation of the surface.  The function supports derivative
+    // calculation through order 2; that is, maxOrder <= 2 is required.  If
+    // you want only the position, pass in maxOrder of 0.  If you want the
+    // position and first-order derivatives, pass in maxOrder of 1, and so on.
+    // The output 'values' are ordered as: position X; first-order derivatives
+    // dX/du, dX/dv; second-order derivatives d2X/du2, d2X/dudv, d2X/dv2.
+    virtual void Evaluate(Real u, Real v, unsigned int maxOrder,
+        Vector<N, Real> values[6]) const = 0;
+
+    // Differential geometric quantities.
+    Vector<N, Real> GetPosition(Real u, Real v) const;
+    Vector<N, Real> GetUTangent(Real u, Real v) const;
+    Vector<N, Real> GetVTangent(Real u, Real v) const;
+
+protected:
+    Real mUMin, mUMax, mVMin, mVMax;
+    bool mRectangular;
+    bool mConstructed;
+};
+
+
+template <int N, typename Real>
+ParametricSurface<N, Real>::ParametricSurface(Real umin, Real umax,
+    Real vmin, Real vmax, bool rectangular)
+    :
+    mUMin(umin),
+    mUMax(umax),
+    mVMin(vmin),
+    mVMax(vmax),
+    mRectangular(rectangular)
+{
+}
+
+template <int N, typename Real>
+ParametricSurface<N, Real>::~ParametricSurface()
+{
+}
+
+template <int N, typename Real>
+ParametricSurface<N, Real>::operator bool() const
+{
+    return mConstructed;
+}
+
+template <int N, typename Real> inline
+Real ParametricSurface<N, Real>::GetUMin() const
+{
+    return mUMin;
+}
+
+template <int N, typename Real> inline
+Real ParametricSurface<N, Real>::GetUMax() const
+{
+    return mUMax;
+}
+
+template <int N, typename Real> inline
+Real ParametricSurface<N, Real>::GetVMin() const
+{
+    return mVMin;
+}
+
+template <int N, typename Real> inline
+Real ParametricSurface<N, Real>::GetVMax() const
+{
+    return mVMax;
+}
+
+template <int N, typename Real> inline
+bool ParametricSurface<N, Real>::IsRectangular() const
+{
+    return mRectangular;
+}
+
+template <int N, typename Real> inline
+Vector<N, Real> ParametricSurface<N, Real>::GetPosition(Real u, Real v) const
+{
+    Vector<N, Real> values[6];
+    Evaluate(u, v, 0, values);
+    return values[0];
+}
+
+template <int N, typename Real> inline
+Vector<N, Real> ParametricSurface<N, Real>::GetUTangent(Real u, Real v) const
+{
+    Vector<N, Real> values[6];
+    Evaluate(u, v, 1, values);
+    Normalize(values[1]);
+    return values[1];
+}
+
+template <int N, typename Real> inline
+Vector<N, Real> ParametricSurface<N, Real>::GetVTangent(Real u, Real v) const
+{
+    Vector<N, Real> values[6];
+    Evaluate(u, v, 1, values);
+    Normalize(values[2]);
+    return values[2];
+}
+
+
+}
diff --git a/Engine/Plugins/Experimental/GeometryProcessing/Source/GeometryAlgorithms/Private/GTEngine/Mathematics/GtePlanarMesh.h b/Engine/Plugins/Experimental/GeometryProcessing/Source/GeometryAlgorithms/Private/GTEngine/Mathematics/GtePlanarMesh.h
new file mode 100644
--- /dev/null
+++ b/Engine/Plugins/Experimental/GeometryProcessing/Source/GeometryAlgorithms/Private/GTEngine/Mathematics/GtePlanarMesh.h
@@ -0,0 +1,500 @@
+// David Eberly, Geometric Tools, Redmond WA 98052
+// Copyright (c) 1998-2019
+// Distributed under the Boost Software License, Version 1.0.
+// http://www.boost.org/LICENSE_1_0.txt
+// http://www.geometrictools.com/License/Boost/LICENSE_1_0.txt
+// File Version: 3.0.0 (2016/06/19)
+
+#pragma once
+
+#include <GTEngine/LowLevel/GteLogger.h>
+#include <GTEngine/Mathematics/GteETManifoldMesh.h>
+#include <GTEngine/Mathematics/GtePrimalQuery2.h>
+#include <GTEngine/Mathematics/GteContPointInPolygon2.h>
+#include <algorithm>
+#include <vector>
+
+// The planar mesh class is convenient for many applications involving
+// searches for triangles containing a specified point.  A couple of
+// issues can show up in practice when the input data to the constructors
+// is very large (number of triangles on the order of 10^5 or larger).
+//
+// The first constructor builds an ETManifoldMesh mesh that contains
+// std::map objects.  When such maps are large, the amount of time it
+// takes to delete them is enormous.  Although you can control the level
+// of debug support in MSVS 2013 (see _ITERATOR_DEBUG_LEVEL), turning off
+// checking might very well affect other classes for which you want
+// iterator checking to be on.  An alternative to reduce debugging time
+// is to dynamically allocate the PlanarMesh object in the main thread but
+// then launch another thread to delete the object and avoid stalling
+// the main thread.  For example,
+//
+//    PlanarMesh<IT,CT,RT>* pmesh =
+//        new PlanarMesh<IT,CT,RT>(numV, vertices, numT, indices);
+//    <make various calls to pmesh>;
+//    std::thread deleter = [pmesh](){ delete pmesh; };
+//    deleter.detach();  // Do not wait for the thread to finish.
+//
+// The second constructor has the mesh passed in, but mTriIndexMap is used
+// in both constructors and can take time to delete.
+//
+// The input mesh should be consistently oriented, say, the triangles are
+// counterclockwise ordered.  The vertices should be consistent with this
+// ordering.  However, floating-point rounding errors in generating the
+// vertices can cause apparent fold-over of the mesh; that is, theoretically
+// the vertex geometry supports counterclockwise geometry but numerical
+// errors cause an inconsistency.  This can manifest in the mQuery.ToLine
+// tests whereby cycles of triangles occur in the linear walk.  When cycles
+// occur, GetContainingTriangle(P,startTriangle) will iterate numTriangle
+// times before reporting that the triangle cannot be found, which is a
+// very slow process (in debug or release builds).  The function
+// GetContainingTriangle(P,startTriangle,visited) is provided to avoid the
+// performance loss, trapping a cycle the first time and exiting, but
+// again reporting that the triangle cannot be found.  If you know that the
+// query should be (theoretically) successful, use the second version of
+// GetContainingTriangle.  If it fails by returning -1, then perform an
+// exhaustive search over the triangles.  For example,
+//
+//    int triangle = pmesh->GetContainingTriangle(P,startTriangle,visited);
+//    if (triangle >= 0)
+//    {
+//        <take action; for example, compute barycenteric coordinates>;
+//    }
+//    else
+//    {
+//        int numTriangles = pmesh->GetNumTriangles();
+//        for (triangle = 0; triangle < numTriangles; ++triangle)
+//        {
+//            if (pmesh->Contains(triangle, P))
+//            {
+//                <take action>;
+//                break;
+//            }
+//        }
+//        if (triangle == numTriangles)
+//        {
+//            <Triangle still not found, take appropriate action>;
+//        }
+//    }
+//
+// The PlanarMesh<*>::Contains function does not require the triangles to
+// be ordered.
+
+namespace gte
+{
+
+template <typename InputType, typename ComputeType, typename RationalType>
+class PlanarMesh
+{
+public:
+    // Construction.  The inputs must represent a manifold mesh of triangles
+    // in the plane.  The index array must have 3*numTriangles elements, each
+    // triple of indices representing a triangle in the mesh.  Each index is
+    // into the 'vertices' array.
+    PlanarMesh(int numVertices, Vector2<InputType> const* vertices, int numTriangles, int const* indices);
+    PlanarMesh(int numVertices, Vector2<InputType> const* vertices, ETManifoldMesh const& mesh);
+
+    // Mesh information.
+    inline int GetNumVertices() const;
+    inline int GetNumTriangles() const;
+    inline Vector2<InputType> const* GetVertices() const;
+    inline int const* GetIndices() const;
+    inline int const* GetAdjacencies() const;
+
+    // Containment queries.  The function GetContainingTriangle works
+    // correctly when the planar mesh is a convex set.  If the mesh is not
+    // convex, it is possible that the linear-walk search algorithm exits the
+    // mesh before finding a containing triangle.  For example, a C-shaped
+    // mesh can contain a point in the top branch of the "C".  A starting
+    // point in the bottom branch of the "C" will lead to the search exiting
+    // the bottom branch and having no path to walk to the top branch.  If
+    // your mesh is not convex and you want a correct containment query, you
+    // will have to append "outside" triangles to your mesh to form a convex
+    // set.
+    int GetContainingTriangle(Vector2<InputType> const& P, int startTriangle = 0) const;
+    int GetContainingTriangle(Vector2<InputType> const& P, int startTriangle, std::set<int>& visited) const;
+    bool GetVertices(int t, std::array<Vector2<InputType>, 3>& vertices) const;
+    bool GetIndices(int t, std::array<int, 3>& indices) const;
+    bool GetAdjacencies(int t, std::array<int, 3>& adjacencies) const;
+    bool GetBarycentrics(int t, Vector2<InputType> const& P, std::array<InputType, 3>& bary) const;
+    bool Contains(int triangle, Vector2<InputType> const& P) const;
+
+public:
+    void CreateVertices(int numVertices, Vector2<InputType> const* vertices);
+
+    int mNumVertices;
+    Vector2<InputType> const* mVertices;
+    int mNumTriangles;
+    std::vector<int> mIndices;
+    ETManifoldMesh mMesh;
+    std::map<TriangleKey<true>, int> mTriIndexMap;
+    std::vector<int> mAdjacencies;
+    std::vector<Vector2<ComputeType>> mComputeVertices;
+    PrimalQuery2<ComputeType> mQuery;
+};
+
+
+template <typename InputType, typename ComputeType, typename RationalType>
+PlanarMesh<InputType, ComputeType, RationalType>::PlanarMesh(int numVertices,
+    Vector2<InputType> const* vertices, int numTriangles, int const* indices)
+    :
+    mNumVertices(0),
+    mVertices(nullptr),
+    mNumTriangles(0)
+{
+    if (numVertices < 3 || !vertices || numTriangles < 1 || !indices)
+    {
+        LogError("Invalid input.");
+        return;
+    }
+
+    // Create a mesh in order to get adjacency information.
+    int const* current = indices;
+    for (int t = 0; t < numTriangles; ++t)
+    {
+        int v0 = *current++;
+        int v1 = *current++;
+        int v2 = *current++;
+        if (!mMesh.Insert(v0, v1, v2))
+        {
+            // The 'mesh' object will assert on nonmanifold inputs.
+            return;
+        }
+    }
+
+    // We have a valid mesh.
+    CreateVertices(numVertices, vertices);
+
+    // Build the adjacency graph using the triangle ordering implied by the
+    // indices, not the mesh triangle map, to preserve the triangle ordering
+    // of the input indices.
+    mNumTriangles = numTriangles;
+    int const numIndices = 3 * numTriangles;
+    mIndices.resize(numIndices);
+
+    std::copy(indices, indices + numIndices, mIndices.begin());
+    for (int t = 0, vIndex = 0; t < numTriangles; ++t)
+    {
+        int v0 = indices[vIndex++];
+        int v1 = indices[vIndex++];
+        int v2 = indices[vIndex++];
+        TriangleKey<true> key(v0, v1, v2);
+        mTriIndexMap.insert(std::make_pair(key, t));
+    }
+
+    mAdjacencies.resize(numIndices);
+    auto const& tmap = mMesh.GetTriangles();
+    for (int t = 0, base = 0; t < numTriangles; ++t, base += 3)
+    {
+        int v0 = indices[base];
+        int v1 = indices[base + 1];
+        int v2 = indices[base + 2];
+        TriangleKey<true> key(v0, v1, v2);
+        auto element = tmap.find(key);
+        for (int i = 0; i < 3; ++i)
+        {
+            auto adj = element->second->T[i].lock();
+            if (adj)
+            {
+                key = TriangleKey<true>(adj->V[0], adj->V[1], adj->V[2]);
+                mAdjacencies[base + i] = mTriIndexMap.find(key)->second;
+            }
+            else
+            {
+                mAdjacencies[base + i] = -1;
+            }
+        }
+    }
+}
+
+template <typename InputType, typename ComputeType, typename RationalType>
+PlanarMesh<InputType, ComputeType, RationalType>::PlanarMesh(int numVertices,
+    Vector2<InputType> const* vertices, ETManifoldMesh const& mesh)
+    :
+    mNumVertices(0),
+    mVertices(nullptr),
+    mNumTriangles(0)
+{
+    if (numVertices < 3 || !vertices || mesh.GetTriangles().size() < 1)
+    {
+        LogError("Invalid input.");
+        return;
+    }
+
+    // We have a valid mesh.
+    CreateVertices(numVertices, vertices);
+
+    // Build the adjacency graph using the triangle ordering implied by the
+    // mesh triangle map.
+    auto const& tmap = mesh.GetTriangles();
+    mNumTriangles = static_cast<int>(tmap.size());
+    mIndices.resize(3 * mNumTriangles);
+
+    int tIndex = 0, vIndex = 0;
+    for (auto const& element : tmap)
+    {
+        mTriIndexMap.insert(std::make_pair(element.first, tIndex++));
+        for (int i = 0; i < 3; ++i, ++vIndex)
+        {
+            mIndices[vIndex] = element.second->V[i];
+        }
+    }
+
+    mAdjacencies.resize(3 * mNumTriangles);
+    vIndex = 0;
+    for (auto const& element : tmap)
+    {
+        for (int i = 0; i < 3; ++i, ++vIndex)
+        {
+            auto adj = element.second->T[i].lock();
+            if (adj)
+            {
+                TriangleKey<true> key(adj->V[0], adj->V[1], adj->V[2]);
+                mAdjacencies[vIndex] = mTriIndexMap.find(key)->second;
+            }
+            else
+            {
+                mAdjacencies[vIndex] = -1;
+            }
+        }
+    }
+}
+
+template <typename InputType, typename ComputeType, typename RationalType>
+inline int PlanarMesh<InputType, ComputeType, RationalType>::
+GetNumVertices() const
+{
+    return mNumVertices;
+}
+
+template <typename InputType, typename ComputeType, typename RationalType>
+inline int PlanarMesh<InputType, ComputeType, RationalType>::GetNumTriangles() const
+{
+    return mNumTriangles;
+}
+
+template <typename InputType, typename ComputeType, typename RationalType>
+inline Vector2<InputType> const* PlanarMesh<InputType, ComputeType, RationalType>::GetVertices() const
+{
+    return mVertices;
+}
+
+template <typename InputType, typename ComputeType, typename RationalType>
+inline int const* PlanarMesh<InputType, ComputeType, RationalType>::GetIndices() const
+{
+    return &mIndices[0];
+}
+
+template <typename InputType, typename ComputeType, typename RationalType>
+inline int const* PlanarMesh<InputType, ComputeType, RationalType>::GetAdjacencies() const
+{
+    return &mAdjacencies[0];
+}
+
+template <typename InputType, typename ComputeType, typename RationalType>
+int PlanarMesh<InputType, ComputeType, RationalType>::GetContainingTriangle(
+    Vector2<InputType> const& P, int startTriangle) const
+{
+    Vector2<ComputeType> test{ P[0], P[1] };
+
+    // Use triangle edges as binary separating lines.
+    int triangle = startTriangle;
+    for (int i = 0; i < mNumTriangles; ++i)
+    {
+        int ibase = 3 * triangle;
+        int const* v = &mIndices[ibase];
+
+        if (mQuery.ToLine(test, v[0], v[1]) > 0)
+        {
+            triangle = mAdjacencies[ibase];
+            if (triangle == -1)
+            {
+                return -1;
+            }
+            continue;
+        }
+
+        if (mQuery.ToLine(test, v[1], v[2]) > 0)
+        {
+            triangle = mAdjacencies[ibase + 1];
+            if (triangle == -1)
+            {
+                return -1;
+            }
+            continue;
+        }
+
+        if (mQuery.ToLine(test, v[2], v[0]) > 0)
+        {
+            triangle = mAdjacencies[ibase + 2];
+            if (triangle == -1)
+            {
+                return -1;
+            }
+            continue;
+        }
+
+        return triangle;
+    }
+
+    return -1;
+}
+
+template <typename InputType, typename ComputeType, typename RationalType>
+int PlanarMesh<InputType, ComputeType, RationalType>::GetContainingTriangle(
+    Vector2<InputType> const& P, int startTriangle, std::set<int>& visited) const
+{
+    Vector2<ComputeType> test{ P[0], P[1] };
+    visited.clear();
+
+    // Use triangle edges as binary separating lines.
+    int triangle = startTriangle;
+    for (int i = 0; i < mNumTriangles; ++i)
+    {
+        visited.insert(triangle);
+        int ibase = 3 * triangle;
+        int const* v = &mIndices[ibase];
+
+        if (mQuery.ToLine(test, v[0], v[1]) > 0)
+        {
+            triangle = mAdjacencies[ibase];
+            if (triangle == -1 || visited.find(triangle) != visited.end())
+            {
+                return -1;
+            }
+            continue;
+        }
+
+        if (mQuery.ToLine(test, v[1], v[2]) > 0)
+        {
+            triangle = mAdjacencies[ibase + 1];
+            if (triangle == -1 || visited.find(triangle) != visited.end())
+            {
+                return -1;
+            }
+            continue;
+        }
+
+        if (mQuery.ToLine(test, v[2], v[0]) > 0)
+        {
+            triangle = mAdjacencies[ibase + 2];
+            if (triangle == -1 || visited.find(triangle) != visited.end())
+            {
+                return -1;
+            }
+            continue;
+        }
+
+        return triangle;
+    }
+
+    return -1;
+}
+
+template <typename InputType, typename ComputeType, typename RationalType>
+bool PlanarMesh<InputType, ComputeType, RationalType>::GetVertices(
+    int t, std::array<Vector2<InputType>, 3>& vertices) const
+{
+    if (0 <= t && t < mNumTriangles)
+    {
+        for (int i = 0, vIndex = 3 * t; i < 3; ++i, ++vIndex)
+        {
+            vertices[i] = mVertices[mIndices[vIndex]];
+        }
+        return true;
+    }
+    return false;
+}
+
+template <typename InputType, typename ComputeType, typename RationalType>
+bool PlanarMesh<InputType, ComputeType, RationalType>::GetIndices(
+    int t, std::array<int, 3>& indices) const
+{
+    if (0 <= t && t < mNumTriangles)
+    {
+        for (int i = 0, vIndex = 3 * t; i < 3; ++i, ++vIndex)
+        {
+            indices[i] = mIndices[vIndex];
+        }
+        return true;
+    }
+    return false;
+}
+
+template <typename InputType, typename ComputeType, typename RationalType>
+bool PlanarMesh<InputType, ComputeType, RationalType>::GetAdjacencies(
+    int t, std::array<int, 3>& adjacencies) const
+{
+    if (0 <= t && t < mNumTriangles)
+    {
+        for (int i = 0, vIndex = 3 * t; i < 3; ++i, ++vIndex)
+        {
+            adjacencies[i] = mAdjacencies[vIndex];
+        }
+        return true;
+    }
+    return false;
+}
+
+template <typename InputType, typename ComputeType, typename RationalType>
+bool PlanarMesh<InputType, ComputeType, RationalType>::GetBarycentrics(
+    int t, Vector2<InputType> const& P, std::array<InputType, 3>& bary) const
+{
+    std::array<int, 3> indices;
+    if (GetIndices(t, indices))
+    {
+        Vector2<RationalType> rtP{ P[0], P[1] };
+        std::array<Vector2<RationalType>, 3> rtV;
+        for (int i = 0; i < 3; ++i)
+        {
+            Vector2<ComputeType> const& V = mComputeVertices[indices[i]];
+            for (int j = 0; j < 2; ++j)
+            {
+                rtV[i][j] = (RationalType)V[j];
+            }
+        };
+
+        RationalType rtBary[3];
+        if (ComputeBarycentrics(rtP, rtV[0], rtV[1], rtV[2], rtBary))
+        {
+            for (int i = 0; i < 3; ++i)
+            {
+                bary[i] = (InputType)rtBary[i];
+            }
+            return true;
+        }
+    }
+    return false;
+}
+
+template <typename InputType, typename ComputeType, typename RationalType>
+bool PlanarMesh<InputType, ComputeType, RationalType>::Contains(
+    int triangle, Vector2<InputType> const& P) const
+{
+    Vector2<ComputeType> test{ P[0], P[1] };
+    Vector2<ComputeType> v[3];
+    v[0] = mComputeVertices[mIndices[3 * triangle + 0]];
+    v[1] = mComputeVertices[mIndices[3 * triangle + 1]];
+    v[2] = mComputeVertices[mIndices[3 * triangle + 2]];
+    PointInPolygon2<ComputeType> pip(3, v);
+    return pip.Contains(test);
+}
+
+template <typename InputType, typename ComputeType, typename RationalType>
+void PlanarMesh<InputType, ComputeType, RationalType>::CreateVertices(
+    int numVertices, Vector2<InputType> const* vertices)
+{
+    mNumVertices = numVertices;
+    mVertices = vertices;
+    mComputeVertices.resize(mNumVertices);
+    for (int i = 0; i < mNumVertices; ++i)
+    {
+        for (int j = 0; j < 2; ++j)
+        {
+            mComputeVertices[i][j] = (ComputeType)mVertices[i][j];
+        }
+    }
+    mQuery.Set(mNumVertices, &mComputeVertices[0]);
+}
+
+}
diff --git a/Engine/Plugins/Experimental/GeometryProcessing/Source/GeometryAlgorithms/Private/GTEngine/Mathematics/GtePolygon2.h b/Engine/Plugins/Experimental/GeometryProcessing/Source/GeometryAlgorithms/Private/GTEngine/Mathematics/GtePolygon2.h
new file mode 100644
--- /dev/null
+++ b/Engine/Plugins/Experimental/GeometryProcessing/Source/GeometryAlgorithms/Private/GTEngine/Mathematics/GtePolygon2.h
@@ -0,0 +1,308 @@
+// David Eberly, Geometric Tools, Redmond WA 98052
+// Copyright (c) 1998-2019
+// Distributed under the Boost Software License, Version 1.0.
+// http://www.boost.org/LICENSE_1_0.txt
+// http://www.geometrictools.com/License/Boost/LICENSE_1_0.txt
+// File Version: 3.0.0 (2016/06/19)
+
+#pragma once
+
+#include <GTEngine/Mathematics/GteVector2.h>
+#include <GTEngine/Mathematics/GteIntrSegment2Segment2.h>
+#include <algorithm>
+#include <set>
+#include <vector>
+
+// The Polygon2 object represents a simple polygon:  No duplicate vertices,
+// closed (each vertex is shared by exactly two edges), and no
+// self-intersections at interior edge points.  The 'vertexPool' array can
+// contain more points than needed to define the polygon, which allows the
+// vertex pool to have multiple polygons associated with it.  Thus, the
+// programmer must ensure that the vertex pool persists as long as any
+// Polygon2 objects exist that depend on the pool.  The number of polygon
+// vertices is 'numIndices' and must be 3 or larger.  The 'indices' array
+// refers to the points in 'vertexPool' that are part of the polygon and must
+// have 'numIndices' unique elements.  The edges of the polygon are pairs of
+// indices into 'vertexPool',
+//   edge[0] = (indices[0], indices[1])
+//   :
+//   edge[numIndices-2] = (indices[numIndices-2], indices[numIndices-1])
+//   edge[numIndices-1] = (indices[numIndices-1], indices[0])
+// The programmer should ensure the polygon is simple.  The geometric
+// queries are valid regardless of whether the polygon is oriented clockwise
+// or counterclockwise.
+//
+// NOTE: Comparison operators are not provided.  The semantics of equal
+// polygons is complicated and (at the moment) not useful.  The vertex pools
+// can be different and indices do not match, but the vertices they reference
+// can match.  Even with a shared vertex pool, the indices can be "rotated",
+// leading to the same polygon abstractly but the data structures do not
+// match.
+
+namespace gte
+{
+
+template <typename Real>
+class Polygon2
+{
+public:
+    // Construction.  The constructor succeeds when 'numIndices' >= 3 and
+    // 'vertexPool' and 'indices' are not null; we cannot test whether you
+    // have a valid number of elements in the input arrays.  A copy is made
+    // of 'indices', but the 'vertexPool' is not copied.  If the constructor
+    // fails, the internal vertex pointer is set to null, the  index array
+    // has no elements, and the orientation is set to clockwise.
+    Polygon2(Vector2<Real> const* vertexPool, int numIndices,
+        int const* indices, bool counterClockwise);
+
+    // To validate construction, create an object as shown:
+    //     Polygon2<Real> polygon(parameters);
+    //     if (!polygon) { <constructor failed, handle accordingly>; }
+    inline operator bool() const;
+
+    // Member access.
+    inline Vector2<Real> const* GetVertexPool() const;
+    inline std::set<int> const& GetVertices() const;
+    inline std::vector<int> const& GetIndices() const;
+    inline bool CounterClockwise() const;
+
+    // Geometric queries.
+    Vector2<Real> ComputeVertexAverage() const;
+    Real ComputePerimeterLength() const;
+    Real ComputeArea() const;
+
+    // Simple polygons have no self-intersections at interior points
+    // of edges.  The test is an exhaustive all-pairs intersection
+    // test for edges, which is inefficient for polygons with a large
+    // number of vertices.  TODO: Provide an efficient algorithm that
+    // uses the algorithm of class RectangleManager.h.
+    bool IsSimple() const;
+
+    // Convex polygons are simple polygons where the angles between
+    // consecutive edges are less than or equal to pi radians.
+    bool IsConvex() const;
+
+private:
+    // These calls have preconditions that mVertexPool is not null and
+    // mIndices.size() > 3.  The heart of the algorithms are implemented
+    // here.
+    bool IsSimpleInternal() const;
+    bool IsConvexInternal() const;
+
+    Vector2<Real> const* mVertexPool;
+    std::set<int> mVertices;
+    std::vector<int> mIndices;
+    bool mCounterClockwise;
+};
+
+
+template <typename Real>
+Polygon2<Real>::Polygon2(Vector2<Real> const* vertexPool, int numIndices,
+    int const* indices, bool counterClockwise)
+    :
+    mVertexPool(vertexPool),
+    mCounterClockwise(counterClockwise)
+{
+    if (numIndices >= 3 && vertexPool && indices)
+    {
+        for (int i = 0; i < numIndices; ++i)
+        {
+            mVertices.insert(indices[i]);
+        }
+
+        if (numIndices == static_cast<int>(mVertices.size()))
+        {
+            mIndices.resize(numIndices);
+            std::copy(indices, indices + numIndices, mIndices.begin());
+            return;
+        }
+
+        // At least one duplicated vertex was encountered, so the polygon
+        // is not simple.  Fail the constructor call.
+        mVertices.clear();
+    }
+
+    // Invalid input to the Polygon2 constructor.
+    mVertexPool = nullptr;
+    mCounterClockwise = false;
+}
+
+template <typename Real> inline
+Polygon2<Real>::operator bool() const
+{
+    return mVertexPool != nullptr;
+}
+
+template <typename Real> inline
+Vector2<Real> const* Polygon2<Real>::GetVertexPool() const
+{
+    return mVertexPool;
+}
+
+template <typename Real> inline
+std::set<int> const& Polygon2<Real>::GetVertices() const
+{
+    return mVertices;
+}
+
+template <typename Real> inline
+std::vector<int> const& Polygon2<Real>::GetIndices() const
+{
+    return mIndices;
+}
+
+template <typename Real> inline
+bool Polygon2<Real>::CounterClockwise() const
+{
+    return mCounterClockwise;
+}
+
+template <typename Real>
+Vector2<Real> Polygon2<Real>::ComputeVertexAverage() const
+{
+    Vector2<Real> average = Vector2<Real>::Zero();
+    if (mVertexPool)
+    {
+        for (int index : mVertices)
+        {
+            average += mVertexPool[index];
+        }
+        average /= static_cast<Real>(mVertices.size());
+    }
+    return average;
+}
+
+template <typename Real>
+Real Polygon2<Real>::ComputePerimeterLength() const
+{
+    Real length = (Real)0;
+    if (mVertexPool)
+    {
+        Vector2<Real> v0 = mVertexPool[mIndices.back()];
+        for (int index : mIndices)
+        {
+            Vector2<Real> v1 = mVertexPool[index];
+            length += Length(v1 - v0);
+            v0 = v1;
+        }
+    }
+    return length;
+}
+
+template <typename Real>
+Real Polygon2<Real>::ComputeArea() const
+{
+    Real area = (Real)0;
+    if (mVertexPool)
+    {
+        int const numIndices = static_cast<int>(mIndices.size());
+        Vector2<Real> v0 = mVertexPool[mIndices[numIndices - 2]];
+        Vector2<Real> v1 = mVertexPool[mIndices[numIndices - 1]];
+        for (int index : mIndices)
+        {
+            Vector2<Real> v2 = mVertexPool[index];
+            area += v1[0] * (v2[1] - v0[1]);
+            v0 = v1;
+            v1 = v2;
+        }
+        area *= (Real)0.5;
+    }
+    return std::abs(area);
+}
+
+template <typename Real>
+bool Polygon2<Real>::IsSimple() const
+{
+    if (!mVertexPool)
+    {
+        return false;
+    }
+
+    // For mVertexPool to be nonnull, the number of indices is guaranteed
+    // to be at least 3.
+    int const numIndices = static_cast<int>(mIndices.size());
+    if (numIndices == 3)
+    {
+        // The polygon is a triangle.
+        return true;
+    }
+
+    return IsSimpleInternal();
+}
+
+template <typename Real>
+bool Polygon2<Real>::IsConvex() const
+{
+    if (!mVertexPool)
+    {
+        return false;
+    }
+
+    // For mVertexPool to be nonnull, the number of indices is guaranteed
+    // to be at least 3.
+    int const numIndices = static_cast<int>(mIndices.size());
+    if (numIndices == 3)
+    {
+        // The polygon is a triangle.
+        return true;
+    }
+
+    return IsSimpleInternal() && IsConvexInternal();
+}
+
+template <typename Real>
+bool Polygon2<Real>::IsSimpleInternal() const
+{
+    Segment2<Real> seg0, seg1;
+    TIQuery<Real, Segment2<Real>, Segment2<Real>> query;
+    typename TIQuery<Real, Segment2<Real>, Segment2<Real>>::Result result;
+
+    int const numIndices = static_cast<int>(mIndices.size());
+    for (int i0 = 0; i0 < numIndices; ++i0)
+    {
+        int i0p1 = (i0 + 1) % numIndices;
+        seg0.p[0] = mVertexPool[mIndices[i0]];
+        seg0.p[1] = mVertexPool[mIndices[i0p1]];
+
+        int i1min = (i0 + 2) % numIndices;
+        int i1max = (i0 - 2 + numIndices) % numIndices;
+        for (int i1 = i1min; i1 <= i1max; ++i1)
+        {
+            int i1p1 = (i1 + 1) % numIndices;
+            seg1.p[0] = mVertexPool[mIndices[i1]];
+            seg1.p[1] = mVertexPool[mIndices[i1p1]];
+
+            result = query(seg0, seg1);
+            if (result.intersect)
+            {
+                return false;
+            }
+        }
+    }
+    return true;
+}
+
+template <typename Real>
+bool Polygon2<Real>::IsConvexInternal() const
+{
+    Real sign = (mCounterClockwise ? (Real)1 : (Real)-1);
+    int const numIndices = static_cast<int>(mIndices.size());
+    for (int i = 0; i < numIndices; ++i)
+    {
+        int iPrev = (i + numIndices - 1) % numIndices;
+        int iNext = (i + 1) % numIndices;
+        Vector2<Real> vPrev = mVertexPool[mIndices[iPrev]];
+        Vector2<Real> vCurr = mVertexPool[mIndices[i]];
+        Vector2<Real> vNext = mVertexPool[mIndices[iNext]];
+        Vector2<Real> edge0 = vCurr - vPrev;
+        Vector2<Real> edge1 = vNext - vCurr;
+        Real test = sign * DotPerp(edge0, edge1);
+        if (test < (Real)0)
+        {
+            return false;
+        }
+    }
+    return true;
+}
+
+}
diff --git a/Engine/Plugins/Experimental/GeometryProcessing/Source/GeometryAlgorithms/Private/GTEngine/Mathematics/GtePolyhedron3.h b/Engine/Plugins/Experimental/GeometryProcessing/Source/GeometryAlgorithms/Private/GTEngine/Mathematics/GtePolyhedron3.h
new file mode 100644
--- /dev/null
+++ b/Engine/Plugins/Experimental/GeometryProcessing/Source/GeometryAlgorithms/Private/GTEngine/Mathematics/GtePolyhedron3.h
@@ -0,0 +1,200 @@
+// David Eberly, Geometric Tools, Redmond WA 98052
+// Copyright (c) 1998-2019
+// Distributed under the Boost Software License, Version 1.0.
+// http://www.boost.org/LICENSE_1_0.txt
+// http://www.geometrictools.com/License/Boost/LICENSE_1_0.txt
+// File Version: 3.0.1 (2018/10/04)
+
+#pragma once
+
+#include <GTEngine/Mathematics/GteVector3.h>
+#include <algorithm>
+#include <memory>
+#include <set>
+#include <vector>
+
+// The Polyhedron3 object represents a simple polyhedron.  The 'vertexPool'
+// array can contain more points than needed to define the polyhedron, which
+// allows the vertex pool to have multiple polyhedra associated with it.
+// Thus, the programmer must ensure that the vertex pool persists as long as
+// any Polyhedron3 objects exist that depend on the pool.  The number of
+// polyhedron indices is 'numIndices' and must be 6 or larger  The 'indices'
+// array refers to the points in 'vertexPool' that form the triangle faces,
+// so 'numIndices' must be a multiple of 3.  The number of vertices is
+// the number of unique elements in 'indices' and is determined during
+// construction.  The programmer should ensure the polyhedron is simple.  The
+// geometric queries are valid regardless of whether the polyhedron triangles
+// are oriented clockwise or counterclockwise.
+//
+// NOTE:  Comparison operators are not provided.  The semantics of equal
+// polyhedra is complicated and (at the moment) not useful.  The vertex pools
+// can be different and indices do not match, but the vertices they reference
+// can match.  Even with a shared vertex pool, the indices can be permuted,
+// leading to the same polyhedron abstractly but the data structures do not
+// match.
+
+namespace gte
+{
+
+template <typename Real>
+class Polyhedron3
+{
+public:
+    // Construction.  The constructor succeeds when 'numIndices >= 12' (at
+    // least 4 triangles), and 'vertexPool' and 'indices' are not null; we
+    // cannot test whether you have a valid number of elements in the input
+    // arrays.  A copy is made of 'indices', but the 'vertexPool' is not
+    // copied.  If the constructor fails, the internal vertex pointer is set
+    // to null, the number of vertices is set to zero, the index array has no
+    // elements, and the triangle face orientation is set to clockwise.
+    Polyhedron3(std::shared_ptr<std::vector<Vector3<Real>>> const& vertexPool,
+        int numIndices, int const* indices, bool counterClockwise);
+
+    // To validate construction, create an object as shown:
+    //     Polyhedron3<Real> polyhedron(parameters);
+    //     if (!polyhedron) { <constructor failed, handle accordingly>; }
+    inline operator bool() const;
+
+    // Member access.
+    inline std::shared_ptr<std::vector<Vector3<Real>>> const& GetVertexPool() const;
+    inline std::vector<Vector3<Real>> const& GetVertices() const;
+    inline std::set<int> const& GetUniqueIndices() const;
+    inline std::vector<int> const& GetIndices() const;
+    inline bool CounterClockwise() const;
+
+    // Geometric queries.
+    Vector3<Real> ComputeVertexAverage() const;
+    Real ComputeSurfaceArea() const;
+    Real ComputeVolume() const;
+
+private:
+    std::shared_ptr<std::vector<Vector3<Real>>> mVertexPool;
+    std::set<int> mUniqueIndices;
+    std::vector<int> mIndices;
+    bool mCounterClockwise;
+};
+
+
+template <typename Real>
+Polyhedron3<Real>::Polyhedron3(std::shared_ptr<std::vector<Vector3<Real>>> const& vertexPool,
+    int numIndices, int const* indices, bool counterClockwise)
+    :
+    mVertexPool(vertexPool),
+    mCounterClockwise(counterClockwise)
+{
+    if (vertexPool && indices && numIndices >= 12 && (numIndices % 3) == 0)
+    {
+        for (int i = 0; i < numIndices; ++i)
+        {
+            mUniqueIndices.insert(indices[i]);
+        }
+
+        mIndices.resize(numIndices);
+        std::copy(indices, indices + numIndices, mIndices.begin());
+    }
+    else
+    {
+        // Encountered an invalid input.
+        mVertexPool = nullptr;
+        mCounterClockwise = false;
+    }
+}
+
+template <typename Real> inline
+Polyhedron3<Real>::operator bool() const
+{
+    return mVertexPool != nullptr;
+}
+
+template <typename Real> inline
+std::shared_ptr<std::vector<Vector3<Real>>> const& Polyhedron3<Real>::GetVertexPool() const
+{
+    return mVertexPool;
+}
+
+template <typename Real> inline
+std::vector<Vector3<Real>> const& Polyhedron3<Real>::GetVertices() const
+{
+    return *mVertexPool.get();
+}
+
+template <typename Real> inline
+std::set<int> const& Polyhedron3<Real>::GetUniqueIndices() const
+{
+    return mUniqueIndices;
+}
+
+template <typename Real> inline
+std::vector<int> const& Polyhedron3<Real>::GetIndices() const
+{
+    return mIndices;
+}
+
+template <typename Real> inline
+bool Polyhedron3<Real>::CounterClockwise() const
+{
+    return mCounterClockwise;
+}
+
+template <typename Real>
+Vector3<Real> Polyhedron3<Real>::ComputeVertexAverage() const
+{
+    Vector3<Real> average = Vector3<Real>::Zero();
+    if (mVertexPool)
+    {
+        auto vertexPool = GetVertices();
+        for (int index : mUniqueIndices)
+        {
+            average += vertexPool[index];
+        }
+        average /= static_cast<Real>(mUniqueIndices.size());
+    }
+    return average;
+}
+
+template <typename Real>
+Real Polyhedron3<Real>::ComputeSurfaceArea() const
+{
+    Real surfaceArea = (Real)0;
+    if (mVertexPool)
+    {
+        auto vertexPool = GetVertices();
+        int const numTriangles = static_cast<int>(mIndices.size()) / 3;
+        int const* indices = mIndices.data();
+        for (int t = 0; t < numTriangles; ++t)
+        {
+            int v0 = *indices++;
+            int v1 = *indices++;
+            int v2 = *indices++;
+            Vector3<Real> edge0 = vertexPool[v1] - vertexPool[v0];
+            Vector3<Real> edge1 = vertexPool[v2] - vertexPool[v0];
+            Vector3<Real> cross = Cross(edge0, edge1);
+            surfaceArea += Length(cross);
+        }
+        surfaceArea *= (Real)0.5;
+    }
+    return surfaceArea;
+}
+
+template <typename Real>
+Real Polyhedron3<Real>::ComputeVolume() const
+{
+    Real volume = (Real)0;
+    if (mVertexPool)
+    {
+        auto vertexPool = GetVertices();
+        int const numTriangles = static_cast<int>(mIndices.size()) / 3;
+        int const* indices = mIndices.data();
+        for (int t = 0; t < numTriangles; ++t)
+        {
+            int v0 = *indices++;
+            int v1 = *indices++;
+            int v2 = *indices++;
+            volume += DotCross(vertexPool[v0], vertexPool[v1], vertexPool[v2]);
+        }
+        volume /= (Real)6;
+    }
+    return std::abs(volume);
+}
+
+}
diff --git a/Engine/Plugins/Experimental/GeometryProcessing/Source/GeometryAlgorithms/Private/GTEngine/Mathematics/GtePolynomial1.h b/Engine/Plugins/Experimental/GeometryProcessing/Source/GeometryAlgorithms/Private/GTEngine/Mathematics/GtePolynomial1.h
new file mode 100644
--- /dev/null
+++ b/Engine/Plugins/Experimental/GeometryProcessing/Source/GeometryAlgorithms/Private/GTEngine/Mathematics/GtePolynomial1.h
@@ -0,0 +1,607 @@
+// David Eberly, Geometric Tools, Redmond WA 98052
+// Copyright (c) 1998-2019
+// Distributed under the Boost Software License, Version 1.0.
+// http://www.boost.org/LICENSE_1_0.txt
+// http://www.geometrictools.com/License/Boost/LICENSE_1_0.txt
+// File Version: 3.0.3 (2018/12/06)
+
+#pragma once
+
+#include <GTEngine/LowLevel/GteLogger.h>
+#include <algorithm>
+#include <initializer_list>
+#include <vector>
+
+namespace gte
+{
+    template <typename Real>
+    class Polynomial1
+    {
+    public:
+        // Construction and destruction.  The first constructor creates a
+        // polynomial of degree 0.  The first and second constructors do not
+        // initialize their coefficients.  In the third constructor, the
+        // degree is the number of initializers plus 1, but then adjusted to
+        // ensure coefficient[degree] is not zero.
+        Polynomial1()
+            :
+            mCoefficient(1)
+        {
+            // uninitialized
+        }
+
+        Polynomial1(unsigned int degree)
+            :
+            mCoefficient(degree + 1)
+        {
+            // uninitialized
+        }
+
+        Polynomial1(std::initializer_list<Real> values)
+        {
+            // C++ 11 will call the default constructor for
+            // Polynomial1<Real> p{}, so it is guaranteed that
+            // values.size() > 0.
+            mCoefficient.resize(values.size());
+            std::copy(values.begin(), values.end(), mCoefficient.begin());
+            EliminateLeadingZeros();
+        }
+
+        // Support for partial construction, where the default constructor is
+        // used when the degree is not yet known.  The coefficients are
+        // uninitialized.
+        void SetDegree(unsigned int degree)
+        {
+            mCoefficient.resize(degree + 1);
+        }
+
+        // Set all coefficients to the specified value.
+        void SetCoefficients(Real value)
+        {
+            std::fill(mCoefficient.begin(), mCoefficient.end(), value);
+        }
+
+        // Member access.
+        inline unsigned int GetDegree() const
+        {
+            // By design, mCoefficient.size() > 0.
+            return static_cast<unsigned int>(mCoefficient.size() - 1);
+        }
+
+        inline Real const& operator[](unsigned int i) const
+        {
+            return mCoefficient[i];
+        }
+
+        inline Real& operator[](unsigned int i)
+        {
+            return mCoefficient[i];
+        }
+
+
+        // Comparisons.
+        inline bool operator==(Polynomial1<Real> const& p) const
+        {
+            return mCoefficient == p.mCoefficient;
+        }
+
+        inline bool operator!=(Polynomial1<Real> const& p) const
+        {
+            return mCoefficient != p.mCoefficient;
+        }
+
+        inline bool operator< (Polynomial1<Real> const& p) const
+        {
+            return mCoefficient < p.mCoefficient;
+        }
+
+        inline bool operator<=(Polynomial1<Real> const& p) const
+        {
+            return mCoefficient <= p.mCoefficient;
+        }
+
+        inline bool operator> (Polynomial1<Real> const& p) const
+        {
+            return mCoefficient > p.mCoefficient;
+        }
+
+        inline bool operator>=(Polynomial1<Real> const& p) const
+        {
+            return mCoefficient >= p.mCoefficient;
+        }
+
+        // Evaluate the polynomial.  If the polynomial is invalid, the
+        // function returns zero.
+        Real operator()(Real t) const
+        {
+            int i = static_cast<int>(mCoefficient.size());
+            Real result = mCoefficient[--i];
+            for (--i; i >= 0; --i)
+            {
+                result *= t;
+                result += mCoefficient[i];
+            }
+            return result;
+        }
+
+        // Compute the derivative of the polynomial.
+        Polynomial1 GetDerivative() const
+        {
+            unsigned int const degree = GetDegree();
+            if (degree > 0)
+            {
+                Polynomial1 result(degree - 1);
+                for (unsigned int i0 = 0, i1 = 1; i0 < degree; ++i0, ++i1)
+                {
+                    result.mCoefficient[i0] = mCoefficient[i1] * (Real)i1;
+                }
+                return result;
+            }
+            else
+            {
+                Polynomial1 result(0);
+                result[0] = (Real)0;
+                return result;
+            }
+        }
+
+        // Inversion (invpoly[i] = poly[degree-i] for 0 <= i <= degree).
+        Polynomial1 GetInversion() const
+        {
+            unsigned int const degree = GetDegree();
+            Polynomial1 result(degree);
+            for (unsigned int i = 0; i <= degree; ++i)
+            {
+                result.mCoefficient[i] = mCoefficient[degree - i];
+            }
+            return result;
+        }
+
+        // Tranlation.  If 'this' is p(t}, return p(t-t0).
+        Polynomial1 GetTranslation(Real t0) const
+        {
+            Polynomial1<Real> factor{ -t0, (Real)1 };  // f(t) = t - t0
+            unsigned int const degree = GetDegree();
+            Polynomial1 result{ mCoefficient[degree] };
+            for (unsigned int i = 1, j = degree - 1; i <= degree; ++i, --j)
+            {
+                result = mCoefficient[j] + factor * result;
+            }
+            return result;
+        }
+
+        // Eliminate any leading zeros in the polynomial, except in the case
+        // the degree is 0 and the coefficient is 0.  The elimination is
+        // necessary when arithmetic operations cause a decrease in the degree
+        // of the result.  For example, (1 + x + x^2) + (1 + 2*x - x^2) =
+        // (2 + 3*x).  The inputs both have degree 2, so the result is created
+        // with degree 2.  After the addition we find that the degree is in
+        // fact 1 and resize the array of coefficients.  This function is
+        // called internally by the arithmetic operators, but it is exposed in
+        // the public interface in case you need it for your own purposes.
+        void EliminateLeadingZeros()
+        {
+            size_t size = mCoefficient.size();
+            if (size > 1)
+            {
+                Real const zero = (Real)0;
+                int leading;
+                for (leading = static_cast<int>(size) - 1; leading > 0; --leading)
+                {
+                    if (mCoefficient[leading] != zero)
+                    {
+                        break;
+                    }
+                }
+
+                mCoefficient.resize(++leading);
+            }
+        }
+
+        // If 'this' is P(t) and the divisor is D(t) with
+        // degree(P) >= degree(D), then P(t) = Q(t)*D(t)+R(t) where Q(t) is
+        // the quotient with degree(Q) = degree(P) - degree(D) and R(t) is the
+        // remainder with degree(R) < degree(D).  If this routine is called
+        // with degree(P) < degree(D), then Q = 0 and R = P are returned.
+        void Divide(Polynomial1 const& divisor, Polynomial1& quotient, Polynomial1& remainder) const
+        {
+            Real const zero = (Real)0;
+            int divisorDegree = static_cast<int>(divisor.GetDegree());
+            int quotientDegree = static_cast<int>(GetDegree()) - divisorDegree;
+            if (quotientDegree >= 0)
+            {
+                quotient.SetDegree(quotientDegree);
+
+                // Temporary storage for the remainder.
+                Polynomial1 tmp = *this;
+
+                // Do the division using the Euclidean algorithm.
+                Real inv = ((Real)1) / divisor[divisorDegree];
+                for (int i = quotientDegree; i >= 0; --i)
+                {
+                    int j = divisorDegree + i;
+                    quotient[i] = inv * tmp[j];
+                    for (j--; j >= i; j--)
+                    {
+                        tmp[j] -= quotient[i] * divisor[j - i];
+                    }
+                }
+
+                // Calculate the correct degree for the remainder.
+                if (divisorDegree >= 1)
+                {
+                    int remainderDegree = divisorDegree - 1;
+                    while (remainderDegree > 0 && tmp[remainderDegree] == zero)
+                    {
+                        --remainderDegree;
+                    }
+
+                    remainder.SetDegree(remainderDegree);
+                    for (int i = 0; i <= remainderDegree; ++i)
+                    {
+                        remainder[i] = tmp[i];
+                    }
+                }
+                else
+                {
+                    remainder.SetDegree(0);
+                    remainder[0] = zero;
+                }
+            }
+            else
+            {
+                quotient.SetDegree(0);
+                quotient[0] = zero;
+                remainder = *this;
+            }
+        }
+
+        // Scale the polynomial so the highest-degree term has coefficient 1.
+        void MakeMonic()
+        {
+            EliminateLeadingZeros();
+            Real const one(1);
+            if (mCoefficient.back() != one)
+            {
+                unsigned int degree = GetDegree();
+                Real invLeading = one / mCoefficient.back();
+                mCoefficient.back() = one;
+                for (unsigned int i = 0; i < degree; ++i)
+                {
+                    mCoefficient[i] *= invLeading;
+                }
+            }
+        }
+
+    protected:
+        // The class is designed so that mCoefficient.size() >= 1.
+        std::vector<Real> mCoefficient;
+    };
+
+    // Compute the greatest common divisor of two polynomials.  The returned
+    // polynomial has leading coefficient 1 (except when zero-valued
+    // polynomials are passed to the function.
+    template <typename Real>
+    Polynomial1<Real> GreatestCommonDivisor(Polynomial1<Real> const& p0, Polynomial1<Real> const& p1)
+    {
+        // The numerator should be the polynomial of larger degree.
+        Polynomial1<Real> a, b;
+        if (p0.GetDegree() >= p1.GetDegree())
+        {
+            a = p0;
+            b = p1;
+        }
+        else
+        {
+            a = p1;
+            b = p0;
+        }
+
+        Polynomial1<Real> const zero{ (Real)0 };
+        if (a == zero || b == zero)
+        {
+            return (a != zero ? a : zero);
+        }
+
+        // Make the polynomials monic to keep the coefficients reasonable size
+        // when computing with floating-point Real.
+        a.MakeMonic();
+        b.MakeMonic();
+
+        Polynomial1<Real> q, r;
+        for (;;)
+        {
+            a.Divide(b, q, r);
+            if (r != zero)
+            {
+                // a = q * b + r, so gcd(a,b) = gcd(b, r)
+                a = b;
+                b = r;
+                b.MakeMonic();
+            }
+            else
+            {
+                b.MakeMonic();
+                break;
+            }
+        }
+
+        return b;
+    }
+
+    // Factor f = factor[0]*factor[1]^2*factor[2]^3*...*factor[n-1]^n
+    // according to the square-free factorization algorithm
+    // https://en.wikipedia.org/wiki/Square-free_polynomial
+    template <typename Real>
+    void SquareFreeFactorization(Polynomial1<Real> const& f, std::vector<Polynomial1<Real>>& factors)
+    {
+        // In the call to Divide(...), we know that the divisor exactly
+        // divides the numerator, so r = 0 after all such calls.
+        Polynomial1<Real> fder = f.GetDerivative();
+        Polynomial1<Real> a, b, c, d, q, r;
+
+        a = GreatestCommonDivisor(f, fder);
+        f.Divide(a, b, r);  // b = f / a
+        fder.Divide(a, c, r);  // c = fder / a
+        d = c - b.GetDerivative();
+
+        do
+        {
+            a = GreatestCommonDivisor(b, d);
+            factors.emplace_back(a);
+            b.Divide(a, q, r);  // q = b / a
+            b = std::move(q);
+            d.Divide(a, c, r);  // c = d / a
+            d = c - b.GetDerivative();
+        } while (b.GetDegree() > 0);
+    }
+
+    // Unary operations.
+    template <typename Real>
+    Polynomial1<Real> operator+(Polynomial1<Real> const& p)
+    {
+        return p;
+    }
+
+    template <typename Real>
+    Polynomial1<Real> operator-(Polynomial1<Real> const& p)
+    {
+        unsigned int const degree = p.GetDegree();
+        Polynomial1<Real> result(degree);
+        for (unsigned int i = 0; i <= degree; ++i)
+        {
+            result[i] = -p[i];
+        }
+        return result;
+    }
+
+    // Linear-algebraic operations.
+    template <typename Real>
+    Polynomial1<Real> operator+(Polynomial1<Real> const& p0, Polynomial1<Real> const& p1)
+    {
+        unsigned int const p0Degree = p0.GetDegree(), p1Degree = p1.GetDegree();
+        unsigned int i;
+        if (p0Degree >= p1Degree)
+        {
+            Polynomial1<Real> result(p0Degree);
+            for (i = 0; i <= p1Degree; ++i)
+            {
+                result[i] = p0[i] + p1[i];
+            }
+            for (/**/; i <= p0Degree; ++i)
+            {
+                result[i] = p0[i];
+            }
+            result.EliminateLeadingZeros();
+            return result;
+        }
+        else
+        {
+            Polynomial1<Real> result(p1Degree);
+            for (i = 0; i <= p0Degree; ++i)
+            {
+                result[i] = p0[i] + p1[i];
+            }
+            for (/**/; i <= p1Degree; ++i)
+            {
+                result[i] = p1[i];
+            }
+            result.EliminateLeadingZeros();
+            return result;
+        }
+    }
+
+    template <typename Real>
+    Polynomial1<Real> operator-(Polynomial1<Real> const& p0, Polynomial1<Real> const& p1)
+    {
+        unsigned int const p0Degree = p0.GetDegree(), p1Degree = p1.GetDegree();
+        unsigned int i;
+        if (p0Degree >= p1Degree)
+        {
+            Polynomial1<Real> result(p0Degree);
+            for (i = 0; i <= p1Degree; ++i)
+            {
+                result[i] = p0[i] - p1[i];
+            }
+            for (/**/; i <= p0Degree; ++i)
+            {
+                result[i] = p0[i];
+            }
+            result.EliminateLeadingZeros();
+            return result;
+        }
+        else
+        {
+            Polynomial1<Real> result(p1Degree);
+            for (i = 0; i <= p0Degree; ++i)
+            {
+                result[i] = p0[i] - p1[i];
+            }
+            for (/**/; i <= p1Degree; ++i)
+            {
+                result[i] = -p1[i];
+            }
+            result.EliminateLeadingZeros();
+            return result;
+        }
+    }
+
+    template <typename Real>
+    Polynomial1<Real> operator*(Polynomial1<Real> const& p0, Polynomial1<Real> const& p1)
+    {
+        unsigned int const p0Degree = p0.GetDegree(), p1Degree = p1.GetDegree();
+        Polynomial1<Real> result(p0Degree + p1Degree);
+        result.SetCoefficients((Real)0);
+        for (unsigned int i0 = 0; i0 <= p0Degree; ++i0)
+        {
+            for (unsigned int i1 = 0; i1 <= p1Degree; ++i1)
+            {
+                result[i0 + i1] += p0[i0] * p1[i1];
+            }
+        }
+        return result;
+    }
+
+    template <typename Real>
+    Polynomial1<Real> operator+(Polynomial1<Real> const& p, Real scalar)
+    {
+        unsigned int const degree = p.GetDegree();
+        Polynomial1<Real> result(degree);
+        result[0] = p[0] + scalar;
+        for (unsigned int i = 1; i <= degree; ++i)
+        {
+            result[i] = p[i];
+        }
+        return result;
+    }
+
+    template <typename Real>
+    Polynomial1<Real> operator+(Real scalar, Polynomial1<Real> const& p)
+    {
+        unsigned int const degree = p.GetDegree();
+        Polynomial1<Real> result(degree);
+        result[0] = p[0] + scalar;
+        for (unsigned int i = 1; i <= degree; ++i)
+        {
+            result[i] = p[i];
+        }
+        return result;
+    }
+
+    template <typename Real>
+    Polynomial1<Real> operator-(Polynomial1<Real> const& p, Real scalar)
+    {
+        unsigned int const degree = p.GetDegree();
+        Polynomial1<Real> result(degree);
+        result[0] = p[0] - scalar;
+        for (unsigned int i = 1; i <= degree; ++i)
+        {
+            result[i] = p[i];
+        }
+        return result;
+    }
+
+    template <typename Real>
+    Polynomial1<Real> operator-(Real scalar, Polynomial1<Real> const& p)
+    {
+        unsigned int const degree = p.GetDegree();
+        Polynomial1<Real> result(degree);
+        result[0] = scalar - p[0];
+        for (unsigned int i = 1; i <= degree; ++i)
+        {
+            result[i] = -p[i];
+        }
+        return result;
+    }
+
+    template <typename Real>
+    Polynomial1<Real> operator*(Polynomial1<Real> const& p, Real scalar)
+    {
+        unsigned int const degree = p.GetDegree();
+        Polynomial1<Real> result(degree);
+        for (unsigned int i = 0; i <= degree; ++i)
+        {
+            result[i] = scalar * p[i];
+        }
+        return result;
+    }
+
+    template <typename Real>
+    Polynomial1<Real> operator*(Real scalar, Polynomial1<Real> const& p)
+    {
+        unsigned int const degree = p.GetDegree();
+        Polynomial1<Real> result(degree);
+        for (unsigned int i = 0; i <= degree; ++i)
+        {
+            result[i] = scalar * p[i];
+        }
+        return result;
+    }
+
+    template <typename Real>
+    Polynomial1<Real> operator/(Polynomial1<Real> const& p, Real scalar)
+    {
+        if (scalar == (Real)0)
+        {
+            LogWarning("Division by zero in operator/(p,scalar).");
+        }
+
+        unsigned int const degree = p.GetDegree();
+        Real invScalar = ((Real)1) / scalar;
+        Polynomial1<Real> result(degree);
+        for (unsigned int i = 0; i <= degree; ++i)
+        {
+            result[i] = invScalar * p[i];
+        }
+        return result;
+    }
+
+    template <typename Real>
+    Polynomial1<Real>& operator+=(Polynomial1<Real>& p0, Polynomial1<Real> const& p1)
+    {
+        p0 = p0 + p1;
+        return p0;
+    }
+
+    template <typename Real>
+    Polynomial1<Real>& operator-=(Polynomial1<Real>& p0, Polynomial1<Real> const& p1)
+    {
+        p0 = p0 - p1;
+        return p0;
+    }
+
+    template <typename Real>
+    Polynomial1<Real>& operator*=(Polynomial1<Real>& p0, Polynomial1<Real> const& p1)
+    {
+        p0 = p0 * p1;
+        return p0;
+    }
+
+    template <typename Real>
+    Polynomial1<Real>& operator+=(Polynomial1<Real>& p, Real scalar)
+    {
+        p[0] += scalar;
+        return p;
+    }
+
+    template <typename Real>
+    Polynomial1<Real>& operator-=(Polynomial1<Real>& p, Real scalar)
+    {
+        p[0] -= scalar;
+        return p;
+    }
+
+    template <typename Real>
+    Polynomial1<Real>& operator*=(Polynomial1<Real>& p, Real scalar)
+    {
+        p = p * scalar;
+        return p;
+    }
+
+    template <typename Real>
+    Polynomial1<Real> & operator/=(Polynomial1<Real>& p, Real scalar)
+    {
+        p = p / scalar;
+        return p;
+    }
+}
diff --git a/Engine/Plugins/Experimental/GeometryProcessing/Source/GeometryAlgorithms/Private/GTEngine/Mathematics/GtePrimalQuery2.h b/Engine/Plugins/Experimental/GeometryProcessing/Source/GeometryAlgorithms/Private/GTEngine/Mathematics/GtePrimalQuery2.h
new file mode 100644
--- /dev/null
+++ b/Engine/Plugins/Experimental/GeometryProcessing/Source/GeometryAlgorithms/Private/GTEngine/Mathematics/GtePrimalQuery2.h
@@ -0,0 +1,437 @@
+// David Eberly, Geometric Tools, Redmond WA 98052
+// Copyright (c) 1998-2019
+// Distributed under the Boost Software License, Version 1.0.
+// http://www.boost.org/LICENSE_1_0.txt
+// http://www.geometrictools.com/License/Boost/LICENSE_1_0.txt
+// File Version: 3.0.0 (2016/06/19)
+
+#pragma once
+
+#include <GTEngine/Mathematics/GteVector2.h>
+
+// Queries about the relation of a point to various geometric objects.  The
+// choices for N when using UIntegerFP32<N> for either BSNumber of BSRational
+// are determined in GeometricTools/GTEngine/Tools/PrecisionCalculator.  These
+// N-values are worst case scenarios. Your specific input data might require
+// much smaller N, in which case you can modify PrecisionCalculator to use the
+// BSPrecision(int32_t,int32_t,int32_t,bool) constructors.
+
+namespace gte
+{
+
+template <typename Real>
+class PrimalQuery2
+{
+public:
+    // The caller is responsible for ensuring that the array is not empty
+    // before calling queries and that the indices passed to the queries are
+    // valid.  The class does no range checking.
+    PrimalQuery2();
+    PrimalQuery2(int numVertices, Vector2<Real> const* vertices);
+
+    // Member access.
+    inline void Set(int numVertices, Vector2<Real> const* vertices);
+    inline int GetNumVertices() const;
+    inline Vector2<Real> const* GetVertices() const;
+
+    // In the following, point P refers to vertices[i] or 'test' and Vi refers
+    // to vertices[vi].
+
+    // For a line with origin V0 and direction <V0,V1>, ToLine returns
+    //   +1, P on right of line
+    //   -1, P on left of line
+    //    0, P on the line
+    //
+    // Choice of N for UIntegerFP32<N>.
+    //    input type | compute type | N
+    //    -----------+--------------+-----
+    //    float      | BSNumber     |   18
+    //    double     | BSNumber     |  132
+    //    float      | BSRational   |  214
+    //    double     | BSRational   | 1587
+    int ToLine(int i, int v0, int v1) const;
+    int ToLine(Vector2<Real> const& test, int v0, int v1) const;
+
+    // For a line with origin V0 and direction <V0,V1>, ToLine returns
+    //   +1, P on right of line
+    //   -1, P on left of line
+    //    0, P on the line
+    // The 'order' parameter is
+    //   -3, points not collinear, P on left of line
+    //   -2, P strictly left of V0 on the line
+    //   -1, P = V0
+    //    0, P interior to line segment [V0,V1]
+    //   +1, P = V1
+    //   +2, P strictly right of V0 on the line
+    //
+    // Choice of N for UIntegerFP32<N>.
+    //    input type | compute type | N
+    //    -----------+--------------+-----
+    //    float      | BSNumber     |   18
+    //    double     | BSNumber     |  132
+    //    float      | BSRational   |  214
+    //    double     | BSRational   | 1587
+    // This is the same as the first-listed ToLine calls because the worst-case
+    // path has the same computational complexity.
+    int ToLine(int i, int v0, int v1, int& order) const;
+    int ToLine(Vector2<Real> const& test, int v0, int v1, int& order) const;
+
+    // For a triangle with counterclockwise vertices V0, V1, and V2,
+    // ToTriangle returns
+    //   +1, P outside triangle
+    //   -1, P inside triangle
+    //    0, P on triangle
+    //
+    // Choice of N for UIntegerFP32<N>.
+    //    input type | compute type | N
+    //    -----------+--------------+-----
+    //    float      | BSNumber     |   18
+    //    double     | BSNumber     |  132
+    //    float      | BSRational   |  214
+    //    double     | BSRational   | 1587
+    // The query involves three calls to ToLine, so the numbers match those
+    // of ToLine.
+    int ToTriangle(int i, int v0, int v1, int v2) const;
+    int ToTriangle(Vector2<Real> const& test, int v0, int v1, int v2) const;
+
+    // For a triangle with counterclockwise vertices V0, V1, and V2,
+    // ToCircumcircle returns
+    //   +1, P outside circumcircle of triangle
+    //   -1, P inside circumcircle of triangle
+    //    0, P on circumcircle of triangle
+    //
+    // Choice of N for UIntegerFP32<N>.
+    //    input type | compute type | N
+    //    -----------+--------------+------
+    //    float      | BSNumber     |    35
+    //    double     | BSNumber     |   263
+    //    float      | BSRational   | 12302
+    //    double     | BSRational   | 92827
+    // The query involves three calls of ToLine, so the numbers match those
+    // of ToLine.
+    int ToCircumcircle(int i, int v0, int v1, int v2) const;
+    int ToCircumcircle(Vector2<Real> const& test, int v0, int v1, int v2) const;
+
+    // An extended classification of the relationship of a point to a line
+    // segment.  For noncollinear points, the return value is
+    //   ORDER_POSITIVE when <P,Q0,Q1> is a counterclockwise triangle
+    //   ORDER_NEGATIVE when <P,Q0,Q1> is a clockwise triangle
+    // For collinear points, the line direction is Q1-Q0.  The return value is
+    //   ORDER_COLLINEAR_LEFT when the line ordering is <P,Q0,Q1>
+    //   ORDER_COLLINEAR_RIGHT when the line ordering is <Q0,Q1,P>
+    //   ORDER_COLLINEAR_CONTAIN when the line ordering is <Q0,P,Q1>
+    enum OrderType
+    {
+        ORDER_Q0_EQUALS_Q1,
+        ORDER_P_EQUALS_Q0,
+        ORDER_P_EQUALS_Q1,
+        ORDER_POSITIVE,
+        ORDER_NEGATIVE,
+        ORDER_COLLINEAR_LEFT,
+        ORDER_COLLINEAR_RIGHT,
+        ORDER_COLLINEAR_CONTAIN
+    };
+
+    // Choice of N for UIntegerFP32<N>.
+    //    input type | compute type | N
+    //    -----------+--------------+-----
+    //    float      | BSNumber     |   18
+    //    double     | BSNumber     |  132
+    //    float      | BSRational   |  214
+    //    double     | BSRational   | 1587
+    // This is the same as the first-listed ToLine calls because the worst-case
+    // path has the same computational complexity.
+    OrderType ToLineExtended(Vector2<Real> const& P, Vector2<Real> const& Q0, Vector2<Real> const& Q1) const;
+
+private:
+    int mNumVertices;
+    Vector2<Real> const* mVertices;
+};
+
+
+template <typename Real>
+PrimalQuery2<Real>::PrimalQuery2()
+    :
+    mNumVertices(0),
+    mVertices(nullptr)
+{
+}
+
+template <typename Real>
+PrimalQuery2<Real>::PrimalQuery2(int numVertices,
+    Vector2<Real> const* vertices)
+    :
+    mNumVertices(numVertices),
+    mVertices(vertices)
+{
+}
+
+template <typename Real> inline
+void PrimalQuery2<Real>::Set(int numVertices, Vector2<Real> const* vertices)
+{
+    mNumVertices = numVertices;
+    mVertices = vertices;
+}
+
+template <typename Real> inline
+int PrimalQuery2<Real>::GetNumVertices() const
+{
+    return mNumVertices;
+}
+
+template <typename Real> inline
+Vector2<Real> const* PrimalQuery2<Real>::GetVertices() const
+{
+    return mVertices;
+}
+
+template <typename Real>
+int PrimalQuery2<Real>::ToLine(int i, int v0, int v1) const
+{
+    return ToLine(mVertices[i], v0, v1);
+}
+
+template <typename Real>
+int PrimalQuery2<Real>::ToLine(Vector2<Real> const& test, int v0, int v1) const
+{
+    Vector2<Real> const& vec0 = mVertices[v0];
+    Vector2<Real> const& vec1 = mVertices[v1];
+
+    Real x0 = test[0] - vec0[0];
+    Real y0 = test[1] - vec0[1];
+    Real x1 = vec1[0] - vec0[0];
+    Real y1 = vec1[1] - vec0[1];
+    Real x0y1 = x0*y1;
+    Real x1y0 = x1*y0;
+    Real det = x0y1 - x1y0;
+    Real const zero(0);
+
+    return (det > zero ? +1 : (det < zero ? -1 : 0));
+}
+
+template <typename Real>
+int PrimalQuery2<Real>::ToLine(int i, int v0, int v1, int& order) const
+{
+    return ToLine(mVertices[i], v0, v1, order);
+}
+
+template <typename Real>
+int PrimalQuery2<Real>::ToLine(Vector2<Real> const& test, int v0, int v1, int& order) const
+{
+    Vector2<Real> const& vec0 = mVertices[v0];
+    Vector2<Real> const& vec1 = mVertices[v1];
+
+    Real x0 = test[0] - vec0[0];
+    Real y0 = test[1] - vec0[1];
+    Real x1 = vec1[0] - vec0[0];
+    Real y1 = vec1[1] - vec0[1];
+    Real x0y1 = x0*y1;
+    Real x1y0 = x1*y0;
+    Real det = x0y1 - x1y0;
+    Real const zero(0);
+
+    if (det > zero)
+    {
+        order = +3;
+        return +1;
+    }
+
+    if (det < zero)
+    {
+        order = -3;
+        return -1;
+    }
+
+    Real x0x1 = x0*x1;
+    Real y0y1 = y0*y1;
+    Real dot = x0x1 + y0y1;
+    if (dot == zero)
+    {
+        order = -1;
+    }
+    else if (dot < zero)
+    {
+        order = -2;
+    }
+    else
+    {
+        Real x0x0 = x0*x0;
+        Real y0y0 = y0*y0;
+        Real sqrLength = x0x0 + y0y0;
+        if (dot == sqrLength)
+        {
+            order = +1;
+        }
+        else if (dot > sqrLength)
+        {
+            order = +2;
+        }
+        else
+        {
+            order = 0;
+        }
+    }
+
+    return 0;
+}
+
+template <typename Real>
+int PrimalQuery2<Real>::ToTriangle(int i, int v0, int v1, int v2) const
+{
+    return ToTriangle(mVertices[i], v0, v1, v2);
+}
+
+template <typename Real>
+int PrimalQuery2<Real>::ToTriangle(Vector2<Real> const& test, int v0, int v1, int v2) const
+{
+    int sign0 = ToLine(test, v1, v2);
+    if (sign0 > 0)
+    {
+        return +1;
+    }
+
+    int sign1 = ToLine(test, v0, v2);
+    if (sign1 < 0)
+    {
+        return +1;
+    }
+
+    int sign2 = ToLine(test, v0, v1);
+    if (sign2 > 0)
+    {
+        return +1;
+    }
+
+    return ((sign0 && sign1 && sign2) ? -1 : 0);
+}
+
+template <typename Real>
+int PrimalQuery2<Real>::ToCircumcircle(int i, int v0, int v1, int v2) const
+{
+    return ToCircumcircle(mVertices[i], v0, v1, v2);
+}
+
+template <typename Real>
+int PrimalQuery2<Real>::ToCircumcircle(Vector2<Real> const& test, int v0, int v1, int v2) const
+{
+    Vector2<Real> const& vec0 = mVertices[v0];
+    Vector2<Real> const& vec1 = mVertices[v1];
+    Vector2<Real> const& vec2 = mVertices[v2];
+
+    Real x0 = vec0[0] - test[0];
+    Real y0 = vec0[1] - test[1];
+    Real s00 = vec0[0] + test[0];
+    Real s01 = vec0[1] + test[1];
+    Real t00 = s00*x0;
+    Real t01 = s01*y0;
+    Real z0 = t00 + t01;
+
+    Real x1 = vec1[0] - test[0];
+    Real y1 = vec1[1] - test[1];
+    Real s10 = vec1[0] + test[0];
+    Real s11 = vec1[1] + test[1];
+    Real t10 = s10*x1;
+    Real t11 = s11*y1;
+    Real z1 = t10 + t11;
+
+    Real x2 = vec2[0] - test[0];
+    Real y2 = vec2[1] - test[1];
+    Real s20 = vec2[0] + test[0];
+    Real s21 = vec2[1] + test[1];
+    Real t20 = s20*x2;
+    Real t21 = s21*y2;
+    Real z2 = t20 + t21;
+
+    Real y0z1 = y0*z1;
+    Real y0z2 = y0*z2;
+    Real y1z0 = y1*z0;
+    Real y1z2 = y1*z2;
+    Real y2z0 = y2*z0;
+    Real y2z1 = y2*z1;
+    Real c0 = y1z2 - y2z1;
+    Real c1 = y2z0 - y0z2;
+    Real c2 = y0z1 - y1z0;
+    Real x0c0 = x0*c0;
+    Real x1c1 = x1*c1;
+    Real x2c2 = x2*c2;
+    Real term = x0c0 + x1c1;
+    Real det = term + x2c2;
+    Real const zero(0);
+
+    return (det < zero ? 1 : (det > zero ? -1 : 0));
+}
+
+template <typename Real>
+typename PrimalQuery2<Real>::OrderType PrimalQuery2<Real>::ToLineExtended(
+    Vector2<Real> const& P, Vector2<Real> const& Q0, Vector2<Real> const& Q1) const
+{
+    Real const zero(0);
+
+    Real x0 = Q1[0] - Q0[0];
+    Real y0 = Q1[1] - Q0[1];
+    if (x0 == zero && y0 == zero)
+    {
+        return ORDER_Q0_EQUALS_Q1;
+    }
+
+    Real x1 = P[0] - Q0[0];
+    Real y1 = P[1] - Q0[1];
+    if (x1 == zero && y1 == zero)
+    {
+        return ORDER_P_EQUALS_Q0;
+    }
+
+    Real x2 = P[0] - Q1[0];
+    Real y2 = P[1] - Q1[1];
+    if (x2 == zero && y2 == zero)
+    {
+        return ORDER_P_EQUALS_Q1;
+    }
+
+    // The theoretical classification relies on computing exactly the sign of
+    // the determinant.  Numerical roundoff errors can cause misclassification.
+    Real x0y1 = x0 * y1;
+    Real x1y0 = x1 * y0;
+    Real det = x0y1 - x1y0;
+
+    if (det != zero)
+    {
+        if (det > zero)
+        {
+            // The points form a counterclockwise triangle <P,Q0,Q1>.
+            return ORDER_POSITIVE;
+        }
+        else
+        {
+            // The points form a clockwise triangle <P,Q1,Q0>.
+            return ORDER_NEGATIVE;
+        }
+    }
+    else
+    {
+        // The points are collinear; P is on the line through Q0 and Q1.
+        Real x0x1 = x0 * x1;
+        Real y0y1 = y0 * y1;
+        Real dot = x0x1 + y0y1;
+        if (dot < zero)
+        {
+            // The line ordering is <P,Q0,Q1>.
+            return ORDER_COLLINEAR_LEFT;
+        }
+
+        Real x0x0 = x0 * x0;
+        Real y0y0 = y0 * y0;
+        Real sqrLength = x0x0 + y0y0;
+        if (dot > sqrLength)
+        {
+            // The line ordering is <Q0,Q1,P>.
+            return ORDER_COLLINEAR_RIGHT;
+        }
+
+        // The line ordering is <Q0,P,Q1> with P strictly between Q0 and Q1.
+        return ORDER_COLLINEAR_CONTAIN;
+    }
+}
+
+
+}
diff --git a/Engine/Plugins/Experimental/GeometryProcessing/Source/GeometryAlgorithms/Private/GTEngine/Mathematics/GtePrimalQuery3.h b/Engine/Plugins/Experimental/GeometryProcessing/Source/GeometryAlgorithms/Private/GTEngine/Mathematics/GtePrimalQuery3.h
new file mode 100644
--- /dev/null
+++ b/Engine/Plugins/Experimental/GeometryProcessing/Source/GeometryAlgorithms/Private/GTEngine/Mathematics/GtePrimalQuery3.h
@@ -0,0 +1,330 @@
+// David Eberly, Geometric Tools, Redmond WA 98052
+// Copyright (c) 1998-2019
+// Distributed under the Boost Software License, Version 1.0.
+// http://www.boost.org/LICENSE_1_0.txt
+// http://www.geometrictools.com/License/Boost/LICENSE_1_0.txt
+// File Version: 3.0.0 (2016/06/19)
+
+#pragma once
+
+#include <GTEngine/Mathematics/GteVector3.h>
+
+// Queries about the relation of a point to various geometric objects.  The
+// choices for N when using UIntegerFP32<N> for either BSNumber of BSRational
+// are determined in GeometricTools/GTEngine/Tools/PrecisionCalculator.  These
+// N-values are worst case scenarios. Your specific input data might require
+// much smaller N, in which case you can modify PrecisionCalculator to use the
+// BSPrecision(int32_t,int32_t,int32_t,bool) constructors.
+
+namespace gte
+{
+
+template <typename Real>
+class PrimalQuery3
+{
+public:
+    // The caller is responsible for ensuring that the array is not empty
+    // before calling queries and that the indices passed to the queries are
+    // valid.  The class does no range checking.
+    PrimalQuery3();
+    PrimalQuery3(int numVertices, Vector3<Real> const* vertices);
+
+    // Member access.
+    inline void Set(int numVertices, Vector3<Real> const* vertices);
+    inline int GetNumVertices() const;
+    inline Vector3<Real> const* GetVertices() const;
+
+    // In the following, point P refers to vertices[i] or 'test' and Vi refers
+    // to vertices[vi].
+
+    // For a plane with origin V0 and normal N = Cross(V1-V0,V2-V0), ToPlane
+    // returns
+    //   +1, P on positive side of plane (side to which N points)
+    //   -1, P on negative side of plane (side to which -N points)
+    //    0, P on the plane
+    //
+    // Choice of N for UIntegerFP32<N>.
+    //    input type | compute type | N
+    //    -----------+--------------+------
+    //    float      | BSNumber     |    27
+    //    double     | BSNumber     |   197
+    //    float      | BSRational   |  2882
+    //    double     | BSRational   | 21688
+    int ToPlane(int i, int v0, int v1, int v2) const;
+    int ToPlane(Vector3<Real> const& test, int v0, int v1, int v2) const;
+
+    // For a tetrahedron with vertices ordered as described in the file
+    // GteTetrahedronKey.h, the function returns
+    //   +1, P outside tetrahedron
+    //   -1, P inside tetrahedron
+    //    0, P on tetrahedron
+    //
+    // Choice of N for UIntegerFP32<N>.
+    //    input type | compute type | N
+    //    -----------+--------------+----
+    //    float      | BSNumber     |    27
+    //    double     | BSNumber     |   197
+    //    float      | BSRational   |  2882
+    //    double     | BSRational   | 21688
+    // The query involves four calls of ToPlane, so the numbers match those
+    // of ToPlane.
+    int ToTetrahedron(int i, int v0, int v1, int v2, int v3) const;
+    int ToTetrahedron(Vector3<Real> const& test, int v0, int v1, int v2, int v3) const;
+
+    // For a tetrahedron with vertices ordered as described in the file
+    // GteTetrahedronKey.h, the function returns
+    //   +1, P outside circumsphere of tetrahedron
+    //   -1, P inside circumsphere of tetrahedron
+    //    0, P on circumsphere of tetrahedron
+    //
+    // Choice of N for UIntegerFP32<N>.
+    //    input type | compute type | N
+    //    -----------+--------------+--------
+    //    float      | BSNumber     |      44
+    //    float      | BSRational   |     329
+    //    double     | BSNumber     |  298037
+    //    double     | BSRational   | 2254442
+    int ToCircumsphere(int i, int v0, int v1, int v2, int v3) const;
+    int ToCircumsphere(Vector3<Real> const& test, int v0, int v1, int v2, int v3) const;
+
+private:
+    int mNumVertices;
+    Vector3<Real> const* mVertices;
+};
+
+
+template <typename Real>
+PrimalQuery3<Real>::PrimalQuery3()
+    :
+    mNumVertices(0),
+    mVertices(nullptr)
+{
+}
+
+template <typename Real>
+PrimalQuery3<Real>::PrimalQuery3(int numVertices,
+    Vector3<Real> const* vertices)
+    :
+    mNumVertices(numVertices),
+    mVertices(vertices)
+{
+}
+
+template <typename Real> inline
+void PrimalQuery3<Real>::Set(int numVertices, Vector3<Real> const* vertices)
+{
+    mNumVertices = numVertices;
+    mVertices = vertices;
+}
+
+template <typename Real> inline
+int PrimalQuery3<Real>::GetNumVertices() const
+{
+    return mNumVertices;
+}
+
+template <typename Real> inline
+Vector3<Real> const* PrimalQuery3<Real>::GetVertices() const
+{
+    return mVertices;
+}
+
+template <typename Real>
+int PrimalQuery3<Real>::ToPlane(int i, int v0, int v1, int v2) const
+{
+    return ToPlane(mVertices[i], v0, v1, v2);
+}
+
+template <typename Real>
+int PrimalQuery3<Real>::ToPlane(Vector3<Real> const& test, int v0, int v1,
+    int v2) const
+{
+    Vector3<Real> const& vec0 = mVertices[v0];
+    Vector3<Real> const& vec1 = mVertices[v1];
+    Vector3<Real> const& vec2 = mVertices[v2];
+
+    Real x0 = test[0] - vec0[0];
+    Real y0 = test[1] - vec0[1];
+    Real z0 = test[2] - vec0[2];
+    Real x1 = vec1[0] - vec0[0];
+    Real y1 = vec1[1] - vec0[1];
+    Real z1 = vec1[2] - vec0[2];
+    Real x2 = vec2[0] - vec0[0];
+    Real y2 = vec2[1] - vec0[1];
+    Real z2 = vec2[2] - vec0[2];
+    Real y1z2 = y1*z2;
+    Real y2z1 = y2*z1;
+    Real y2z0 = y2*z0;
+    Real y0z2 = y0*z2;
+    Real y0z1 = y0*z1;
+    Real y1z0 = y1*z0;
+    Real c0 = y1z2 - y2z1;
+    Real c1 = y2z0 - y0z2;
+    Real c2 = y0z1 - y1z0;
+    Real x0c0 = x0*c0;
+    Real x1c1 = x1*c1;
+    Real x2c2 = x2*c2;
+    Real term = x0c0 + x1c1;
+    Real det = term + x2c2;
+    Real const zero(0);
+
+    return (det > zero ? +1 : (det < zero ? -1 : 0));
+}
+
+template <typename Real>
+int PrimalQuery3<Real>::ToTetrahedron(int i, int v0, int v1, int v2, int v3)
+    const
+{
+    return ToTetrahedron(mVertices[i], v0, v1, v2, v3);
+}
+
+template <typename Real>
+int PrimalQuery3<Real>::ToTetrahedron(Vector3<Real> const& test, int v0,
+    int v1, int v2, int v3) const
+{
+    int sign0 = ToPlane(test, v1, v2, v3);
+    if (sign0 > 0)
+    {
+        return +1;
+    }
+
+    int sign1 = ToPlane(test, v0, v2, v3);
+    if (sign1 < 0)
+    {
+        return +1;
+    }
+
+    int sign2 = ToPlane(test, v0, v1, v3);
+    if (sign2 > 0)
+    {
+        return +1;
+    }
+
+    int sign3 = ToPlane(test, v0, v1, v2);
+    if (sign3 < 0)
+    {
+        return +1;
+    }
+
+    return ((sign0 && sign1 && sign2 && sign3) ? -1 : 0);
+}
+
+template <typename Real>
+int PrimalQuery3<Real>::ToCircumsphere(int i, int v0, int v1, int v2, int v3)
+const
+{
+    return ToCircumsphere(mVertices[i], v0, v1, v2, v3);
+}
+
+template <typename Real>
+int PrimalQuery3<Real>::ToCircumsphere(Vector3<Real> const& test, int v0,
+    int v1, int v2, int v3) const
+{
+    Vector3<Real> const& vec0 = mVertices[v0];
+    Vector3<Real> const& vec1 = mVertices[v1];
+    Vector3<Real> const& vec2 = mVertices[v2];
+    Vector3<Real> const& vec3 = mVertices[v3];
+
+    Real x0 = vec0[0] - test[0];
+    Real y0 = vec0[1] - test[1];
+    Real z0 = vec0[2] - test[2];
+    Real s00 = vec0[0] + test[0];
+    Real s01 = vec0[1] + test[1];
+    Real s02 = vec0[2] + test[2];
+    Real t00 = s00*x0;
+    Real t01 = s01*y0;
+    Real t02 = s02*z0;
+    Real t00pt01 = t00 + t01;
+    Real w0 = t00pt01 + t02;
+
+    Real x1 = vec1[0] - test[0];
+    Real y1 = vec1[1] - test[1];
+    Real z1 = vec1[2] - test[2];
+    Real s10 = vec1[0] + test[0];
+    Real s11 = vec1[1] + test[1];
+    Real s12 = vec1[2] + test[2];
+    Real t10 = s10*x1;
+    Real t11 = s11*y1;
+    Real t12 = s12*z1;
+    Real t10pt11 = t10 + t11;
+    Real w1 = t10pt11 + t12;
+
+    Real x2 = vec2[0] - test[0];
+    Real y2 = vec2[1] - test[1];
+    Real z2 = vec2[2] - test[2];
+    Real s20 = vec2[0] + test[0];
+    Real s21 = vec2[1] + test[1];
+    Real s22 = vec2[2] + test[2];
+    Real t20 = s20*x2;
+    Real t21 = s21*y2;
+    Real t22 = s22*z2;
+    Real t20pt21 = t20 + t21;
+    Real w2 = t20pt21 + t22;
+
+    Real x3 = vec3[0] - test[0];
+    Real y3 = vec3[1] - test[1];
+    Real z3 = vec3[2] - test[2];
+    Real s30 = vec3[0] + test[0];
+    Real s31 = vec3[1] + test[1];
+    Real s32 = vec3[2] + test[2];
+    Real t30 = s30*x3;
+    Real t31 = s31*y3;
+    Real t32 = s32*z3;
+    Real t30pt31 = t30 + t31;
+    Real w3 = t30pt31 + t32;
+
+    Real x0y1 = x0*y1;
+    Real x0y2 = x0*y2;
+    Real x0y3 = x0*y3;
+    Real x1y0 = x1*y0;
+    Real x1y2 = x1*y2;
+    Real x1y3 = x1*y3;
+    Real x2y0 = x2*y0;
+    Real x2y1 = x2*y1;
+    Real x2y3 = x2*y3;
+    Real x3y0 = x3*y0;
+    Real x3y1 = x3*y1;
+    Real x3y2 = x3*y2;
+    Real a0 = x0y1 - x1y0;
+    Real a1 = x0y2 - x2y0;
+    Real a2 = x0y3 - x3y0;
+    Real a3 = x1y2 - x2y1;
+    Real a4 = x1y3 - x3y1;
+    Real a5 = x2y3 - x3y2;
+
+    Real z0w1 = z0*w1;
+    Real z0w2 = z0*w2;
+    Real z0w3 = z0*w3;
+    Real z1w0 = z1*w0;
+    Real z1w2 = z1*w2;
+    Real z1w3 = z1*w3;
+    Real z2w0 = z2*w0;
+    Real z2w1 = z2*w1;
+    Real z2w3 = z2*w3;
+    Real z3w0 = z3*w0;
+    Real z3w1 = z3*w1;
+    Real z3w2 = z3*w2;
+    Real b0 = z0w1 - z1w0;
+    Real b1 = z0w2 - z2w0;
+    Real b2 = z0w3 - z3w0;
+    Real b3 = z1w2 - z2w1;
+    Real b4 = z1w3 - z3w1;
+    Real b5 = z2w3 - z3w2;
+    Real a0b5 = a0*b5;
+    Real a1b4 = a1*b4;
+    Real a2b3 = a2*b3;
+    Real a3b2 = a3*b2;
+    Real a4b1 = a4*b1;
+    Real a5b0 = a5*b0;
+    Real term0 = a0b5 - a1b4;
+    Real term1 = term0 + a2b3;
+    Real term2 = term1 + a3b2;
+    Real term3 = term2 - a4b1;
+    Real det = term3 + a5b0;
+    Real const zero(0);
+
+    return (det > zero ? 1 : (det < zero ? -1 : 0));
+}
+
+
+}
diff --git a/Engine/Plugins/Experimental/GeometryProcessing/Source/GeometryAlgorithms/Private/GTEngine/Mathematics/GteProjection.h b/Engine/Plugins/Experimental/GeometryProcessing/Source/GeometryAlgorithms/Private/GTEngine/Mathematics/GteProjection.h
new file mode 100644
--- /dev/null
+++ b/Engine/Plugins/Experimental/GeometryProcessing/Source/GeometryAlgorithms/Private/GTEngine/Mathematics/GteProjection.h
@@ -0,0 +1,71 @@
+// David Eberly, Geometric Tools, Redmond WA 98052
+// Copyright (c) 1998-2019
+// Distributed under the Boost Software License, Version 1.0.
+// http://www.boost.org/LICENSE_1_0.txt
+// http://www.geometrictools.com/License/Boost/LICENSE_1_0.txt
+// File Version: 3.0.1 (2018/10/05)
+
+#pragma once
+
+#include <GTEngine/Mathematics/GteHyperellipsoid.h>
+#include <GTEngine/Mathematics/GteLine.h>
+
+namespace gte
+{
+
+// Project an ellipse onto a line.  The projection interval is [smin,smax]
+// and corresponds to the line segment P+s*D, where smin <= s <= smax.
+template <typename Real>
+void Project(Ellipse2<Real> const& ellipse, Line2<Real> const& line,
+    Real& smin, Real& smax);
+
+// Project an ellipsoid onto a line.  The projection interval is [smin,smax]
+// and corresponds to the line segment P+s*D, where smin <= s <= smax.
+template <typename Real>
+void Project(Ellipsoid3<Real> const& ellipsoid,
+    Line3<Real> const& line, Real& smin, Real& smax);
+
+
+template <typename Real>
+void Project(Ellipse2<Real> const& ellipse, Line2<Real> const& line,
+    Real& smin, Real& smax)
+{
+    // Center of projection interval.
+    Real center = Dot(line.direction, ellipse.center - line.origin);
+
+    // Radius of projection interval.
+    Real tmp[2] =
+    {
+        ellipse.extent[0] * Dot(line.direction, ellipse.axis[0]),
+        ellipse.extent[1] * Dot(line.direction, ellipse.axis[1])
+    };
+    Real rSqr = tmp[0] * tmp[0] + tmp[1] * tmp[1];
+    Real radius = std::sqrt(rSqr);
+
+    smin = center - radius;
+    smax = center + radius;
+}
+
+template <typename Real>
+void Project(Ellipsoid3<Real> const& ellipsoid, Line3<Real> const& line,
+    Real& smin, Real& smax)
+{
+    // Center of projection interval.
+    Real center = Dot(line.direction, ellipsoid.center - line.origin);
+
+    // Radius of projection interval.
+    Real tmp[3] =
+    {
+        ellipsoid.extent[0] * Dot(line.direction, ellipsoid.axis[0]),
+        ellipsoid.extent[1] * Dot(line.direction, ellipsoid.axis[1]),
+        ellipsoid.extent[2] * Dot(line.direction, ellipsoid.axis[2])
+    };
+    Real rSqr = tmp[0] * tmp[0] + tmp[1] * tmp[1] + tmp[2] * tmp[2];
+    Real radius = std::sqrt(rSqr);
+
+    smin = center - radius;
+    smax = center + radius;
+}
+
+
+}
