move GeometryProcessing plugin from NotForLicensees to Experimental

diff --git a/Engine/Plugins/Experimental/GeometryProcessing/Source/GeometryAlgorithms/Private/GTEngine/Mathematics/GteQuadraticField.h b/Engine/Plugins/Experimental/GeometryProcessing/Source/GeometryAlgorithms/Private/GTEngine/Mathematics/GteQuadraticField.h
new file mode 100644
--- /dev/null
+++ b/Engine/Plugins/Experimental/GeometryProcessing/Source/GeometryAlgorithms/Private/GTEngine/Mathematics/GteQuadraticField.h
@@ -0,0 +1,368 @@
+// David Eberly, Geometric Tools, Redmond WA 98052
+// Copyright (c) 1998-2019
+// Distributed under the Boost Software License, Version 1.0.
+// http://www.boost.org/LICENSE_1_0.txt
+// http://www.geometrictools.com/License/Boost/LICENSE_1_0.txt
+// File Version: 3.21.0 (2019/01/21)
+
+#pragma once
+
+#include <GTEngine/Mathematics/GteMath.h>
+
+// A real quadratic field has elements of the form z = x + y * sqrt(d), where
+// x, y and d are rational numbers and where d > 0.  In abstract algebra,
+// it is require that d not be the square of a rational number.  When this
+// is the case, the representation of z is unique.  For GTEngine, the
+// square-free constraint is not essential, because we need only the
+// arithmetic properties of the field (addition, subtraction, multiplication
+// and division).  We also need comparisons between elements of the quadratic
+// field, but these can be reformulated to use only operations for rational
+// arithmetic.
+
+namespace gte
+{
+    // The type T can be rational (for arbitrary precision arithmetic) or
+    // floating-point (the quadratic field operations are valid, although
+    // with floating-point rounding errors).
+    template <typename T>
+    class QuadraticField
+    {
+    public:
+        // The quadratic field descriptor that simply stores d.
+        QuadraticField()
+            :
+            d(0)
+        {
+        }
+
+        QuadraticField(T const& inD)
+            :
+            d(inD)
+        {
+        }
+
+        T d;
+
+        // Elements of the quadratic field.  The arithmetic and comparisons
+        // are managed by the quadratic field itself to avoid having multiple
+        // copies of d (one copy per Element).
+        class Element
+        {
+        public:
+            Element()
+                :
+                x(0),
+                y(0)
+            {
+            }
+
+            Element(T const& inX, T const& inY)
+                :
+                x(inX),
+                y(inY)
+            {
+            }
+
+            T x, y;
+        };
+
+        T Convert(Element const& e) const
+        {
+            return e.x + e.y * std::sqrt(d);
+        }
+
+        Element Negate(Element const& e) const
+        {
+            return Element(-e.x, -e.y);
+        }
+
+        Element Add(Element const& e0, Element const& e1) const
+        {
+            T x = e0.x + e1.x;
+            T y = e0.y + e1.y;
+            return Element(x, y);
+        }
+
+        Element Add(Element const& e, T const& s) const
+        {
+            T x = e.x + s;
+            T y = e.y;
+            return Element(x, y);
+        }
+
+        Element Add(T const& s, Element const& e) const
+        {
+            T x = s + e.x;
+            T y = e.y;
+            return Element(x, y);
+        }
+
+        Element Sub(Element const& e0, Element const& e1) const
+        {
+            T x = e0.x - e1.x;
+            T y = e0.y - e1.y;
+            return Element(x, y);
+        }
+
+        Element Sub(Element const& e, T const& s) const
+        {
+            T x = e.x - s;
+            T y = e.y;
+            return Element(x, y);
+        }
+
+        Element Sub(T const& s, Element const& e) const
+        {
+            T x = s - e.x;
+            T y = -e.y;
+            return Element(x, y);
+        }
+
+        Element Mul(Element const& e0, Element const& e1) const
+        {
+            T x = e0.x * e1.x + e0.y * e1.y * d;
+            T y = e0.x * e1.y + e0.y * e1.x;
+            return Element(x, y);
+        }
+
+        Element Mul(Element const& e, T const& s)
+        {
+            return Element(e.x * s, e.y * s);
+        }
+
+        Element Mul(T const& s, Element const& e)
+        {
+            return Element(s * e.x, s * e.y);
+        }
+
+        // Expose division only when T has a division operator.
+        template <typename Dummy = T>
+        typename std::enable_if<has_division_operator<Dummy>::value, Element>::type
+        Div(Element const& e0, Element const& e1) const
+        {
+            T z = e1.x * e1.x - e1.y * e1.y * d;
+            T x = (e0.x * e1.x - e0.y * e1.y * d) / z;
+            T y = (e0.y * e1.x - e0.x * e1.y) / z;
+            return Element(x, y);
+        }
+
+        template <typename Dummy = T>
+        typename std::enable_if<has_division_operator<Dummy>::value, Element>::type
+        Div(Element const& e, T const& s)
+        {
+            return Element(e.x / s, e.y / s);
+        }
+
+        // Comparison functions.
+        bool EqualZero(Element const& e) const
+        {
+            T const zero(0);
+            if (d == zero || e.y == zero)
+            {
+                return e.x == zero;
+            }
+            else if (e.y > zero)
+            {
+                if (e.x >= zero)
+                {
+                    return false;
+                }
+                else
+                {
+                    return e.x * e.x - e.y * e.y * d == zero;
+                }
+            }
+            else // e.y < zero
+            {
+                if (e.x <= zero)
+                {
+                    return false;
+                }
+                else
+                {
+                    return e.x * e.x - e.y * e.y * d == zero;
+                }
+            }
+        }
+
+        bool NotEqualZero(Element const& e) const
+        {
+            T const zero(0);
+            if (d == zero || e.y == zero)
+            {
+                return e.x != zero;
+            }
+            else if (e.y > zero)
+            {
+                if (e.x >= zero)
+                {
+                    return true;
+                }
+                else
+                {
+                    return e.x * e.x - e.y * e.y * d != zero;
+                }
+            }
+            else // e.y < zero
+            {
+                if (e.x <= zero)
+                {
+                    return true;
+                }
+                else
+                {
+                    return e.x * e.x - e.y * e.y * d != zero;
+                }
+            }
+        }
+
+        bool LessThanZero(Element const& e) const
+        {
+            T const zero(0);
+            if (d == zero || e.y == zero)
+            {
+                return e.x < zero;
+            }
+            else if (e.y > zero)
+            {
+                if (e.x >= zero)
+                {
+                    return false;
+                }
+                else
+                {
+                    return e.x * e.x - e.y * e.y * d > zero;
+                }
+            }
+            else // e.y < zero
+            {
+                if (e.x <= zero)
+                {
+                    return true;
+                }
+                else
+                {
+                    return e.x * e.x - e.y * e.y * d < zero;
+                }
+            }
+        }
+
+        bool GreaterThanZero(Element const& e) const
+        {
+            T const zero(0);
+            if (d == zero || e.y == zero)
+            {
+                return e.x > zero;
+            }
+            else if (e.y > zero)
+            {
+                if (e.x >= zero)
+                {
+                    return true;
+                }
+                else
+                {
+                    return e.x * e.x - e.y * e.y * d < zero;
+                }
+            }
+            else // e.y < zero
+            {
+                if (e.x <= zero)
+                {
+                    return false;
+                }
+                else
+                {
+                    return e.x * e.x - e.y * e.y * d > zero;
+                }
+            }
+        }
+
+        bool LessThanOrEqualZero(Element const& e) const
+        {
+            T const zero(0);
+            if (d == zero || e.y == zero)
+            {
+                return e.x <= zero;
+            }
+            else if (e.y > zero)
+            {
+                if (e.x >= zero)
+                {
+                    return false;
+                }
+                else
+                {
+                    return e.x * e.x - e.y * e.y * d >= zero;
+                }
+            }
+            else // e.y < zero
+            {
+                if (e.x <= zero)
+                {
+                    return true;
+                }
+                else
+                {
+                    return e.x * e.x - e.y * e.y * d <= zero;
+                }
+            }
+        }
+
+        bool GreaterThanOrEqualZero(Element const& e) const
+        {
+            T const zero(0);
+            if (d == zero || e.y == zero)
+            {
+                return e.x >= zero;
+            }
+            else if (e.y > zero)
+            {
+                if (e.x >= zero)
+                {
+                    return true;
+                }
+                else
+                {
+                    return e.x * e.x - e.y * e.y * d <= zero;
+                }
+            }
+            else // e.y < zero
+            {
+                if (e.x <= zero)
+                {
+                    return false;
+                }
+                else
+                {
+                    return e.x * e.x - e.y * e.y * d >= zero;
+                }
+            }
+        }
+
+        bool Equal(Element const& e0, Element const& e1) const
+        {
+            return EqualZero(Sub(e0, e1));
+        }
+
+        bool LessThan(Element const& e0, Element const& e1) const
+        {
+            return LessThanZero(Sub(e0, e1));
+        }
+
+        bool GreaterThan(Element const& e0, Element const& e1) const
+        {
+            return GreaterThanZero(Sub(e0, e1));
+        }
+
+        bool LessThanOrEqual(Element const& e0, Element const& e1) const
+        {
+            return LessThanOrEqualZero(Sub(e0, e1));
+        }
+
+        bool GreaterThanOrEqual(Element const& e0, Element const& e1) const
+        {
+            return GreaterThanOrEqualZero(Sub(e0, e1));
+        }
+    };
+}
diff --git a/Engine/Plugins/Experimental/GeometryProcessing/Source/GeometryAlgorithms/Private/GTEngine/Mathematics/GteQuarticRootsQR.h b/Engine/Plugins/Experimental/GeometryProcessing/Source/GeometryAlgorithms/Private/GTEngine/Mathematics/GteQuarticRootsQR.h
new file mode 100644
--- /dev/null
+++ b/Engine/Plugins/Experimental/GeometryProcessing/Source/GeometryAlgorithms/Private/GTEngine/Mathematics/GteQuarticRootsQR.h
@@ -0,0 +1,305 @@
+// David Eberly, Geometric Tools, Redmond WA 98052
+// Copyright (c) 1998-2019
+// Distributed under the Boost Software License, Version 1.0.
+// http://www.boost.org/LICENSE_1_0.txt
+// http://www.geometrictools.com/License/Boost/LICENSE_1_0.txt
+// File Version: 3.0.1 (2018/10/05)
+
+#pragma once
+
+#include "GteCubicRootsQR.h"
+
+// An implementation of the QR algorithm described in "Matrix Computations,
+// 2nd edition" by G. H. Golub and C. F. Van Loan, The Johns Hopkins
+// University Press, Baltimore MD, Fourth Printing 1993.  In particular,
+// the implementation is based on Chapter 7 (The Unsymmetric Eigenvalue
+// Problem), Section 7.5 (The Practical QR Algorithm).  The algorithm is
+// specialized for the companion matrix associated with a quartic polynomial.
+
+namespace gte
+{
+
+template <typename Real>
+class QuarticRootsQR
+{
+public:
+    typedef std::array<std::array<Real, 4>, 4> Matrix;
+
+    // Solve p(x) = c0 + c1 * x + c2 * x^2 + c3 * x^3 + x^4 = 0.
+    uint32_t operator() (uint32_t maxIterations, Real c0, Real c1, Real c2, Real c3,
+        uint32_t& numRoots, std::array<Real, 4>& roots);
+
+    // Compute the real eigenvalues of the upper Hessenberg matrix A.  The
+    // matrix is modified by in-place operations, so if you need to remember
+    // A, you must make your own copy before calling this function.
+    uint32_t operator() (uint32_t maxIterations, Matrix& A,
+        uint32_t& numRoots, std::array<Real, 4>& roots);
+
+private:
+    void DoIteration(std::array<Real, 3> const& V, Matrix& A);
+
+    template <int N>
+    std::array<Real, N> House(std::array<Real, N> const& X);
+
+    template <int N>
+    void RowHouse(int rmin, int rmax, int cmin, int cmax,
+        std::array<Real, N> const& V, std::array<Real, N> const& MV, Matrix& A);
+
+    template <int N>
+    void ColHouse(int rmin, int rmax, int cmin, int cmax,
+        std::array<Real, N> const& V, std::array<Real, N> const& MV, Matrix& A);
+
+    void GetQuadraticRoots(int i0, int i1, Matrix const& A,
+        uint32_t& numRoots, std::array<Real, 4>& roots);
+};
+
+
+template <typename Real>
+uint32_t QuarticRootsQR<Real>::operator() (uint32_t maxIterations, Real c0, Real c1, Real c2, Real c3,
+    uint32_t& numRoots, std::array<Real, 4>& roots)
+{
+    // Create the companion matrix for the polynomial.  The matrix is in upper
+    // Hessenberg form.
+    Matrix A;
+    A[0][0] = (Real)0;
+    A[0][1] = (Real)0;
+    A[0][2] = (Real)0;
+    A[0][3] = -c0;
+    A[1][0] = (Real)1;
+    A[1][1] = (Real)0;
+    A[1][2] = (Real)0;
+    A[1][3] = -c1;
+    A[2][0] = (Real)0;
+    A[2][1] = (Real)1;
+    A[2][2] = (Real)0;
+    A[2][3] = -c2;
+    A[3][0] = (Real)0;
+    A[3][1] = (Real)0;
+    A[3][2] = (Real)1;
+    A[3][3] = -c3;
+
+    // Avoid the QR-cycle when c1 = c2 = 0 and avoid the slow convergence
+    // when c1 and c2 are nearly zero.
+    std::array<Real, 3> V{
+        (Real)1,
+        (Real)0.36602540378443865,
+        (Real)0.36602540378443865 };
+    DoIteration(V, A);
+
+    return operator()(maxIterations, A, numRoots, roots);
+}
+
+template <typename Real>
+uint32_t QuarticRootsQR<Real>::operator() (uint32_t maxIterations, Matrix& A,
+    uint32_t& numRoots, std::array<Real, 4>& roots)
+{
+    numRoots = 0;
+    std::fill(roots.begin(), roots.end(), (Real)0);
+
+    for (uint32_t numIterations = 0; numIterations < maxIterations; ++numIterations)
+    {
+        // Apply a Francis QR iteration.
+        Real tr = A[2][2] + A[3][3];
+        Real det = A[2][2] * A[3][3] - A[2][3] * A[3][2];
+        std::array<Real, 3> X{
+            A[0][0] * A[0][0] + A[0][1] * A[1][0] - tr * A[0][0] + det,
+            A[1][0] * (A[0][0] + A[1][1] - tr),
+            A[1][0] * A[2][1] };
+        std::array<Real, 3> V = House<3>(X);
+        DoIteration(V, A);
+
+        // Test for uncoupling of A.
+        Real tr12 = A[1][1] + A[2][2];
+        if (tr12 + A[2][1] == tr12)
+        {
+            GetQuadraticRoots(0, 1, A, numRoots, roots);
+            GetQuadraticRoots(2, 3, A, numRoots, roots);
+            return numIterations;
+        }
+
+        Real tr01 = A[0][0] + A[1][1];
+        if (tr01 + A[1][0] == tr01)
+        {
+            numRoots = 1;
+            roots[0] = A[0][0];
+
+            // TODO: The cubic solver is not designed to process 3x3 submatrices
+            // of an NxN matrix, so the copy of a submatrix of A to B is a simple
+            // workaround for running the solver.  Write general root-finding
+            // code that avoids such copying.
+            uint32_t subMaxIterations = maxIterations - numIterations;
+            typename CubicRootsQR<Real>::Matrix B;
+            for (int r = 0; r < 3; ++r)
+            {
+                for (int c = 0; c < 3; ++c)
+                {
+                    B[r][c] = A[r + 1][c + 1];
+                }
+            }
+
+            uint32_t numSubroots = 0;
+            std::array<Real, 3> subroots;
+            uint32_t numSubiterations = CubicRootsQR<Real>()(subMaxIterations, B,
+                numSubroots, subroots);
+            for (uint32_t i = 0; i < numSubroots; ++i)
+            {
+                roots[numRoots++] = subroots[i];
+            }
+            return numIterations + numSubiterations;
+        }
+
+        Real tr23 = A[2][2] + A[3][3];
+        if (tr23 + A[3][2] == tr23)
+        {
+            numRoots = 1;
+            roots[0] = A[3][3];
+
+            // TODO: The cubic solver is not designed to process 3x3 submatrices
+            // of an NxN matrix, so the copy of a submatrix of A to B is a simple
+            // workaround for running the solver.  Write general root-finding
+            // code that avoids such copying.
+            uint32_t subMaxIterations = maxIterations - numIterations;
+            typename CubicRootsQR<Real>::Matrix B;
+            for (int r = 0; r < 3; ++r)
+            {
+                for (int c = 0; c < 3; ++c)
+                {
+                    B[r][c] = A[r][c];
+                }
+            }
+
+            uint32_t numSubroots = 0;
+            std::array<Real, 3> subroots;
+            uint32_t numSubiterations = CubicRootsQR<Real>()(subMaxIterations, B,
+                numSubroots, subroots);
+            for (uint32_t i = 0; i < numSubroots; ++i)
+            {
+                roots[numRoots++] = subroots[i];
+            }
+            return numIterations + numSubiterations;
+        }
+    }
+    return maxIterations;
+}
+
+template <typename Real>
+void QuarticRootsQR<Real>::DoIteration(std::array<Real, 3> const& V, Matrix& A)
+{
+    Real multV = ((Real)-2) / (V[0] * V[0] + V[1] * V[1] + V[2] * V[2]);
+    std::array<Real, 3> MV{ multV * V[0], multV * V[1], multV * V[2] };
+    RowHouse<3>(0, 2, 0, 3, V, MV, A);
+    ColHouse<3>(0, 3, 0, 2, V, MV, A);
+
+    std::array<Real, 3> X{ A[1][0], A[2][0], A[3][0] };
+    std::array<Real, 3> locV = House<3>(X);
+    multV = ((Real)-2) / (locV[0] * locV[0] + locV[1] * locV[1] + locV[2] * locV[2]);
+    MV = { multV * locV[0], multV * locV[1], multV * locV[2] };
+    RowHouse<3>(1, 3, 0, 3, locV, MV, A);
+    ColHouse<3>(0, 3, 1, 3, locV, MV, A);
+
+    std::array<Real, 2> Y{ A[2][1], A[3][1] };
+    std::array<Real, 2> W = House<2>(Y);
+    Real multW = ((Real)-2) / (W[0] * W[0] + W[1] * W[1]);
+    std::array<Real, 2> MW = { multW * W[0], multW * W[1] };
+    RowHouse<2>(2, 3, 0, 3, W, MW, A);
+    ColHouse<2>(0, 3, 2, 3, W, MW, A);
+}
+
+template <typename Real>
+template <int N>
+std::array<Real, N> QuarticRootsQR<Real>::House(std::array<Real, N> const & X)
+{
+    std::array<Real, N> V;
+    Real length = (Real)0;
+    for (int i = 0; i < N; ++i)
+    {
+        length += X[i] * X[i];
+    }
+    length = std::sqrt(length);
+    if (length != (Real)0)
+    {
+        Real sign = (X[0] >= (Real)0 ? (Real)1 : (Real)-1);
+        Real denom = X[0] + sign * length;
+        for (int i = 1; i < N; ++i)
+        {
+            V[i] = X[i] / denom;
+        }
+    }
+    V[0] = (Real)1;
+    return V;
+}
+
+template <typename Real>
+template <int N>
+void QuarticRootsQR<Real>::RowHouse(int rmin, int rmax, int cmin, int cmax,
+    std::array<Real, N> const& V, std::array<Real, N> const& MV, Matrix& A)
+{
+    std::array<Real, 4> W;  // only elements cmin through cmax are used
+
+    for (int c = cmin; c <= cmax; ++c)
+    {
+        W[c] = (Real)0;
+        for (int r = rmin, k = 0; r <= rmax; ++r, ++k)
+        {
+            W[c] += V[k] * A[r][c];
+        }
+    }
+
+    for (int r = rmin, k = 0; r <= rmax; ++r, ++k)
+    {
+        for (int c = cmin; c <= cmax; ++c)
+        {
+            A[r][c] += MV[k] * W[c];
+        }
+    }
+}
+
+template <typename Real>
+template <int N>
+void QuarticRootsQR<Real>::ColHouse(int rmin, int rmax, int cmin, int cmax,
+    std::array<Real, N> const& V, std::array<Real, N> const& MV, Matrix& A)
+{
+    std::array<Real, 4> W;  // only elements rmin through rmax are used
+
+    for (int r = rmin; r <= rmax; ++r)
+    {
+        W[r] = (Real)0;
+        for (int c = cmin, k = 0; c <= cmax; ++c, ++k)
+        {
+            W[r] += V[k] * A[r][c];
+        }
+    }
+
+    for (int r = rmin; r <= rmax; ++r)
+    {
+        for (int c = cmin, k = 0; c <= cmax; ++c, ++k)
+        {
+            A[r][c] += W[r] * MV[k];
+        }
+    }
+}
+
+template <typename Real>
+void QuarticRootsQR<Real>::GetQuadraticRoots(int i0, int i1, Matrix const& A,
+    uint32_t& numRoots, std::array<Real, 4>& roots)
+{
+    // Solve x^2 - t * x + d = 0, where t is the trace and d is the
+    // determinant of the 2x2 matrix defined by indices i0 and i1.  The
+    // discriminant is D = (t/2)^2 - d.  When D >= 0, the roots are real
+    // values t/2 - sqrt(D) and t/2 + sqrt(D).  To avoid potential numerical
+    // issues with subtractive cancellation, the roots are computed as
+    //   r0 = t/2 + sign(t/2)*sqrt(D), r1 = trace - r0.
+    Real trace = A[i0][i0] + A[i1][i1];
+    Real halfTrace = trace * (Real)0.5;
+    Real determinant = A[i0][i0] * A[i1][i1] - A[i0][i1] * A[i1][i0];
+    Real discriminant = halfTrace * halfTrace - determinant;
+    if (discriminant >= (Real)0)
+    {
+        Real sign = (trace >= (Real)0 ? (Real)1 : (Real)-1);
+        Real root = halfTrace + sign * std::sqrt(discriminant);
+        roots[numRoots++] = root;
+        roots[numRoots++] = trace - root;
+    }
+}
+
+}
diff --git a/Engine/Plugins/Experimental/GeometryProcessing/Source/GeometryAlgorithms/Private/GTEngine/Mathematics/GteQuaternion.h b/Engine/Plugins/Experimental/GeometryProcessing/Source/GeometryAlgorithms/Private/GTEngine/Mathematics/GteQuaternion.h
new file mode 100644
--- /dev/null
+++ b/Engine/Plugins/Experimental/GeometryProcessing/Source/GeometryAlgorithms/Private/GTEngine/Mathematics/GteQuaternion.h
@@ -0,0 +1,455 @@
+// David Eberly, Geometric Tools, Redmond WA 98052
+// Copyright (c) 1998-2019
+// Distributed under the Boost Software License, Version 1.0.
+// http://www.boost.org/LICENSE_1_0.txt
+// http://www.geometrictools.com/License/Boost/LICENSE_1_0.txt
+// File Version: 3.0.1 (2018/12/11)
+
+#pragma once
+
+#include <GTEngine/Mathematics/GteVector.h>
+#include <GTEngine/Mathematics/GteMatrix.h>
+#include <GTEngine/Mathematics/GteChebyshevRatio.h>
+
+// A quaternion is of the form
+//   q = x * i + y * j + z * k + w * 1 = x * i + y * j + z * k + w
+// where w, x, y, and z are real numbers.  The scalar and vector parts are
+//   Vector(q) = x * i + y * j + z * k
+//   Scalar(q) = w
+//   q = Vector(q) + Scalar(q)
+// I assume that you are familiar with the arithmetic and algebraic properties
+// of quaternions.  See
+// https://www.geometrictools.com/Documentation/Quaternions.pdf
+
+namespace gte
+{
+    template <typename Real>
+    class Quaternion
+    {
+    public:
+        // The quaternions are of the form q = x*i + y*j + z*k + w.  In tuple
+        // form, q = (x,y,z,w).
+
+        // Construction.  The default constructor does not initialize the
+        // members.
+        Quaternion() = default;
+
+        Quaternion(Real x, Real y, Real z, Real w)
+        {
+            mTuple[0] = x;
+            mTuple[1] = y;
+            mTuple[2] = z;
+            mTuple[3] = w;
+        }
+
+        // Member access.
+        inline Real const& operator[](int i) const
+        {
+            return mTuple[i];
+        }
+
+        inline Real& operator[](int i)
+        {
+            return mTuple[i];
+        }
+
+        // Comparisons.
+        inline bool operator==(Quaternion const& q) const
+        {
+            return mTuple == q.mTuple;
+        }
+
+        inline bool operator!=(Quaternion const& q) const
+        {
+            return mTuple != q.mTuple;
+        }
+
+        inline bool operator<(Quaternion const& q) const
+        {
+            return mTuple < q.mTuple;
+        }
+
+        inline bool operator<=(Quaternion const& q) const
+        {
+            return mTuple <= q.mTuple;
+        }
+
+        inline bool operator>(Quaternion const& q) const
+        {
+            return mTuple > q.mTuple;
+        }
+
+        inline bool operator>=(Quaternion const& q) const
+        {
+            return mTuple >= q.mTuple;
+        }
+
+        // Special quaternions.
+
+        // z = 0*i + 0*j + 0*k + 0
+        static Quaternion Zero()
+        {
+            return Quaternion((Real)0, (Real)0, (Real)0, (Real)0);
+        }
+
+        // i = 1*i + 0*j + 0*k + 0
+        static Quaternion I()
+        {
+            return Quaternion((Real)1, (Real)0, (Real)0, (Real)0);
+        }
+
+        // j = 0*i + 1*j + 0*k + 0
+        static Quaternion J()
+        {
+            return Quaternion((Real)0, (Real)1, (Real)0, (Real)0);
+        }
+
+        // k = 0*i + 0*j + 1*k + 0
+        static Quaternion K()
+        {
+            return Quaternion((Real)0, (Real)0, (Real)1, (Real)0);
+        }
+
+        // 1 = 0*i + 0*j + 0*k + 1
+        static Quaternion Identity()
+        {
+            return Quaternion((Real)0, (Real)0, (Real)0, (Real)1);
+        }
+
+    protected:
+        std::array<Real, 4> mTuple;
+    };
+
+    // Unary operations.
+    template <typename Real>
+    Quaternion<Real> operator+(Quaternion<Real> const& q)
+    {
+        return q;
+    }
+
+    template <typename Real>
+    Quaternion<Real> operator-(Quaternion<Real> const& q)
+    {
+        Quaternion<Real> result;
+        for (int i = 0; i < 4; ++i)
+        {
+            result[i] = -q[i];
+        }
+        return result;
+    }
+
+    // Linear algebraic operations.
+    template <typename Real>
+    Quaternion<Real> operator+(Quaternion<Real> const& q0, Quaternion<Real> const& q1)
+    {
+        Quaternion<Real> result = q0;
+        return result += q1;
+    }
+
+    template <typename Real>
+    Quaternion<Real> operator-(Quaternion<Real> const& q0, Quaternion<Real> const& q1)
+    {
+        Quaternion<Real> result = q0;
+        return result -= q1;
+    }
+
+    template <typename Real>
+    Quaternion<Real> operator*(Quaternion<Real> const& q, Real scalar)
+    {
+        Quaternion<Real> result = q;
+        return result *= scalar;
+    }
+
+    template <typename Real>
+    Quaternion<Real> operator*(Real scalar, Quaternion<Real> const& q)
+    {
+        Quaternion<Real> result = q;
+        return result *= scalar;
+    }
+
+    template <typename Real>
+    Quaternion<Real> operator/(Quaternion<Real> const& q, Real scalar)
+    {
+        Quaternion<Real> result = q;
+        return result /= scalar;
+    }
+
+    template <typename Real>
+    Quaternion<Real>& operator+=(Quaternion<Real>& q0, Quaternion<Real> const& q1)
+    {
+        for (int i = 0; i < 4; ++i)
+        {
+            q0[i] += q1[i];
+        }
+        return q0;
+    }
+
+    template <typename Real>
+    Quaternion<Real>& operator-=(Quaternion<Real>& q0, Quaternion<Real> const& q1)
+    {
+        for (int i = 0; i < 4; ++i)
+        {
+            q0[i] -= q1[i];
+        }
+        return q0;
+    }
+
+    template <typename Real>
+    Quaternion<Real>& operator*=(Quaternion<Real>& q, Real scalar)
+    {
+        for (int i = 0; i < 4; ++i)
+        {
+            q[i] *= scalar;
+        }
+        return q;
+    }
+
+    template <typename Real>
+    Quaternion<Real>& operator/=(Quaternion<Real>& q, Real scalar)
+    {
+        if (scalar != (Real)0)
+        {
+            for (int i = 0; i < 4; ++i)
+            {
+                q[i] /= scalar;
+            }
+        }
+        else
+        {
+            for (int i = 0; i < 4; ++i)
+            {
+                q[i] = (Real)0;
+            }
+        }
+        return q;
+    }
+
+    // Geometric operations.
+    template <typename Real>
+    Real Dot(Quaternion<Real> const& q0, Quaternion<Real> const& q1)
+    {
+        Real dot = q0[0] * q1[0];
+        for (int i = 1; i < 4; ++i)
+        {
+            dot += q0[i] * q1[i];
+        }
+        return dot;
+    }
+
+    template <typename Real>
+    Real Length(Quaternion<Real> const& q)
+    {
+        return std::sqrt(Dot(q, q));
+    }
+
+    template <typename Real>
+    Real Normalize(Quaternion<Real>& q)
+    {
+        Real length = std::sqrt(Dot(q, q));
+        if (length > (Real)0)
+        {
+            q /= length;
+        }
+        else
+        {
+            for (int i = 0; i < 4; ++i)
+            {
+                q[i] = (Real)0;
+            }
+        }
+        return length;
+    }
+
+    // Multiplication of quaternions.  This operation is not generally
+    // commutative; that is, q0*q1 and q1*q0 are not usually the same value.
+    // (x0*i + y0*j + z0*k + w0)*(x1*i + y1*j + z1*k + w1)
+    // =
+    // i*(+x0*w1 + y0*z1 - z0*y1 + w0*x1) +
+    // j*(-x0*z1 + y0*w1 + z0*x1 + w0*y1) +
+    // k*(+x0*y1 - y0*x1 + z0*w1 + w0*z1) +
+    // 1*(-x0*x1 - y0*y1 - z0*z1 + w0*w1)
+    template <typename Real>
+    Quaternion<Real> operator*(Quaternion<Real> const& q0, Quaternion<Real> const& q1)
+    {
+        // (x0*i + y0*j + z0*k + w0)*(x1*i + y1*j + z1*k + w1)
+        // =
+        // i*(+x0*w1 + y0*z1 - z0*y1 + w0*x1) +
+        // j*(-x0*z1 + y0*w1 + z0*x1 + w0*y1) +
+        // k*(+x0*y1 - y0*x1 + z0*w1 + w0*z1) +
+        // 1*(-x0*x1 - y0*y1 - z0*z1 + w0*w1)
+
+        return Quaternion<Real>
+        (
+            +q0[0] * q1[3] + q0[1] * q1[2] - q0[2] * q1[1] + q0[3] * q1[0],
+            -q0[0] * q1[2] + q0[1] * q1[3] + q0[2] * q1[0] + q0[3] * q1[1],
+            +q0[0] * q1[1] - q0[1] * q1[0] + q0[2] * q1[3] + q0[3] * q1[2],
+            -q0[0] * q1[0] - q0[1] * q1[1] - q0[2] * q1[2] + q0[3] * q1[3]
+        );
+    }
+
+    // For a nonzero quaternion q = (x,y,z,w), inv(q) = (-x,-y,-z,w)/|q|^2, where
+    // |q| is the length of the quaternion.  When q is zero, the function returns
+    // zero, which is considered to be an improbable case.
+    template <typename Real>
+    Quaternion<Real> Inverse(Quaternion<Real> const& q)
+    {
+        Real sqrLen = Dot(q, q);
+        if (sqrLen > (Real)0)
+        {
+            Quaternion<Real> inverse = Conjugate(q) / sqrLen;
+            return inverse;
+        }
+        else
+        {
+            return Quaternion<Real>::Zero();
+        }
+    }
+
+    // The conjugate of q = (x,y,z,w) is conj(q) = (-x,-y,-z,w).
+    template <typename Real>
+    Quaternion<Real> Conjugate(Quaternion<Real> const& q)
+    {
+        return Quaternion<Real>(-q[0], -q[1], -q[2], +q[3]);
+    }
+
+    // Rotate a vector using quaternion multiplication.  The input quaternion must
+    // be unit length.
+    template <typename Real>
+    Vector<4, Real> Rotate(Quaternion<Real> const& q, Vector<4, Real> const& v)
+    {
+        Quaternion<Real> input(v[0], v[1], v[2], (Real)0);
+        Quaternion<Real> output = q * input * Conjugate(q);
+        Vector<4, Real> u{ output[0], output[1], output[2], (Real)0 };
+        return u;
+    }
+
+    // The spherical linear interpolation (slerp) of unit-length quaternions
+    // q0 and q1 for t in [0,1] is
+    //     slerp(t,q0,q1) = [sin(t*theta)*q0 + sin((1-t)*theta)*q1]/sin(theta)
+    // where theta is the angle between q0 and q1 [cos(theta) = Dot(q0,q1)].
+    // This function is a parameterization of the great spherical arc between
+    // q0 and q1 on the unit hypersphere.  Moreover, the parameterization is
+    // one of normalized arclength--a particle traveling along the arc through
+    // time t does so with constant speed.
+    //
+    // When using slerp in animations involving sequences of quaternions, it is
+    // typical that the quaternions are preprocessed so that consecutive ones
+    // form an acute angle A in [0,pi/2].  Other preprocessing can help with
+    // performance.  See the function comments below.
+    //
+    // See GteSlerpEstimate.{h,inl} for various approximations, including
+    // SLERP<Real>::EstimateRPH that gives good performance and accurate results
+    // for preprocessed quaternions.
+
+    // The angle between q0 and q1 is in [0,pi).  There are no angle restrictions
+    // restrictions and nothing is precomputed.
+    template <typename Real>
+    Quaternion<Real> Slerp(Real t, Quaternion<Real> const& q0, Quaternion<Real> const& q1)
+    {
+        Real cosA = Dot(q0, q1);
+        Real sign;
+        if (cosA >= (Real)0)
+        {
+            sign = (Real)1;
+        }
+        else
+        {
+            cosA = -cosA;
+            sign = (Real)-1;
+        }
+
+        Real f0, f1;
+        ChebyshevRatio<Real>::Get(t, cosA, f0, f1);
+        return q0 * f0 + q1 * (sign * f1);
+    }
+
+    // The angle between q0 and q1 must be in [0,pi/2].  The suffix R is for
+    // 'Restricted'.  The preprocessing code is
+    //   Quaternion<Real> q[n];  // assuming initialized
+    //   for (i0 = 0, i1 = 1; i1 < n; i0 = i1++)
+    //   {
+    //       cosA = Dot(q[i0], q[i1]);
+    //       if (cosA < 0)
+    //       {
+    //           q[i1] = -q[i1];  // now Dot(q[i0], q[i]1) >= 0
+    //       }
+    //   }
+    template <typename Real>
+    Quaternion<Real> SlerpR(Real t, Quaternion<Real> const& q0, Quaternion<Real> const& q1)
+    {
+        Real f0, f1;
+        ChebyshevRatio<Real>::Get(t, Dot(q0, q1), f0, f1);
+        return q0 * f0 + q1 * f1;
+    }
+
+    // The angle between q0 and q1 must be in [0,pi/2].  The suffix R is for
+    // 'Restricted' and the suffix P is for 'Preprocessed'.  The preprocessing
+    // code is
+    //   Quaternion<Real> q[n];  // assuming initialized
+    //   Real cosA[n-1], omcosA[n-1];  // to be precomputed
+    //   for (i0 = 0, i1 = 1; i1 < n; i0 = i1++)
+    //   {
+    //       cs = Dot(q[i0], q[i1]);
+    //       if (cosA[i0] < 0)
+    //       {
+    //           q[i1] = -q[i1];
+    //           cs = -cs;
+    //       }
+    //
+    //       // for Quaterion<Real>::SlerpRP
+    //       cosA[i0] = cs;
+    //
+    //       // for SLERP<Real>::EstimateRP
+    //       omcosA[i0] = 1 - cs;
+    //   }
+    template <typename Real>
+    Quaternion<Real> SlerpRP(Real t, Quaternion<Real> const& q0, Quaternion<Real> const& q1, Real cosA)
+    {
+        Real f0, f1;
+        ChebyshevRatio<Real>::Get(t, cosA, f0, f1);
+        return q0 * f0 + q1 * f1;
+    }
+
+    // The angle between q0 and q1 is A and must be in [0,pi/2].  The suffic R
+    // is for 'Restricted', the suffix P is for 'Preprocessed' and the suffix
+    // H is for 'Half' (the quaternion qh halfway between q0 and q1 is
+    // precomputed).  Quaternion qh is slerp(1/2,q0,q1) = (q0+q1)/|q0+q1|, so
+    // the angle between q0 and qh is A/2 and the angle between qh and q1 is
+    // A/2.  The preprocessing code is
+    //   Quaternion<Real> q[n];  // assuming initialized
+    //   Quaternion<Real> qh[n-1];  // to be precomputed
+    //   Real omcosAH[n-1];  // to be precomputed
+    //   for (i0 = 0, i1 = 1; i1 < n; i0 = i1++)
+    //   {
+    //       cosA = Dot(q[i0], q[i1]);
+    //       if (cosA < 0)
+    //       {
+    //           q[i1] = -q[i1];
+    //           cosA = -cosA;
+    //       }
+    //
+    //       // for Quaternion<Real>::SlerpRPH and SLERP<Real>::EstimateRPH
+    //       cosAH[i0] = sqrt((1+cosA)/2);
+    //       qh[i0] = (q0 + q1) / (2 * cosAH[i0]);
+    //
+    //       // for SLERP<Real>::EstimateRPH
+    //       omcosAH[i0] = 1 - cosAH[i0];
+    //   }
+    template <typename Real>
+    Quaternion<Real> SlerpRPH(Real t, Quaternion<Real> const& q0, Quaternion<Real> const& q1,
+        Quaternion<Real> const& qh, Real cosAH)
+    {
+        Real f0, f1;
+        Real twoT = static_cast<Real>(2) * t;
+        if (twoT <= static_cast<Real>(1))
+        {
+            ChebyshevRatio<Real>::Get(twoT, cosAH, f0, f1);
+            return q0 * f0 + qh * f1;
+        }
+        else
+        {
+            ChebyshevRatio<Real>::Get(twoT - static_cast<Real>(1), cosAH, f0, f1);
+            return qh * f0 + q1 * f1;
+        }
+    }
+}
diff --git a/Engine/Plugins/Experimental/GeometryProcessing/Source/GeometryAlgorithms/Private/GTEngine/Mathematics/GteRay.h b/Engine/Plugins/Experimental/GeometryProcessing/Source/GeometryAlgorithms/Private/GTEngine/Mathematics/GteRay.h
new file mode 100644
--- /dev/null
+++ b/Engine/Plugins/Experimental/GeometryProcessing/Source/GeometryAlgorithms/Private/GTEngine/Mathematics/GteRay.h
@@ -0,0 +1,112 @@
+// David Eberly, Geometric Tools, Redmond WA 98052
+// Copyright (c) 1998-2019
+// Distributed under the Boost Software License, Version 1.0.
+// http://www.boost.org/LICENSE_1_0.txt
+// http://www.geometrictools.com/License/Boost/LICENSE_1_0.txt
+// File Version: 3.0.0 (2016/06/19)
+
+#pragma once
+
+#include <GTEngine/Mathematics/GteVector.h>
+
+// The ray is represented as P+t*D, where P is the ray origin, D is a
+// unit-length direction vector, and t >= 0.  The user must ensure that D is
+// unit length.
+
+namespace gte
+{
+
+template <int N, typename Real>
+class Ray
+{
+public:
+    // Construction and destruction.  The default constructor sets the origin
+    // to (0,...,0) and the ray direction to (1,0,...,0).
+    Ray();
+    Ray(Vector<N, Real> const& inOrigin, Vector<N, Real> const& inDirection);
+
+    // Public member access.  The direction must be unit length.
+    Vector<N, Real> origin, direction;
+
+public:
+    // Comparisons to support sorted containers.
+    bool operator==(Ray const& ray) const;
+    bool operator!=(Ray const& ray) const;
+    bool operator< (Ray const& ray) const;
+    bool operator<=(Ray const& ray) const;
+    bool operator> (Ray const& ray) const;
+    bool operator>=(Ray const& ray) const;
+};
+
+// Template aliases for convenience.
+template <typename Real>
+using Ray2 = Ray<2, Real>;
+
+template <typename Real>
+using Ray3 = Ray<3, Real>;
+
+
+template <int N, typename Real>
+Ray<N, Real>::Ray()
+{
+    origin.MakeZero();
+    direction.MakeUnit(0);
+}
+
+template <int N, typename Real>
+Ray<N, Real>::Ray(Vector<N, Real> const& inOrigin,
+    Vector<N, Real> const& inDirection)
+    :
+    origin(inOrigin),
+    direction(inDirection)
+{
+}
+
+template <int N, typename Real>
+bool Ray<N, Real>::operator==(Ray const& ray) const
+{
+    return origin == ray.origin && direction == ray.direction;
+}
+
+template <int N, typename Real>
+bool Ray<N, Real>::operator!=(Ray const& ray) const
+{
+    return !operator==(ray);
+}
+
+template <int N, typename Real>
+bool Ray<N, Real>::operator<(Ray const& ray) const
+{
+    if (origin < ray.origin)
+    {
+        return true;
+    }
+
+    if (origin > ray.origin)
+    {
+        return false;
+    }
+
+    return direction < ray.direction;
+}
+
+template <int N, typename Real>
+bool Ray<N, Real>::operator<=(Ray const& ray) const
+{
+    return operator<(ray) || operator==(ray);
+}
+
+template <int N, typename Real>
+bool Ray<N, Real>::operator>(Ray const& ray) const
+{
+    return !operator<=(ray);
+}
+
+template <int N, typename Real>
+bool Ray<N, Real>::operator>=(Ray const& ray) const
+{
+    return !operator<(ray);
+}
+
+
+}
diff --git a/Engine/Plugins/Experimental/GeometryProcessing/Source/GeometryAlgorithms/Private/GTEngine/Mathematics/GteRectangle.h b/Engine/Plugins/Experimental/GeometryProcessing/Source/GeometryAlgorithms/Private/GTEngine/Mathematics/GteRectangle.h
new file mode 100644
--- /dev/null
+++ b/Engine/Plugins/Experimental/GeometryProcessing/Source/GeometryAlgorithms/Private/GTEngine/Mathematics/GteRectangle.h
@@ -0,0 +1,172 @@
+// David Eberly, Geometric Tools, Redmond WA 98052
+// Copyright (c) 1998-2019
+// Distributed under the Boost Software License, Version 1.0.
+// http://www.boost.org/LICENSE_1_0.txt
+// http://www.geometrictools.com/License/Boost/LICENSE_1_0.txt
+// File Version: 3.0.0 (2016/06/19)
+
+#pragma once
+
+#include <GTEngine/Mathematics/GteVector.h>
+
+// Points are R(s0,s1) = C + s0*A0 + s1*A1, where C is the center of the
+// rectangle and A0 and A1 are unit-length and perpendicular axes.  The
+// parameters s0 and s1 are constrained by |s0| <= e0 and |s1| <= e1,
+// where e0 > 0 and e1 > 0 are the extents of the rectangle.
+
+namespace gte
+{
+
+template <int N, typename Real>
+class Rectangle
+{
+public:
+    // Construction and destruction.  The default constructor sets the origin
+    // to (0,...,0), axis A0 to (1,0,...,0), axis A1 to (0,1,0,...0), and both
+    // extents to 1.
+    Rectangle();
+    Rectangle(Vector<N, Real> const& inCenter,
+        std::array<Vector<N, Real>, 2> const& inAxis,
+        Vector<2, Real> const& inExtent);
+
+    // Compute the vertices of the rectangle.  If index i has the bit pattern
+    // i = b[1]b[0], then
+    //   vertex[i] = center + sum_{d=0}^{1} sign[d] * extent[d] * axis[d]
+    // where sign[d] = 2*b[d] - 1.
+    void GetVertices(std::array<Vector<N, Real>, 4>& vertex) const;
+
+    Vector<N, Real> center;
+    std::array<Vector<N, Real>, 2> axis;
+    Vector<2, Real> extent;
+
+public:
+    // Comparisons to support sorted containers.
+    bool operator==(Rectangle const& rectangle) const;
+    bool operator!=(Rectangle const& rectangle) const;
+    bool operator< (Rectangle const& rectangle) const;
+    bool operator<=(Rectangle const& rectangle) const;
+    bool operator> (Rectangle const& rectangle) const;
+    bool operator>=(Rectangle const& rectangle) const;
+};
+
+// Template alias for convenience.
+template <typename Real>
+using Rectangle3 = Rectangle<3, Real>;
+
+
+template <int N, typename Real>
+Rectangle<N, Real>::Rectangle()
+{
+    center.MakeZero();
+    for (int i = 0; i < 2; ++i)
+    {
+        axis[i].MakeUnit(i);
+        extent[i] = (Real)1;
+    }
+}
+
+template <int N, typename Real>
+Rectangle<N, Real>::Rectangle(Vector<N, Real> const& inCenter,
+    std::array<Vector<N, Real>, 2> const& inAxis,
+    Vector<2, Real> const& inExtent)
+    :
+    center(inCenter),
+    axis(inAxis),
+    extent(inExtent)
+{
+}
+
+template <int N, typename Real>
+void Rectangle<N, Real>::GetVertices(
+    std::array<Vector<N, Real>, 4>& vertex) const
+{
+    Vector<N, Real> product0 = extent[0] * axis[0];
+    Vector<N, Real> product1 = extent[1] * axis[1];
+    Vector<N, Real> sum = product0 + product1;
+    Vector<N, Real> dif = product0 - product1;
+
+    vertex[0] = center - sum;
+    vertex[1] = center + dif;
+    vertex[2] = center - dif;
+    vertex[3] = center + sum;
+}
+
+template <int N, typename Real>
+bool Rectangle<N, Real>::operator==(Rectangle const& rectangle) const
+{
+    if (center != rectangle.center)
+    {
+        return false;
+    }
+
+    for (int i = 0; i < 2; ++i)
+    {
+        if (axis[i] != rectangle.axis[i])
+        {
+            return false;
+        }
+    }
+
+    for (int i = 0; i < 2; ++i)
+    {
+        if (extent[i] != rectangle.extent[i])
+        {
+            return false;
+        }
+    }
+
+    return true;
+}
+
+template <int N, typename Real>
+bool Rectangle<N, Real>::operator!=(Rectangle const& rectangle) const
+{
+    return !operator==(rectangle);
+}
+
+template <int N, typename Real>
+bool Rectangle<N, Real>::operator<(Rectangle const& rectangle) const
+{
+    if (center < rectangle.center)
+    {
+        return true;
+    }
+
+    if (center > rectangle.center)
+    {
+        return false;
+    }
+
+    if (axis < rectangle.axis)
+    {
+        return true;
+    }
+
+    if (axis > rectangle.axis)
+    {
+        return false;
+    }
+
+    return extent < rectangle.extent;
+}
+
+template <int N, typename Real>
+bool Rectangle<N, Real>::operator<=(Rectangle const& rectangle) const
+{
+    return operator<(rectangle) || operator==(rectangle);
+}
+
+template <int N, typename Real>
+bool Rectangle<N, Real>::operator>(Rectangle const& rectangle) const
+{
+    return !operator<=(rectangle);
+}
+
+template <int N, typename Real>
+bool Rectangle<N, Real>::operator>=(Rectangle const& rectangle) const
+{
+    return !operator<(rectangle);
+}
+
+
+}
diff --git a/Engine/Plugins/Experimental/GeometryProcessing/Source/GeometryAlgorithms/Private/GTEngine/Mathematics/GteRectangleMesh.h b/Engine/Plugins/Experimental/GeometryProcessing/Source/GeometryAlgorithms/Private/GTEngine/Mathematics/GteRectangleMesh.h
new file mode 100644
--- /dev/null
+++ b/Engine/Plugins/Experimental/GeometryProcessing/Source/GeometryAlgorithms/Private/GTEngine/Mathematics/GteRectangleMesh.h
@@ -0,0 +1,170 @@
+// David Eberly, Geometric Tools, Redmond WA 98052
+// Copyright (c) 1998-2019
+// Distributed under the Boost Software License, Version 1.0.
+// http://www.boost.org/LICENSE_1_0.txt
+// http://www.geometrictools.com/License/Boost/LICENSE_1_0.txt
+// File Version: 3.0.4 (2016/08/29)
+
+#pragma once
+
+#include <GTEngine/Mathematics/GteMesh.h>
+#include <GTEngine/Mathematics/GteRectangle.h>
+#include <GTEngine/Mathematics/GteVector3.h>
+#include <memory>
+
+namespace gte
+{
+
+template <typename Real>
+class RectangleMesh : public Mesh<Real>
+{
+public:
+    // Create a mesh that tessellates a rectangle.
+    RectangleMesh(MeshDescription const& description, Rectangle<3, Real> const& rectangle);
+
+    // Member access.
+    inline Rectangle<3, Real> const& GetRectangle() const;
+
+protected:
+    void InitializeTCoords();
+    void InitializePositions();
+    void InitializeNormals();
+    void InitializeFrame();
+
+    Rectangle<3, Real> mRectangle;
+
+    // If the client does not request texture coordinates, they will be
+    // computed internally for use in evaluation of the surface geometry.
+    std::vector<Vector2<Real>> mDefaultTCoords;
+};
+
+
+template <typename Real>
+RectangleMesh<Real>::RectangleMesh(MeshDescription const& description,
+    Rectangle<3, Real> const& rectangle)
+    :
+    Mesh<Real>(description, { MeshTopology::RECTANGLE }),
+    mRectangle(rectangle)
+{
+    if (!this->mDescription.constructed)
+    {
+        // The logger system will report these errors in the Mesh constructor.
+        return;
+    }
+
+    if (!this->mTCoords)
+    {
+        mDefaultTCoords.resize(this->mDescription.numVertices);
+        this->mTCoords = mDefaultTCoords.data();
+        this->mTCoordStride = sizeof(Vector2<Real>);
+
+        this->mDescription.allowUpdateFrame = this->mDescription.wantDynamicTangentSpaceUpdate;
+        if (this->mDescription.allowUpdateFrame)
+        {
+            if (!this->mDescription.hasTangentSpaceVectors)
+            {
+                this->mDescription.allowUpdateFrame = false;
+            }
+
+            if (!this->mNormals)
+            {
+                this->mDescription.allowUpdateFrame = false;
+            }
+        }
+    }
+
+    this->ComputeIndices();
+    InitializeTCoords();
+    InitializePositions();
+    if (this->mDescription.allowUpdateFrame)
+    {
+        InitializeFrame();
+    }
+    else if (this->mNormals)
+    {
+        InitializeNormals();
+    }
+}
+
+template <typename Real>
+inline Rectangle<3, Real> const& RectangleMesh<Real>::GetRectangle() const
+{
+    return mRectangle;
+}
+
+template <typename Real>
+void RectangleMesh<Real>::InitializeTCoords()
+{
+    Vector2<Real> tcoord;
+    for (uint32_t r = 0, i = 0; r < this->mDescription.numRows; ++r)
+    {
+        tcoord[1] = (Real)r / (Real)(this->mDescription.numRows - 1);
+        for (uint32_t c = 0; c < this->mDescription.numCols; ++c, ++i)
+        {
+            tcoord[0] = (Real)c / (Real)(this->mDescription.numCols - 1);
+            this->TCoord(i) = tcoord;
+        }
+    }
+}
+
+template <typename Real>
+void RectangleMesh<Real>::InitializePositions()
+{
+    for (uint32_t r = 0, i = 0; r < this->mDescription.numRows; ++r)
+    {
+        for (uint32_t c = 0; c < this->mDescription.numCols; ++c, ++i)
+        {
+            Vector2<Real> tcoord = this->TCoord(i);
+            Real w0 = ((Real)2 * tcoord[0] - (Real)1) * mRectangle.extent[0];
+            Real w1 = ((Real)2 * tcoord[1] - (Real)1) * mRectangle.extent[1];
+            this->Position(i) = mRectangle.center + w0 * mRectangle.axis[0] + w1 * mRectangle.axis[1];
+        }
+    }
+}
+
+template <typename Real>
+void RectangleMesh<Real>::InitializeNormals()
+{
+    Vector3<Real> normal = UnitCross(mRectangle.axis[0], mRectangle.axis[1]);
+    for (uint32_t i = 0; i < this->mDescription.numVertices; ++i)
+    {
+        this->Normal(i) = normal;
+    }
+}
+
+template <typename Real>
+void RectangleMesh<Real>::InitializeFrame()
+{
+    Vector3<Real> normal = UnitCross(mRectangle.axis[0], mRectangle.axis[1]);
+    Vector3<Real> tangent{ (Real)1, (Real)0, (Real)0 };
+    Vector3<Real> bitangent{ (Real)0, (Real)1, (Real)0 };  // = Cross(normal,tangent)
+    for (uint32_t i = 0; i < this->mDescription.numVertices; ++i)
+    {
+        if (this->mNormals)
+        {
+            this->Normal(i) = normal;
+        }
+
+        if (this->mTangents)
+        {
+            this->Tangent(i) = tangent;
+        }
+
+        if (this->mBitangents)
+        {
+            this->Bitangent(i) = bitangent;
+        }
+
+        if (this->mDPDUs)
+        {
+            this->DPDU(i) = tangent;
+        }
+
+        if (this->mDPDVs)
+        {
+            this->DPDV(i) = bitangent;
+        }
+    }
+}
+
+}
diff --git a/Engine/Plugins/Experimental/GeometryProcessing/Source/GeometryAlgorithms/Private/GTEngine/Mathematics/GteRectanglePatchMesh.h b/Engine/Plugins/Experimental/GeometryProcessing/Source/GeometryAlgorithms/Private/GTEngine/Mathematics/GteRectanglePatchMesh.h
new file mode 100644
--- /dev/null
+++ b/Engine/Plugins/Experimental/GeometryProcessing/Source/GeometryAlgorithms/Private/GTEngine/Mathematics/GteRectanglePatchMesh.h
@@ -0,0 +1,226 @@
+// David Eberly, Geometric Tools, Redmond WA 98052
+// Copyright (c) 1998-2019
+// Distributed under the Boost Software License, Version 1.0.
+// http://www.boost.org/LICENSE_1_0.txt
+// http://www.geometrictools.com/License/Boost/LICENSE_1_0.txt
+// File Version: 3.0.4 (2016/08/29)
+
+#pragma once
+
+#include <GTEngine/Mathematics/GteMesh.h>
+#include <GTEngine/Mathematics/GteParametricSurface.h>
+#include <memory>
+
+namespace gte
+{
+
+template <typename Real>
+class RectanglePatchMesh : public Mesh<Real>
+{
+public:
+    // Create a mesh (x(u,v),y(u,v),z(u,v)) defined by the specified surface.
+    // It is required that surface->IsRectangular() return 'true'.
+    RectanglePatchMesh(MeshDescription const& description,
+        std::shared_ptr<ParametricSurface<3, Real>> const& surface);
+
+    // Member access.
+    inline std::shared_ptr<ParametricSurface<3, Real>> const& GetSurface() const;
+
+protected:
+    void InitializeTCoords();
+    void InitializePositions();
+    void InitializeNormals();
+    void InitializeFrame();
+    virtual void UpdatePositions() override;
+    virtual void UpdateNormals() override;
+    virtual void UpdateFrame() override;
+
+    std::shared_ptr<ParametricSurface<3, Real>> mSurface;
+
+    // If the client does not request texture coordinates, they will be
+    // computed internally for use in evaluation of the surface geometry.
+    std::vector<Vector2<Real>> mDefaultTCoords;
+};
+
+
+template <typename Real>
+RectanglePatchMesh<Real>::RectanglePatchMesh(MeshDescription const& description,
+    std::shared_ptr<ParametricSurface<3, Real>> const& surface)
+    :
+    Mesh<Real>(description, { MeshTopology::RECTANGLE }),
+    mSurface(surface)
+{
+    if (!this->mDescription.constructed)
+    {
+        // The logger system will report these errors in the Mesh constructor.
+        mSurface = nullptr;
+        return;
+    }
+
+    if (!mSurface || !mSurface->IsRectangular())
+    {
+        LogError("A nonnull rectangular surface is required.");
+        mSurface = nullptr;
+        this->mDescription.constructed = false;
+        return;
+    }
+
+    if (!this->mTCoords)
+    {
+        mDefaultTCoords.resize(this->mDescription.numVertices);
+        this->mTCoords = mDefaultTCoords.data();
+        this->mTCoordStride = sizeof(Vector2<Real>);
+
+        this->mDescription.allowUpdateFrame = this->mDescription.wantDynamicTangentSpaceUpdate;
+        if (this->mDescription.allowUpdateFrame)
+        {
+            if (!this->mDescription.hasTangentSpaceVectors)
+            {
+                this->mDescription.allowUpdateFrame = false;
+            }
+
+            if (!this->mNormals)
+            {
+                this->mDescription.allowUpdateFrame = false;
+            }
+        }
+    }
+
+    this->ComputeIndices();
+    InitializeTCoords();
+    InitializePositions();
+    if (this->mDescription.allowUpdateFrame)
+    {
+        InitializeFrame();
+    }
+    else if (this->mNormals)
+    {
+        InitializeNormals();
+    }
+}
+
+template <typename Real>
+inline std::shared_ptr<ParametricSurface<3, Real>> const&
+RectanglePatchMesh<Real>::GetSurface() const
+{
+    return mSurface;
+}
+
+template <typename Real>
+void RectanglePatchMesh<Real>::InitializeTCoords()
+{
+    Real uMin = mSurface->GetUMin();
+    Real uDelta = (mSurface->GetUMax() - uMin) / static_cast<Real>(this->mDescription.numCols - 1);
+    Real vMin = mSurface->GetVMin();
+    Real vDelta = (mSurface->GetVMax() - vMin) / static_cast<Real>(this->mDescription.numRows - 1);
+    Vector2<Real> tcoord;
+    for (uint32_t r = 0, i = 0; r < this->mDescription.numRows; ++r)
+    {
+        tcoord[1] = vMin + vDelta * (Real)r;
+        for (uint32_t c = 0; c < this->mDescription.numCols; ++c, ++i)
+        {
+            tcoord[0] = uMin + uDelta * (Real)c;
+            this->TCoord(i) = tcoord;
+        }
+    }
+}
+
+template <typename Real>
+void RectanglePatchMesh<Real>::InitializePositions()
+{
+    for (uint32_t r = 0, i = 0; r < this->mDescription.numRows; ++r)
+    {
+        for (uint32_t c = 0; c < this->mDescription.numCols; ++c, ++i)
+        {
+            Vector2<Real> tcoord = this->TCoord(i);
+            this->Position(i) = mSurface->GetPosition(tcoord[0], tcoord[1]);
+        }
+    }
+}
+
+template <typename Real>
+void RectanglePatchMesh<Real>::InitializeNormals()
+{
+    for (uint32_t r = 0, i = 0; r < this->mDescription.numRows; ++r)
+    {
+        for (uint32_t c = 0; c < this->mDescription.numCols; ++c, ++i)
+        {
+            Vector2<Real> tcoord = this->TCoord(i);
+            Vector3<Real> values[6];
+            mSurface->Evaluate(tcoord[0], tcoord[1], 1, values);
+            Normalize(values[1], true);
+            Normalize(values[2], true);
+            this->Normal(i) = UnitCross(values[1], values[2], true);
+        }
+    }
+}
+
+template <typename Real>
+void RectanglePatchMesh<Real>::InitializeFrame()
+{
+    Vector3<Real> normal, tangent, bitangent;
+    for (unsigned int r = 0, i = 0; r < this->mDescription.numRows; ++r)
+    {
+        for (unsigned int c = 0; c < this->mDescription.numCols; ++c, ++i)
+        {
+            Vector2<Real> tcoord = this->TCoord(i);
+            Vector3<Real> values[6];
+            mSurface->Evaluate(tcoord[0], tcoord[1], 1, values);
+            Normalize(values[1], true);
+            Normalize(values[2], true);
+
+            if (this->mDPDUs)
+            {
+                this->DPDU(i) = values[1];
+            }
+            if (this->mDPDVs)
+            {
+                this->DPDV(i) = values[2];
+            }
+
+            ComputeOrthogonalComplement(2, &values[1], true);
+
+            if (this->mNormals)
+            {
+                this->Normal(i) = values[3];
+            }
+            if (this->mTangents)
+            {
+                this->Tangent(i) = values[1];
+            }
+            if (this->mBitangents)
+            {
+                this->Bitangent(i) = values[2];
+            }
+        }
+    }
+}
+
+template <typename Real>
+void RectanglePatchMesh<Real>::UpdatePositions()
+{
+    if (mSurface)
+    {
+        InitializePositions();
+    }
+}
+
+template <typename Real>
+void RectanglePatchMesh<Real>::UpdateNormals()
+{
+    if (mSurface)
+    {
+        InitializeNormals();
+    }
+}
+
+template <typename Real>
+void RectanglePatchMesh<Real>::UpdateFrame()
+{
+    if (mSurface)
+    {
+        InitializeFrame();
+    }
+}
+
+}
diff --git a/Engine/Plugins/Experimental/GeometryProcessing/Source/GeometryAlgorithms/Private/GTEngine/Mathematics/GteRevolutionMesh.h b/Engine/Plugins/Experimental/GeometryProcessing/Source/GeometryAlgorithms/Private/GTEngine/Mathematics/GteRevolutionMesh.h
new file mode 100644
--- /dev/null
+++ b/Engine/Plugins/Experimental/GeometryProcessing/Source/GeometryAlgorithms/Private/GTEngine/Mathematics/GteRevolutionMesh.h
@@ -0,0 +1,347 @@
+// David Eberly, Geometric Tools, Redmond WA 98052
+// Copyright (c) 1998-2019
+// Distributed under the Boost Software License, Version 1.0.
+// http://www.boost.org/LICENSE_1_0.txt
+// http://www.geometrictools.com/License/Boost/LICENSE_1_0.txt
+// File Version: 3.0.4 (2018/10/05)
+
+#pragma once
+
+#include <GTEngine/Mathematics/GteMesh.h>
+#include <GTEngine/Mathematics/GteParametricCurve.h>
+#include <memory>
+
+namespace gte
+{
+
+template <typename Real>
+class RevolutionMesh : public Mesh<Real>
+{
+public:
+    // The axis of revolution is the z-axis.  The curve of revolution is
+    // p(t) = (x(t),z(t)), where t in [tmin,tmax], x(t) > 0 for t in
+    // (tmin,tmax), x(tmin) >= 0, and x(tmax) >= 0.  The values tmin and tmax
+    // are those for the curve object passed to the constructor.  The curve
+    // must be non-self-intersecting, except possibly at its endpoints.  The
+    // curve is closed when p(tmin) = p(tmax), in which case the surface of
+    // revolution has torus topology.  The curve is open when p(tmin) !=
+    // p(tmax).  For an open curve, define x0 = x(tmin) and x1 = x(tmax).  The
+    // surface has cylinder topology when x0 > 0 and x1 > 0, disk topology
+    // when exactly one of x0 or x1 is zero, or sphere topology when x0 and
+    // x1 are both zero.  However, to simplify the design, the mesh is always
+    // built using cylinder topology.  The row samples correspond to curve
+    // points and the column samples correspond to the points on the circles
+    // of revolution.
+    RevolutionMesh(MeshDescription const& description,
+        std::shared_ptr<ParametricCurve<2, Real>> const& curve,
+        bool sampleByArcLength = false);
+
+    // Member access.
+    inline std::shared_ptr<ParametricCurve<2, Real>> const& GetCurve() const;
+    inline bool IsSampleByArcLength() const;
+
+private:
+    void CreateSampler();
+    void InitializeTCoords();
+    virtual void UpdatePositions() override;
+    void UpdateCylinderPositions();
+    void UpdateTorusPositions();
+    void UpdateDiskPositions();
+    void UpdateSpherePositions();
+
+    std::shared_ptr<ParametricCurve<2, Real>> mCurve;
+    bool mSampleByArcLength;
+    std::vector<Real> mCosAngle, mSinAngle;
+    std::function<Real(unsigned int)> mTSampler;
+    std::vector<Vector3<Real>> mSamples;
+
+    // If the client does not request texture coordinates, they will be
+    // computed internally for use in evaluation of the surface geometry.
+    std::vector<Vector2<Real>> mDefaultTCoords;
+};
+
+
+template <typename Real>
+RevolutionMesh<Real>::RevolutionMesh(MeshDescription const& description,
+    std::shared_ptr<ParametricCurve<2, Real>> const& curve, bool sampleByArcLength)
+    :
+    Mesh<Real>(description,
+        { MeshTopology::CYLINDER, MeshTopology::TORUS, MeshTopology::DISK, MeshTopology::SPHERE }),
+    mCurve(curve),
+    mSampleByArcLength(sampleByArcLength)
+{
+    if (!this->mDescription.constructed)
+    {
+        // The logger system will report these errors in the Mesh constructor.
+        mCurve = nullptr;
+        return;
+    }
+
+    if (!mCurve)
+    {
+        LogWarning("A nonnull revolution curve is required.");
+        this->mDescription.constructed = false;
+        return;
+    }
+
+    // The four supported topologies all wrap around in the column direction.
+    mCosAngle.resize(this->mDescription.numCols + 1);
+    mSinAngle.resize(this->mDescription.numCols + 1);
+    Real invRadialSamples = (Real)1 / (Real)this->mDescription.numCols;
+    for (unsigned int c = 0; c < this->mDescription.numCols; ++c)
+    {
+        Real angle = c * invRadialSamples * (Real)GTE_C_TWO_PI;
+        mCosAngle[c] = std::cos(angle);
+        mSinAngle[c] = std::sin(angle);
+    }
+    mCosAngle[this->mDescription.numCols] = mCosAngle[0];
+    mSinAngle[this->mDescription.numCols] = mSinAngle[0];
+
+    CreateSampler();
+
+    if (!this->mTCoords)
+    {
+        mDefaultTCoords.resize(this->mDescription.numVertices);
+        this->mTCoords = mDefaultTCoords.data();
+        this->mTCoordStride = sizeof(Vector2<Real>);
+
+        this->mDescription.allowUpdateFrame = this->mDescription.wantDynamicTangentSpaceUpdate;
+        if (this->mDescription.allowUpdateFrame)
+        {
+            if (!this->mDescription.hasTangentSpaceVectors)
+            {
+                this->mDescription.allowUpdateFrame = false;
+            }
+
+            if (!this->mNormals)
+            {
+                this->mDescription.allowUpdateFrame = false;
+            }
+        }
+    }
+
+    this->ComputeIndices();
+    InitializeTCoords();
+    UpdatePositions();
+    if (this->mDescription.allowUpdateFrame)
+    {
+        this->UpdateFrame();
+    }
+    else if (this->mNormals)
+    {
+        this->UpdateNormals();
+    }
+}
+
+template <typename Real>
+inline std::shared_ptr<ParametricCurve<2, Real>> const& RevolutionMesh<Real>::GetCurve() const
+{
+    return mCurve;
+}
+
+template <typename Real>
+inline bool RevolutionMesh<Real>::IsSampleByArcLength() const
+{
+    return mSampleByArcLength;
+}
+
+template <typename Real>
+void RevolutionMesh<Real>::CreateSampler()
+{
+    if (this->mDescription.topology == MeshTopology::CYLINDER
+        || this->mDescription.topology == MeshTopology::TORUS)
+    {
+        mSamples.resize(this->mDescription.rMax + 1);
+    }
+    else if (this->mDescription.topology == MeshTopology::DISK)
+    {
+        mSamples.resize(this->mDescription.rMax + 2);
+    }
+    else if (this->mDescription.topology == MeshTopology::SPHERE)
+    {
+        mSamples.resize(this->mDescription.rMax + 3);
+    }
+
+    Real invDenom = ((Real)1) / (Real)(mSamples.size() - 1);
+    if (mSampleByArcLength)
+    {
+        Real factor = mCurve->GetTotalLength() * invDenom;
+        mTSampler = [this, factor](unsigned int i)
+        {
+            return mCurve->GetTime(i * factor);
+        };
+    }
+    else
+    {
+        Real factor = (mCurve->GetTMax() - mCurve->GetTMin()) * invDenom;
+        mTSampler = [this, factor](unsigned int i)
+        {
+            return mCurve->GetTMin() + i * factor;
+        };
+    }
+}
+
+template <typename Real>
+void RevolutionMesh<Real>::InitializeTCoords()
+{
+    Vector2<Real>tcoord;
+
+    switch (this->mDescription.topology)
+    {
+    case MeshTopology::CYLINDER:
+    {
+        for (unsigned int r = 0, i = 0; r < this->mDescription.numRows; ++r)
+        {
+            tcoord[1] = (Real)r / (Real)(this->mDescription.numRows - 1);
+            for (unsigned int c = 0; c <= this->mDescription.numCols; ++c, ++i)
+            {
+                tcoord[0] = (Real)c / (Real)this->mDescription.numCols;
+                this->TCoord(i) = tcoord;
+            }
+        }
+        break;
+    }
+    case MeshTopology::TORUS:
+    {
+        for (unsigned int r = 0, i = 0; r <= this->mDescription.numRows; ++r)
+        {
+            tcoord[1] = (Real)r / (Real)this->mDescription.numRows;
+            for (unsigned int c = 0; c <= this->mDescription.numCols; ++c, ++i)
+            {
+                tcoord[0] = (Real)c / (Real)this->mDescription.numCols;
+                this->TCoord(i) = tcoord;
+            }
+        }
+        break;
+    }
+    case MeshTopology::DISK:
+    {
+        Vector2<Real> origin{ (Real)0.5, (Real)0.5 };
+        unsigned int i = 0;
+        for (unsigned int r = 0; r < this->mDescription.numRows; ++r)
+        {
+            Real radius = (Real)(r + 1) / (Real)(2 * this->mDescription.numRows);
+            radius = std::min(radius, (Real)0.5);
+            for (unsigned int c = 0; c <= this->mDescription.numCols; ++c, ++i)
+            {
+                Real angle = (Real)GTE_C_TWO_PI * (Real)c / (Real)this->mDescription.numCols;
+                this->TCoord(i) = { radius * std::cos(angle), radius * std::sin(angle) };
+            }
+        }
+        this->TCoord(i) = origin;
+        break;
+    }
+    case MeshTopology::SPHERE:
+    {
+        unsigned int i = 0;
+        for (unsigned int r = 0; r < this->mDescription.numRows; ++r)
+        {
+            tcoord[1] = (Real)r / (Real)(this->mDescription.numRows - 1);
+            for (unsigned int c = 0; c <= this->mDescription.numCols; ++c, ++i)
+            {
+                tcoord[0] = (Real)c / (Real)this->mDescription.numCols;
+                this->TCoord(i) = tcoord;
+            }
+        }
+        this->TCoord(i++) = { (Real)0.5, (Real)0 };
+        this->TCoord(i) = { (Real)0.5, (Real)1 };
+        break;
+    }
+    default:
+        // Invalid topology is reported by the Mesh constructor, so there is
+        // no need to log a message here.
+        break;
+    }
+}
+
+template <typename Real>
+void RevolutionMesh<Real>::UpdatePositions()
+{
+    unsigned int const numSamples = static_cast<unsigned int>(mSamples.size());
+    for (unsigned int i = 0; i < numSamples; ++i)
+    {
+        Real t = mTSampler(i);
+        Vector2<Real> position = mCurve->GetPosition(t);
+        mSamples[i][0] = position[0];
+        mSamples[i][1] = (Real)0;
+        mSamples[i][2] = position[1];
+    }
+
+    switch (this->mDescription.topology)
+    {
+    case MeshTopology::CYLINDER:
+        UpdateCylinderPositions();
+        break;
+    case MeshTopology::TORUS:
+        UpdateTorusPositions();
+        break;
+    case MeshTopology::DISK:
+        UpdateDiskPositions();
+        break;
+    case MeshTopology::SPHERE:
+        UpdateSpherePositions();
+        break;
+    default:
+        break;
+    }
+}
+
+template <typename Real>
+void RevolutionMesh<Real>::UpdateCylinderPositions()
+{
+    for (unsigned int r = 0, i = 0; r <= this->mDescription.rMax; ++r)
+    {
+        Real radius = mSamples[r][0];
+        for (unsigned int c = 0; c <= this->mDescription.cMax; ++c, ++i)
+        {
+            this->Position(i) = { radius * mCosAngle[c], radius * mSinAngle[c], mSamples[r][2] };
+        }
+    }
+}
+
+template <typename Real>
+void RevolutionMesh<Real>::UpdateTorusPositions()
+{
+    for (unsigned int r = 0, i = 0; r <= this->mDescription.rMax; ++r)
+    {
+        Real radius = mSamples[r][0];
+        for (unsigned int c = 0; c <= this->mDescription.cMax; ++c, ++i)
+        {
+            this->Position(i) = { radius * mCosAngle[c], radius * mSinAngle[c], mSamples[r][2] };
+        }
+    }
+}
+
+template <typename Real>
+void RevolutionMesh<Real>::UpdateDiskPositions()
+{
+    for (unsigned int r = 0, rp1 = 1, i = 0; r <= this->mDescription.rMax; ++r, ++rp1)
+    {
+        Real radius = mSamples[rp1][0];
+        for (unsigned int c = 0; c <= this->mDescription.cMax; ++c, ++i)
+        {
+            this->Position(i) = { radius * mCosAngle[c], radius * mSinAngle[c], mSamples[rp1][2] };
+        }
+    }
+
+    this->Position(this->mDescription.numVertices - 1) = { (Real)0, (Real)0, mSamples.front()[2] };
+}
+
+template <typename Real>
+void RevolutionMesh<Real>::UpdateSpherePositions()
+{
+    for (unsigned int r = 0, rp1 = 1, i = 0; r <= this->mDescription.rMax; ++r, ++rp1)
+    {
+        Real radius = mSamples[rp1][0];
+        for (unsigned int c = 0; c <= this->mDescription.cMax; ++c, ++i)
+        {
+            this->Position(i) = { radius * mCosAngle[c], radius * mSinAngle[c], mSamples[rp1][2] };
+        }
+    }
+
+    this->Position(this->mDescription.numVertices - 2) = { (Real)0, (Real)0, mSamples.front()[2] };
+    this->Position(this->mDescription.numVertices - 1) = { (Real)0, (Real)0, mSamples.back()[2] };
+}
+
+
+}
diff --git a/Engine/Plugins/Experimental/GeometryProcessing/Source/GeometryAlgorithms/Private/GTEngine/Mathematics/GteRootsBisection.h b/Engine/Plugins/Experimental/GeometryProcessing/Source/GeometryAlgorithms/Private/GTEngine/Mathematics/GteRootsBisection.h
new file mode 100644
--- /dev/null
+++ b/Engine/Plugins/Experimental/GeometryProcessing/Source/GeometryAlgorithms/Private/GTEngine/Mathematics/GteRootsBisection.h
@@ -0,0 +1,172 @@
+// David Eberly, Geometric Tools, Redmond WA 98052
+// Copyright (c) 1998-2019
+// Distributed under the Boost Software License, Version 1.0.
+// http://www.boost.org/LICENSE_1_0.txt
+// http://www.geometrictools.com/License/Boost/LICENSE_1_0.txt
+// File Version: 3.0.0 (2016/06/19)
+
+#pragma once
+
+#include <GTEngine/GTEngineDEF.h>
+#include <functional>
+
+// Compute a root of a function F(t) on an interval [t0, t1].  The caller
+// specifies the maximum number of iterations, in case you want limited
+// accuracy for the root.  However, the function is designed for native types
+// (Real = float/double).  If you specify a sufficiently large number of
+// iterations, the root finder bisects until either F(t) is identically zero
+// [a condition dependent on how you structure F(t) for evaluation] or the
+// midpoint (t0 + t1)/2 rounds numerically to tmin or tmax.  Of course, it
+// is required that t0 < t1.  The return value of Find is:
+//   0: F(t0)*F(t1) > 0, we cannot determine a root
+//   1: F(t0) = 0 or F(t1) = 0
+//   2..maxIterations:  the number of bisections plus one
+//   maxIterations+1:  the loop executed without a break (no convergence)
+
+namespace gte
+{
+
+template <typename Real>
+class RootsBisection
+{
+public:
+    // Use this function when F(t0) and F(t1) are not already known.
+    static unsigned int Find(std::function<Real(Real)> const& F, Real t0,
+        Real t1, unsigned int maxIterations, Real& root);
+
+    // If f0 = F(t0) and f1 = F(t1) are already known, pass them to the
+    // bisector.  This is useful when |f0| or |f1| is infinite, and you can
+    // pass sign(f0) or sign(f1) rather than then infinity because the
+    // bisector cares only about the signs of f.
+    static unsigned int Find(std::function<Real(Real)> const& F, Real t0,
+        Real t1, Real f0, Real f1, unsigned int maxIterations, Real& root);
+};
+
+
+template <typename Real>
+unsigned int RootsBisection<Real>::Find(std::function<Real(Real)> const& F,
+    Real t0, Real t1, unsigned int maxIterations, Real& root)
+{
+    if (t0 < t1)
+    {
+        // Test the endpoints to see whether F(t) is zero.
+        Real f0 = F(t0);
+        if (f0 == (Real)0)
+        {
+            root = t0;
+            return 1;
+        }
+
+        Real f1 = F(t1);
+        if (f1 == (Real)0)
+        {
+            root = t1;
+            return 1;
+        }
+
+        if (f0*f1 > (Real)0)
+        {
+            // It is not known whether the interval bounds a root.
+            return 0;
+        }
+
+        unsigned int i;
+        for (i = 2; i <= maxIterations; ++i)
+        {
+            root = ((Real)0.5) * (t0 + t1);
+            if (root == t0 || root == t1)
+            {
+                // The numbers t0 and t1 are consecutive floating-point
+                // numbers.
+                break;
+            }
+
+            Real fm = F(root);
+            Real product = fm * f0;
+            if (product < (Real)0)
+            {
+                t1 = root;
+                f1 = fm;
+            }
+            else if (product > (Real)0)
+            {
+                t0 = root;
+                f0 = fm;
+            }
+            else
+            {
+                break;
+            }
+        }
+        return i;
+    }
+    else
+    {
+        return 0;
+    }
+}
+
+template <typename Real>
+unsigned int RootsBisection<Real>::Find(std::function<Real(Real)> const& F,
+    Real t0, Real t1, Real f0, Real f1, unsigned int maxIterations,
+    Real& root)
+{
+    if (t0 < t1)
+    {
+        // Test the endpoints to see whether F(t) is zero.
+        if (f0 == (Real)0)
+        {
+            root = t0;
+            return 1;
+        }
+
+        if (f1 == (Real)0)
+        {
+            root = t1;
+            return 1;
+        }
+
+        if (f0*f1 > (Real)0)
+        {
+            // It is not known whether the interval bounds a root.
+            return 0;
+        }
+
+        unsigned int i;
+        for (i = 2; i <= maxIterations; ++i)
+        {
+            root = ((Real)0.5) * (t0 + t1);
+            if (root == t0 || root == t1)
+            {
+                // The numbers t0 and t1 are consecutive floating-point
+                // numbers.
+                break;
+            }
+
+            Real fm = F(root);
+            Real product = fm * f0;
+            if (product < (Real)0)
+            {
+                t1 = root;
+                f1 = fm;
+            }
+            else if (product >(Real)0)
+            {
+                t0 = root;
+                f0 = fm;
+            }
+            else
+            {
+                break;
+            }
+        }
+        return i;
+    }
+    else
+    {
+        return 0;
+    }
+}
+
+
+}
diff --git a/Engine/Plugins/Experimental/GeometryProcessing/Source/GeometryAlgorithms/Private/GTEngine/Mathematics/GteRootsBrentsMethod.h b/Engine/Plugins/Experimental/GeometryProcessing/Source/GeometryAlgorithms/Private/GTEngine/Mathematics/GteRootsBrentsMethod.h
new file mode 100644
--- /dev/null
+++ b/Engine/Plugins/Experimental/GeometryProcessing/Source/GeometryAlgorithms/Private/GTEngine/Mathematics/GteRootsBrentsMethod.h
@@ -0,0 +1,232 @@
+// David Eberly, Geometric Tools, Redmond WA 98052
+// Copyright (c) 1998-2019
+// Distributed under the Boost Software License, Version 1.0.
+// http://www.boost.org/LICENSE_1_0.txt
+// http://www.geometrictools.com/License/Boost/LICENSE_1_0.txt
+// File Version: 3.0.0 (2016/06/19)
+
+#pragma once
+
+#include <GTEngine/GTEngineDEF.h>
+#include <cmath>
+#include <functional>
+
+// This is an implementation of Brent's Method for computing a root of a
+// function on an interval [t0,t1] for which F(t0)*F(t1) < 0.  The method
+// uses inverse quadratic interpolation to generate a root estimate but
+// falls back to inverse linear interpolation (secant method) if
+// necessary.  Moreover, based on previous iterates, the method will fall
+// back to bisection when it appears the interpolated estimate is not of
+// sufficient quality.
+//
+//   maxIterations:
+//       The maximum number of iterations used to locate a root.  This
+//       should be positive.
+//   negFTolerance, posFTolerance:
+//       The root estimate t is accepted when the function value F(t)
+//       satisfies negFTolerance <= F(t) <= posFTolerance.  The values
+//       must satisfy:  negFTolerance <= 0, posFTolerance >= 0.
+//   stepTTolerance:
+//       Brent's Method requires additional tests before an interpolated
+//       t-value is accepted as the next root estimate.  One of these
+//       tests compares the difference of consecutive iterates and
+//       requires it to be larger than a user-specified t-tolerance (to
+//       ensure progress is made).  This parameter is that tolerance
+//       and should be nonnegative.
+//   convTTolerance:
+//       The root search is allowed to terminate when the current
+//       subinterval [tsub0,tsub1] is sufficiently small, say,
+//       |tsub1 - tsub0| <= tolerance.  This parameter is that tolerance
+//       and should be nonnegative.
+
+namespace gte
+{
+
+template <typename Real>
+class RootsBrentsMethod
+{
+public:
+    // It is necessary that F(t0)*F(t1) <= 0, in which case the function
+    // returns 'true' and the 'root' is valid; otherwise, the function returns
+    // 'false' and 'root' is invalid (do not use it).  When F(t0)*F(t1) > 0,
+    // the interval may very well contain a root but we cannot know that.  The
+    // function also returns 'false' if t0 >= t1.
+
+    static bool Find(std::function<Real(Real)> const& F, Real t0, Real t1,
+        unsigned int maxIterations, Real negFTolerance, Real posFTolerance,
+        Real stepTTolerance, Real convTTolerance, Real& root);
+};
+
+
+template <typename Real>
+bool RootsBrentsMethod<Real>::Find(std::function<Real(Real)> const& F,
+    Real t0, Real t1, unsigned int maxIterations, Real negFTolerance,
+    Real posFTolerance, Real stepTTolerance, Real convTTolerance, Real& root)
+{
+    // Parameter validation.
+    if (t1 <= t0
+        || maxIterations == 0
+        || negFTolerance > (Real)0
+        || posFTolerance < (Real)0
+        || stepTTolerance < (Real)0
+        || convTTolerance < (Real)0)
+    {
+        // The input is invalid.
+        return false;
+    }
+
+    Real f0 = F(t0);
+    if (negFTolerance <= f0 && f0 <= posFTolerance)
+    {
+        // This endpoint is an approximate root that satisfies the function
+        // tolerance.
+        root = t0;
+        return true;
+    }
+
+    Real f1 = F(t1);
+    if (negFTolerance <= f1 && f1 <= posFTolerance)
+    {
+        // This endpoint is an approximate root that satisfies the function
+        // tolerance.
+        root = t1;
+        return true;
+    }
+
+    if (f0*f1 > (Real)0)
+    {
+        // The input interval must bound a root.
+        return false;
+    }
+
+    if (std::abs(f0) < std::abs(f1))
+    {
+        // Swap t0 and t1 so that |F(t1)| <= |F(t0)|.  The number t1 is
+        // considered to be the best estimate of the root.
+        std::swap(t0, t1);
+        std::swap(f0, f1);
+    }
+
+    // Initialize values for the root search.
+    Real t2 = t0, t3 = t0, f2 = f0;
+    bool prevBisected = true;
+
+    // The root search.
+    for (unsigned int i = 0; i < maxIterations; ++i)
+    {
+        Real fDiff01 = f0 - f1, fDiff02 = f0 - f2, fDiff12 = f1 - f2;
+        Real invFDiff01 = ((Real)1) / fDiff01;
+        Real s;
+        if (fDiff02 != (Real)0 && fDiff12 != (Real)0)
+        {
+            // Use inverse quadratic interpolation.
+            Real infFDiff02 = ((Real)1) / fDiff02;
+            Real invFDiff12 = ((Real)1) / fDiff12;
+            s =
+                t0 * f1 * f2 * invFDiff01 * infFDiff02 -
+                t1 * f0 * f2 * invFDiff01 * invFDiff12 +
+                t2 * f0 * f1 * infFDiff02 * invFDiff12;
+        }
+        else
+        {
+            // Use inverse linear interpolation (secant method).
+            s = (t1 * f0 - t0 * f1) * invFDiff01;
+        }
+
+        // Compute values need in the accept-or-reject tests.
+        Real tDiffSAvr = s - ((Real)0.75) * t0 - ((Real)0.25) * t1;
+        Real tDiffS1 = s - t1;
+        Real absTDiffS1 = std::abs(tDiffS1);
+        Real absTDiff12 = std::abs(t1 - t2);
+        Real absTDiff23 = std::abs(t2 - t3);
+
+        bool currBisected = false;
+        if (tDiffSAvr * tDiffS1 > (Real)0)
+        {
+            // The value s is not between 0.75*t0 + 0.25*t1 and t1.  NOTE:
+            // The algorithm sometimes has t0 < t1 but sometimes t1 < t0, so
+            // the between-ness test does not use simple comparisons.
+            currBisected = true;
+        }
+        else if (prevBisected)
+        {
+            // The first of Brent's tests to determine whether to accept the
+            // interpolated s-value.
+            currBisected =
+                (absTDiffS1 >= ((Real)0.5) * absTDiff12) ||
+                (absTDiff12 <= stepTTolerance);
+        }
+        else
+        {
+            // The second of Brent's tests to determine whether to accept the
+            // interpolated s-value.
+            currBisected =
+                (absTDiffS1 >= ((Real)0.5) * absTDiff23) ||
+                (absTDiff23 <= stepTTolerance);
+        }
+
+        if (currBisected)
+        {
+            // One of the additional tests failed, so reject the interpolated
+            // s-value and use bisection instead.
+            s = ((Real)0.5) * (t0 + t1);
+            if (s == t0 || s == t1)
+            {
+                // The numbers t0 and t1 are consecutive floating-point
+                // numbers.
+                root = s;
+                return true;
+            }
+            prevBisected = true;
+        }
+        else
+        {
+            prevBisected = false;
+        }
+
+        // Evaluate the function at the new estimate and test for convergence.
+        Real fs = F(s);
+        if (negFTolerance <= fs && fs <= posFTolerance)
+        {
+            root = s;
+            return true;
+        }
+
+        // Update the subinterval to include the new estimate as an endpoint.
+        t3 = t2;
+        t2 = t1;
+        f2 = f1;
+        if (f0 * fs < (Real)0)
+        {
+            t1 = s;
+            f1 = fs;
+        }
+        else
+        {
+            t0 = s;
+            f0 = fs;
+        }
+
+        // Allow the algorithm to terminate when the subinterval is
+        // sufficiently small.
+        if (std::abs(t1 - t0) <= convTTolerance)
+        {
+            root = t1;
+            return true;
+        }
+
+        // A loop invariant is that t1 is the root estimate, F(t0)*F(t1) < 0,
+        // and |F(t1)| <= |F(t0)|.
+        if (std::abs(f0) < std::abs(f1))
+        {
+            std::swap(t0, t1);
+            std::swap(f0, f1);
+        }
+    }
+
+    // Failed to converge in the specified number of iterations.
+    return false;
+}
+
+
+}
