move GeometryProcessing plugin from NotForLicensees to Experimental

diff --git a/Engine/Plugins/Experimental/GeometryProcessing/Source/GeometryAlgorithms/Private/GTEngine/Mathematics/GteRootsPolynomial.h b/Engine/Plugins/Experimental/GeometryProcessing/Source/GeometryAlgorithms/Private/GTEngine/Mathematics/GteRootsPolynomial.h
new file mode 100644
--- /dev/null
+++ b/Engine/Plugins/Experimental/GeometryProcessing/Source/GeometryAlgorithms/Private/GTEngine/Mathematics/GteRootsPolynomial.h
@@ -0,0 +1,1158 @@
+// David Eberly, Geometric Tools, Redmond WA 98052
+// Copyright (c) 1998-2019
+// Distributed under the Boost Software License, Version 1.0.
+// http://www.boost.org/LICENSE_1_0.txt
+// http://www.geometrictools.com/License/Boost/LICENSE_1_0.txt
+// File Version: 3.0.1 (2018/10/05)
+
+#pragma once
+
+#include <GTEngine/GTEngineDEF.h>
+#include <cmath>
+#include <map>
+#include <vector>
+
+// The Find functions return the number of roots, if any, and this number
+// of elements of the outputs are valid.  If the polynomial is identically
+// zero, Find returns 1.
+//
+// Some root-bounding algorithms for real-valued roots are mentioned next for
+// the polynomial p(t) = c[0] + c[1]*t + ... + c[d-1]*t^{d-1} + c[d]*t^d.
+//
+// 1. The roots must be contained by the interval [-M,M] where
+//   M = 1 + max{|c[0]|, ..., |c[d-1]|}/|c[d]| >= 1
+// is called the Cauchy bound.
+//
+// 2. You may search for roots in the interval [-1,1].  Define
+//   q(t) = t^d*p(1/t) = c[0]*t^d + c[1]*t^{d-1} + ... + c[d-1]*t + c[d]
+// The roots of p(t) not in [-1,1] are the roots of q(t) in [-1,1].
+//
+// 3. Between two consecutive roots of the derivative p'(t), say, r0 < r1,
+// the function p(t) is strictly monotonic on the open interval (r0,r1).
+// If additionally, p(r0) * p(r1) <= 0, then p(x) has a unique root on
+// the closed interval [r0,r1].  Thus, one can compute the derivatives
+// through order d for p(t), find roots for the derivative of order k+1,
+// then use these to bound roots for the derivative of order k.
+//
+// 4. Sturm sequences of polynomials may be used to determine bounds on the
+// roots.  This is a more sophisticated approach to root bounding than item 3.
+// Moreover, a Sturm sequence allows you to compute the number of real-valued
+// roots on a specified interval.
+//
+// 5. For the low-degree Solve* functions, see
+// http://www.geometrictools.com/Documentation/LowDegreePolynomialRoots.pdf
+
+namespace gte
+{
+
+template <typename Real>
+class RootsPolynomial
+{
+public:
+    // Low-degree root finders.  These use exact rational arithmetic for
+    // theoretically correct root classification.  The roots themselves are
+    // computed with mixed types (rational and floating-point arithmetic).
+    // The Rational type must support rational arithmetic (+, -, *, /); for
+    // example, BSRational<UIntegerAP32> suffices.  The Rational class must
+    // have single-input constructors where the input is type Real.  This
+    // ensures you can call the Solve* functions with floating-point inputs;
+    // they will be converted to Rational implicitly.  The highest-order
+    // coefficients must be nonzero (p2 != 0 for quadratic, p3 != 0 for
+    // cubic, and p4 != 0 for quartic).
+
+    template <typename Rational>
+    static void SolveQuadratic(Rational const& p0, Rational const& p1,
+        Rational const& p2, std::map<Real, int>& rmMap);
+
+    template <typename Rational>
+    static void SolveCubic(Rational const& p0, Rational const& p1,
+        Rational const& p2, Rational const& p3, std::map<Real, int>& rmMap);
+
+    template <typename Rational>
+    static void SolveQuartic(Rational const& p0, Rational const& p1,
+        Rational const& p2, Rational const& p3, Rational const& p4,
+        std::map<Real, int>& rmMap);
+
+    // Return only the number of real-valued roots and their multiplicities.
+    // info.size() is the number of real-valued roots and info[i] is the
+    // multiplicity of root corresponding to index i.
+    template <typename Rational>
+    static void GetRootInfoQuadratic(Rational const& p0, Rational const& p1,
+        Rational const& p2, std::vector<int>& info);
+
+    template <typename Rational>
+    static void GetRootInfoCubic(Rational const& p0, Rational const& p1,
+        Rational const& p2, Rational const& p3, std::vector<int>& info);
+
+    template <typename Rational>
+    static void GetRootInfoQuartic(Rational const& p0, Rational const& p1,
+        Rational const& p2, Rational const& p3, Rational const& p4,
+        std::vector<int>& info);
+
+
+    // General equations: sum_{i=0}^{d} c(i)*t^i = 0.  The input array 'c'
+    // must have at least d+1 elements and the output array 'root' must have
+    // at least d elements.
+
+    // Find the roots on (-infinity,+infinity).
+    static int Find(int degree, Real const* c, unsigned int maxIterations,
+        Real* roots);
+
+    // If you know that p(tmin) * p(tmax) <= 0, then there must be at least
+    // one root in [tmin, tmax].  Compute it using bisection.
+    static bool Find(int degree, Real const* c, Real tmin, Real tmax,
+        unsigned int maxIterations, Real& root);
+
+private:
+    // Support for the Solve* functions.
+    template <typename Rational>
+    static void SolveDepressedQuadratic(Rational const& c0,
+        std::map<Rational, int>& rmMap);
+
+    template <typename Rational>
+    static void SolveDepressedCubic(Rational const& c0, Rational const& c1,
+        std::map<Rational, int>& rmMap);
+
+    template <typename Rational>
+    static void SolveDepressedQuartic(Rational const& c0, Rational const& c1,
+        Rational const& c2, std::map<Rational, int>& rmMap);
+
+    template <typename Rational>
+    static void SolveBiquadratic(Rational const& c0, Rational const& c2,
+        std::map<Rational, int>& rmMap);
+
+    // Support for the GetNumRoots* functions.
+    template <typename Rational>
+    static void GetRootInfoDepressedQuadratic(Rational const& c0,
+        std::vector<int>& info);
+
+    template <typename Rational>
+    static void GetRootInfoDepressedCubic(Rational const& c0,
+        Rational const& c1, std::vector<int>& info);
+
+    template <typename Rational>
+    static void GetRootInfoDepressedQuartic(Rational const& c0,
+        Rational const& c1, Rational const& c2, std::vector<int>& info);
+
+    template <typename Rational>
+    static void GetRootInfoBiquadratic(Rational const& c0,
+        Rational const& c2, std::vector<int>& info);
+
+    // Support for the Find functions.
+    static int FindRecursive(int degree, Real const* c, Real tmin, Real tmax,
+        unsigned int maxIterations, Real* roots);
+
+    static Real Evaluate(int degree, Real const* c, Real t);
+};
+
+// FOR INTERNAL USE ONLY.  Do not define GTE_ROOTS_LOW_DEGREE_UNIT_TEST in
+// your own code.
+#if defined(GTE_ROOTS_LOW_DEGREE_UNIT_TEST)
+extern void RootsLowDegreeBlock(int);
+#define GTE_ROOTS_LOW_DEGREE_BLOCK(block) RootsLowDegreeBlock(block)
+#else
+#define GTE_ROOTS_LOW_DEGREE_BLOCK(block)
+#endif
+
+
+template <typename Real>
+template <typename Rational>
+void RootsPolynomial<Real>::SolveQuadratic(Rational const& p0,
+    Rational const& p1, Rational const& p2, std::map<Real, int>& rmMap)
+{
+    Rational const rat2 = 2;
+    Rational q0 = p0 / p2;
+    Rational q1 = p1 / p2;
+    Rational q1half = q1 / rat2;
+    Rational c0 = q0 - q1half * q1half;
+
+    std::map<Rational, int> rmLocalMap;
+    SolveDepressedQuadratic(c0, rmLocalMap);
+
+    rmMap.clear();
+    for (auto& rm : rmLocalMap)
+    {
+        Rational root = rm.first - q1half;
+        rmMap.insert(std::make_pair((Real)root, rm.second));
+    }
+}
+
+template <typename Real>
+template <typename Rational>
+void RootsPolynomial<Real>::SolveCubic(Rational const& p0,
+    Rational const& p1, Rational const& p2, Rational const& p3,
+    std::map<Real, int>& rmMap)
+{
+    Rational const rat2 = 2, rat3 = 3;
+    Rational q0 = p0 / p3;
+    Rational q1 = p1 / p3;
+    Rational q2 = p2 / p3;
+    Rational q2third = q2 / rat3;
+    Rational c0 = q0 - q2third * (q1 - rat2 * q2third * q2third);
+    Rational c1 = q1 - q2 * q2third;
+
+    std::map<Rational, int> rmLocalMap;
+    SolveDepressedCubic(c0, c1, rmLocalMap);
+
+    rmMap.clear();
+    for (auto& rm : rmLocalMap)
+    {
+        Rational root = rm.first - q2third;
+        rmMap.insert(std::make_pair((Real)root, rm.second));
+    }
+}
+
+template <typename Real>
+template <typename Rational>
+void RootsPolynomial<Real>::SolveQuartic(Rational const& p0,
+    Rational const& p1, Rational const& p2, Rational const& p3,
+    Rational const& p4, std::map<Real, int>& rmMap)
+{
+    Rational const rat2 = 2, rat3 = 3, rat4 = 4, rat6 = 6;
+    Rational q0 = p0 / p4;
+    Rational q1 = p1 / p4;
+    Rational q2 = p2 / p4;
+    Rational q3 = p3 / p4;
+    Rational q3fourth = q3 / rat4;
+    Rational q3fourthSqr = q3fourth *  q3fourth;
+    Rational c0 = q0 - q3fourth * (q1 - q3fourth * (q2 - q3fourthSqr * rat3));
+    Rational c1 = q1 - rat2 * q3fourth * (q2 - rat4 * q3fourthSqr);
+    Rational c2 = q2 - rat6 * q3fourthSqr;
+
+    std::map<Rational, int> rmLocalMap;
+    SolveDepressedQuartic(c0, c1, c2, rmLocalMap);
+
+    rmMap.clear();
+    for (auto& rm : rmLocalMap)
+    {
+        Rational root = rm.first - q3fourth;
+        rmMap.insert(std::make_pair((Real)root, rm.second));
+    }
+}
+
+template <typename Real>
+template <typename Rational>
+void RootsPolynomial<Real>::GetRootInfoQuadratic(Rational const& p0,
+    Rational const& p1, Rational const& p2, std::vector<int>& info)
+{
+    Rational const rat2 = 2;
+    Rational q0 = p0 / p2;
+    Rational q1 = p1 / p2;
+    Rational q1half = q1 / rat2;
+    Rational c0 = q0 - q1half * q1half;
+
+    info.clear();
+    info.reserve(2);
+    GetRootInfoDepressedQuadratic(c0, info);
+}
+
+template <typename Real>
+template <typename Rational>
+void RootsPolynomial<Real>::GetRootInfoCubic(Rational const& p0,
+    Rational const& p1, Rational const& p2, Rational const& p3,
+    std::vector<int>& info)
+{
+    Rational const rat2 = 2, rat3 = 3;
+    Rational q0 = p0 / p3;
+    Rational q1 = p1 / p3;
+    Rational q2 = p2 / p3;
+    Rational q2third = q2 / rat3;
+    Rational c0 = q0 - q2third * (q1 - rat2 * q2third * q2third);
+    Rational c1 = q1 - q2 * q2third;
+
+    info.clear();
+    info.reserve(3);
+    GetRootInfoDepressedCubic(c0, c1, info);
+}
+
+template <typename Real>
+template <typename Rational>
+void RootsPolynomial<Real>::GetRootInfoQuartic(Rational const& p0,
+    Rational const& p1, Rational const& p2, Rational const& p3,
+    Rational const& p4, std::vector<int>& info)
+{
+    Rational const rat2 = 2, rat3 = 3, rat4 = 4, rat6 = 6;
+    Rational q0 = p0 / p4;
+    Rational q1 = p1 / p4;
+    Rational q2 = p2 / p4;
+    Rational q3 = p3 / p4;
+    Rational q3fourth = q3 / rat4;
+    Rational q3fourthSqr = q3fourth *  q3fourth;
+    Rational c0 = q0 - q3fourth * (q1 - q3fourth * (q2 - q3fourthSqr * rat3));
+    Rational c1 = q1 - rat2 * q3fourth * (q2 - rat4 * q3fourthSqr);
+    Rational c2 = q2 - rat6 * q3fourthSqr;
+
+    info.clear();
+    info.reserve(4);
+    GetRootInfoDepressedQuartic(c0, c1, c2, info);
+}
+
+template <typename Real>
+int RootsPolynomial<Real>::Find(int degree, Real const* c,
+    unsigned int maxIterations, Real* roots)
+{
+    if (degree >= 0 && c)
+    {
+        Real const zero = (Real)0;
+        while (degree >= 0 && c[degree] == zero)
+        {
+            --degree;
+        }
+
+        if (degree > 0)
+        {
+            // Compute the Cauchy bound.
+            Real const one = (Real)1;
+            Real invLeading = one / c[degree];
+            Real maxValue = zero;
+            for (int i = 0; i < degree; ++i)
+            {
+                Real value = std::abs(c[i] * invLeading);
+                if (value > maxValue)
+                {
+                    maxValue = value;
+                }
+            }
+            Real bound = one + maxValue;
+
+            return FindRecursive(degree, c, -bound, bound, maxIterations,
+                roots);
+        }
+        else if (degree == 0)
+        {
+            // The polynomial is a nonzero constant.
+            return 0;
+        }
+        else
+        {
+            // The polynomial is identically zero.
+            roots[0] = zero;
+            return 1;
+        }
+    }
+    else
+    {
+        // Invalid degree or c.
+        return 0;
+    }
+}
+
+template <typename Real>
+bool RootsPolynomial<Real>::Find(int degree, Real const* c, Real tmin,
+    Real tmax, unsigned int maxIterations, Real& root)
+{
+    Real const zero = (Real)0;
+    Real pmin = Evaluate(degree, c, tmin);
+    if (pmin == zero)
+    {
+        root = tmin;
+        return true;
+    }
+    Real pmax = Evaluate(degree, c, tmax);
+    if (pmax == zero)
+    {
+        root = tmax;
+        return true;
+    }
+
+    if (pmin*pmax > zero)
+    {
+        // It is not known whether the interval bounds a root.
+        return false;
+    }
+
+    if (tmin >= tmax)
+    {
+        // Invalid ordering of interval endpoitns. 
+        return false;
+    }
+
+    for (unsigned int i = 1; i <= maxIterations; ++i)
+    {
+        root = ((Real)0.5) * (tmin + tmax);
+
+        // This test is designed for 'float' or 'double' when tmin and tmax
+        // are consecutive floating-point numbers.
+        if (root == tmin || root == tmax)
+        {
+            break;
+        }
+
+        Real p = Evaluate(degree, c, root);
+        Real product = p * pmin;
+        if (product < zero)
+        {
+            tmax = root;
+            pmax = p;
+        }
+        else if (product > zero)
+        {
+            tmin = root;
+            pmin = p;
+        }
+        else
+        {
+            break;
+        }
+    }
+
+    return true;
+}
+
+template <typename Real>
+template <typename Rational>
+void RootsPolynomial<Real>::SolveDepressedQuadratic(Rational const& c0,
+    std::map<Rational, int>& rmMap)
+{
+    Rational const zero = 0;
+    if (c0 < zero)
+    {
+        // Two simple roots.
+        Rational root1 = (Rational)std::sqrt(-(double)c0);
+        Rational root0 = -root1;
+        rmMap.insert(std::make_pair(root0, 1));
+        rmMap.insert(std::make_pair(root1, 1));
+        GTE_ROOTS_LOW_DEGREE_BLOCK(0);
+    }
+    else if (c0 == zero)
+    {
+        // One double root.
+        rmMap.insert(std::make_pair(zero, 2));
+        GTE_ROOTS_LOW_DEGREE_BLOCK(1);
+    }
+    else  // c0 > 0
+    {
+        // A complex-conjugate pair of roots.
+        // Complex z0 = -q1/2 - i*sqrt(c0);
+        // Complex z0conj = -q1/2 + i*sqrt(c0);
+        GTE_ROOTS_LOW_DEGREE_BLOCK(2);
+    }
+}
+
+template <typename Real>
+template <typename Rational>
+void RootsPolynomial<Real>::SolveDepressedCubic(Rational const& c0,
+    Rational const& c1, std::map<Rational, int>& rmMap)
+{
+    // Handle the special case of c0 = 0, in which case the polynomial
+    // reduces to a depressed quadratic.
+    Rational const zero = 0;
+    if (c0 == zero)
+    {
+        SolveDepressedQuadratic(c1, rmMap);
+        auto iter = rmMap.find(zero);
+        if (iter != rmMap.end())
+        {
+            // The quadratic has a root of zero, so the multiplicity must be
+            // increased.
+            ++iter->second;
+            GTE_ROOTS_LOW_DEGREE_BLOCK(3);
+        }
+        else
+        {
+            // The quadratic does not have a root of zero.  Insert the one
+            // for the cubic.
+            rmMap.insert(std::make_pair(zero, 1));
+            GTE_ROOTS_LOW_DEGREE_BLOCK(4);
+        }
+        return;
+    }
+
+    // Handle the special case of c0 != 0 and c1 = 0.
+    double const oneThird = 1.0 / 3.0;
+    if (c1 == zero)
+    {
+        // One simple real root.
+        Rational root0;
+        if (c0 > zero)
+        {
+            root0 = (Rational)-std::pow((double)c0, oneThird);
+            GTE_ROOTS_LOW_DEGREE_BLOCK(5);
+        }
+        else
+        {
+            root0 = (Rational)std::pow(-(double)c0, oneThird);
+            GTE_ROOTS_LOW_DEGREE_BLOCK(6);
+        }
+        rmMap.insert(std::make_pair(root0, 1));
+
+        // One complex conjugate pair.
+        // Complex z0 = root0*(-1 - i*sqrt(3))/2;
+        // Complex z0conj = root0*(-1 + i*sqrt(3))/2;
+        return;
+    }
+
+    // At this time, c0 != 0 and c1 != 0.
+    Rational const rat2 = 2, rat3 = 3, rat4 = 4, rat27 = 27, rat108 = 108;
+    Rational delta = -(rat4 * c1 * c1 * c1 + rat27 * c0 * c0);
+    if (delta > zero)
+    {
+        // Three simple roots.
+        Rational deltaDiv108 = delta / rat108;
+        Rational betaRe = -c0 / rat2;
+        Rational betaIm = std::sqrt(deltaDiv108);
+        Rational theta = std::atan2(betaIm, betaRe);
+        Rational thetaDiv3 = theta / rat3;
+        double angle = (double)thetaDiv3;
+        Rational cs = (Rational)std::cos(angle);
+        Rational sn = (Rational)std::sin(angle);
+        Rational rhoSqr = betaRe * betaRe + betaIm * betaIm;
+        Rational rhoPowThird = (Rational)std::pow((double)rhoSqr, 1.0 / 6.0);
+        Rational temp0 = rhoPowThird * cs;
+        Rational temp1 = rhoPowThird * sn * (Rational)std::sqrt(3.0);
+        Rational root0 = rat2 * temp0;
+        Rational root1 = -temp0 - temp1;
+        Rational root2 = -temp0 + temp1;
+        rmMap.insert(std::make_pair(root0, 1));
+        rmMap.insert(std::make_pair(root1, 1));
+        rmMap.insert(std::make_pair(root2, 1));
+        GTE_ROOTS_LOW_DEGREE_BLOCK(7);
+    }
+    else if (delta < zero)
+    {
+        // One simple root.
+        Rational deltaDiv108 = delta / rat108;
+        Rational temp0 = -c0 / rat2;
+        Rational temp1 = (Rational)std::sqrt(-(double)deltaDiv108);
+        Rational temp2 = temp0 - temp1;
+        Rational temp3 = temp0 + temp1;
+        if (temp2 >= zero)
+        {
+            temp2 = (Rational)std::pow((double)temp2, oneThird);
+            GTE_ROOTS_LOW_DEGREE_BLOCK(8);
+        }
+        else
+        {
+            temp2 = (Rational)-std::pow(-(double)temp2, oneThird);
+            GTE_ROOTS_LOW_DEGREE_BLOCK(9);
+        }
+        if (temp3 >= zero)
+        {
+            temp3 = (Rational)std::pow((double)temp3, oneThird);
+            GTE_ROOTS_LOW_DEGREE_BLOCK(10);
+        }
+        else
+        {
+            temp3 = (Rational)-std::pow(-(double)temp3, oneThird);
+            GTE_ROOTS_LOW_DEGREE_BLOCK(11);
+        }
+        Rational root0 = temp2 + temp3;
+        rmMap.insert(std::make_pair(root0, 1));
+
+        // One complex conjugate pair.
+        // Complex z0 = (-root0 - i*sqrt(3*root0*root0+4*c1))/2;
+        // Complex z0conj = (-root0 + i*sqrt(3*root0*root0+4*c1))/2;
+    }
+    else  // delta = 0
+    {
+        // One simple root and one double root.
+        Rational root0 = -rat3 * c0 / (rat2 * c1);
+        Rational root1 = -rat2 * root0;
+        rmMap.insert(std::make_pair(root0, 2));
+        rmMap.insert(std::make_pair(root1, 1));
+        GTE_ROOTS_LOW_DEGREE_BLOCK(12);
+    }
+}
+
+template <typename Real>
+template <typename Rational>
+void RootsPolynomial<Real>::SolveDepressedQuartic(Rational const& c0,
+    Rational const& c1, Rational const& c2, std::map<Rational, int>& rmMap)
+{
+    // Handle the special case of c0 = 0, in which case the polynomial
+    // reduces to a depressed cubic.
+    Rational const zero = 0;
+    if (c0 == zero)
+    {
+        SolveDepressedCubic(c1, c2, rmMap);
+        auto iter = rmMap.find(zero);
+        if (iter != rmMap.end())
+        {
+            // The cubic has a root of zero, so the multiplicity must be
+            // increased.
+            ++iter->second;
+            GTE_ROOTS_LOW_DEGREE_BLOCK(13);
+        }
+        else
+        {
+            // The cubic does not have a root of zero.  Insert the one
+            // for the quartic.
+            rmMap.insert(std::make_pair(zero, 1));
+            GTE_ROOTS_LOW_DEGREE_BLOCK(14);
+        }
+        return;
+    }
+
+    // Handle the special case of c1 = 0, in which case the quartic is a
+    // biquadratic x^4 + c1*x^2 + c0 = (x^2 + c2/2)^2 + (c0 - c2^2/4).
+    if (c1 == zero)
+    {
+        SolveBiquadratic(c0, c2, rmMap);
+        return;
+    }
+
+    // At this time, c0 != 0 and c1 != 0, which is a requirement for the
+    // general solver that must use a root of a special cubic polynomial.
+    Rational const rat2 = 2, rat4 = 4, rat8 = 8, rat12 = 12, rat16 = 16;
+    Rational const rat27 = 27, rat36 = 36;
+    Rational c0sqr = c0 * c0, c1sqr = c1 * c1, c2sqr = c2 * c2;
+    Rational delta = c1sqr * (-rat27 * c1sqr + rat4 * c2 *
+        (rat36 * c0 - c2sqr)) + rat16 * c0 * (c2sqr * (c2sqr - rat8 * c0) +
+        rat16 * c0sqr);
+    Rational a0 = rat12 * c0 + c2sqr;
+    Rational a1 = rat4 * c0 - c2sqr;
+
+    if (delta > zero)
+    {
+        if (c2 < zero && a1 < zero)
+        {
+            // Four simple real roots.
+            std::map<Real, int> rmCubicMap;
+            SolveCubic(c1sqr - rat4 * c0 * c2, rat8 * c0, rat4 * c2, -rat8,
+                rmCubicMap);
+            Rational t = (Rational)rmCubicMap.rbegin()->first;
+            Rational alphaSqr = rat2 * t - c2;
+            Rational alpha = (Rational)std::sqrt((double)alphaSqr);
+            double sgnC1;
+            if (c1 > zero)
+            {
+                sgnC1 = 1.0;
+                GTE_ROOTS_LOW_DEGREE_BLOCK(15);
+            }
+            else
+            {
+                sgnC1 = -1.0;
+                GTE_ROOTS_LOW_DEGREE_BLOCK(16);
+            }
+            Rational arg = t * t - c0;
+            Rational beta = (Rational)(sgnC1 * std::sqrt(std::max((double)arg, 0.0)));
+            Rational D0 = alphaSqr - rat4 * (t + beta);
+            Rational sqrtD0 = (Rational)std::sqrt(std::max((double)D0, 0.0));
+            Rational D1 = alphaSqr - rat4 * (t - beta);
+            Rational sqrtD1 = (Rational)std::sqrt(std::max((double)D1, 0.0));
+            Rational root0 = (alpha - sqrtD0) / rat2;
+            Rational root1 = (alpha + sqrtD0) / rat2;
+            Rational root2 = (-alpha - sqrtD1) / rat2;
+            Rational root3 = (-alpha + sqrtD1) / rat2;
+            rmMap.insert(std::make_pair(root0, 1));
+            rmMap.insert(std::make_pair(root1, 1));
+            rmMap.insert(std::make_pair(root2, 1));
+            rmMap.insert(std::make_pair(root3, 1));
+        }
+        else // c2 >= 0 or a1 >= 0
+        {
+            // Two complex-conjugate pairs.  The values alpha, D0, and D1 are
+            // those of the if-block.
+            // Complex z0 = (alpha - i*sqrt(-D0))/2;
+            // Complex z0conj = (alpha + i*sqrt(-D0))/2;
+            // Complex z1 = (-alpha - i*sqrt(-D1))/2;
+            // Complex z1conj = (-alpha + i*sqrt(-D1))/2;
+            GTE_ROOTS_LOW_DEGREE_BLOCK(17);
+        }
+    }
+    else if (delta < zero)
+    {
+        // Two simple real roots, one complex-conjugate pair.
+        std::map<Real, int> rmCubicMap;
+        SolveCubic(c1sqr - rat4 * c0 * c2, rat8 * c0, rat4 * c2, -rat8,
+            rmCubicMap);
+        Rational t = (Rational)rmCubicMap.rbegin()->first;
+        Rational alphaSqr = rat2 * t - c2;
+        Rational alpha = (Rational)std::sqrt(std::max((double)alphaSqr, 0.0));
+        double sgnC1;
+        if (c1 > zero)
+        {
+            sgnC1 = 1.0;  // Leads to BLOCK(18)
+        }
+        else
+        {
+            sgnC1 = -1.0;  // Leads to BLOCK(19)
+        }
+        Rational arg = t * t - c0;
+        Rational beta = (Rational)(sgnC1 * std::sqrt(std::max((double)arg, 0.0)));
+        Rational root0, root1;
+        if (sgnC1 > 0.0)
+        {
+            Rational D1 = alphaSqr - rat4 * (t - beta);
+            Rational sqrtD1 = (Rational)std::sqrt(std::max((double)D1, 0.0));
+            root0 = (-alpha - sqrtD1) / rat2;
+            root1 = (-alpha + sqrtD1) / rat2;
+
+            // One complex conjugate pair.
+            // Complex z0 = (alpha - i*sqrt(-D0))/2;
+            // Complex z0conj = (alpha + i*sqrt(-D0))/2;
+            GTE_ROOTS_LOW_DEGREE_BLOCK(18);
+        }
+        else
+        {
+            Rational D0 = alphaSqr - rat4 * (t + beta);
+            Rational sqrtD0 = (Rational)std::sqrt(std::max((double)D0, 0.0));
+            root0 = (alpha - sqrtD0) / rat2;
+            root1 = (alpha + sqrtD0) / rat2;
+
+            // One complex conjugate pair.
+            // Complex z0 = (-alpha - i*sqrt(-D1))/2;
+            // Complex z0conj = (-alpha + i*sqrt(-D1))/2;
+            GTE_ROOTS_LOW_DEGREE_BLOCK(19);
+        }
+        rmMap.insert(std::make_pair(root0, 1));
+        rmMap.insert(std::make_pair(root1, 1));
+    }
+    else  // delta = 0
+    {
+        if (a1 > zero || (c2 > zero && (a1 != zero || c1 != zero)))
+        {
+            // One double real root, one complex-conjugate pair.
+            Rational const rat9 = 9;
+            Rational root0 = -c1 * a0 / (rat9 * c1sqr - rat2 * c2 * a1);
+            rmMap.insert(std::make_pair(root0, 2));
+
+            // One complex conjugate pair.
+            // Complex z0 = -root0 - i*sqrt(c2 + root0^2);
+            // Complex z0conj = -root0 + i*sqrt(c2 + root0^2);
+            GTE_ROOTS_LOW_DEGREE_BLOCK(20);
+        }
+        else
+        {
+            Rational const rat3 = 3;
+            if (a0 != zero)
+            {
+                // One double real root, two simple real roots.
+                Rational const rat9 = 9;
+                Rational root0 = -c1 * a0 / (rat9 * c1sqr - rat2 * c2 * a1);
+                Rational alpha = rat2 * root0;
+                Rational beta = c2 + rat3 * root0 * root0;
+                Rational discr = alpha * alpha - rat4 * beta;
+                Rational temp1 = (Rational)std::sqrt((double)discr);
+                Rational root1 = (-alpha - temp1) / rat2;
+                Rational root2 = (-alpha + temp1) / rat2;
+                rmMap.insert(std::make_pair(root0, 2));
+                rmMap.insert(std::make_pair(root1, 1));
+                rmMap.insert(std::make_pair(root2, 1));
+                GTE_ROOTS_LOW_DEGREE_BLOCK(21);
+            }
+            else
+            {
+                // One triple real root, one simple real root.
+                Rational root0 = -rat3 * c1 / (rat4 * c2);
+                Rational root1 = -rat3 * root0;
+                rmMap.insert(std::make_pair(root0, 3));
+                rmMap.insert(std::make_pair(root1, 1));
+                GTE_ROOTS_LOW_DEGREE_BLOCK(22);
+            }
+        }
+    }
+}
+
+template <typename Real>
+template <typename Rational>
+void RootsPolynomial<Real>::SolveBiquadratic(Rational const& c0,
+    Rational const& c2, std::map<Rational, int>& rmMap)
+{
+    // Solve 0 = x^4 + c2*x^2 + c0 = (x^2 + c2/2)^2 + a1, where
+    // a1 = c0 - c2^2/2.  We know that c0 != 0 at the time of the function
+    // call, so x = 0 is not a root.  The condition c1 = 0 implies the quartic
+    // Delta = 256*c0*a1^2.
+
+    Rational const zero = 0, rat2 = 2, rat256 = 256;
+    Rational c2Half = c2 / rat2;
+    Rational a1 = c0 - c2Half * c2Half;
+    Rational delta = rat256 * c0 * a1 * a1;
+    if (delta > zero)
+    {
+        if (c2 < zero)
+        {
+            if (a1 < zero)
+            {
+                // Four simple roots.
+                Rational temp0 = (Rational)std::sqrt(-(double)a1);
+                Rational temp1 = -c2Half - temp0;
+                Rational temp2 = -c2Half + temp0;
+                Rational root1 = (Rational)std::sqrt((double)temp1);
+                Rational root0 = -root1;
+                Rational root2 = (Rational)std::sqrt((double)temp2);
+                Rational root3 = -root2;
+                rmMap.insert(std::make_pair(root0, 1));
+                rmMap.insert(std::make_pair(root1, 1));
+                rmMap.insert(std::make_pair(root2, 1));
+                rmMap.insert(std::make_pair(root3, 1));
+                GTE_ROOTS_LOW_DEGREE_BLOCK(23);
+            }
+            else  // a1 > 0
+            {
+                // Two simple complex conjugate pairs.
+                // double thetaDiv2 = atan2(sqrt(a1), -c2/2) / 2.0;
+                // double cs = cos(thetaDiv2), sn = sin(thetaDiv2);
+                // double length = pow(c0, 0.25);
+                // Complex z0 = length*(cs + i*sn);
+                // Complex z0conj = length*(cs - i*sn);
+                // Complex z1 = length*(-cs + i*sn);
+                // Complex z1conj = length*(-cs - i*sn);
+                GTE_ROOTS_LOW_DEGREE_BLOCK(24);
+            }
+        }
+        else  // c2 >= 0
+        {
+            // Two simple complex conjugate pairs.
+            // Complex z0 = -i*sqrt(c2/2 - sqrt(-a1));
+            // Complex z0conj = +i*sqrt(c2/2 - sqrt(-a1));
+            // Complex z1 = -i*sqrt(c2/2 + sqrt(-a1));
+            // Complex z1conj = +i*sqrt(c2/2 + sqrt(-a1));
+            GTE_ROOTS_LOW_DEGREE_BLOCK(25);
+        }
+    }
+    else if (delta < zero)
+    {
+        // Two simple real roots.
+        Rational temp0 = (Rational)std::sqrt(-(double)a1);
+        Rational temp1 = -c2Half + temp0;
+        Rational root1 = (Rational)std::sqrt((double)temp1);
+        Rational root0 = -root1;
+        rmMap.insert(std::make_pair(root0, 1));
+        rmMap.insert(std::make_pair(root1, 1));
+
+        // One complex conjugate pair.
+        // Complex z0 = -i*sqrt(c2/2 + sqrt(-a1));
+        // Complex z0conj = +i*sqrt(c2/2 + sqrt(-a1));
+        GTE_ROOTS_LOW_DEGREE_BLOCK(26);
+    }
+    else  // delta = 0
+    {
+        if (c2 < zero)
+        {
+            // Two double real roots.
+            Rational root1 = (Rational)std::sqrt(-(double)c2Half);
+            Rational root0 = -root1;
+            rmMap.insert(std::make_pair(root0, 2));
+            rmMap.insert(std::make_pair(root1, 2));
+            GTE_ROOTS_LOW_DEGREE_BLOCK(27);
+        }
+        else  // c2 > 0
+        {
+            // Two double complex conjugate pairs.
+            // Complex z0 = -i*sqrt(c2/2);  // multiplicity 2
+            // Complex z0conj = +i*sqrt(c2/2);  // multiplicity 2
+            GTE_ROOTS_LOW_DEGREE_BLOCK(28);
+        }
+    }
+}
+
+template <typename Real>
+template <typename Rational>
+void RootsPolynomial<Real>::GetRootInfoDepressedQuadratic(Rational const& c0,
+    std::vector<int>& info)
+{
+    Rational const zero = 0;
+    if (c0 < zero)
+    {
+        // Two simple roots.
+        info.push_back(1);
+        info.push_back(1);
+    }
+    else if (c0 == zero)
+    {
+        // One double root.
+        info.push_back(2);  // root is zero
+    }
+    else  // c0 > 0
+    {
+        // A complex-conjugate pair of roots.
+    }
+}
+
+template <typename Real>
+template <typename Rational>
+void RootsPolynomial<Real>::GetRootInfoDepressedCubic(Rational const& c0,
+    Rational const& c1, std::vector<int>& info)
+{
+    // Handle the special case of c0 = 0, in which case the polynomial
+    // reduces to a depressed quadratic.
+    Rational const zero = 0;
+    if (c0 == zero)
+    {
+        if (c1 == zero)
+        {
+            info.push_back(3);  // triple root of zero
+        }
+        else
+        {
+            info.push_back(1);  // simple root of zero
+            GetRootInfoDepressedQuadratic(c1, info);
+        }
+        return;
+    }
+
+    Rational const rat4 = 4, rat27 = 27;
+    Rational delta = -(rat4 * c1 * c1 * c1 + rat27 * c0 * c0);
+    if (delta > zero)
+    {
+        // Three simple real roots.
+        info.push_back(1);
+        info.push_back(1);
+        info.push_back(1);
+    }
+    else if (delta < zero)
+    {
+        // One simple real root.
+        info.push_back(1);
+    }
+    else  // delta = 0
+    {
+        // One simple real root and one double real root.
+        info.push_back(1);
+        info.push_back(2);
+    }
+}
+
+template <typename Real>
+template <typename Rational>
+void RootsPolynomial<Real>::GetRootInfoDepressedQuartic(Rational const& c0,
+    Rational const& c1, Rational const& c2, std::vector<int>& info)
+{
+    // Handle the special case of c0 = 0, in which case the polynomial
+    // reduces to a depressed cubic.
+    Rational const zero = 0;
+    if (c0 == zero)
+    {
+        if (c1 == zero)
+        {
+            if (c2 == zero)
+            {
+                info.push_back(4);  // quadruple root of zero
+            }
+            else
+            {
+                info.push_back(2);  // double root of zero
+                GetRootInfoDepressedQuadratic(c2, info);
+            }
+        }
+        else
+        {
+            info.push_back(1);  // simple root of zero
+            GetRootInfoDepressedCubic(c1, c2, info);
+        }
+        return;
+    }
+
+    // Handle the special case of c1 = 0, in which case the quartic is a
+    // biquadratic x^4 + c1*x^2 + c0 = (x^2 + c2/2)^2 + (c0 - c2^2/4).
+    if (c1 == zero)
+    {
+        GetRootInfoBiquadratic(c0, c2, info);
+        return;
+    }
+
+    // At this time, c0 != 0 and c1 != 0, which is a requirement for the
+    // general solver that must use a root of a special cubic polynomial.
+    Rational const rat4 = 4, rat8 = 8, rat12 = 12, rat16 = 16;
+    Rational const rat27 = 27, rat36 = 36;
+    Rational c0sqr = c0 * c0, c1sqr = c1 * c1, c2sqr = c2 * c2;
+    Rational delta = c1sqr * (-rat27 * c1sqr + rat4 * c2 *
+        (rat36 * c0 - c2sqr)) + rat16 * c0 * (c2sqr * (c2sqr - rat8 * c0) +
+        rat16 * c0sqr);
+    Rational a0 = rat12 * c0 + c2sqr;
+    Rational a1 = rat4 * c0 - c2sqr;
+
+    if (delta > zero)
+    {
+        if (c2 < zero && a1 < zero)
+        {
+            // Four simple real roots.
+            info.push_back(1);
+            info.push_back(1);
+            info.push_back(1);
+            info.push_back(1);
+        }
+        else // c2 >= 0 or a1 >= 0
+        {
+            // Two complex-conjugate pairs.
+        }
+    }
+    else if (delta < zero)
+    {
+        // Two simple real roots, one complex-conjugate pair.
+        info.push_back(1);
+        info.push_back(1);
+    }
+    else  // delta = 0
+    {
+        if (a1 > zero || (c2 > zero && (a1 != zero || c1 != zero)))
+        {
+            // One double real root, one complex-conjugate pair.
+            info.push_back(2);
+        }
+        else
+        {
+            if (a0 != zero)
+            {
+                // One double real root, two simple real roots.
+                info.push_back(2);
+                info.push_back(1);
+                info.push_back(1);
+            }
+            else
+            {
+                // One triple real root, one simple real root.
+                info.push_back(3);
+                info.push_back(1);
+            }
+        }
+    }
+}
+
+template <typename Real>
+template <typename Rational>
+void RootsPolynomial<Real>::GetRootInfoBiquadratic(Rational const& c0,
+    Rational const& c2, std::vector<int>& info)
+{
+    // Solve 0 = x^4 + c2*x^2 + c0 = (x^2 + c2/2)^2 + a1, where
+    // a1 = c0 - c2^2/2.  We know that c0 != 0 at the time of the function
+    // call, so x = 0 is not a root.  The condition c1 = 0 implies the quartic
+    // Delta = 256*c0*a1^2.
+
+    Rational const zero = 0, rat2 = 2, rat256 = 256;
+    Rational c2Half = c2 / rat2;
+    Rational a1 = c0 - c2Half * c2Half;
+    Rational delta = rat256 * c0 * a1 * a1;
+    if (delta > zero)
+    {
+        if (c2 < zero)
+        {
+            if (a1 < zero)
+            {
+                // Four simple roots.
+                info.push_back(1);
+                info.push_back(1);
+                info.push_back(1);
+                info.push_back(1);
+            }
+            else  // a1 > 0
+            {
+                // Two simple complex conjugate pairs.
+            }
+        }
+        else  // c2 >= 0
+        {
+            // Two simple complex conjugate pairs.
+        }
+    }
+    else if (delta < zero)
+    {
+        // Two simple real roots, one complex conjugate pair.
+        info.push_back(1);
+        info.push_back(1);
+    }
+    else  // delta = 0
+    {
+        if (c2 < zero)
+        {
+            // Two double real roots.
+            info.push_back(2);
+            info.push_back(2);
+        }
+        else  // c2 > 0
+        {
+            // Two double complex conjugate pairs.
+        }
+    }
+}
+
+template <typename Real>
+int RootsPolynomial<Real>::FindRecursive(int degree, Real const* c,
+    Real tmin, Real tmax, unsigned int maxIterations, Real* roots)
+{
+    // The base of the recursion.
+    Real const zero = (Real)0;
+    Real root = zero;
+    if (degree == 1)
+    {
+        int numRoots;
+        if (c[1] != zero)
+        {
+            root = -c[0] / c[1];
+            numRoots = 1;
+        }
+        else if (c[0] == zero)
+        {
+            root = zero;
+            numRoots = 1;
+        }
+        else
+        {
+            numRoots = 0;
+        }
+
+        if (numRoots > 0 && tmin <= root && root <= tmax)
+        {
+            roots[0] = root;
+            return 1;
+        }
+        return 0;
+    }
+
+    // Find the roots of the derivative polynomial scaled by 1/degree.  The
+    // scaling avoids the factorial growth in the coefficients; for example,
+    // without the scaling, the high-order term x^d becomes (d!)*x through
+    // multiple differentiations.  With the scaling we instead get x.  This
+    // leads to better numerical behavior of the root finder.
+    int derivDegree = degree - 1;
+    std::vector<Real> derivCoeff(derivDegree + 1);
+    std::vector<Real> derivRoots(derivDegree);
+    for (int i = 0; i <= derivDegree; ++i)
+    {
+        derivCoeff[i] = c[i + 1] * (Real)(i + 1) / (Real)degree;
+    }
+    int numDerivRoots = FindRecursive(degree - 1, &derivCoeff[0], tmin, tmax,
+        maxIterations, &derivRoots[0]);
+
+    int numRoots = 0;
+    if (numDerivRoots > 0)
+    {
+        // Find root on [tmin,derivRoots[0]].
+        if (Find(degree, c, tmin, derivRoots[0], maxIterations, root))
+        {
+            roots[numRoots++] = root;
+        }
+
+        // Find root on [derivRoots[i],derivRoots[i+1]].
+        for (int i = 0; i <= numDerivRoots - 2; ++i)
+        {
+            if (Find(degree, c, derivRoots[i], derivRoots[i + 1],
+                maxIterations, root))
+            {
+                roots[numRoots++] = root;
+            }
+        }
+
+        // Find root on [derivRoots[numDerivRoots-1],tmax].
+        if (Find(degree, c, derivRoots[numDerivRoots - 1], tmax,
+            maxIterations, root))
+        {
+            roots[numRoots++] = root;
+        }
+    }
+    else
+    {
+        // The polynomial is monotone on [tmin,tmax], so has at most one root.
+        if (Find(degree, c, tmin, tmax, maxIterations, root))
+        {
+            roots[numRoots++] = root;
+        }
+    }
+    return numRoots;
+}
+
+template <typename Real>
+Real RootsPolynomial<Real>::Evaluate(int degree, Real const* c, Real t)
+{
+    int i = degree;
+    Real result = c[i];
+    while (--i >= 0)
+    {
+        result = t * result + c[i];
+    }
+    return result;
+}
+
+
+}
diff --git a/Engine/Plugins/Experimental/GeometryProcessing/Source/GeometryAlgorithms/Private/GTEngine/Mathematics/GteRotation.h b/Engine/Plugins/Experimental/GeometryProcessing/Source/GeometryAlgorithms/Private/GTEngine/Mathematics/GteRotation.h
new file mode 100644
--- /dev/null
+++ b/Engine/Plugins/Experimental/GeometryProcessing/Source/GeometryAlgorithms/Private/GTEngine/Mathematics/GteRotation.h
@@ -0,0 +1,892 @@
+// David Eberly, Geometric Tools, Redmond WA 98052
+// Copyright (c) 1998-2019
+// Distributed under the Boost Software License, Version 1.0.
+// http://www.boost.org/LICENSE_1_0.txt
+// http://www.geometrictools.com/License/Boost/LICENSE_1_0.txt
+// File Version: 3.0.1 (2018/10/05)
+
+#pragma once
+
+#include <GTEngine/Mathematics/GteAxisAngle.h>
+#include <GTEngine/Mathematics/GteEulerAngles.h>
+#include <GTEngine/Mathematics/GteMatrix.h>
+#include <GTEngine/Mathematics/GteQuaternion.h>
+#include <GTEngine/Mathematics/GteMath.h>
+
+namespace gte
+{
+
+// Conversions among various representations of rotations.  The value of
+// N must be 3 or 4.  The latter case supports affine algebra when you use
+// 4-tuple vectors (w-component is 1 for points and 0 for vector) and 4x4
+// matrices for affine transformations.  Rotation axes must be unit length.
+// The angles are in radians.  The Euler angles are in world coordinates;
+// we have not yet added support for body coordinates.
+
+template <int N, typename Real>
+class Rotation
+{
+public:
+    // Create rotations from various representations.
+    Rotation(Matrix<N,N,Real> const& matrix);
+    Rotation(Quaternion<Real> const& quaternion);
+    Rotation(AxisAngle<N,Real> const& axisAngle);
+    Rotation(EulerAngles<Real> const& eulerAngles);
+
+    // Convert one representation to another.
+    operator Matrix<N,N,Real>() const;
+    operator Quaternion<Real>() const;
+    operator AxisAngle<N,Real>() const;
+    EulerAngles<Real> const& operator()(int i0, int i1, int i2) const;
+
+private:
+    enum RepresentationType
+    {
+        IS_MATRIX,
+        IS_QUATERNION,
+        IS_AXIS_ANGLE,
+        IS_EULER_ANGLES
+    };
+
+    RepresentationType mType;
+    mutable Matrix<N,N,Real> mMatrix;
+    mutable Quaternion<Real> mQuaternion;
+    mutable AxisAngle<N,Real> mAxisAngle;
+    mutable EulerAngles<Real> mEulerAngles;
+
+    // Convert a rotation matrix to a quaternion.
+    //
+    // x^2 = (+r00 - r11 - r22 + 1)/4
+    // y^2 = (-r00 + r11 - r22 + 1)/4
+    // z^2 = (-r00 - r11 + r22 + 1)/4
+    // w^2 = (+r00 + r11 + r22 + 1)/4
+    // x^2 + y^2 = (1 - r22)/2
+    // z^2 + w^2 = (1 + r22)/2
+    // y^2 - x^2 = (r11 - r00)/2
+    // w^2 - z^2 = (r11 + r00)/2
+    // x*y = (r01 + r10)/4
+    // x*z = (r02 + r20)/4
+    // y*z = (r12 + r21)/4
+    // [GTE_USE_MAT_VEC]
+    //   x*w = (r21 - r12)/4
+    //   y*w = (r02 - r20)/4
+    //   z*w = (r10 - r01)/4
+    // [GTE_USE_VEC_MAT]
+    //   x*w = (r12 - r21)/4
+    //   y*w = (r20 - r02)/4
+    //   z*w = (r01 - r10)/4
+    //
+    // If Q is the 4x1 column vector (x,y,z,w), the previous equations give us
+    //         +-                  -+
+    //         | x*x  x*y  x*z  x*w |
+    // Q*Q^T = | y*x  y*y  y*z  y*w |
+    //         | z*x  z*y  z*z  z*w |
+    //         | w*x  w*y  w*z  w*w |
+    //         +-                  -+
+    // The code extracts the row of maximum length, normalizing it to obtain
+    // the result q.
+    static void Convert(Matrix<N,N,Real> const& r, Quaternion<Real>& q);
+
+    // Convert a quaterion q = x*i + y*j + z*k + w to a rotation matrix.
+    // [GTE_USE_MAT_VEC]
+    //     +-           -+   +-                                     -+
+    // R = | r00 r01 r02 | = | 1-2y^2-2z^2  2(xy-zw)     2(xz+yw)    |
+    //     | r10 r11 r12 |   | 2(xy+zw)     1-2x^2-2z^2  2(yz-xw)    |
+    //     | r20 r21 r22 |   | 2(xz-yw)     2(yz+xw)     1-2x^2-2y^2 |
+    //     +-           -+   +-                                     -+
+    // [GTE_USE_VEC_MAT]
+    //     +-           -+   +-                                     -+
+    // R = | r00 r01 r02 | = | 1-2y^2-2z^2  2(xy+zw)     2(xz-yw)    |
+    //     | r10 r11 r12 |   | 2(xy-zw)     1-2x^2-2z^2  2(yz+xw)    |
+    //     | r20 r21 r22 |   | 2(xz+yw)     2(yz-xw)     1-2x^2-2y^2 |
+    //     +-           -+   +-                                     -+
+    static void Convert(Quaternion<Real> const& q, Matrix<N,N,Real>& r);
+
+    // Convert a rotation matrix to an axis-angle pair.  Let (x0,x1,x2) be the
+    // axis let t be an angle of rotation.  The rotation matrix is
+    // [GTE_USE_MAT_VEC]
+    //   R = I + sin(t)*S + (1-cos(t))*S^2
+    // or
+    // [GTE_USE_VEC_MAT]
+    //   R = I - sin(t)*S + (1-cos(t))*S^2
+    // where I is the identity and S = {{0,-x2,x1},{x2,0,-x0},{-x1,x0,0}}
+    // where the inner-brace triples are the rows of the matrix.  If t > 0,
+    // R represents a counterclockwise rotation; see the comments for the
+    // constructor Matrix3x3(axis,angle).  It may be shown that cos(t) =
+    // (trace(R)-1)/2 and R - Transpose(R) = 2*sin(t)*S.  As long as sin(t) is
+    // not zero, we may solve for S in the second equation, which produces the
+    // axis direction U = (S21,S02,S10).  When t = 0, the rotation is the
+    // identity, in which case any axis direction is valid; we choose (1,0,0).
+    // When t = pi, it must be that R - Transpose(R) = 0, which prevents us
+    // from extracting the axis.  Instead, note that (R+I)/2 = I+S^2 = U*U^T,
+    // where U is a unit-length axis direction.
+    static void Convert(Matrix<N,N,Real> const& r, AxisAngle<N,Real>& a);
+
+    // Convert an axis-angle pair to a rotation matrix.  Assuming (x0,x1,x2)
+    // is for a right-handed world (x0 to right, x1 up, x2 out of plane of
+    // page), a positive angle corresponds to a counterclockwise rotation from
+    // the perspective of an observer looking at the origin of the plane of
+    // rotation and having view direction the negative of the axis direction.
+    // The coordinate-axis rotations are the following, where
+    // unit(0) = (1,0,0), unit(1) = (0,1,0), unit(2) = (0,0,1),
+    // [GTE_USE_MAT_VEC]
+    //   R(unit(0),t) = {{ 1, 0, 0}, { 0, c,-s}, { 0, s, c}}
+    //   R(unit(1),t) = {{ c, 0, s}, { 0, 1, 0}, {-s, 0, c}}
+    //   R(unit(2),t) = {{ c,-s, 0}, { s, c, 0}, { 0, 0, 1}}
+    // or
+    // [GTE_USE_VEC_MAT]
+    //   R(unit(0),t) = {{ 1, 0, 0}, { 0, c, s}, { 0,-s, c}}
+    //   R(unit(1),t) = {{ c, 0,-s}, { 0, 1, 0}, { s, 0, c}}
+    //   R(unit(2),t) = {{ c, s, 0}, {-s, c, 0}, { 0, 0, 1}}
+    // where c = cos(t), s = sin(t), and the inner-brace triples are rows of
+    // the matrix.  The general matrix is
+    // [GTE_USE_MAT_VEC]
+    //       +-                                                          -+
+    //   R = | (1-c)*x0^2  + c     (1-c)*x0*x1 - s*x2  (1-c)*x0*x2 + s*x1 |
+    //       | (1-c)*x0*x1 + s*x2  (1-c)*x1^2  + c     (1-c)*x1*x2 - s*x0 |
+    //       | (1-c)*x0*x2 - s*x1  (1-c)*x1*x2 + s*x0  (1-c)*x2^2  + c    |
+    //       +-                                                          -+
+    // [GTE_USE_VEC_MAT]
+    //       +-                                                          -+
+    //   R = | (1-c)*x0^2  + c     (1-c)*x0*x1 + s*x2  (1-c)*x0*x2 - s*x1 |
+    //       | (1-c)*x0*x1 - s*x2  (1-c)*x1^2  + c     (1-c)*x1*x2 + s*x0 |
+    //       | (1-c)*x0*x2 + s*x1  (1-c)*x1*x2 - s*x0  (1-c)*x2^2  + c    |
+    //       +-                                                          -+
+    static void Convert(AxisAngle<N,Real> const& a, Matrix<N,N,Real>& r);
+
+    // Convert a rotation matrix to Euler angles.  Factorization into Euler
+    // angles is not necessarily unique.  If the result is ER_NOT_UNIQUE_SUM,
+    // then the multiple solutions occur because angleN2+angleN0 is constant.
+    // If the result is ER_NOT_UNIQUE_DIF, then the multiple solutions occur
+    // because angleN2-angleN0 is constant.  In either type of nonuniqueness,
+    // the function returns angleN0=0.
+    static void Convert(Matrix<N,N,Real> const& r, EulerAngles<Real>& e);
+
+    // Convert Euler angles to a rotation matrix.  The three integer inputs
+    // are in {0,1,2} and correspond to world directions unit(0) = (1,0,0),
+    // unit(1) = (0,1,0), or unit(2) = (0,0,1).  The triples (N0,N1,N2) must
+    // be in the following set,
+    //   {(0,1,2),(0,2,1),(1,0,2),(1,2,0),(2,0,1),(2,1,0),
+    //    (0,1,0),(0,2,0),(1,0,1),(1,2,1),(2,0,2),(2,1,2)}
+    // The rotation matrix is
+    //   [GTE_USE_MAT_VEC]
+    //     R(unit(N2),angleN2)*R(unit(N1),angleN1)*R(unit(N0),angleN0)
+    // or
+    //   [GTE_USE_VEC_MAT]
+    //     R(unit(N0),angleN0)*R(unit(N1),angleN1)*R(unit(N2),angleN2)
+    // The conventions of constructor Matrix3(axis,angle) apply here as well.
+    //
+    // NOTE:  The reversal of order is chosen so that a rotation matrix built
+    // with one multiplication convention is the transpose of the rotation
+    // matrix built with the other multiplication convention.  Thus,
+    //   [GTE_USE_MAT_VEC]
+    //     Matrix3x3<Real> R_mvconvention(N0,N1,N2,angleN0,angleN1,angleN2);
+    //     Vector3<Real> V(...);
+    //     Vector3<Real> U = R_mvconvention*V;  // (u0,u1,u2) = R2*R1*R0*V
+    //   [GTE_USE_VEC_MAT]
+    //     Matrix3x3<Real> R_vmconvention(N0,N1,N2,angleN0,angleN1,angleN2);
+    //     Vector3<Real> V(...);
+    //     Vector3<Real> U = R_mvconvention*V;  // (u0,u1,u2) = V*R0*R1*R2
+    // In either convention, you get the same 3-tuple U.
+    static void Convert(EulerAngles<Real> const& e, Matrix<N,N,Real>& r);
+
+    // Convert a quaternion to an axis-angle pair, where
+    //   q = sin(angle/2)*(axis[0]*i + axis[1]*j + axis[2]*k) + cos(angle/2)
+    static void Convert(Quaternion<Real> const& q, AxisAngle<N,Real>& a);
+
+    // Convert an axis-angle pair to a quaternion, where
+    //   q = sin(angle/2)*(axis[0]*i + axis[1]*j + axis[2]*k) + cos(angle/2)
+    static void Convert(AxisAngle<N,Real> const& a, Quaternion<Real>& q);
+
+    // Convert a quaternion to Euler angles.  The quaternion is converted to
+    // a matrix which is then converted to Euler angles.
+    static void Convert(Quaternion<Real> const& q, EulerAngles<Real>& e);
+
+    // Convert Euler angles to a quaternion.  The Euler angles are converted
+    // to a matrix which is then converted to a quaternion.
+    static void Convert(EulerAngles<Real> const& e, Quaternion<Real>& q);
+
+    // Convert an axis-angle pair to Euler angles.  The axis-angle pair is
+    // converted to a quaternion which is then converted to Euler angles.
+    static void Convert(AxisAngle<N,Real> const& a, EulerAngles<Real>& e);
+
+    // Convert Euler angles to an axis-angle pair.  The Euler angles are
+    // converted to a quaternion which is then converted to an axis-angle
+    // pair.
+    static void Convert(EulerAngles<Real> const& e, AxisAngle<N,Real>& a);
+};
+
+
+template <int N, typename Real>
+Rotation<N, Real>::Rotation(Matrix<N, N, Real> const& matrix)
+    :
+    mType(IS_MATRIX),
+    mMatrix(matrix)
+{
+    static_assert(N == 3 || N == 4, "Dimension must be 3 or 4.");
+}
+
+template <int N, typename Real>
+Rotation<N, Real>::Rotation(Quaternion<Real> const& quaternion)
+    :
+    mType(IS_QUATERNION),
+    mQuaternion(quaternion)
+{
+    static_assert(N == 3 || N == 4, "Dimension must be 3 or 4.");
+}
+
+template <int N, typename Real>
+Rotation<N, Real>::Rotation(AxisAngle<N, Real> const& axisAngle)
+    :
+    mType(IS_AXIS_ANGLE),
+    mAxisAngle(axisAngle)
+{
+    static_assert(N == 3 || N == 4, "Dimension must be 3 or 4.");
+}
+
+template <int N, typename Real>
+Rotation<N, Real>::Rotation(EulerAngles<Real> const& eulerAngles)
+    :
+    mType(IS_EULER_ANGLES),
+    mEulerAngles(eulerAngles)
+{
+    static_assert(N == 3 || N == 4, "Dimension must be 3 or 4.");
+}
+
+template <int N, typename Real>
+Rotation<N, Real>::operator Matrix<N, N, Real>() const
+{
+    static_assert(N == 3 || N == 4, "Dimension must be 3 or 4.");
+
+    switch (mType)
+    {
+    case IS_MATRIX:
+        break;
+    case IS_QUATERNION:
+        Convert(mQuaternion, mMatrix);
+        break;
+    case IS_AXIS_ANGLE:
+        Convert(mAxisAngle, mMatrix);
+        break;
+    case IS_EULER_ANGLES:
+        Convert(mEulerAngles, mMatrix);
+        break;
+    }
+
+    return mMatrix;
+}
+
+template <int N, typename Real>
+Rotation<N, Real>::operator Quaternion<Real>() const
+{
+    static_assert(N == 3 || N == 4, "Dimension must be 3 or 4.");
+
+    switch (mType)
+    {
+    case IS_MATRIX:
+        Convert(mMatrix, mQuaternion);
+        break;
+    case IS_QUATERNION:
+        break;
+    case IS_AXIS_ANGLE:
+        Convert(mAxisAngle, mQuaternion);
+        break;
+    case IS_EULER_ANGLES:
+        Convert(mEulerAngles, mQuaternion);
+        break;
+    }
+
+    return mQuaternion;
+}
+
+template <int N, typename Real>
+Rotation<N, Real>::operator AxisAngle<N, Real>() const
+{
+    static_assert(N == 3 || N == 4, "Dimension must be 3 or 4.");
+
+    switch (mType)
+    {
+    case IS_MATRIX:
+        Convert(mMatrix, mAxisAngle);
+        break;
+    case IS_QUATERNION:
+        Convert(mQuaternion, mAxisAngle);
+        break;
+    case IS_AXIS_ANGLE:
+        break;
+    case IS_EULER_ANGLES:
+        Convert(mEulerAngles, mAxisAngle);
+        break;
+    }
+
+    return mAxisAngle;
+}
+
+template <int N, typename Real>
+EulerAngles<Real> const& Rotation<N, Real>::operator()(int i0, int i1,
+    int i2) const
+{
+    static_assert(N == 3 || N == 4, "Dimension must be 3 or 4.");
+
+    mEulerAngles.axis[0] = i0;
+    mEulerAngles.axis[1] = i1;
+    mEulerAngles.axis[2] = i2;
+
+    switch (mType)
+    {
+    case IS_MATRIX:
+        Convert(mMatrix, mEulerAngles);
+        break;
+    case IS_QUATERNION:
+        Convert(mQuaternion, mEulerAngles);
+        break;
+    case IS_AXIS_ANGLE:
+        Convert(mAxisAngle, mEulerAngles);
+        break;
+    case IS_EULER_ANGLES:
+        break;
+    }
+
+    return mEulerAngles;
+}
+
+template <int N, typename Real>
+void Rotation<N, Real>::Convert(Matrix<N, N, Real> const& r,
+    Quaternion<Real>& q)
+{
+    static_assert(N == 3 || N == 4, "Dimension must be 3 or 4.");
+
+    Real r22 = r(2, 2);
+    if (r22 <= (Real)0)  // x^2 + y^2 >= z^2 + w^2
+    {
+        Real dif10 = r(1, 1) - r(0, 0);
+        Real omr22 = (Real)1 - r22;
+        if (dif10 <= (Real)0)  // x^2 >= y^2
+        {
+            Real fourXSqr = omr22 - dif10;
+            Real inv4x = ((Real)0.5) / std::sqrt(fourXSqr);
+            q[0] = fourXSqr*inv4x;
+            q[1] = (r(0, 1) + r(1, 0))*inv4x;
+            q[2] = (r(0, 2) + r(2, 0))*inv4x;
+#if defined(GTE_USE_MAT_VEC)
+            q[3] = (r(2, 1) - r(1, 2))*inv4x;
+#else
+            q[3] = (r(1, 2) - r(2, 1))*inv4x;
+#endif
+        }
+        else  // y^2 >= x^2
+        {
+            Real fourYSqr = omr22 + dif10;
+            Real inv4y = ((Real)0.5) / std::sqrt(fourYSqr);
+            q[0] = (r(0, 1) + r(1, 0))*inv4y;
+            q[1] = fourYSqr*inv4y;
+            q[2] = (r(1, 2) + r(2, 1))*inv4y;
+#if defined(GTE_USE_MAT_VEC)
+            q[3] = (r(0, 2) - r(2, 0))*inv4y;
+#else
+            q[3] = (r(2, 0) - r(0, 2))*inv4y;
+#endif
+        }
+    }
+    else  // z^2 + w^2 >= x^2 + y^2
+    {
+        Real sum10 = r(1, 1) + r(0, 0);
+        Real opr22 = (Real)1 + r22;
+        if (sum10 <= (Real)0)  // z^2 >= w^2
+        {
+            Real fourZSqr = opr22 - sum10;
+            Real inv4z = ((Real)0.5) / std::sqrt(fourZSqr);
+            q[0] = (r(0, 2) + r(2, 0))*inv4z;
+            q[1] = (r(1, 2) + r(2, 1))*inv4z;
+            q[2] = fourZSqr*inv4z;
+#if defined(GTE_USE_MAT_VEC)
+            q[3] = (r(1, 0) - r(0, 1))*inv4z;
+#else
+            q[3] = (r(0, 1) - r(1, 0))*inv4z;
+#endif
+        }
+        else  // w^2 >= z^2
+        {
+            Real fourWSqr = opr22 + sum10;
+            Real inv4w = ((Real)0.5) / std::sqrt(fourWSqr);
+#if defined(GTE_USE_MAT_VEC)
+            q[0] = (r(2, 1) - r(1, 2))*inv4w;
+            q[1] = (r(0, 2) - r(2, 0))*inv4w;
+            q[2] = (r(1, 0) - r(0, 1))*inv4w;
+#else
+            q[0] = (r(1, 2) - r(2, 1))*inv4w;
+            q[1] = (r(2, 0) - r(0, 2))*inv4w;
+            q[2] = (r(0, 1) - r(1, 0))*inv4w;
+#endif
+            q[3] = fourWSqr*inv4w;
+        }
+    }
+}
+
+template <int N, typename Real>
+void Rotation<N, Real>::Convert(Quaternion<Real> const& q, Matrix<N, N, Real>& r)
+{
+    static_assert(N == 3 || N == 4, "Dimension must be 3 or 4.");
+
+    r.MakeIdentity();
+
+    Real twoX = ((Real)2)*q[0];
+    Real twoY = ((Real)2)*q[1];
+    Real twoZ = ((Real)2)*q[2];
+    Real twoXX = twoX*q[0];
+    Real twoXY = twoX*q[1];
+    Real twoXZ = twoX*q[2];
+    Real twoXW = twoX*q[3];
+    Real twoYY = twoY*q[1];
+    Real twoYZ = twoY*q[2];
+    Real twoYW = twoY*q[3];
+    Real twoZZ = twoZ*q[2];
+    Real twoZW = twoZ*q[3];
+
+#if defined(GTE_USE_MAT_VEC)
+    r(0, 0) = (Real)1 - twoYY - twoZZ;
+    r(0, 1) = twoXY - twoZW;
+    r(0, 2) = twoXZ + twoYW;
+    r(1, 0) = twoXY + twoZW;
+    r(1, 1) = (Real)1 - twoXX - twoZZ;
+    r(1, 2) = twoYZ - twoXW;
+    r(2, 0) = twoXZ - twoYW;
+    r(2, 1) = twoYZ + twoXW;
+    r(2, 2) = (Real)1 - twoXX - twoYY;
+#else
+    r(0, 0) = (Real)1 - twoYY - twoZZ;
+    r(1, 0) = twoXY - twoZW;
+    r(2, 0) = twoXZ + twoYW;
+    r(0, 1) = twoXY + twoZW;
+    r(1, 1) = (Real)1 - twoXX - twoZZ;
+    r(2, 1) = twoYZ - twoXW;
+    r(0, 2) = twoXZ - twoYW;
+    r(1, 2) = twoYZ + twoXW;
+    r(2, 2) = (Real)1 - twoXX - twoYY;
+#endif
+}
+
+template <int N, typename Real>
+void Rotation<N, Real>::Convert(Matrix<N, N, Real> const& r,
+    AxisAngle<N, Real>& a)
+{
+    static_assert(N == 3 || N == 4, "Dimension must be 3 or 4.");
+
+    Real trace = r(0, 0) + r(1, 1) + r(2, 2);
+    Real half = (Real)0.5;
+    Real cs = half*(trace - (Real)1);
+    cs = std::max(std::min(cs, (Real)1), (Real)-1);
+    a.angle = std::acos(cs);  // The angle is in [0,pi].
+    a.axis.MakeZero();
+
+    if (a.angle > (Real)0)
+    {
+        if (a.angle < (Real)GTE_C_PI)
+        {
+            // The angle is in (0,pi).
+#if defined(GTE_USE_MAT_VEC)
+            a.axis[0] = r(2, 1) - r(1, 2);
+            a.axis[1] = r(0, 2) - r(2, 0);
+            a.axis[2] = r(1, 0) - r(0, 1);
+            Normalize(a.axis);
+#else
+            a.axis[0] = r(1, 2) - r(2, 1);
+            a.axis[1] = r(2, 0) - r(0, 2);
+            a.axis[2] = r(0, 1) - r(1, 0);
+            Normalize(a.axis);
+#endif
+        }
+        else
+        {
+            // The angle is pi, in which case R is symmetric and
+            // R+I = 2*(I+S^2) = 2*U*U^T, where U = (u0,u1,u2) is the
+            // unit-length direction of the rotation axis.  Determine the
+            // largest diagonal entry of R+I and normalize the
+            // corresponding row to produce U.  It does not matter the
+            // sign on u[d] for chosen diagonal d, because R(U,pi) = R(-U,pi).
+            Real one = (Real)1;
+            if (r(0, 0) >= r(1, 1))
+            {
+                if (r(0, 0) >= r(2, 2))
+                {
+                    // r00 is maximum diagonal term
+                    a.axis[0] = r(0, 0) + one;
+                    a.axis[1] = half*(r(0, 1) + r(1, 0));
+                    a.axis[2] = half*(r(0, 2) + r(2, 0));
+                }
+                else
+                {
+                    // r22 is maximum diagonal term
+                    a.axis[0] = half*(r(2, 0) + r(0, 2));
+                    a.axis[1] = half*(r(2, 1) + r(1, 2));
+                    a.axis[2] = r(2, 2) + one;
+                }
+            }
+            else
+            {
+                if (r(1, 1) >= r(2, 2))
+                {
+                    // r11 is maximum diagonal term
+                    a.axis[0] = half*(r(1, 0) + r(0, 1));
+                    a.axis[1] = r(1, 1) + one;
+                    a.axis[2] = half*(r(1, 2) + r(2, 1));
+                }
+                else
+                {
+                    // r22 is maximum diagonal term
+                    a.axis[0] = half*(r(2, 0) + r(0, 2));
+                    a.axis[1] = half*(r(2, 1) + r(1, 2));
+                    a.axis[2] = r(2, 2) + one;
+                }
+            }
+            Normalize(a.axis);
+        }
+    }
+    else
+    {
+        // The angle is 0 and the matrix is the identity.  Any axis will
+        // work, so choose the Unit(0) axis.
+        a.axis[0] = (Real)1;
+    }
+}
+
+template <int N, typename Real>
+void Rotation<N, Real>::Convert(AxisAngle<N, Real> const& a,
+    Matrix<N, N, Real>& r)
+{
+    static_assert(N == 3 || N == 4, "Dimension must be 3 or 4.");
+
+    r.MakeIdentity();
+
+    Real cs = std::cos(a.angle);
+    Real sn = std::sin(a.angle);
+    Real oneMinusCos = ((Real)1) - cs;
+    Real x0sqr = a.axis[0] * a.axis[0];
+    Real x1sqr = a.axis[1] * a.axis[1];
+    Real x2sqr = a.axis[2] * a.axis[2];
+    Real x0x1m = a.axis[0] * a.axis[1] * oneMinusCos;
+    Real x0x2m = a.axis[0] * a.axis[2] * oneMinusCos;
+    Real x1x2m = a.axis[1] * a.axis[2] * oneMinusCos;
+    Real x0Sin = a.axis[0] * sn;
+    Real x1Sin = a.axis[1] * sn;
+    Real x2Sin = a.axis[2] * sn;
+
+#if defined(GTE_USE_MAT_VEC)
+    r(0, 0) = x0sqr*oneMinusCos + cs;
+    r(0, 1) = x0x1m - x2Sin;
+    r(0, 2) = x0x2m + x1Sin;
+    r(1, 0) = x0x1m + x2Sin;
+    r(1, 1) = x1sqr*oneMinusCos + cs;
+    r(1, 2) = x1x2m - x0Sin;
+    r(2, 0) = x0x2m - x1Sin;
+    r(2, 1) = x1x2m + x0Sin;
+    r(2, 2) = x2sqr*oneMinusCos + cs;
+#else
+    r(0, 0) = x0sqr*oneMinusCos + cs;
+    r(1, 0) = x0x1m - x2Sin;
+    r(2, 0) = x0x2m + x1Sin;
+    r(0, 1) = x0x1m + x2Sin;
+    r(1, 1) = x1sqr*oneMinusCos + cs;
+    r(2, 1) = x1x2m - x0Sin;
+    r(0, 2) = x0x2m - x1Sin;
+    r(1, 2) = x1x2m + x0Sin;
+    r(2, 2) = x2sqr*oneMinusCos + cs;
+#endif
+}
+
+template <int N, typename Real>
+void Rotation<N, Real>::Convert(Matrix<N, N, Real> const& r,
+    EulerAngles<Real>& e)
+{
+    static_assert(N == 3 || N == 4, "Dimension must be 3 or 4.");
+
+    if (0 <= e.axis[0] && e.axis[0] < 3
+        && 0 <= e.axis[1] && e.axis[1] < 3
+        && 0 <= e.axis[2] && e.axis[2] < 3
+        && e.axis[1] != e.axis[0]
+        && e.axis[1] != e.axis[2])
+    {
+        if (e.axis[0] != e.axis[2])
+        {
+#if defined(GTE_USE_MAT_VEC)
+            // Map (0,1,2), (1,2,0), and (2,0,1) to +1.
+            // Map (0,2,1), (2,1,0), and (1,0,2) to -1.
+            int parity = (((e.axis[2] | (e.axis[1] << 2)) >> e.axis[0]) & 1);
+            Real const sgn = (parity & 1 ? (Real)-1 : (Real)+1);
+
+            if (r(e.axis[2], e.axis[0]) < (Real)1)
+            {
+                if (r(e.axis[2], e.axis[0]) > (Real)-1)
+                {
+                    e.angle[2] = std::atan2(sgn*r(e.axis[1], e.axis[0]),
+                        r(e.axis[0], e.axis[0]));
+                    e.angle[1] = std::asin(-sgn*r(e.axis[2], e.axis[0]));
+                    e.angle[0] = std::atan2(sgn*r(e.axis[2], e.axis[1]),
+                        r(e.axis[2], e.axis[2]));
+                    e.result = ER_UNIQUE;
+                }
+                else
+                {
+                    e.angle[2] = (Real)0;
+                    e.angle[1] = sgn*(Real)GTE_C_HALF_PI;
+                    e.angle[0] = std::atan2(-sgn*r(e.axis[1], e.axis[2]),
+                        r(e.axis[1], e.axis[1]));
+                    e.result = ER_NOT_UNIQUE_DIF;
+                }
+            }
+            else
+            {
+                e.angle[2] = (Real)0;
+                e.angle[1] = -sgn*(Real)GTE_C_HALF_PI;
+                e.angle[0] = std::atan2(-sgn*r(e.axis[1], e.axis[2]),
+                    r(e.axis[1], e.axis[1]));
+                e.result = ER_NOT_UNIQUE_SUM;
+            }
+#else
+            // Map (0,1,2), (1,2,0), and (2,0,1) to +1.
+            // Map (0,2,1), (2,1,0), and (1,0,2) to -1.
+            int parity = (((e.axis[0] | (e.axis[1] << 2)) >> e.axis[2]) & 1);
+            Real const sgn = (parity & 1 ? (Real)+1 : (Real)-1);
+
+            if (r(e.axis[0], e.axis[2]) < (Real)1)
+            {
+                if (r(e.axis[0], e.axis[2]) > (Real)-1)
+                {
+                    e.angle[0] = std::atan2(sgn*r(e.axis[1], e.axis[2]),
+                        r(e.axis[2], e.axis[2]));
+                    e.angle[1] = std::asin(-sgn*r(e.axis[0], e.axis[2]));
+                    e.angle[2] = std::atan2(sgn*r(e.axis[0], e.axis[1]),
+                        r(e.axis[0], e.axis[0]));
+                    e.result = ER_UNIQUE;
+                }
+                else
+                {
+                    e.angle[0] = (Real)0;
+                    e.angle[1] = sgn*(Real)GTE_C_HALF_PI;
+                    e.angle[2] = std::atan2(-sgn*r(e.axis[1], e.axis[0]),
+                        r(e.axis[1], e.axis[1]));
+                    e.result = ER_NOT_UNIQUE_DIF;
+                }
+            }
+            else
+            {
+                e.angle[0] = (Real)0;
+                e.angle[1] = -sgn*(Real)GTE_C_HALF_PI;
+                e.angle[2] = std::atan2(-sgn*r(e.axis[1], e.axis[0]),
+                    r(e.axis[1], e.axis[1]));
+                e.result = ER_NOT_UNIQUE_SUM;
+            }
+#endif
+        }
+        else
+        {
+#if defined(GTE_USE_MAT_VEC)
+            // Map (0,2,0), (1,0,1), and (2,1,2) to +1.
+            // Map (0,1,0), (1,2,1), and (2,0,2) to -1.
+            int b0 = 3 - e.axis[1] - e.axis[2];
+            int parity = (((b0 | (e.axis[1] << 2)) >> e.axis[2]) & 1);
+            Real const sgn = (parity & 1 ? (Real)+1 : (Real)-1);
+
+            if (r(e.axis[2], e.axis[2]) < (Real)1)
+            {
+                if (r(e.axis[2], e.axis[2]) > (Real)-1)
+                {
+                    e.angle[2] = std::atan2(r(e.axis[1], e.axis[2]),
+                        sgn*r(b0, e.axis[2]));
+                    e.angle[1] = std::acos(r(e.axis[2], e.axis[2]));
+                    e.angle[0] = std::atan2(r(e.axis[2], e.axis[1]),
+                        -sgn*r(e.axis[2], b0));
+                    e.result = ER_UNIQUE;
+                }
+                else
+                {
+                    e.angle[2] = (Real)0;
+                    e.angle[1] = (Real)GTE_C_PI;
+                    e.angle[0] = std::atan2(sgn*r(e.axis[1], b0),
+                        r(e.axis[1], e.axis[1]));
+                    e.result = ER_NOT_UNIQUE_DIF;
+                }
+            }
+            else
+            {
+                e.angle[2] = (Real)0;
+                e.angle[1] = (Real)0;
+                e.angle[0] = std::atan2(sgn*r(e.axis[1], b0),
+                    r(e.axis[1], e.axis[1]));
+                e.result = ER_NOT_UNIQUE_SUM;
+            }
+#else
+            // Map (0,2,0), (1,0,1), and (2,1,2) to -1.
+            // Map (0,1,0), (1,2,1), and (2,0,2) to +1.
+            int b2 = 3 - e.axis[0] - e.axis[1];
+            int parity = (((b2 | (e.axis[1] << 2)) >> e.axis[0]) & 1);
+            Real const sgn = (parity & 1 ? (Real)-1 : (Real)+1);
+
+            if (r(e.axis[0], e.axis[0]) < (Real)1)
+            {
+                if (r(e.axis[0], e.axis[0]) > (Real)-1)
+                {
+                    e.angle[0] = std::atan2(r(e.axis[1], e.axis[0]),
+                        sgn*r(b2, e.axis[0]));
+                    e.angle[1] = std::acos(r(e.axis[0], e.axis[0]));
+                    e.angle[2] = std::atan2(r(e.axis[0], e.axis[1]),
+                        -sgn*r(e.axis[0], b2));
+                    e.result = ER_UNIQUE;
+                }
+                else
+                {
+                    e.angle[0] = (Real)0;
+                    e.angle[1] = (Real)GTE_C_PI;
+                    e.angle[2] = std::atan2(sgn*r(e.axis[1], b2),
+                        r(e.axis[1], e.axis[1]));
+                    e.result = ER_NOT_UNIQUE_DIF;
+                }
+            }
+            else
+            {
+                e.angle[0] = (Real)0;
+                e.angle[1] = (Real)0;
+                e.angle[2] = std::atan2(sgn*r(e.axis[1], b2),
+                    r(e.axis[1], e.axis[1]));
+                e.result = ER_NOT_UNIQUE_SUM;
+            }
+#endif
+        }
+    }
+    else
+    {
+        // Invalid angles.
+        e.angle[0] = (Real)0;
+        e.angle[1] = (Real)0;
+        e.angle[2] = (Real)0;
+        e.result = ER_INVALID;
+    }
+}
+
+template <int N, typename Real>
+void Rotation<N, Real>::Convert(EulerAngles<Real> const& e,
+    Matrix<N, N, Real>& r)
+{
+    static_assert(N == 3 || N == 4, "Dimension must be 3 or 4.");
+
+    if (0 <= e.axis[0] && e.axis[0] < 3
+        && 0 <= e.axis[1] && e.axis[1] < 3
+        && 0 <= e.axis[2] && e.axis[2] < 3
+        && e.axis[1] != e.axis[0]
+        && e.axis[1] != e.axis[2])
+    {
+        Matrix<N, N, Real> r0, r1, r2;
+        Convert(AxisAngle<N, Real>(Vector<N, Real>::Unit(e.axis[0]),
+            e.angle[0]), r0);
+        Convert(AxisAngle<N, Real>(Vector<N, Real>::Unit(e.axis[1]),
+            e.angle[1]), r1);
+        Convert(AxisAngle<N, Real>(Vector<N, Real>::Unit(e.axis[2]),
+            e.angle[2]), r2);
+#if defined(GTE_USE_MAT_VEC)
+        r = r2*r1*r0;
+#else
+        r = r0*r1*r2;
+#endif
+    }
+    else
+    {
+        // Invalid angles.
+        r.MakeIdentity();
+    }
+}
+
+template <int N, typename Real>
+void Rotation<N, Real>::Convert(Quaternion<Real> const& q,
+    AxisAngle<N, Real>& a)
+{
+    static_assert(N == 3 || N == 4, "Dimension must be 3 or 4.");
+
+    a.axis.MakeZero();
+
+    Real axisSqrLen = q[0] * q[0] + q[1] * q[1] + q[2] * q[2];
+    if (axisSqrLen > (Real)0)
+    {
+#if defined(GTE_USE_MAT_VEC)
+        Real adjust = ((Real)1) / std::sqrt(axisSqrLen);
+#else
+        Real adjust = ((Real)-1) / std::sqrt(axisSqrLen);
+#endif
+        a.axis[0] = q[0] * adjust;
+        a.axis[1] = q[1] * adjust;
+        a.axis[2] = q[2] * adjust;
+        Real cs = std::max(std::min(q[3], (Real)1), (Real)-1);
+        a.angle = (Real)2 * std::acos(cs);
+    }
+    else
+    {
+        // The angle is 0 (modulo 2*pi). Any axis will work, so choose the
+        // Unit(0) axis.
+        a.axis[0] = (Real)1;
+        a.angle = (Real)0;
+    }
+}
+
+template <int N, typename Real>
+void Rotation<N, Real>::Convert(AxisAngle<N, Real> const& a,
+    Quaternion<Real>& q)
+{
+    static_assert(N == 3 || N == 4, "Dimension must be 3 or 4.");
+
+#if defined(GTE_USE_MAT_VEC)
+    Real halfAngle = ((Real)0.5)*a.angle;
+#else
+    Real halfAngle = ((Real)-0.5)*a.angle;
+#endif
+    Real sn = std::sin(halfAngle);
+    q[0] = sn*a.axis[0];
+    q[1] = sn*a.axis[1];
+    q[2] = sn*a.axis[2];
+    q[3] = std::cos(halfAngle);
+}
+
+template <int N, typename Real>
+void Rotation<N, Real>::Convert(Quaternion<Real> const& q,
+    EulerAngles<Real>& e)
+{
+    static_assert(N == 3 || N == 4, "Dimension must be 3 or 4.");
+
+    Matrix<N, N, Real> r;
+    Convert(q, r);
+    Convert(r, e);
+}
+
+template <int N, typename Real>
+void Rotation<N, Real>::Convert(EulerAngles<Real> const& e,
+    Quaternion<Real>& q)
+{
+    static_assert(N == 3 || N == 4, "Dimension must be 3 or 4.");
+
+    Matrix<N, N, Real> r;
+    Convert(e, r);
+    Convert(r, q);
+}
+
+template <int N, typename Real>
+void Rotation<N, Real>::Convert(AxisAngle<N, Real> const& a,
+    EulerAngles<Real>& e)
+{
+    static_assert(N == 3 || N == 4, "Dimension must be 3 or 4.");
+
+    Quaternion<Real> q;
+    Convert(a, q);
+    Convert(q, e);
+}
+
+template <int N, typename Real>
+void Rotation<N, Real>::Convert(EulerAngles<Real> const& e,
+    AxisAngle<N, Real>& a)
+{
+    static_assert(N == 3 || N == 4, "Dimension must be 3 or 4.");
+
+    Quaternion<Real> q;
+    Convert(e, q);
+    Convert(q, a);
+}
+
+
+}
diff --git a/Engine/Plugins/Experimental/GeometryProcessing/Source/GeometryAlgorithms/Private/GTEngine/Mathematics/GteSector2.h b/Engine/Plugins/Experimental/GeometryProcessing/Source/GeometryAlgorithms/Private/GTEngine/Mathematics/GteSector2.h
new file mode 100644
--- /dev/null
+++ b/Engine/Plugins/Experimental/GeometryProcessing/Source/GeometryAlgorithms/Private/GTEngine/Mathematics/GteSector2.h
@@ -0,0 +1,165 @@
+// David Eberly, Geometric Tools, Redmond WA 98052
+// Copyright (c) 1998-2019
+// Distributed under the Boost Software License, Version 1.0.
+// http://www.boost.org/LICENSE_1_0.txt
+// http://www.geometrictools.com/License/Boost/LICENSE_1_0.txt
+// File Version: 3.0.1 (2018/10/05)
+
+#pragma once
+
+#include <GTEngine/Mathematics/GteMath.h>
+#include <GTEngine/Mathematics/GteVector2.h>
+
+namespace gte
+{
+
+// A solid sector is the intersection of a disk and a 2D cone.  The disk
+// has center C, radius R, and contains points X for which |X-C| <= R.  The
+// 2D cone has vertex C, unit-length axis direction D, angle A in (0,pi)
+// measured from D, and contains points X for which Dot(D,(X-C)/|X-C|) >= cos(A).
+// Sector points X satisfy both inequality constraints.
+
+template <typename Real>
+class Sector2
+{
+public:
+    // Construction and destruction.  The default constructor sets the vertex
+    // to (0,0), radius to 1, axis direction to (1,0), and angle to pi,
+    // all of which define a disk.
+    Sector2();
+    Sector2(Vector2<Real> const& inVertex, Real inRadius,
+        Vector2<Real> const& inDirection, Real inAngle);
+
+    // Set the angle and cos(angle) simultaneously.
+    void SetAngle(Real inAngle);
+
+    // Test whether P is in the sector. 
+    bool Contains(Vector2<Real> const& p) const;
+
+    // The cosine and sine of the angle are used in queries, so all o
+    // angle, cos(angle), and sin(angle) are stored. If you set 'angle'
+    // via the public members, you must set all to be consistent.  You
+    // can also call SetAngle(...) to ensure consistency.
+    Vector2<Real> vertex;
+    Real radius;
+    Vector2<Real> direction;
+    Real angle, cosAngle, sinAngle;
+
+public:
+    // Comparisons to support sorted containers.
+    bool operator==(Sector2 const& sector) const;
+    bool operator!=(Sector2 const& sector) const;
+    bool operator< (Sector2 const& sector) const;
+    bool operator<=(Sector2 const& sector) const;
+    bool operator> (Sector2 const& sector) const;
+    bool operator>=(Sector2 const& sector) const;
+};
+
+
+template <typename Real>
+Sector2<Real>::Sector2()
+    :
+    vertex(Vector2<Real>::Zero()),
+    radius((Real)1),
+    direction(Vector2<Real>::Unit(0)),
+    angle((Real)GTE_C_PI),
+    cosAngle((Real)-1),
+    sinAngle((Real)0)
+{
+}
+
+template <typename Real>
+Sector2<Real>::Sector2(Vector2<Real> const& inCenter, Real inRadius,
+    Vector2<Real> const& inDirection, Real inAngle)
+    :
+    vertex(inCenter),
+    radius(inRadius),
+    direction(inDirection)
+{
+    SetAngle(inAngle);
+}
+
+template <typename Real>
+void Sector2<Real>::SetAngle(Real inAngle)
+{
+    angle = inAngle;
+    cosAngle = std::cos(angle);
+    sinAngle = std::sin(angle);
+}
+
+template <typename Real>
+bool Sector2<Real>::Contains(Vector2<Real> const& p) const
+{
+    Vector2<Real> diff = p - vertex;
+    Real length = Length(diff);
+    return length <= radius && Dot(direction, diff) >= length * cosAngle;
+}
+
+template <typename Real>
+bool Sector2<Real>::operator==(Sector2 const& sector) const
+{
+    return vertex == sector.vertex && radius == sector.radius
+        && direction == sector.direction && angle == sector.angle;
+}
+
+template <typename Real>
+bool Sector2<Real>::operator!=(Sector2 const& sector) const
+{
+    return !operator==(sector);
+}
+
+template <typename Real>
+bool Sector2<Real>::operator<(Sector2 const& sector) const
+{
+    if (vertex < sector.vertex)
+    {
+        return true;
+    }
+
+    if (vertex > sector.vertex)
+    {
+        return false;
+    }
+
+    if (radius < sector.radius)
+    {
+        return true;
+    }
+
+    if (radius > sector.radius)
+    {
+        return false;
+    }
+
+    if (direction < sector.direction)
+    {
+        return true;
+    }
+
+    if (direction > sector.direction)
+    {
+        return false;
+    }
+
+    return angle < sector.angle;
+}
+
+template <typename Real>
+bool Sector2<Real>::operator<=(Sector2 const& sector) const
+{
+    return operator<(sector) || operator==(sector);
+}
+
+template <typename Real>
+bool Sector2<Real>::operator>(Sector2 const& sector) const
+{
+    return !operator<=(sector);
+}
+
+template <typename Real>
+bool Sector2<Real>::operator>=(Sector2 const& sector) const
+{
+    return !operator<(sector);
+}
+
+}
diff --git a/Engine/Plugins/Experimental/GeometryProcessing/Source/GeometryAlgorithms/Private/GTEngine/Mathematics/GteSegment.h b/Engine/Plugins/Experimental/GeometryProcessing/Source/GeometryAlgorithms/Private/GTEngine/Mathematics/GteSegment.h
new file mode 100644
--- /dev/null
+++ b/Engine/Plugins/Experimental/GeometryProcessing/Source/GeometryAlgorithms/Private/GTEngine/Mathematics/GteSegment.h
@@ -0,0 +1,148 @@
+// David Eberly, Geometric Tools, Redmond WA 98052
+// Copyright (c) 1998-2019
+// Distributed under the Boost Software License, Version 1.0.
+// http://www.boost.org/LICENSE_1_0.txt
+// http://www.geometrictools.com/License/Boost/LICENSE_1_0.txt
+// File Version: 3.0.0 (2016/06/19)
+
+#pragma once
+
+#include <GTEngine/Mathematics/GteVector.h>
+
+// The segment is represented by (1-t)*P0 + t*P1, where P0 and P1 are the
+// endpoints of the segment and 0 <= t <= 1.  Some algorithms prefer a
+// centered representation that is similar to how oriented bounding boxes are
+// defined.  This representation is C + s*D, where C = (P0 + P1)/2 is the
+// center of the segment, D = (P1 - P0)/|P1 - P0| is a unit-length direction
+// vector for the segment, and |t| <= e.  The value e = |P1 - P0|/2 is the
+// extent (or radius or half-length) of the segment.
+
+namespace gte
+{
+
+template <int N, typename Real>
+class Segment
+{
+public:
+    // Construction and destruction.  The default constructor sets p0 to
+    // (-1,0,...,0) and p1 to (1,0,...,0).  NOTE:  If you set p0 and p1;
+    // compute C, D, and e; and then recompute q0 = C-e*D and q1 = C+e*D,
+    // numerical round-off errors can lead to q0 not exactly equal to p0
+    // and q1 not exactly equal to p1.
+    Segment();
+    Segment(Vector<N, Real> const& p0, Vector<N, Real> const& p1);
+    Segment(std::array<Vector<N, Real>, 2> const& inP);
+    Segment(Vector<N, Real> const& center, Vector<N, Real> const& direction,
+        Real extent);
+
+    // Manipulation via the centered form.
+    void SetCenteredForm(Vector<N, Real> const& center,
+        Vector<N, Real> const& direction, Real extent);
+
+    void GetCenteredForm(Vector<N, Real>& center, Vector<N, Real>& direction,
+        Real& extent) const;
+
+    // Public member access.
+    std::array<Vector<N, Real>, 2> p;
+
+public:
+    // Comparisons to support sorted containers.
+    bool operator==(Segment const& segment) const;
+    bool operator!=(Segment const& segment) const;
+    bool operator< (Segment const& segment) const;
+    bool operator<=(Segment const& segment) const;
+    bool operator> (Segment const& segment) const;
+    bool operator>=(Segment const& segment) const;
+};
+
+// Template aliases for convenience.
+template <typename Real>
+using Segment2 = Segment<2, Real>;
+
+template <typename Real>
+using Segment3 = Segment<3, Real>;
+
+
+template <int N, typename Real>
+Segment<N, Real>::Segment()
+{
+    p[1].MakeUnit(0);
+    p[0] = -p[1];
+}
+
+template <int N, typename Real>
+Segment<N, Real>::Segment(Vector<N, Real> const& p0,
+    Vector<N, Real> const& p1)
+{
+    p[0] = p0;
+    p[1] = p1;
+}
+
+template <int N, typename Real>
+Segment<N, Real>::Segment(std::array<Vector<N, Real>, 2> const& inP)
+{
+    p = inP;
+}
+
+template <int N, typename Real>
+Segment<N, Real>::Segment(Vector<N, Real> const& center,
+    Vector<N, Real> const& direction, Real extent)
+{
+    SetCenteredForm(center, direction, extent);
+}
+
+template <int N, typename Real>
+void Segment<N, Real>::SetCenteredForm(Vector<N, Real> const& center,
+    Vector<N, Real> const& direction, Real extent)
+{
+    p[0] = center - extent * direction;
+    p[1] = center + extent * direction;
+}
+
+template <int N, typename Real>
+void Segment<N, Real>::GetCenteredForm(Vector<N, Real>& center,
+    Vector<N, Real>& direction, Real& extent) const
+{
+    center = ((Real)0.5)*(p[0] + p[1]);
+    direction = p[1] - p[0];
+    extent = ((Real)0.5)*Normalize(direction);
+}
+
+template <int N, typename Real>
+bool Segment<N, Real>::operator==(Segment const& segment) const
+{
+    return p == segment.p;
+}
+
+template <int N, typename Real>
+bool Segment<N, Real>::operator!=(Segment const& segment) const
+{
+    return !operator==(segment);
+}
+
+template <int N, typename Real>
+bool Segment<N, Real>::operator<(Segment const& segment) const
+{
+    return p < segment.p;
+}
+
+template <int N, typename Real>
+bool Segment<N, Real>::operator<=(Segment const& segment) const
+{
+    return operator<(segment) || operator==(segment);
+}
+
+template <int N, typename Real>
+bool Segment<N, Real>::operator>(Segment const& segment) const
+{
+    return !operator<=(segment);
+}
+
+template <int N, typename Real>
+bool Segment<N, Real>::operator>=(Segment const& segment) const
+{
+    return !operator<(segment);
+}
+
+
+}
diff --git a/Engine/Plugins/Experimental/GeometryProcessing/Source/GeometryAlgorithms/Private/GTEngine/Mathematics/GteSeparatePoints2.h b/Engine/Plugins/Experimental/GeometryProcessing/Source/GeometryAlgorithms/Private/GTEngine/Mathematics/GteSeparatePoints2.h
new file mode 100644
--- /dev/null
+++ b/Engine/Plugins/Experimental/GeometryProcessing/Source/GeometryAlgorithms/Private/GTEngine/Mathematics/GteSeparatePoints2.h
@@ -0,0 +1,221 @@
+// David Eberly, Geometric Tools, Redmond WA 98052
+// Copyright (c) 1998-2019
+// Distributed under the Boost Software License, Version 1.0.
+// http://www.boost.org/LICENSE_1_0.txt
+// http://www.geometrictools.com/License/Boost/LICENSE_1_0.txt
+// File Version: 3.0.0 (2016/06/19)
+
+#pragma once
+
+#include <GTEngine/Mathematics/GteConvexHull2.h>
+
+// Separate two point sets, if possible, by computing a line for which the
+// point sets lie on opposite sides.  The algorithm computes the convex hull
+// of the point sets, then uses the method of separating axes to determine if
+// the two convex polygons are disjoint.  The ComputeType is for the
+// ConvexHull2 class.
+
+namespace gte
+{
+
+template <typename Real, typename ComputeType>
+class SeparatePoints2
+{
+public:
+    // The return value is 'true' if and only if there is a separation.  If
+    // 'true', the returned line is a separating line.  The code assumes that
+    // each point set has at least 3 noncollinear points.
+    bool operator()(int numPoints0, Vector2<Real> const* points0,
+        int numPoints1, Vector2<Real> const* points1,
+        Line2<Real>& separatingLine) const;
+
+private:
+    int OnSameSide(Vector2<Real> const& lineNormal, Real lineConstant,
+        int numEdges, int const* edges, Vector2<Real> const* points) const;
+
+    int WhichSide(Vector2<Real> const& lineNormal, Real lineConstant,
+        int numEdges, int const* edges, Vector2<Real> const* points) const;
+};
+
+
+template <typename Real, typename ComputeType>
+bool SeparatePoints2<Real, ComputeType>::operator()(int numPoints0,
+    Vector2<Real> const* points0, int numPoints1,
+    Vector2<Real> const* points1, Line2<Real>& separatingLine) const
+{
+    // Construct convex hull of point set 0.
+    ConvexHull2<Real, ComputeType> ch0;
+    ch0(numPoints0, points0, (Real)0);
+    if (ch0.GetDimension() != 2)
+    {
+        return false;
+    }
+
+    // Construct convex hull of point set 1.
+    ConvexHull2<Real, ComputeType> ch1;
+    ch1(numPoints1, points1, (Real)0);
+    if (ch1.GetDimension() != 2)
+    {
+        return false;
+    }
+
+    int numEdges0 = static_cast<int>(ch0.GetHull().size());
+    int const* edges0 = &ch0.GetHull()[0];
+    int numEdges1 = static_cast<int>(ch1.GetHull().size());
+    int const* edges1 = &ch1.GetHull()[0];
+
+    // Test edges of hull 0 for possible separation of points.
+    int j0, j1, i0, i1, side0, side1;
+    Vector2<Real> lineNormal;
+    Real lineConstant;
+    for (j1 = 0, j0 = numEdges0 - 1; j1 < numEdges0; j0 = j1++)
+    {
+        // Look up edge (assert: i0 != i1 ).
+        i0 = edges0[j0];
+        i1 = edges0[j1];
+
+        // Compute potential separating line (assert: (xNor,yNor) != (0,0)).
+        separatingLine.origin = points0[i0];
+        separatingLine.direction = points0[i1] - points0[i0];
+        Normalize(separatingLine.direction);
+        lineNormal = Perp(separatingLine.direction);
+        lineConstant = Dot(lineNormal, separatingLine.origin);
+
+        // Determine whether hull 1 is on same side of line.
+        side1 = OnSameSide(lineNormal, lineConstant, numEdges1, edges1,
+            points1);
+
+        if (side1)
+        {
+            // Determine on which side of line hull 0 lies.
+            side0 = WhichSide(lineNormal, lineConstant, numEdges0,
+                edges0, points0);
+
+            if (side0*side1 <= 0)  // Line separates hulls.
+            {
+                return true;
+            }
+        }
+    }
+
+    // Test edges of hull 1 for possible separation of points.
+    for (j1 = 0, j0 = numEdges1 - 1; j1 < numEdges1; j0 = j1++)
+    {
+        // Look up edge (assert: i0 != i1 ).
+        i0 = edges1[j0];
+        i1 = edges1[j1];
+
+        // Compute perpendicular to edge (assert: (xNor,yNor) != (0,0)).
+        separatingLine.origin = points1[i0];
+        separatingLine.direction = points1[i1] - points1[i0];
+        Normalize(separatingLine.direction);
+        lineNormal = Perp(separatingLine.direction);
+        lineConstant = Dot(lineNormal, separatingLine.origin);
+
+        // Determine whether hull 0 is on same side of line.
+        side0 = OnSameSide(lineNormal, lineConstant, numEdges0, edges0,
+            points0);
+
+        if (side0)
+        {
+            // Determine on which side of line hull 1 lies.
+            side1 = WhichSide(lineNormal, lineConstant, numEdges1,
+                edges1, points1);
+
+            if (side0*side1 <= 0)  // Line separates hulls.
+            {
+                return true;
+            }
+        }
+    }
+
+    return false;
+}
+
+template <typename Real, typename ComputeType>
+int SeparatePoints2<Real, ComputeType>::OnSameSide(
+    Vector2<Real> const& lineNormal, Real lineConstant, int numEdges,
+    int const* edges, Vector2<Real> const* points) const
+{
+    // Test whether all points on same side of line Dot(N,X) = c.
+    Real c0;
+    int posSide = 0, negSide = 0;
+
+    for (int i1 = 0, i0 = numEdges - 1; i1 < numEdges; i0 = i1++)
+    {
+        c0 = Dot(lineNormal, points[edges[i0]]);
+        if (c0 > lineConstant)
+        {
+            ++posSide;
+        }
+        else if (c0 < lineConstant)
+        {
+            ++negSide;
+        }
+
+        if (posSide && negSide)
+        {
+            // Line splits point set.
+            return 0;
+        }
+
+        c0 = Dot(lineNormal, points[edges[i1]]);
+        if (c0 > lineConstant)
+        {
+            ++posSide;
+        }
+        else if (c0 < lineConstant)
+        {
+            ++negSide;
+        }
+
+        if (posSide && negSide)
+        {
+            // Line splits point set.
+            return 0;
+        }
+    }
+
+    return (posSide ? +1 : -1);
+}
+
+template <typename Real, typename ComputeType>
+int SeparatePoints2<Real, ComputeType>::WhichSide(
+    Vector2<Real> const& lineNormal, Real lineConstant, int numEdges,
+    int const* edges, Vector2<Real> const* points) const
+{
+    // Establish which side of line hull is on.
+    Real c0;
+    for (int i1 = 0, i0 = numEdges - 1; i1 < numEdges; i0 = i1++)
+    {
+        c0 = Dot(lineNormal, points[edges[i0]]);
+        if (c0 > lineConstant)
+        {
+            // Hull on positive side.
+            return +1;
+        }
+        if (c0 < lineConstant)
+        {
+            // Hull on negative side.
+            return -1;
+        }
+
+        c0 = Dot(lineNormal, points[edges[i1]]);
+        if (c0 > lineConstant)
+        {
+            // Hull on positive side.
+            return +1;
+        }
+        if (c0 < lineConstant)
+        {
+            // Hull on negative side.
+            return -1;
+        }
+    }
+
+    // Hull is effectively collinear.
+    return 0;
+}
+
+
+}
diff --git a/Engine/Plugins/Experimental/GeometryProcessing/Source/GeometryAlgorithms/Private/GTEngine/Mathematics/GteSeparatePoints3.h b/Engine/Plugins/Experimental/GeometryProcessing/Source/GeometryAlgorithms/Private/GTEngine/Mathematics/GteSeparatePoints3.h
new file mode 100644
--- /dev/null
+++ b/Engine/Plugins/Experimental/GeometryProcessing/Source/GeometryAlgorithms/Private/GTEngine/Mathematics/GteSeparatePoints3.h
@@ -0,0 +1,244 @@
+// David Eberly, Geometric Tools, Redmond WA 98052
+// Copyright (c) 1998-2019
+// Distributed under the Boost Software License, Version 1.0.
+// http://www.boost.org/LICENSE_1_0.txt
+// http://www.geometrictools.com/License/Boost/LICENSE_1_0.txt
+// File Version: 3.0.0 (2016/06/19)
+
+#pragma once
+
+#include <GTEngine/Mathematics/GteConvexHull3.h>
+
+// Separate two point sets, if possible, by computing a plane for which the
+// point sets lie on opposite sides.  The algorithm computes the convex hull
+// of the point sets, then uses the method of separating axes to determine if
+// the two convex polyhedra are disjoint.  The ComputeType is for the
+// ConvexHull3 class.
+
+namespace gte
+{
+
+template <typename Real, typename ComputeType>
+class SeparatePoints3
+{
+public:
+    // The return value is 'true' if and only if there is a separation.  If
+    // 'true', the returned plane is a separating plane.  The code assumes
+    // that each point set has at least 4 noncoplanar points.
+    bool operator()(int numPoints0, Vector3<Real> const* points0,
+        int numPoints1, Vector3<Real> const* points1,
+        Plane3<Real>& separatingPlane) const;
+
+private:
+    int OnSameSide(Plane3<Real> const& plane, int numTriangles,
+        int const* indices, Vector3<Real> const* points) const;
+
+    int WhichSide(Plane3<Real> const& plane, int numTriangles,
+        int const* indices, Vector3<Real> const* points) const;
+};
+
+
+template <typename Real, typename ComputeType>
+bool SeparatePoints3<Real, ComputeType>::operator()(int numPoints0,
+    Vector3<Real> const* points0, int numPoints1,
+    Vector3<Real> const* points1, Plane3<Real>& separatingPlane) const
+{
+    // Construct convex hull of point set 0.
+    ConvexHull3<Real, ComputeType> ch0;
+    ch0(numPoints0, points0, (Real)0);
+    if (ch0.GetDimension() != 3)
+    {
+        return false;
+    }
+
+    // Construct convex hull of point set 1.
+    ConvexHull3<Real, ComputeType> ch1;
+    ch1(numPoints1, points1, (Real)0);
+    if (ch1.GetDimension() != 3)
+    {
+        return false;
+    }
+
+    auto const& hull0 = ch0.GetHullUnordered();
+    auto const& hull1 = ch1.GetHullUnordered();
+    int numTriangles0 = static_cast<int>(hull0.size());
+    int const* indices0 = reinterpret_cast<int const*>(&hull0[0]);
+    int numTriangles1 = static_cast<int>(hull1.size());
+    int const* indices1 = reinterpret_cast<int const*>(&hull1[0]);
+
+    // Test faces of hull 0 for possible separation of points.
+    int i, i0, i1, i2, side0, side1;
+    Vector3<Real> diff0, diff1;
+    for (i = 0; i < numTriangles0; ++i)
+    {
+        // Look up face (assert: i0 != i1 && i0 != i2 && i1 != i2).
+        i0 = indices0[3 * i];
+        i1 = indices0[3 * i + 1];
+        i2 = indices0[3 * i + 2];
+
+        // Compute potential separating plane (assert: normal != (0,0,0)).
+        separatingPlane = Plane3<Real>({ points0[i0], points0[i1],
+            points0[i2] });
+
+        // Determine whether hull 1 is on same side of plane.
+        side1 = OnSameSide(separatingPlane, numTriangles1, indices1, points1);
+
+        if (side1)
+        {
+            // Determine on which side of plane hull 0 lies.
+            side0 = WhichSide(separatingPlane, numTriangles0, indices0,
+                points0);
+            if (side0*side1 <= 0)  // Plane separates hulls.
+            {
+                return true;
+            }
+        }
+    }
+
+    // Test faces of hull 1 for possible separation of points.
+    for (i = 0; i < numTriangles1; ++i)
+    {
+        // Look up edge (assert: i0 != i1 && i0 != i2 && i1 != i2).
+        i0 = indices1[3 * i];
+        i1 = indices1[3 * i + 1];
+        i2 = indices1[3 * i + 2];
+
+        // Compute perpendicular to face (assert: normal != (0,0,0)).
+        separatingPlane = Plane3<Real>({ points1[i0], points1[i1],
+            points1[i2] });
+
+        // Determine whether hull 0 is on same side of plane.
+        side0 = OnSameSide(separatingPlane, numTriangles0, indices0, points0);
+        if (side0)
+        {
+            // Determine on which side of plane hull 1 lies.
+            side1 = WhichSide(separatingPlane, numTriangles1, indices1,
+                points1);
+            if (side0*side1 <= 0)  // Plane separates hulls.
+            {
+                return true;
+            }
+        }
+    }
+
+    // Build edge set for hull 0.
+    std::set<std::pair<int, int>> edgeSet0;
+    for (i = 0; i < numTriangles0; ++i)
+    {
+        // Look up face (assert: i0 != i1 && i0 != i2 && i1 != i2).
+        i0 = indices0[3 * i];
+        i1 = indices0[3 * i + 1];
+        i2 = indices0[3 * i + 2];
+        edgeSet0.insert(std::make_pair(i0, i1));
+        edgeSet0.insert(std::make_pair(i0, i2));
+        edgeSet0.insert(std::make_pair(i1, i2));
+    }
+
+    // Build edge list for hull 1.
+    std::set<std::pair<int, int>> edgeSet1;
+    for (i = 0; i < numTriangles1; ++i)
+    {
+        // Look up face (assert: i0 != i1 && i0 != i2 && i1 != i2).
+        i0 = indices1[3 * i];
+        i1 = indices1[3 * i + 1];
+        i2 = indices1[3 * i + 2];
+        edgeSet1.insert(std::make_pair(i0, i1));
+        edgeSet1.insert(std::make_pair(i0, i2));
+        edgeSet1.insert(std::make_pair(i1, i2));
+    }
+
+    // Test planes whose normals are cross products of two edges, one from
+    // each hull.
+    for (auto const& e0 : edgeSet0)
+    {
+        // Get edge.
+        diff0 = points0[e0.second] - points0[e0.first];
+
+        for (auto const& e1 : edgeSet1)
+        {
+            diff1 = points1[e1.second] - points1[e1.first];
+
+            // Compute potential separating plane.
+            separatingPlane.normal = UnitCross(diff0, diff1);
+            separatingPlane.constant = Dot(separatingPlane.normal,
+                points0[e0.first]);
+
+            // Determine if hull 0 is on same side of plane.
+            side0 = OnSameSide(separatingPlane, numTriangles0, indices0,
+                points0);
+            side1 = OnSameSide(separatingPlane, numTriangles1, indices1,
+                points1);
+
+            if (side0*side1 < 0)  // Plane separates hulls.
+            {
+                return true;
+            }
+        }
+    }
+
+    return false;
+}
+
+template <typename Real, typename ComputeType>
+int SeparatePoints3<Real, ComputeType>::OnSameSide(Plane3<Real> const& plane,
+    int numTriangles, int const* indices, Vector3<Real> const* points) const
+{
+    // test if all points on same side of plane (nx,ny,nz)*(x,y,z) = c
+    int posSide = 0, negSide = 0;
+
+    for (int t = 0; t < numTriangles; ++t)
+    {
+        for (int i = 0; i < 3; ++i)
+        {
+            int v = indices[3 * t + i];
+            Real c0 = Dot(plane.normal, points[v]);
+            if (c0 > plane.constant)
+            {
+                ++posSide;
+            }
+            else if (c0 < plane.constant)
+            {
+                ++negSide;
+            }
+
+            if (posSide && negSide)
+            {
+                // Plane splits point set.
+                return 0;
+            }
+        }
+    }
+
+    return (posSide ? +1 : -1);
+}
+
+template <typename Real, typename ComputeType>
+int SeparatePoints3<Real, ComputeType>::WhichSide(Plane3<Real> const& plane,
+    int numTriangles, int const* indices, Vector3<Real> const* points) const
+{
+    // Establish which side of plane hull is on.
+    for (int t = 0; t < numTriangles; ++t)
+    {
+        for (int i = 0; i < 3; ++i)
+        {
+            int v = indices[3 * t + i];
+            Real c0 = Dot(plane.normal, points[v]);
+            if (c0 > plane.constant)
+            {
+                // Positive side.
+                return +1;
+            }
+            if (c0 < plane.constant)
+            {
+                // Negative side.
+                return -1;
+            }
+        }
+    }
+
+    // Hull is effectively collinear.
+    return 0;
+}
+
+
+}
