move GeometryProcessing plugin from NotForLicensees to Experimental

diff --git a/Engine/Plugins/Experimental/GeometryProcessing/Source/GeometryAlgorithms/Private/GTEngine/Mathematics/GteSinEstimate.h b/Engine/Plugins/Experimental/GeometryProcessing/Source/GeometryAlgorithms/Private/GTEngine/Mathematics/GteSinEstimate.h
new file mode 100644
--- /dev/null
+++ b/Engine/Plugins/Experimental/GeometryProcessing/Source/GeometryAlgorithms/Private/GTEngine/Mathematics/GteSinEstimate.h
@@ -0,0 +1,161 @@
+// David Eberly, Geometric Tools, Redmond WA 98052
+// Copyright (c) 1998-2019
+// Distributed under the Boost Software License, Version 1.0.
+// http://www.boost.org/LICENSE_1_0.txt
+// http://www.geometrictools.com/License/Boost/LICENSE_1_0.txt
+// File Version: 3.0.1 (2018/10/05)
+
+#pragma once
+
+#include <GTEngine/Mathematics/GteMath.h>
+
+// Minimax polynomial approximations to sin(x).  The polynomial p(x) of
+// degree D has only odd-power terms, is required to have linear term x,
+// and p(pi/2) = sin(pi/2) = 1.  It minimizes the quantity
+// maximum{|sin(x) - p(x)| : x in [-pi/2,pi/2]} over all polynomials of
+// degree D subject to the constraints mentioned.
+
+namespace gte
+{
+
+template <typename Real>
+class SinEstimate
+{
+public:
+    // The input constraint is x in [-pi/2,pi/2].  For example,
+    //   float x; // in [-pi/2,pi/2]
+    //   float result = SinEstimate<float>::Degree<3>(x);
+    template <int D>
+    inline static Real Degree(Real x);
+
+    // The input x can be any real number.  Range reduction is used to
+    // generate a value y in [-pi/2,pi/2] for which sin(y) = sin(x).
+    // For example,
+    //   float x;  // x any real number
+    //   float result = SinEstimate<float>::DegreeRR<3>(x);
+    template <int D>
+    inline static Real DegreeRR(Real x);
+
+private:
+    // Metaprogramming and private implementation to allow specialization of
+    // a template member function.
+    template <int D> struct degree {};
+    inline static Real Evaluate(degree<3>, Real x);
+    inline static Real Evaluate(degree<5>, Real x);
+    inline static Real Evaluate(degree<7>, Real x);
+    inline static Real Evaluate(degree<9>, Real x);
+    inline static Real Evaluate(degree<11>, Real x);
+
+    // Support for range reduction.
+    inline static Real Reduce(Real x);
+};
+
+
+template <typename Real>
+template <int D>
+inline Real SinEstimate<Real>::Degree(Real x)
+{
+    return Evaluate(degree<D>(), x);
+}
+
+template <typename Real>
+template <int D>
+inline Real SinEstimate<Real>::DegreeRR(Real x)
+{
+    return Degree<D>(Reduce(x));
+}
+
+template <typename Real>
+inline Real SinEstimate<Real>::Evaluate(degree<3>, Real x)
+{
+    Real xsqr = x * x;
+    Real poly;
+    poly = (Real)GTE_C_SIN_DEG3_C1;
+    poly = (Real)GTE_C_SIN_DEG3_C0 + poly * xsqr;
+    poly = poly * x;
+    return poly;
+}
+
+template <typename Real>
+inline Real SinEstimate<Real>::Evaluate(degree<5>, Real x)
+{
+    Real xsqr = x * x;
+    Real poly;
+    poly = (Real)GTE_C_SIN_DEG5_C2;
+    poly = (Real)GTE_C_SIN_DEG5_C1 + poly * xsqr;
+    poly = (Real)GTE_C_SIN_DEG5_C0 + poly * xsqr;
+    poly = poly * x;
+    return poly;
+}
+
+template <typename Real>
+inline Real SinEstimate<Real>::Evaluate(degree<7>, Real x)
+{
+    Real xsqr = x * x;
+    Real poly;
+    poly = (Real)GTE_C_SIN_DEG7_C3;
+    poly = (Real)GTE_C_SIN_DEG7_C2 + poly * xsqr;
+    poly = (Real)GTE_C_SIN_DEG7_C1 + poly * xsqr;
+    poly = (Real)GTE_C_SIN_DEG7_C0 + poly * xsqr;
+    poly = poly * x;
+    return poly;
+}
+
+template <typename Real>
+inline Real SinEstimate<Real>::Evaluate(degree<9>, Real x)
+{
+    Real xsqr = x * x;
+    Real poly;
+    poly = (Real)GTE_C_SIN_DEG9_C4;
+    poly = (Real)GTE_C_SIN_DEG9_C3 + poly * xsqr;
+    poly = (Real)GTE_C_SIN_DEG9_C2 + poly * xsqr;
+    poly = (Real)GTE_C_SIN_DEG9_C1 + poly * xsqr;
+    poly = (Real)GTE_C_SIN_DEG9_C0 + poly * xsqr;
+    poly = poly * x;
+    return poly;
+}
+
+template <typename Real>
+inline Real SinEstimate<Real>::Evaluate(degree<11>, Real x)
+{
+    Real xsqr = x * x;
+    Real poly;
+    poly = (Real)GTE_C_SIN_DEG11_C5;
+    poly = (Real)GTE_C_SIN_DEG11_C4 + poly * xsqr;
+    poly = (Real)GTE_C_SIN_DEG11_C3 + poly * xsqr;
+    poly = (Real)GTE_C_SIN_DEG11_C2 + poly * xsqr;
+    poly = (Real)GTE_C_SIN_DEG11_C1 + poly * xsqr;
+    poly = (Real)GTE_C_SIN_DEG11_C0 + poly * xsqr;
+    poly = poly * x;
+    return poly;
+}
+
+template <typename Real>
+inline Real SinEstimate<Real>::Reduce(Real x)
+{
+    // Map x to y in [-pi,pi], x = 2*pi*quotient + remainder.
+    Real quotient = (Real)GTE_C_INV_TWO_PI * x;
+    if (x >= (Real)0)
+    {
+        quotient = (Real)((int)(quotient + (Real)0.5));
+    }
+    else
+    {
+        quotient = (Real)((int)(quotient - (Real)0.5));
+    }
+    Real y = x - (Real)GTE_C_TWO_PI * quotient;
+
+    // Map y to [-pi/2,pi/2] with sin(y) = sin(x).
+    if (y > (Real)GTE_C_HALF_PI)
+    {
+        y = (Real)GTE_C_PI - y;
+    }
+    else if (y < (Real)-GTE_C_HALF_PI)
+    {
+        y = (Real)-GTE_C_PI - y;
+    }
+    return y;
+}
+
+
+}
diff --git a/Engine/Plugins/Experimental/GeometryProcessing/Source/GeometryAlgorithms/Private/GTEngine/Mathematics/GteSingularValueDecomposition.h b/Engine/Plugins/Experimental/GeometryProcessing/Source/GeometryAlgorithms/Private/GTEngine/Mathematics/GteSingularValueDecomposition.h
new file mode 100644
--- /dev/null
+++ b/Engine/Plugins/Experimental/GeometryProcessing/Source/GeometryAlgorithms/Private/GTEngine/Mathematics/GteSingularValueDecomposition.h
@@ -0,0 +1,1147 @@
+// David Eberly, Geometric Tools, Redmond WA 98052
+// Copyright (c) 1998-2019
+// Distributed under the Boost Software License, Version 1.0.
+// http://www.boost.org/LICENSE_1_0.txt
+// http://www.geometrictools.com/License/Boost/LICENSE_1_0.txt
+// File Version: 3.0.1 (2018/10/05)
+
+#pragma once
+
+#include <GTEngine/LowLevel/GteRangeIteration.h>
+#include <GTEngine/LowLevel/GteWrapper.h>
+#include <algorithm>
+#include <cmath>
+#include <cstring>
+#include <vector>
+
+// The SingularValueDecomposition class is an implementation of Algorithm
+// 8.3.2 (The SVD Algorithm) described in "Matrix Computations, 2nd
+// edition" by G. H. Golub and Charles F. Van Loan, The Johns Hopkins
+// Press, Baltimore MD, Fourth Printing 1993.  Algorithm 5.4.2 (Householder
+// Bidiagonalization) is used to reduce A to bidiagonal B.  Algorithm 8.3.1
+// (Golub-Kahan SVD Step) is used for the iterative reduction from bidiagonal
+// to diagonal.  If A is the original matrix, S is the matrix whose diagonal
+// entries are the singular values, and U and V are corresponding matrices,
+// then theoretically U^T*A*V = S.  Numerically, we have errors
+// E = U^T*A*V - S.  Algorithm 8.3.2 mentions that one expects |E| is
+// approximately u*|A|, where |M| denotes the Frobenius norm of M and where
+// u is the unit roundoff for the floating-point arithmetic: 2^{-23} for
+// 'float', which is FLT_EPSILON = 1.192092896e-7f, and 2^{-52} for'double',
+// which is DBL_EPSILON = 2.2204460492503131e-16.
+//
+// The condition |a(i,i+1)| <= epsilon*(|a(i,i) + a(i+1,i+1)|) used to
+// determine when the reduction decouples to smaller problems is implemented
+// as:  sum = |a(i,i)| + |a(i+1,i+1)|; sum + |a(i,i+1)| == sum.  The idea is
+// that the superdiagonal term is small relative to its diagonal neighbors,
+// and so it is effectively zero.  The unit tests have shown that this
+// interpretation of decoupling is effective.
+//
+// The condition |a(i,i)| <= epsilon*|B| used to determine when the
+// reduction decouples (with a zero singular value) is implemented using
+// the Frobenius norm of B and epsilon = multiplier*u, where for now the
+// multiplier is hard-coded in Solve(...) as 8.
+//
+// The authors suggest that once you have the bidiagonal matrix, a practical
+// implementation will store the diagonal and superdiagonal entries in linear
+// arrays, ignoring the theoretically zero values not in the 2-band.  This is
+// good for cache coherence, and we have used the suggestion.  The essential
+// parts of the Householder u-vectors are stored in the lower-triangular
+// portion of the matrix and the essential parts of the Householder v-vectors
+// are stored in the upper-triangular portion of the matrix.  To avoid having
+// to recompute 2/Dot(u,u) and 2/Dot(v,v) when constructing orthogonal U and
+// V, we store these quantities in additional memory during bidiagonalization.
+//
+// For matrices with randomly generated values in [0,1], the unit tests
+// produce the following information for N-by-N matrices.
+//
+// N  |A|     |E|        |E|/|A|    iterations
+// -------------------------------------------
+//  2  1.4831 4.1540e-16 2.8007e-16  1
+//  3  2.1065 3.5024e-16 1.6626e-16  4
+//  4  2.4979 7.4605e-16 2.9867e-16  6
+//  5  3.6591 1.8305e-15 5.0025e-16  9
+//  6  4.0572 2.0571e-15 5.0702e-16 10
+//  7  4.7745 2.9057e-15 6.0859e-16 12
+//  8  5.1964 2.7958e-15 5.3803e-16 13
+//  9  5.7599 3.3128e-15 5.7514e-16 16
+// 10  6.2700 3.7209e-15 5.9344e-16 16
+// 11  6.8220 5.0580e-15 7.4142e-16 18
+// 12  7.4540 5.2493e-15 7.0422e-16 21
+// 13  8.1225 5.6043e-15 6.8997e-16 24
+// 14  8.5883 5.8553e-15 6.8177e-16 26
+// 15  9.1337 6.9663e-15 7.6270e-16 27
+// 16  9.7884 9.1358e-15 9.3333e-16 29
+// 17 10.2407 8.2715e-15 8.0771e-16 34
+// 18 10.7147 8.9748e-15 8.3761e-16 33
+// 19 11.1887 1.0094e-14 9.0220e-16 32
+// 20 11.7739 9.7000e-15 8.2386e-16 35
+// 21 12.2822 1.1217e-14 9.1329e-16 36
+// 22 12.7649 1.1071e-14 8.6732e-16 37
+// 23 13.3366 1.1271e-14 8.4513e-16 41
+// 24 13.8505 1.2806e-14 9.2460e-16 43
+// 25 14.4332 1.3081e-14 9.0637e-16 43
+// 26 14.9301 1.4882e-14 9.9680e-16 46
+// 27 15.5214 1.5571e-14 1.0032e-15 48
+// 28 16.1029 1.7553e-14 1.0900e-15 49
+// 29 16.6407 1.6219e-14 9.7467e-16 53
+// 30 17.1891 1.8560e-14 1.0797e-15 55
+// 31 17.7773 1.8522e-14 1.0419e-15 56
+//
+// The singularvalues and |E|/|A| values were compared to those generated by
+// Mathematica Version 9.0; Wolfram Research, Inc., Champaign IL, 2012.
+// The results were all comparable with singular values agreeing to a large
+// number of decimal places.
+//
+// Timing on an Intel (R) Core (TM) i7-3930K CPU @ 3.20 GHZ (in seconds)
+// for NxN matrices:
+//
+// N    |E|/|A|    iters bidiag  QR     U-and-V    comperr
+// -------------------------------------------------------
+//  512 3.8632e-15  848   0.341  0.016    1.844      2.203
+// 1024 5.6456e-15 1654   4.279  0.032   18.765     20.844 
+// 2048 7.5499e-15 3250  40.421  0.141  186.607    213.216
+//
+// where iters is the number of QR steps taken, bidiag is the computation
+// of the Householder reflection vectors, U-and-V is the composition of
+// Householder reflections and Givens rotations to obtain the orthogonal
+// matrices of the decomposigion, and comperr is the computation E =
+// U^T*A*V - S.
+
+namespace gte
+{
+
+template <typename Real>
+class SingularValueDecomposition
+{
+public:
+    // The solver processes MxN symmetric matrices, where M >= N > 1
+    // ('numRows' is M and 'numCols' is N) and the matrix is stored in
+    // row-major order.  The maximum number of iterations ('maxIterations')
+    // must be specified for the reduction of a bidiagonal matrix to a
+    // diagonal matrix.  The goal is to compute MxM orthogonal U, NxN
+    // orthogonal V, and MxN matrix S for which U^T*A*V = S.  The only
+    // nonzero entries of S are on the diagonal; the diagonal entries are
+    // the singular values of the original matrix.
+    SingularValueDecomposition(int numRows, int numCols,
+        unsigned int maxIterations);
+
+    // A copy of the MxN input is made internally.  The order of the singular
+    // values is specified by sortType: -1 (decreasing), 0 (no sorting), or +1
+    // (increasing).  When sorted, the columns of the orthogonal matrices
+    // are ordered accordingly.  The return value is the number of iterations
+    // consumed when convergence occurred, 0xFFFFFFFF when convergence did not
+    // occur or 0 when N <= 1 or M < N was passed to the constructor.
+    unsigned int Solve(Real const* input, int sortType);
+
+    // Get the singular values of the matrix passed to Solve(...).  The input
+    // 'singularValues' must have N elements.
+    void GetSingularValues(Real* singularValues) const;
+
+    // Accumulate the Householder reflections, the Givens rotations, and the
+    // diagonal fix-up matrix to compute the orthogonal matrices U and V for
+    // which U^T*A*V = S.  The input uMatrix must be MxM and the input vMatrix
+    // must be NxN, both stored in row-major order.
+    void GetU(Real* uMatrix) const;
+    void GetV(Real* vMatrix) const;
+
+    // Compute a single column of U or V.  The reflections and rotations are
+    // applied incrementally.  This is useful when you want only a small
+    // number of the singular values or vectors.
+    void GetUColumn(int index, Real* uColumn) const;
+    void GetVColumn(int index, Real* vColumn) const;
+    Real GetSingularValue(int index) const;
+
+private:
+    // Bidiagonalize using Householder reflections.  On input, mMatrix is a
+    // copy of the input matrix and has one extra row.  On output, the
+    // diagonal and superdiagonal contain the bidiagonalized results.  The
+    // lower-triangular portion stores the essential parts of the Householder
+    // u vectors (the elements of u after the leading 1-valued component) and
+    // the upper-triangular portion stores the essential parts of the
+    // Householder v vectors.  To avoid recomputing 2/Dot(u,u) and 2/Dot(v,v),
+    // these quantities are stored in mTwoInvUTU and mTwoInvVTV.
+    void Bidiagonalize();
+
+    // A helper for generating Givens rotation sine and cosine robustly.
+    void GetSinCos(Real u, Real v, Real& cs, Real& sn);
+
+    // Test for (effectively) zero-valued diagonal entries (through all but
+    // the last).  For each such entry, the B matrix decouples.  Perform
+    // that decoupling.  If there are no zero-valued entries, then the
+    // Golub-Kahan step must be performed.
+    bool DiagonalEntriesNonzero(int imin, int imax, Real threshold);
+
+    // This is Algorithm 8.3.1 in "Matrix Computations, 2nd edition" by
+    // G. H. Golub and C. F. Van Loan.
+    void DoGolubKahanStep(int imin, int imax);
+
+    // The diagonal entries are not guaranteed to be nonnegative during the
+    // construction.  After convergence to a diagonal matrix S, test for
+    // negative entries and build a diagonal matrix that reverses the sign
+    // on the S-entry.
+    void EnsureNonnegativeDiagonal();
+
+    // Sort the singular values and compute the corresponding permutation of
+    // the indices of the array storing the singular values.  The permutation
+    // is used for reordering the singular values and the corresponding
+    // columns of the orthogonal matrix in the calls to GetSingularValues(...)
+    // and GetOrthogonalMatrices(...).
+    void ComputePermutation(int sortType);
+
+    // The number rows and columns of the matrices to be processed.
+    int mNumRows, mNumCols;
+
+    // The maximum number of iterations for reducing the bidiagonal matrix
+    // to a diagonal matrix.
+    unsigned int mMaxIterations;
+
+    // The internal copy of a matrix passed to the solver.  See the comments
+    // about function Bidiagonalize() about what is stored in the matrix.
+    std::vector<Real> mMatrix;  // MxN elements
+
+    // After the initial bidiagonalization by Householder reflections, we no
+    // longer need the full mMatrix.  Copy the diagonal and superdiagonal
+    // entries to linear arrays in order to be cache friendly.
+    std::vector<Real> mDiagonal;  // N elements
+    std::vector<Real> mSuperdiagonal;  // N-1 elements
+
+    // The Givens rotations used to reduce the initial bidiagonal matrix to
+    // a diagonal matrix.  A rotation is the identity with the following
+    // replacement entries:  R(index0,index0) = cs, R(index0,index1) = sn,
+    // R(index1,index0) = -sn, and R(index1,index1) = cs.  If N is the
+    // number of matrix columns and K is the maximum number of iterations, the
+    // maximum number of right or left Givens rotations is K*(N-1).  The
+    // maximum amount of memory is allocated to store these.  However, we also
+    // potentially need left rotations to decouple the matrix when a diagonal
+    // terms are zero.  Worst case is a number of matrices quadratic in N, so
+    // for now we just use std::vector<Rotation> whose initial capacity is
+    // K*(N-1).
+    struct GivensRotation
+    {
+        GivensRotation();
+        GivensRotation(int inIndex0, int inIndex1, Real inCs, Real inSn);
+        int index0, index1;
+        Real cs, sn;
+    };
+
+    std::vector<GivensRotation> mRGivens;
+    std::vector<GivensRotation> mLGivens;
+
+    // The diagonal matrix that is used to convert S-entries to nonnegative.
+    std::vector<Real> mFixupDiagonal;  // N elements
+
+    // When sorting is requested, the permutation associated with the sort is
+    // stored in mPermutation.  When sorting is not requested, mPermutation[0]
+    // is set to -1.  mVisited is used for finding cycles in the permutation.
+    std::vector<int> mPermutation;  // N elements
+    mutable std::vector<int> mVisited;  // N elements
+
+    // Temporary storage to compute Householder reflections and to support
+    // sorting of columns of the orthogonal matrices.
+    std::vector<Real> mTwoInvUTU;  // N elements
+    std::vector<Real> mTwoInvVTV;  // N-2 elements
+    mutable std::vector<Real> mUVector;  // M elements
+    mutable std::vector<Real> mVVector;  // N elements
+    mutable std::vector<Real> mWVector;  // max(M,N) elements
+};
+
+
+template <typename Real>
+SingularValueDecomposition<Real>::SingularValueDecomposition(int numRows,
+    int numCols, unsigned int maxIterations)
+    :
+    mNumRows(0),
+    mNumCols(0),
+    mMaxIterations(0)
+{
+    if (numCols > 1 && numRows >= numCols && maxIterations > 0)
+    {
+        mNumRows = numRows;
+        mNumCols = numCols;
+        mMaxIterations = maxIterations;
+        mMatrix.resize(numRows * numCols);
+        mDiagonal.resize(numCols);
+        mSuperdiagonal.resize(numCols - 1);
+        mRGivens.reserve(maxIterations*(numCols - 1));
+        mLGivens.reserve(maxIterations*(numCols - 1));
+        mFixupDiagonal.resize(numCols);
+        mPermutation.resize(numCols);
+        mVisited.resize(numCols);
+        mTwoInvUTU.resize(numCols);
+        mTwoInvVTV.resize(numCols - 2);
+        mUVector.resize(numRows);
+        mVVector.resize(numCols);
+        mWVector.resize(numRows);
+    }
+}
+
+template <typename Real>
+unsigned int SingularValueDecomposition<Real>::Solve(Real const* input,
+    int sortType)
+{
+    if (mNumRows > 0)
+    {
+        int numElements = mNumRows * mNumCols;
+        std::copy(input, input + numElements, mMatrix.begin());
+        Bidiagonalize();
+
+        // Compute 'threshold = multiplier*epsilon*|B|' as the threshold for
+        // diagonal entries effectively zero; that is, |d| <= |threshold|
+        // implies that d is (effectively) zero.  TODO: Allow the caller to
+        // pass 'multiplier' to the constructor.
+        //
+        // We will use the L2-norm |B|, which is the length of the elements
+        // of B treated as an NM-tuple.  The following code avoids overflow
+        // when accumulating the squares of the elements when those elements
+        // are large.
+        Real maxAbsComp = std::abs(input[0]);
+        for (int i = 1; i < numElements; ++i)
+        {
+            Real absComp = std::abs(input[i]);
+            if (absComp > maxAbsComp)
+            {
+                maxAbsComp = absComp;
+            }
+        }
+
+        Real norm = (Real)0;
+        if (maxAbsComp > (Real)0)
+        {
+            Real invMaxAbsComp = ((Real)1) / maxAbsComp;
+            for (int i = 0; i < numElements; ++i)
+            {
+                Real ratio = input[i] * invMaxAbsComp;
+                norm += ratio * ratio;
+            }
+            norm = maxAbsComp * std::sqrt(norm);
+        }
+
+        Real const multiplier = (Real)8;  // TODO: Expose to caller.
+        Real const epsilon = std::numeric_limits<Real>::epsilon();
+        Real const threshold = multiplier * epsilon * norm;
+
+        mRGivens.clear();
+        mLGivens.clear();
+        for (unsigned int j = 0; j < mMaxIterations; ++j)
+        {
+            int imin = -1, imax = -1;
+            for (int i = mNumCols - 2; i >= 0; --i)
+            {
+                // When a01 is much smaller than its diagonal neighbors, it is
+                // effectively zero.
+                Real a00 = mDiagonal[i];
+                Real a01 = mSuperdiagonal[i];
+                Real a11 = mDiagonal[i + 1];
+                Real sum = std::abs(a00) + std::abs(a11);
+                if (sum + std::abs(a01) != sum)
+                {
+                    if (imax == -1)
+                    {
+                        imax = i;
+                    }
+                    imin = i;
+                }
+                else
+                {
+                    // The superdiagonal term is effectively zero compared to
+                    // the neighboring diagonal terms.
+                    if (imin >= 0)
+                    {
+                        break;
+                    }
+                }
+            }
+
+            if (imax == -1)
+            {
+                // The algorithm has converged.
+                EnsureNonnegativeDiagonal();
+                ComputePermutation(sortType);
+                return j;
+            }
+
+            // We need to test diagonal entries of B for zero.  For each zero
+            // diagonal entry, zero the superdiagonal.
+            if (DiagonalEntriesNonzero(imin, imax, threshold))
+            {
+                // Process the lower-right-most unreduced bidiagonal block.
+                DoGolubKahanStep(imin, imax);
+            }
+        }
+        return 0xFFFFFFFF;
+    }
+    else
+    {
+        return 0;
+    }
+}
+
+template <typename Real>
+void SingularValueDecomposition<Real>::GetSingularValues(
+    Real* singularValues) const
+{
+    if (singularValues && mNumCols > 0)
+    {
+        if (mPermutation[0] >= 0)
+        {
+            // Sorting was requested.
+            for (int i = 0; i < mNumCols; ++i)
+            {
+                int p = mPermutation[i];
+                singularValues[i] = mDiagonal[p];
+            }
+        }
+        else
+        {
+            // Sorting was not requested.
+            size_t numBytes = mNumCols*sizeof(Real);
+            Memcpy(singularValues, &mDiagonal[0], numBytes);
+        }
+    }
+}
+
+template <typename Real>
+void SingularValueDecomposition<Real>::GetU(Real* uMatrix) const
+{
+    if (!uMatrix || mNumCols == 0)
+    {
+        // Invalid input or the constructor failed.
+        return;
+    }
+
+    // Start with the identity matrix.
+    std::fill(uMatrix, uMatrix + mNumRows*mNumRows, (Real)0);
+    for (int d = 0; d < mNumRows; ++d)
+    {
+        uMatrix[d + mNumRows*d] = (Real)1;
+    }
+
+    // Multiply the Householder reflections using backward accumulation.
+    int r, c;
+    for (int i0 = mNumCols - 1, i1 = i0 + 1; i0 >= 0; --i0, --i1)
+    {
+        // Copy the u vector and 2/Dot(u,u) from the matrix.
+        Real twoinvudu = mTwoInvUTU[i0];
+        Real const* column = &mMatrix[i0];
+        mUVector[i0] = (Real)1;
+        for (r = i1; r < mNumRows; ++r)
+        {
+            mUVector[r] = column[mNumCols*r];
+        }
+
+        // Compute the w vector.
+        mWVector[i0] = twoinvudu;
+        for (r = i1; r < mNumRows; ++r)
+        {
+            mWVector[r] = (Real)0;
+            for (c = i1; c < mNumRows; ++c)
+            {
+                mWVector[r] += mUVector[c] * uMatrix[r + mNumRows*c];
+            }
+            mWVector[r] *= twoinvudu;
+        }
+
+        // Update the matrix, U <- U - u*w^T.
+        for (r = i0; r < mNumRows; ++r)
+        {
+            for (c = i0; c < mNumRows; ++c)
+            {
+                uMatrix[c + mNumRows*r] -= mUVector[r] * mWVector[c];
+            }
+        }
+    }
+
+    // Multiply the Givens rotations.
+    for (auto const& givens : mLGivens)
+    {
+        int j0 = givens.index0;
+        int j1 = givens.index1;
+        for (r = 0; r < mNumRows; ++r, j0 += mNumRows, j1 += mNumRows)
+        {
+            Real& q0 = uMatrix[j0];
+            Real& q1 = uMatrix[j1];
+            Real prd0 = givens.cs * q0 - givens.sn * q1;
+            Real prd1 = givens.sn * q0 + givens.cs * q1;
+            q0 = prd0;
+            q1 = prd1;
+        }
+    }
+
+    if (mPermutation[0] >= 0)
+    {
+        // Sorting was requested.
+        std::fill(mVisited.begin(), mVisited.end(), 0);
+        for (c = 0; c < mNumCols; ++c)
+        {
+            if (mVisited[c] == 0 && mPermutation[c] != c)
+            {
+                // The item starts a cycle with 2 or more elements.
+                int start = c, current = c, next;
+                for (r = 0; r < mNumRows; ++r)
+                {
+                    mWVector[r] = uMatrix[c + mNumRows*r];
+                }
+                while ((next = mPermutation[current]) != start)
+                {
+                    mVisited[current] = 1;
+                    for (r = 0; r < mNumRows; ++r)
+                    {
+                        uMatrix[current + mNumRows*r] =
+                            uMatrix[next + mNumRows*r];
+                    }
+                    current = next;
+                }
+                mVisited[current] = 1;
+                for (r = 0; r < mNumRows; ++r)
+                {
+                    uMatrix[current + mNumRows*r] = mWVector[r];
+                }
+            }
+        }
+    }
+}
+
+template <typename Real>
+void SingularValueDecomposition<Real>::GetV(Real* vMatrix) const
+{
+    if (!vMatrix || mNumCols == 0)
+    {
+        // Invalid input or the constructor failed.
+        return;
+    }
+
+    // Start with the identity matrix.
+    std::fill(vMatrix, vMatrix + mNumCols*mNumCols, (Real)0);
+    for (int d = 0; d < mNumCols; ++d)
+    {
+        vMatrix[d + mNumCols*d] = (Real)1;
+    }
+
+    // Multiply the Householder reflections using backward accumulation.
+    int i0 = mNumCols - 3;
+    int i1 = i0 + 1;
+    int i2 = i0 + 2;
+    int r, c;
+    for (/**/; i0 >= 0; --i0, --i1, --i2)
+    {
+        // Copy the v vector and 2/Dot(v,v) from the matrix.
+        Real twoinvvdv = mTwoInvVTV[i0];
+        Real const* row = &mMatrix[mNumCols*i0];
+        mVVector[i1] = (Real)1;
+        for (r = i2; r < mNumCols; ++r)
+        {
+            mVVector[r] = row[r];
+        }
+
+        // Compute the w vector.
+        mWVector[i1] = twoinvvdv;
+        for (r = i2; r < mNumCols; ++r)
+        {
+            mWVector[r] = (Real)0;
+            for (c = i2; c < mNumCols; ++c)
+            {
+                mWVector[r] += mVVector[c] * vMatrix[r + mNumCols*c];
+            }
+            mWVector[r] *= twoinvvdv;
+        }
+
+        // Update the matrix, V <- V - v*w^T.
+        for (r = i1; r < mNumCols; ++r)
+        {
+            for (c = i1; c < mNumCols; ++c)
+            {
+                vMatrix[c + mNumCols*r] -= mVVector[r] * mWVector[c];
+            }
+        }
+    }
+
+    // Multiply the Givens rotations.
+    for (auto const& givens : mRGivens)
+    {
+        int j0 = givens.index0;
+        int j1 = givens.index1;
+        for (c = 0; c < mNumCols; ++c, j0 += mNumCols, j1 += mNumCols)
+        {
+            Real& q0 = vMatrix[j0];
+            Real& q1 = vMatrix[j1];
+            Real prd0 = givens.cs * q0 - givens.sn * q1;
+            Real prd1 = givens.sn * q0 + givens.cs * q1;
+            q0 = prd0;
+            q1 = prd1;
+        }
+    }
+
+    // Fix-up the diagonal.
+    for (r = 0; r < mNumCols; ++r)
+    {
+        for (c = 0; c < mNumCols; ++c)
+        {
+            vMatrix[c + mNumCols*r] *= mFixupDiagonal[c];
+        }
+    }
+
+    if (mPermutation[0] >= 0)
+    {
+        // Sorting was requested.
+        std::fill(mVisited.begin(), mVisited.end(), 0);
+        for (c = 0; c < mNumCols; ++c)
+        {
+            if (mVisited[c] == 0 && mPermutation[c] != c)
+            {
+                // The item starts a cycle with 2 or more elements.
+                int start = c, current = c, next;
+                for (r = 0; r < mNumCols; ++r)
+                {
+                    mWVector[r] = vMatrix[c + mNumCols*r];
+                }
+                while ((next = mPermutation[current]) != start)
+                {
+                    mVisited[current] = 1;
+                    for (r = 0; r < mNumCols; ++r)
+                    {
+                        vMatrix[current + mNumCols*r] =
+                            vMatrix[next + mNumCols*r];
+                    }
+                    current = next;
+                }
+                mVisited[current] = 1;
+                for (r = 0; r < mNumCols; ++r)
+                {
+                    vMatrix[current + mNumCols*r] = mWVector[r];
+                }
+            }
+        }
+    }
+}
+
+template <typename Real>
+void SingularValueDecomposition<Real>::GetUColumn(int index,
+    Real* uColumn) const
+{
+    if (0 <= index && index < mNumRows)
+    {
+        // y = H*x, then x and y are swapped for the next H
+        Real* x = uColumn;
+        Real* y = &mWVector[0];
+
+        // Start with the Euclidean basis vector.
+        memset(x, 0, mNumRows * sizeof(Real));
+        if (index < mNumCols && mPermutation[0] >= 0)
+        {
+            // Sorting was requested.
+            x[mPermutation[index]] = (Real)1;
+        }
+        else
+        {
+            x[index] = (Real)1;
+        }
+
+        // Apply the Givens rotations.
+        for (auto const& givens : gte::reverse(mLGivens))
+        {
+            Real& xr0 = x[givens.index0];
+            Real& xr1 = x[givens.index1];
+            Real tmp0 = givens.cs * xr0 + givens.sn * xr1;
+            Real tmp1 = -givens.sn * xr0 + givens.cs * xr1;
+            xr0 = tmp0;
+            xr1 = tmp1;
+        }
+
+        // Apply the Householder reflections.
+        for (int c = mNumCols - 1; c >= 0; --c)
+        {
+            // Get the Householder vector u.
+            int r;
+            for (r = 0; r < c; ++r)
+            {
+                y[r] = x[r];
+            }
+
+            // Compute s = Dot(x,u) * 2/u^T*u.
+            Real s = x[r];  // r = c, u[r] = 1
+            for (int j = r + 1; j < mNumRows; ++j)
+            {
+                s += x[j] * mMatrix[c + mNumCols * j];
+            }
+            s *= mTwoInvUTU[c];
+
+            // r = c, y[r] = x[r]*u[r] - s = x[r] - s because u[r] = 1
+            y[r] = x[r] - s;
+
+            // Compute the remaining components of y.
+            for (++r; r < mNumRows; ++r)
+            {
+                y[r] = x[r] - s * mMatrix[c + mNumCols * r];
+            }
+
+            std::swap(x, y);
+        }
+
+        // The final product is stored in x.
+        if (x != uColumn)
+        {
+            size_t numBytes = mNumRows*sizeof(Real);
+            Memcpy(uColumn, x, numBytes);
+        }
+    }
+}
+
+template <typename Real>
+void SingularValueDecomposition<Real>::GetVColumn(int index,
+    Real* vColumn) const
+{
+    if (0 <= index && index < mNumCols)
+    {
+        // y = H*x, then x and y are swapped for the next H
+        Real* x = vColumn;
+        Real* y = &mWVector[0];
+
+        // Start with the Euclidean basis vector.
+        memset(x, 0, mNumCols * sizeof(Real));
+        if (mPermutation[0] >= 0)
+        {
+            // Sorting was requested.
+            int p = mPermutation[index];
+            x[p] = mFixupDiagonal[p];
+        }
+        else
+        {
+            x[index] = mFixupDiagonal[index];
+        }
+
+        // Apply the Givens rotations.
+        for (auto const& givens : gte::reverse(mRGivens))
+        {
+            Real& xr0 = x[givens.index0];
+            Real& xr1 = x[givens.index1];
+            Real tmp0 = givens.cs * xr0 + givens.sn * xr1;
+            Real tmp1 = -givens.sn * xr0 + givens.cs * xr1;
+            xr0 = tmp0;
+            xr1 = tmp1;
+        }
+
+        // Apply the Householder reflections.
+        for (int r = mNumCols - 3; r >= 0; --r)
+        {
+            // Get the Householder vector v.
+            int c;
+            for (c = 0; c < r + 1; ++c)
+            {
+                y[c] = x[c];
+            }
+
+            // Compute s = Dot(x,v) * 2/v^T*v.
+            Real s = x[c];  // c = r+1, v[c] = 1
+            for (int j = c + 1; j < mNumCols; ++j)
+            {
+                s += x[j] * mMatrix[j + mNumCols * r];
+            }
+            s *= mTwoInvVTV[r];
+
+            // c = r+1, y[c] = x[c]*v[c] - s = x[c] - s because v[c] = 1
+            y[c] = x[c] - s;
+
+            // Compute the remaining components of y.
+            for (++c; c < mNumCols; ++c)
+            {
+                y[c] = x[c] - s * mMatrix[c + mNumCols * r];
+            }
+
+            std::swap(x, y);
+        }
+
+        // The final product is stored in x.
+        if (x != vColumn)
+        {
+            size_t numBytes = mNumCols*sizeof(Real);
+            Memcpy(vColumn, x, numBytes);
+        }
+    }
+}
+
+template <typename Real>
+Real SingularValueDecomposition<Real>::GetSingularValue(int index) const
+{
+    if (0 <= index && index < mNumCols)
+    {
+        if (mPermutation[0] >= 0)
+        {
+            // Sorting was requested.
+            return mDiagonal[mPermutation[index]];
+        }
+        else
+        {
+            // Sorting was not requested.
+            return mDiagonal[index];
+        }
+    }
+    else
+    {
+        return (Real)0;
+    }
+}
+
+template <typename Real>
+void SingularValueDecomposition<Real>::Bidiagonalize()
+{
+    int r, c;
+    for (int i = 0, ip1 = 1; i < mNumCols; ++i, ++ip1)
+    {
+        // Compute the U-Householder vector.
+        Real length = (Real)0;
+        for (r = i; r < mNumRows; ++r)
+        {
+            Real ur = mMatrix[i + mNumCols*r];
+            mUVector[r] = ur;
+            length += ur * ur;
+        }
+        Real udu = (Real)1;
+        length = std::sqrt(length);
+        if (length >(Real)0)
+        {
+            Real& u1 = mUVector[i];
+            Real sgn = (u1 >= (Real)0 ? (Real)1 : (Real)-1);
+            Real invDenom = ((Real)1) / (u1 + sgn * length);
+            u1 = (Real)1;
+            for (r = ip1; r < mNumRows; ++r)
+            {
+                Real& ur = mUVector[r];
+                ur *= invDenom;
+                udu += ur * ur;
+            }
+        }
+
+        // Compute the rank-1 offset u*w^T.
+        Real invudu = (Real)1 / udu;
+        Real twoinvudu = invudu * (Real)2;
+        for (c = i; c < mNumCols; ++c)
+        {
+            mWVector[c] = (Real)0;
+            for (r = i; r < mNumRows; ++r)
+            {
+                mWVector[c] += mMatrix[c + mNumCols*r] * mUVector[r];
+            }
+            mWVector[c] *= twoinvudu;
+        }
+
+        // Update the input matrix.
+        for (r = i; r < mNumRows; ++r)
+        {
+            for (c = i; c < mNumCols; ++c)
+            {
+                mMatrix[c + mNumCols*r] -= mUVector[r] * mWVector[c];
+            }
+        }
+
+        if (i < mNumCols - 2)
+        {
+            // Compute the V-Householder vectors.
+            length = (Real)0;
+            for (c = ip1; c < mNumCols; ++c)
+            {
+                Real vc = mMatrix[c + mNumCols*i];
+                mVVector[c] = vc;
+                length += vc * vc;
+            }
+            Real vdv = (Real)1;
+            length = std::sqrt(length);
+            if (length >(Real)0)
+            {
+                Real& v1 = mVVector[ip1];
+                Real sgn = (v1 >= (Real)0 ? (Real)1 : (Real)-1);
+                Real invDenom = ((Real)1) / (v1 + sgn * length);
+                v1 = (Real)1;
+                for (c = ip1 + 1; c < mNumCols; ++c)
+                {
+                    Real& vc = mVVector[c];
+                    vc *= invDenom;
+                    vdv += vc * vc;
+                }
+            }
+
+            // Compute the rank-1 offset w*v^T.
+            Real invvdv = (Real)1 / vdv;
+            Real twoinvvdv = invvdv * (Real)2;
+            for (r = i; r < mNumRows; ++r)
+            {
+                mWVector[r] = (Real)0;
+                for (c = ip1; c < mNumCols; ++c)
+                {
+                    mWVector[r] += mMatrix[c + mNumCols*r] * mVVector[c];
+                }
+                mWVector[r] *= twoinvvdv;
+            }
+
+            // Update the input matrix.
+            for (r = i; r < mNumRows; ++r)
+            {
+                for (c = ip1; c < mNumCols; ++c)
+                {
+                    mMatrix[c + mNumCols*r] -= mWVector[r] * mVVector[c];
+                }
+            }
+
+            mTwoInvVTV[i] = twoinvvdv;
+            for (c = i + 2; c < mNumCols; ++c)
+            {
+                mMatrix[c + mNumCols*i] = mVVector[c];
+            }
+        }
+
+        mTwoInvUTU[i] = twoinvudu;
+        for (r = ip1; r < mNumRows; ++r)
+        {
+            mMatrix[i + mNumCols*r] = mUVector[r];
+        }
+    }
+
+    // Copy the diagonal and subdiagonal for cache coherence in the
+    // Golub-Kahan iterations.
+    int k, ksup = mNumCols - 1, index = 0, delta = mNumCols + 1;
+    for (k = 0; k < ksup; ++k, index += delta)
+    {
+        mDiagonal[k] = mMatrix[index];
+        mSuperdiagonal[k] = mMatrix[index + 1];
+    }
+    mDiagonal[k] = mMatrix[index];
+}
+
+template <typename Real>
+void SingularValueDecomposition<Real>::GetSinCos(Real x, Real y, Real& cs,
+    Real& sn)
+{
+    // Solves sn*x + cs*y = 0 robustly.
+    Real tau;
+    if (y != (Real)0)
+    {
+        if (std::abs(y) > std::abs(x))
+        {
+            tau = -x / y;
+            sn = ((Real)1) / std::sqrt(((Real)1) + tau*tau);
+            cs = sn * tau;
+        }
+        else
+        {
+            tau = -y / x;
+            cs = ((Real)1) / std::sqrt(((Real)1) + tau*tau);
+            sn = cs * tau;
+        }
+    }
+    else
+    {
+        cs = (Real)1;
+        sn = (Real)0;
+    }
+}
+
+template <typename Real>
+bool SingularValueDecomposition<Real>::DiagonalEntriesNonzero(int imin,
+    int imax, Real threshold)
+{
+    for (int i = imin; i <= imax; ++i)
+    {
+        if (std::abs(mDiagonal[i]) <= threshold)
+        {
+            // Use planar rotations to case the superdiagonal entry out of
+            // the matrix, thus producing a row of zeros.
+            Real x, z, cs, sn;
+            Real y = mSuperdiagonal[i];
+            mSuperdiagonal[i] = (Real)0;
+            for (int j = i + 1; j <= imax + 1; ++j)
+            {
+                x = mDiagonal[j];
+                GetSinCos(x, y, cs, sn);
+                mLGivens.push_back(GivensRotation(i, j, cs, sn));
+                mDiagonal[j] = cs*x - sn*y;
+                if (j <= imax)
+                {
+                    z = mSuperdiagonal[j];
+                    mSuperdiagonal[j] = cs*z;
+                    y = sn*z;
+                }
+            }
+            return false;
+        }
+    }
+    return true;
+}
+
+template <typename Real>
+void SingularValueDecomposition<Real>::DoGolubKahanStep(int imin, int imax)
+{
+    // The implicit shift.  Compute the eigenvalue u of the lower-right 2x2
+    // block of A = B^T*B that is closer to b11.
+    Real f0 = (imax >= (Real)1 ? mSuperdiagonal[imax - 1] : (Real)0);
+    Real d1 = mDiagonal[imax];
+    Real f1 = mSuperdiagonal[imax];
+    Real d2 = mDiagonal[imax + 1];
+    Real a00 = d1*d1 + f0*f0;
+    Real a01 = d1*f1;
+    Real a11 = d2*d2 + f1*f1;
+    Real dif = (a00 - a11) * (Real)0.5;
+    Real sgn = (dif >= (Real)0 ? (Real)1 : (Real)-1);
+    Real a01sqr = a01 * a01;
+    Real u = a11 - a01sqr / (dif + sgn * std::sqrt(dif*dif + a01sqr));
+    Real x = mDiagonal[imin] * mDiagonal[imin] - u;
+    Real y = mDiagonal[imin] * mSuperdiagonal[imin];
+
+    Real a12, a21, a22, a23, cs, sn;
+    Real a02 = (Real)0;
+    int i0 = imin - 1, i1 = imin, i2 = imin + 1;
+    for (/**/; i1 <= imax; ++i0, ++i1, ++i2)
+    {
+        // Compute the Givens rotation G and save it for use in computing
+        // V in U^T*A*V = S.
+        GetSinCos(x, y, cs, sn);
+        mRGivens.push_back(GivensRotation(i1, i2, cs, sn));
+
+        // Update B0 = B*G.
+        if (i1 > imin)
+        {
+            mSuperdiagonal[i0] = cs*mSuperdiagonal[i0] - sn*a02;
+        }
+
+        a11 = mDiagonal[i1];
+        a12 = mSuperdiagonal[i1];
+        a22 = mDiagonal[i2];
+        mDiagonal[i1] = cs*a11 - sn*a12;
+        mSuperdiagonal[i1] = sn*a11 + cs*a12;
+        mDiagonal[i2] = cs*a22;
+        a21 = -sn*a22;
+
+        // Update the parameters for the next Givens rotations.
+        x = mDiagonal[i1];
+        y = a21;
+
+        // Compute the Givens rotation G and save it for use in computing
+        // U in U^T*A*V = S.
+        GetSinCos(x, y, cs, sn);
+        mLGivens.push_back(GivensRotation(i1, i2, cs, sn));
+
+        // Update B1 = G^T*B0.
+        a11 = mDiagonal[i1];
+        a12 = mSuperdiagonal[i1];
+        a22 = mDiagonal[i2];
+        mDiagonal[i1] = cs*a11 - sn*a21;
+        mSuperdiagonal[i1] = cs*a12 - sn*a22;
+        mDiagonal[i2] = sn*a12 + cs*a22;
+
+        if (i1 < imax)
+        {
+            a23 = mSuperdiagonal[i2];
+            a02 = -sn*a23;
+            mSuperdiagonal[i2] = cs*a23;
+
+            // Update the parameters for the next Givens rotations.
+            x = mSuperdiagonal[i1];
+            y = a02;
+        }
+    }
+}
+
+template <typename Real>
+void SingularValueDecomposition<Real>::EnsureNonnegativeDiagonal()
+{
+    for (int i = 0; i < mNumCols; ++i)
+    {
+        if (mDiagonal[i] >= (Real)0)
+        {
+            mFixupDiagonal[i] = (Real)1;
+        }
+        else
+        {
+            mDiagonal[i] = -mDiagonal[i];
+            mFixupDiagonal[i] = (Real)-1;
+        }
+    }
+}
+
+template <typename Real>
+void SingularValueDecomposition<Real>::ComputePermutation(int sortType)
+{
+    if (sortType == 0)
+    {
+        // Set a flag for GetSingularValues() and GetOrthogonalMatrices() to
+        // know that sorted output was not requested.
+        mPermutation[0] = -1;
+        return;
+    }
+
+    // Compute the permutation induced by sorting.  Initially, we start with
+    // the identity permutation I = (0,1,...,N-1).
+    struct SortItem
+    {
+        Real singularValue;
+        int index;
+    };
+
+    std::vector<SortItem> items(mNumCols);
+    int i;
+    for (i = 0; i < mNumCols; ++i)
+    {
+        items[i].singularValue = mDiagonal[i];
+        items[i].index = i;
+    }
+
+    if (sortType > 0)
+    {
+        std::sort(items.begin(), items.end(),
+            [](SortItem const& item0, SortItem const& item1)
+        {
+            return item0.singularValue < item1.singularValue;
+        }
+        );
+    }
+    else
+    {
+        std::sort(items.begin(), items.end(),
+            [](SortItem const& item0, SortItem const& item1)
+        {
+            return item0.singularValue > item1.singularValue;
+        }
+        );
+    }
+
+    i = 0;
+    for (auto const& item : items)
+    {
+        mPermutation[i++] = item.index;
+    }
+
+    // GetOrthogonalMatrices() has nontrivial code for computing the
+    // orthogonal U and V from the reflections and rotations.  To avoid
+    // complicating the code further when sorting is requested, U and V are
+    // computed as in the unsorted case.  We then need to swap columns of
+    // U and V to be consistent with the sorting of the singular values.  To
+    // minimize copying due to column swaps, we use permutation P.  The
+    // minimum number of transpositions to obtain P from I is N minus the
+    // number of cycles of P.  Each cycle is reordered with a minimum number
+    // of transpositions; that is, the singular items are cyclically swapped,
+    // leading to a minimum amount of copying.  For example, if there is a
+    // cycle i0 -> i1 -> i2 -> i3, then the copying is
+    //   save = singularitem[i0];
+    //   singularitem[i1] = singularitem[i2];
+    //   singularitem[i2] = singularitem[i3];
+    //   singularitem[i3] = save;
+}
+
+template <typename Real>
+SingularValueDecomposition<Real>::GivensRotation::GivensRotation()
+{
+    // No default initialization for fast creation of std::vector of objects
+    // of this type.
+}
+
+template <typename Real>
+SingularValueDecomposition<Real>::GivensRotation::GivensRotation(int inIndex0,
+    int inIndex1, Real inCs, Real inSn)
+    :
+    index0(inIndex0),
+    index1(inIndex1),
+    cs(inCs),
+    sn(inSn)
+{
+}
+
+
+}
diff --git a/Engine/Plugins/Experimental/GeometryProcessing/Source/GeometryAlgorithms/Private/GTEngine/Mathematics/GteSlerpEstimate.h b/Engine/Plugins/Experimental/GeometryProcessing/Source/GeometryAlgorithms/Private/GTEngine/Mathematics/GteSlerpEstimate.h
new file mode 100644
--- /dev/null
+++ b/Engine/Plugins/Experimental/GeometryProcessing/Source/GeometryAlgorithms/Private/GTEngine/Mathematics/GteSlerpEstimate.h
@@ -0,0 +1,156 @@
+// David Eberly, Geometric Tools, Redmond WA 98052
+// Copyright (c) 1998-2019
+// Distributed under the Boost Software License, Version 1.0.
+// http://www.boost.org/LICENSE_1_0.txt
+// http://www.geometrictools.com/License/Boost/LICENSE_1_0.txt
+// File Version: 3.0.1 (2018/12/10)
+
+#pragma once
+
+#include <GTEngine/Mathematics/GteQuaternion.h>
+
+// The spherical linear interpolation (slerp) of unit-length quaternions
+// q0 and q1 for t in [0,1] and theta in (0,pi) is
+//   slerp(t,q0,q1) = [sin((1-t)*theta)*q0 + sin(theta)*q1]/sin(theta)
+// where theta is the angle between q0 and q1 [cos(theta) = Dot(q0,q1)].
+// This function is a parameterization of the great spherical arc between
+// q0 and q1 on the unit hypersphere.  Moreover, the parameterization is
+// one of normalized arclength--a particle traveling along the arc through
+// time t does so with constant speed.
+//
+// Read the comments in GteChebyshevRatio.h regarding estimates for the
+// ratio sin(t*theta)/sin(theta).
+//
+// When using slerp in animations involving sequences of quaternions, it is
+// typical that the quaternions are preprocessed so that consecutive ones
+// form an acute angle A in [0,pi/2].  Other preprocessing can help with
+// performance.  See the function comments in the SLERP class.
+
+namespace gte
+{
+    template <typename Real>
+    class SLERP
+    {
+    public:
+        // The angle between q0 and q1 is in [0,pi).  There are no angle
+        // restrictions and nothing is precomputed.
+        template <int N>
+        inline static Quaternion<Real> Estimate(Real t, Quaternion<Real> const& q0, Quaternion<Real> const& q1)
+        {
+            static_assert(1 <= N && N <= 16, "Invalid degree.");
+
+            Real cs = Dot(q0, q1);
+            Real sign;
+            if (cs >= (Real)0)
+            {
+                sign = (Real)1;
+            }
+            else
+            {
+                cs = -cs;
+                sign = (Real)-1;
+            }
+
+            Real f0, f1;
+            ChebyshevRatio<Real>::template
+                GetEstimate<N>(t, (Real)1 - cs, f0, f1);
+            return q0 * f0 + q1 * (sign * f1);
+        }
+
+
+        // The angle between q0 and q1 must be in [0,pi/2].  The suffix R is for
+        // 'Restricted'.  The preprocessing code is
+        //   Quaternion<Real> q[n];  // assuming initialized
+        //   for (i0 = 0, i1 = 1; i1 < n; i0 = i1++)
+        //   {
+        //       cosA = Dot(q[i0], q[i1]);
+        //       if (cosA < 0)
+        //       {
+        //           q[i1] = -q[i1];  // now Dot(q[i0], q[i]1) >= 0
+        //       }
+        //   }
+        template <int N>
+        inline static Quaternion<Real> EstimateR(Real t, Quaternion<Real> const& q0, Quaternion<Real> const& q1)
+        {
+            static_assert(1 <= N && N <= 16, "Invalid degree.");
+
+            Real f0, f1;
+            ChebyshevRatio<Real>::template
+                GetEstimate<N>(t, (Real)1 - Dot(q0, q1), f0, f1);
+            return q0 * f0 + q1 * f1;
+        }
+
+        // The angle between q0 and q1 must be in [0,pi/2].  The suffix R is for
+        // 'Restricted' and the suffix P is for 'Preprocessed'.  The preprocessing
+        // code is
+        //   Quaternion<Real> q[n];  // assuming initialized
+        //   Real cosA[n-1], omcosA[n-1];  // to be precomputed
+        //   for (i0 = 0, i1 = 1; i1 < n; i0 = i1++)
+        //   {
+        //       cs = Dot(q[i0], q[i1]);
+        //       if (cosA[i0] < 0)
+        //       {
+        //           q[i1] = -q[i1];
+        //           cs = -cs;
+        //       }
+        //
+        //       // for Quaterion<T>::SlerpRP
+        //       cosA[i0] = cs;
+        //
+        //       // for SLERP<T>::EstimateRP
+        //       omcosA[i0] = 1 - cs;
+        //   }
+        template <int N>
+        inline static Quaternion<Real> EstimateRP(Real t, Quaternion<Real> const& q0, Quaternion<Real> const& q1,
+            Real omcosA)
+        {
+            static_assert(1 <= N && N <= 16, "Invalid degree.");
+
+            Real f0, f1;
+            ChebyshevRatio<Real>::template
+                GetEstimate<N>(t, omcosA, f0, f1);
+            return q0 * f0 + q1 * f1;
+        }
+
+        // The angle between q0 and q1 is A and must be in [0,pi/2].  Quaternion
+        // qh is slerp(1/2,q0,q1) = (q0+q1)/|q0+q1|, so the angle between q0 and
+        // qh is A/2 and the angle between qh and q1 is A/2.  The preprocessing
+        // code is
+        //   Quaternion<Real> q[n];  // assuming initialized
+        //   Quaternion<Real> qh[n-1];  // to be precomputed
+        //   Real omcosAH[n-1];  // to be precomputed
+        //   for (i0 = 0, i1 = 1; i1 < n; i0 = i1++)
+        //   {
+        //       cosA = Dot(q[i0], q[i1]);
+        //       if (cosA < 0)
+        //       {
+        //           q[i1] = -q[i1];
+        //           cosA = -cosA;
+        //       }
+        //       cosAH = sqrt((1 + cosA)/2);
+        //       qh[i0] = (q0 + q1) / (2 * cosAH[i0]);
+        //       omcosAH[i0] = 1 - cosAH;
+        //   }
+        template <int N>
+        inline static Quaternion<Real> EstimateRPH(Real t, Quaternion<Real> const& q0, Quaternion<Real> const& q1,
+            Quaternion<Real> const& qh, Real omcosAH)
+        {
+            static_assert(1 <= N && N <= 16, "Invalid degree.");
+
+            Real f0, f1;
+            Real twoT = t * (Real)2;
+            if (twoT <= (Real)1)
+            {
+                ChebyshevRatio<Real>::template
+                    GetEstimate<N>(twoT, omcosAH, f0, f1);
+                return q0 * f0 + qh * f1;
+            }
+            else
+            {
+                ChebyshevRatio<Real>::template
+                    GetEstimate<N>(twoT - (Real)1, omcosAH, f0, f1);
+                return qh * f0 + q1 * f1;
+            }
+        }
+    };
+}
diff --git a/Engine/Plugins/Experimental/GeometryProcessing/Source/GeometryAlgorithms/Private/GTEngine/Mathematics/GteSqrtEstimate.h b/Engine/Plugins/Experimental/GeometryProcessing/Source/GeometryAlgorithms/Private/GTEngine/Mathematics/GteSqrtEstimate.h
new file mode 100644
--- /dev/null
+++ b/Engine/Plugins/Experimental/GeometryProcessing/Source/GeometryAlgorithms/Private/GTEngine/Mathematics/GteSqrtEstimate.h
@@ -0,0 +1,193 @@
+// David Eberly, Geometric Tools, Redmond WA 98052
+// Copyright (c) 1998-2019
+// Distributed under the Boost Software License, Version 1.0.
+// http://www.boost.org/LICENSE_1_0.txt
+// http://www.geometrictools.com/License/Boost/LICENSE_1_0.txt
+// File Version: 3.0.1 (2018/10/05)
+
+#pragma once
+
+#include <GTEngine/Mathematics/GteMath.h>
+
+// Minimax polynomial approximations to sqrt(x).  The polynomial p(x) of
+// degree D minimizes the quantity maximum{|sqrt(x) - p(x)| : x in [1,2]}
+// over all polynomials of degree D.
+
+namespace gte
+{
+
+template <typename Real>
+class SqrtEstimate
+{
+public:
+    // The input constraint is x in [1,2].  For example,
+    //   float x; // in [1,2]
+    //   float result = SqrtEstimate<float>::Degree<3>(x);
+    template <int D>
+    inline static Real Degree(Real x);
+
+    // The input constraint is x >= 0.  Range reduction is used to generate a
+    // value y in [0,1], call Degree(y), and combine the output with the
+    // proper exponent to obtain the approximation.  For example,
+    //   float x;  // x >= 0
+    //   float result = SqrtEstimate<float>::DegreeRR<3>(x);
+    template <int D>
+    inline static Real DegreeRR(Real x);
+
+private:
+    // Metaprogramming and private implementation to allow specialization of
+    // a template member function.
+    template <int D> struct degree {};
+    inline static Real Evaluate(degree<1>, Real t);
+    inline static Real Evaluate(degree<2>, Real t);
+    inline static Real Evaluate(degree<3>, Real t);
+    inline static Real Evaluate(degree<4>, Real t);
+    inline static Real Evaluate(degree<5>, Real t);
+    inline static Real Evaluate(degree<6>, Real t);
+    inline static Real Evaluate(degree<7>, Real t);
+    inline static Real Evaluate(degree<8>, Real t);
+
+    // Support for range reduction.
+    inline static void Reduce(Real x, Real& adj, Real& y, int& p);
+    inline static Real Combine(Real adj, Real y, int p);
+};
+
+
+template <typename Real>
+template <int D>
+inline Real SqrtEstimate<Real>::Degree(Real x)
+{
+    Real t = x - (Real)1;  // t in [0,1]
+    return Evaluate(degree<D>(), t);
+}
+
+template <typename Real>
+template <int D>
+inline Real SqrtEstimate<Real>::DegreeRR(Real x)
+{
+    Real adj, y;
+    int p;
+    Reduce(x, adj, y, p);
+    Real poly = Degree<D>(y);
+    Real result = Combine(adj, poly, p);
+    return result;
+}
+
+template <typename Real>
+inline Real SqrtEstimate<Real>::Evaluate(degree<1>, Real t)
+{
+    Real poly;
+    poly = (Real)GTE_C_SQRT_DEG1_C1;
+    poly = (Real)GTE_C_SQRT_DEG1_C0 + poly * t;
+    return poly;
+}
+
+template <typename Real>
+inline Real SqrtEstimate<Real>::Evaluate(degree<2>, Real t)
+{
+    Real poly;
+    poly = (Real)GTE_C_SQRT_DEG2_C2;
+    poly = (Real)GTE_C_SQRT_DEG2_C1 + poly * t;
+    poly = (Real)GTE_C_SQRT_DEG2_C0 + poly * t;
+    return poly;
+}
+
+template <typename Real>
+inline Real SqrtEstimate<Real>::Evaluate(degree<3>, Real t)
+{
+    Real poly;
+    poly = (Real)GTE_C_SQRT_DEG3_C3;
+    poly = (Real)GTE_C_SQRT_DEG3_C2 + poly * t;
+    poly = (Real)GTE_C_SQRT_DEG3_C1 + poly * t;
+    poly = (Real)GTE_C_SQRT_DEG3_C0 + poly * t;
+    return poly;
+}
+
+template <typename Real>
+inline Real SqrtEstimate<Real>::Evaluate(degree<4>, Real t)
+{
+    Real poly;
+    poly = (Real)GTE_C_SQRT_DEG4_C4;
+    poly = (Real)GTE_C_SQRT_DEG4_C3 + poly * t;
+    poly = (Real)GTE_C_SQRT_DEG4_C2 + poly * t;
+    poly = (Real)GTE_C_SQRT_DEG4_C1 + poly * t;
+    poly = (Real)GTE_C_SQRT_DEG4_C0 + poly * t;
+    return poly;
+}
+
+template <typename Real>
+inline Real SqrtEstimate<Real>::Evaluate(degree<5>, Real t)
+{
+    Real poly;
+    poly = (Real)GTE_C_SQRT_DEG5_C5;
+    poly = (Real)GTE_C_SQRT_DEG5_C4 + poly * t;
+    poly = (Real)GTE_C_SQRT_DEG5_C3 + poly * t;
+    poly = (Real)GTE_C_SQRT_DEG5_C2 + poly * t;
+    poly = (Real)GTE_C_SQRT_DEG5_C1 + poly * t;
+    poly = (Real)GTE_C_SQRT_DEG5_C0 + poly * t;
+    return poly;
+}
+
+template <typename Real>
+inline Real SqrtEstimate<Real>::Evaluate(degree<6>, Real t)
+{
+    Real poly;
+    poly = (Real)GTE_C_SQRT_DEG6_C6;
+    poly = (Real)GTE_C_SQRT_DEG6_C5 + poly * t;
+    poly = (Real)GTE_C_SQRT_DEG6_C4 + poly * t;
+    poly = (Real)GTE_C_SQRT_DEG6_C3 + poly * t;
+    poly = (Real)GTE_C_SQRT_DEG6_C2 + poly * t;
+    poly = (Real)GTE_C_SQRT_DEG6_C1 + poly * t;
+    poly = (Real)GTE_C_SQRT_DEG6_C0 + poly * t;
+    return poly;
+}
+
+template <typename Real>
+inline Real SqrtEstimate<Real>::Evaluate(degree<7>, Real t)
+{
+    Real poly;
+    poly = (Real)GTE_C_SQRT_DEG7_C7;
+    poly = (Real)GTE_C_SQRT_DEG7_C6 + poly * t;
+    poly = (Real)GTE_C_SQRT_DEG7_C5 + poly * t;
+    poly = (Real)GTE_C_SQRT_DEG7_C4 + poly * t;
+    poly = (Real)GTE_C_SQRT_DEG7_C3 + poly * t;
+    poly = (Real)GTE_C_SQRT_DEG7_C2 + poly * t;
+    poly = (Real)GTE_C_SQRT_DEG7_C1 + poly * t;
+    poly = (Real)GTE_C_SQRT_DEG7_C0 + poly * t;
+    return poly;
+}
+
+template <typename Real>
+inline Real SqrtEstimate<Real>::Evaluate(degree<8>, Real t)
+{
+    Real poly;
+    poly = (Real)GTE_C_SQRT_DEG8_C8;
+    poly = (Real)GTE_C_SQRT_DEG8_C7 + poly * t;
+    poly = (Real)GTE_C_SQRT_DEG8_C6 + poly * t;
+    poly = (Real)GTE_C_SQRT_DEG8_C5 + poly * t;
+    poly = (Real)GTE_C_SQRT_DEG8_C4 + poly * t;
+    poly = (Real)GTE_C_SQRT_DEG8_C3 + poly * t;
+    poly = (Real)GTE_C_SQRT_DEG8_C2 + poly * t;
+    poly = (Real)GTE_C_SQRT_DEG8_C1 + poly * t;
+    poly = (Real)GTE_C_SQRT_DEG8_C0 + poly * t;
+    return poly;
+}
+
+template <typename Real>
+inline void SqrtEstimate<Real>::Reduce(Real x, Real& adj, Real& y, int& p)
+{
+    y = std::frexp(x, &p);  // y in [1/2,1)
+    y = ((Real)2)*y;   // y in [1,2)
+    --p;
+    adj = (1 & p)*(Real)GTE_C_SQRT_2 + (1 & ~p)*(Real)1;
+    p >>= 1;
+}
+
+template <typename Real>
+inline Real SqrtEstimate<Real>::Combine(Real adj, Real y, int p)
+{
+    return adj * std::ldexp(y, p);
+}
+
+
+}
diff --git a/Engine/Plugins/Experimental/GeometryProcessing/Source/GeometryAlgorithms/Private/GTEngine/Mathematics/GteSymmetricEigensolver.h b/Engine/Plugins/Experimental/GeometryProcessing/Source/GeometryAlgorithms/Private/GTEngine/Mathematics/GteSymmetricEigensolver.h
new file mode 100644
--- /dev/null
+++ b/Engine/Plugins/Experimental/GeometryProcessing/Source/GeometryAlgorithms/Private/GTEngine/Mathematics/GteSymmetricEigensolver.h
@@ -0,0 +1,835 @@
+// David Eberly, Geometric Tools, Redmond WA 98052
+// Copyright (c) 1998-2019
+// Distributed under the Boost Software License, Version 1.0.
+// http://www.boost.org/LICENSE_1_0.txt
+// http://www.geometrictools.com/License/Boost/LICENSE_1_0.txt
+// File Version: 3.0.2 (2018/10/05)
+
+#pragma once
+
+#include <GTEngine/LowLevel/GteRangeIteration.h>
+#include <GTEngine/LowLevel/GteWrapper.h>
+#include <algorithm>
+#include <cmath>
+#include <cstring>
+#include <vector>
+
+// The SymmetricEigensolver class is an implementation of Algorithm 8.2.3
+// (Symmetric QR Algorithm) described in "Matrix Computations, 2nd edition"
+// by G. H. Golub and C. F. Van Loan, The Johns Hopkins University Press,
+// Baltimore MD, Fourth Printing 1993.  Algorithm 8.2.1 (Householder
+// Tridiagonalization) is used to reduce matrix A to tridiagonal T.
+// Algorithm 8.2.2 (Implicit Symmetric QR Step with Wilkinson Shift) is
+// used for the iterative reduction from tridiagonal to diagonal.  If A is
+// the original matrix, D is the diagonal matrix of eigenvalues, and Q is
+// the orthogonal matrix of eigenvectors, then theoretically Q^T*A*Q = D.
+// Numerically, we have errors E = Q^T*A*Q - D.  Algorithm 8.2.3 mentions
+// that one expects |E| is approximately u*|A|, where |M| denotes the
+// Frobenius norm of M and where u is the unit roundoff for the
+// floating-point arithmetic: 2^{-23} for 'float', which is FLT_EPSILON
+// = 1.192092896e-7f, and 2^{-52} for'double', which is DBL_EPSILON
+// = 2.2204460492503131e-16.
+//
+// The condition |a(i,i+1)| <= epsilon*(|a(i,i) + a(i+1,i+1)|) used to
+// determine when the reduction decouples to smaller problems is implemented
+// as:  sum = |a(i,i)| + |a(i+1,i+1)|; sum + |a(i,i+1)| == sum.  The idea is
+// that the superdiagonal term is small relative to its diagonal neighbors,
+// and so it is effectively zero.  The unit tests have shown that this
+// interpretation of decoupling is effective.
+//
+// The authors suggest that once you have the tridiagonal matrix, a practical
+// implementation will store the diagonal and superdiagonal entries in linear
+// arrays, ignoring the theoretically zero values not in the 3-band.  This is
+// good for cache coherence.  The authors also suggest storing the Householder
+// vectors in the lower-triangular portion of the matrix to save memory.  The
+// implementation uses both suggestions.
+//
+// For matrices with randomly generated values in [0,1], the unit tests
+// produce the following information for N-by-N matrices.
+//
+// N  |A|     |E|        |E|/|A|    iterations
+// -------------------------------------------
+//  2  1.2332 5.5511e-17 4.5011e-17  1
+//  3  2.0024 1.1818e-15 5.9020e-16  5
+//  4  2.8708 9.9287e-16 3.4585e-16  7
+//  5  2.9040 2.5958e-15 8.9388e-16  9
+//  6  4.0427 2.2434e-15 5.5493e-16 12
+//  7  5.0276 3.2456e-15 6.4555e-16 15
+//  8  5.4468 6.5626e-15 1.2048e-15 15
+//  9  6.1510 4.0317e-15 6.5545e-16 17
+// 10  6.7523 4.9334e-15 7.3062e-16 21
+// 11  7.1322 7.1347e-15 1.0003e-15 22
+// 12  7.8324 5.6106e-15 7.1633e-16 24
+// 13  8.1073 5.1366e-15 6.3357e-16 25
+// 14  8.6257 8.3496e-15 9.6798e-16 29
+// 15  9.2603 6.9526e-15 7.5080e-16 31
+// 16  9.9853 6.5807e-15 6.5904e-16 32
+// 17 10.5388 8.0931e-15 7.6793e-16 35
+// 18 11.2377 1.1149e-14 9.9218e-16 38
+// 19 11.7105 1.0711e-14 9.1470e-16 42
+// 20 12.2227 1.7723e-14 1.4500e-15 42
+// 21 12.7411 1.2515e-14 9.8231e-16 47
+// 22 13.4462 1.2645e-14 9.4046e-16 50
+// 23 13.9541 1.2899e-14 9.2444e-16 47
+// 24 14.3298 1.6337e-14 1.1400e-15 53
+// 25 14.8050 1.4760e-14 9.9701e-16 54
+// 26 15.3914 1.7076e-14 1.1094e-15 57
+// 27 15.8430 1.9714e-14 1.2443e-15 60
+// 28 16.4771 1.7148e-14 1.0407e-15 60
+// 29 16.9909 1.7309e-14 1.0187e-15 60
+// 30 17.4456 2.1453e-14 1.2297e-15 64
+// 31 17.9909 2.2069e-14 1.2267e-15 68
+//
+// The eigenvalues and |E|/|A| values were compared to those generated by
+// Mathematica Version 9.0; Wolfram Research, Inc., Champaign IL, 2012.
+// The results were all comparable with eigenvalues agreeing to a large
+// number of decimal places.
+//
+// Timing on an Intel (R) Core (TM) i7-3930K CPU @ 3.20 GHZ (in seconds):
+//
+// N    |E|/|A|    iters tridiag QR     evecs    evec[N]  comperr
+// --------------------------------------------------------------
+//  512 4.4149e-15 1017   0.180  0.005    1.151    0.848    2.166
+// 1024 6.1691e-15 1990   1.775  0.031   11.894   12.759   21.179
+// 2048 8.5108e-15 3849  16.592  0.107  119.744  116.56   212.227
+//
+// where iters is the number of QR steps taken, tridiag is the computation
+// of the Householder reflection vectors, evecs is the composition of
+// Householder reflections and Givens rotations to obtain the matrix of
+// eigenvectors, evec[N] is N calls to get the eigenvectors separately, and
+// comperr is the computation E = Q^T*A*Q - D.  The construction of the full
+// eigenvector matrix is, of course, quite expensive.  If you need only a
+// small number of eigenvectors, use function GetEigenvector(int,Real*).
+
+namespace gte
+{
+
+template <typename Real>
+class SymmetricEigensolver
+{
+public:
+    // The solver processes NxN symmetric matrices, where N > 1 ('size' is N)
+    // and the matrix is stored in row-major order.  The maximum number of
+    // iterations ('maxIterations') must be specified for the reduction of a
+    // tridiagonal matrix to a diagonal matrix.  The goal is to compute
+    // NxN orthogonal Q and NxN diagonal D for which Q^T*A*Q = D.
+    SymmetricEigensolver(int size, unsigned int maxIterations);
+
+    // A copy of the NxN symmetric input is made internally.  The order of
+    // the eigenvalues is specified by sortType: -1 (decreasing), 0 (no
+    // sorting), or +1 (increasing).  When sorted, the eigenvectors are
+    // ordered accordingly.  The return value is the number of iterations
+    // consumed when convergence occurred, 0xFFFFFFFF when convergence did
+    // not occur, or 0 when N <= 1 was passed to the constructor.
+    unsigned int Solve(Real const* input, int sortType);
+
+    // Get the eigenvalues of the matrix passed to Solve(...).  The input
+    // 'eigenvalues' must have N elements.
+    void GetEigenvalues(Real* eigenvalues) const;
+
+    // Accumulate the Householder reflections and Givens rotations to produce
+    // the orthogonal matrix Q for which Q^T*A*Q = D.  The input
+    // 'eigenvectors' must have N*N elements.  The array is filled in as
+    // if the eigenvector matrix is stored in row-major order.  The i-th
+    // eigenvector is
+    //   (eigenvectors[i+size*0], ... eigenvectors[i+size*(size - 1)])
+    // which is the i-th column of 'eigenvectors' as an NxN matrix stored
+    // in row-major order.
+    void GetEigenvectors(Real* eigenvectors) const;
+
+    // The eigenvector matrix is a rotation (return +1) or a reflection
+    // (return 0).  If the input 'size' to the constructor is 0 or the
+    // input 'eigenvectors' to GetEigenvectors is null, the returned value
+    // is -1.
+    int GetEigenvectorMatrixType() const;
+
+    // Compute a single eigenvector, which amounts to computing column c
+    // of matrix Q.  The reflections and rotations are applied incrementally.
+    // This is useful when you want only a small number of the eigenvectors.
+    void GetEigenvector(int c, Real* eigenvector) const;
+    Real GetEigenvalue(int c) const;
+
+private:
+    // Tridiagonalize using Householder reflections.  On input, mMatrix is a
+    // copy of the input matrix.  On output, the upper-triangular part of
+    // mMatrix including the diagonal stores the tridiagonalization.  The
+    // lower-triangular part contains 2/Dot(v,v) that are used in computing
+    // eigenvectors and the part below the subdiagonal stores the essential
+    // parts of the Householder vectors v (the elements of v after the
+    // leading 1-valued component).
+    void Tridiagonalize();
+
+    // A helper for generating Givens rotation sine and cosine robustly.
+    void GetSinCos(Real u, Real v, Real& cs, Real& sn);
+
+    // The QR step with implicit shift.  Generally, the initial T is unreduced
+    // tridiagonal (all subdiagonal entries are nonzero).  If a QR step causes
+    // a superdiagonal entry to become zero, the matrix decouples into a block
+    // diagonal matrix with two tridiagonal blocks.  These blocks can be
+    // reduced independently of each other, which allows for parallelization
+    // of the algorithm.  The inputs imin and imax identify the subblock of T
+    // to be processed.   That block has upper-left element T(imin,imin) and
+    // lower-right element T(imax,imax).
+    void DoQRImplicitShift(int imin, int imax);
+
+    // Sort the eigenvalues and compute the corresponding permutation of the
+    // indices of the array storing the eigenvalues.  The permutation is used
+    // for reordering the eigenvalues and eigenvectors in the calls to
+    // GetEigenvalues(...) and GetEigenvectors(...).
+    void ComputePermutation(int sortType);
+
+    // The number N of rows and columns of the matrices to be processed.
+    int mSize;
+
+    // The maximum number of iterations for reducing the tridiagonal mtarix
+    // to a diagonal matrix.
+    unsigned int mMaxIterations;
+
+    // The internal copy of a matrix passed to the solver.  See the comments
+    // about function Tridiagonalize() about what is stored in the matrix.
+    std::vector<Real> mMatrix;  // NxN elements
+
+    // After the initial tridiagonalization by Householder reflections, we no
+    // longer need the full mMatrix.  Copy the diagonal and superdiagonal
+    // entries to linear arrays in order to be cache friendly.
+    std::vector<Real> mDiagonal;  // N elements
+    std::vector<Real> mSuperdiagonal;  // N-1 elements
+
+    // The Givens rotations used to reduce the initial tridiagonal matrix to
+    // a diagonal matrix.  A rotation is the identity with the following
+    // replacement entries:  R(index,index) = cs, R(index,index+1) = sn,
+    // R(index+1,index) = -sn, and R(index+1,index+1) = cs.  If N is the
+    // matrix size and K is the maximum number of iterations, the maximum
+    // number of Givens rotations is K*(N-1).  The maximum amount of memory
+    // is allocated to store these.
+    struct GivensRotation
+    {
+        GivensRotation();
+        GivensRotation(int inIndex, Real inCs, Real inSn);
+        int index;
+        Real cs, sn;
+    };
+
+    std::vector<GivensRotation> mGivens;  // K*(N-1) elements
+
+    // When sorting is requested, the permutation associated with the sort is
+    // stored in mPermutation.  When sorting is not requested, mPermutation[0]
+    // is set to -1.  mVisited is used for finding cycles in the permutation.
+    // mEigenvectorMatrixType is +1 if GetEigenvectors returns a rotation
+    // matrix, 0 if GetEigenvectors returns a reflection matrix or -1 if an
+    // input to the constructor or to GetEigenvectors is invalid.
+    std::vector<int> mPermutation;  // N elements
+    mutable std::vector<int> mVisited;  // N elements
+    mutable int mEigenvectorMatrixType;
+
+    // Temporary storage to compute Householder reflections and to support
+    // sorting of eigenvectors.
+    mutable std::vector<Real> mPVector;  // N elements
+    mutable std::vector<Real> mVVector;  // N elements
+    mutable std::vector<Real> mWVector;  // N elements
+};
+
+
+template <typename Real>
+SymmetricEigensolver<Real>::SymmetricEigensolver(int size,
+    unsigned int maxIterations)
+    :
+    mSize(0),
+    mMaxIterations(0),
+    mEigenvectorMatrixType(-1)
+{
+    if (size > 1 && maxIterations > 0)
+    {
+        mSize = size;
+        mMaxIterations = maxIterations;
+        mMatrix.resize(size*size);
+        mDiagonal.resize(size);
+        mSuperdiagonal.resize(size - 1);
+        mGivens.reserve(maxIterations * (size - 1));
+        mPermutation.resize(size);
+        mVisited.resize(size);
+        mPVector.resize(size);
+        mVVector.resize(size);
+        mWVector.resize(size);
+    }
+}
+
+template <typename Real>
+unsigned int SymmetricEigensolver<Real>::Solve(Real const* input,
+    int sortType)
+{
+    mEigenvectorMatrixType = -1;
+
+    if (mSize > 0)
+    {
+        std::copy(input, input + mSize*mSize, mMatrix.begin());
+        Tridiagonalize();
+
+        mGivens.clear();
+        for (unsigned int j = 0; j < mMaxIterations; ++j)
+        {
+            int imin = -1, imax = -1;
+            for (int i = mSize - 2; i >= 0; --i)
+            {
+                // When a01 is much smaller than its diagonal neighbors, it is
+                // effectively zero.
+                Real a00 = mDiagonal[i];
+                Real a01 = mSuperdiagonal[i];
+                Real a11 = mDiagonal[i + 1];
+                Real sum = std::abs(a00) + std::abs(a11);
+                if (sum + std::abs(a01) != sum)
+                {
+                    if (imax == -1)
+                    {
+                        imax = i;
+                    }
+                    imin = i;
+                }
+                else
+                {
+                    // The superdiagonal term is effectively zero compared to
+                    // the neighboring diagonal terms.
+                    if (imin >= 0)
+                    {
+                        break;
+                    }
+                }
+            }
+
+            if (imax == -1)
+            {
+                // The algorithm has converged.
+                ComputePermutation(sortType);
+                return j;
+            }
+
+            // Process the lower-right-most unreduced tridiagonal block.
+            DoQRImplicitShift(imin, imax);
+        }
+        return 0xFFFFFFFF;
+    }
+    else
+    {
+        return 0;
+    }
+}
+
+template <typename Real>
+void SymmetricEigensolver<Real>::GetEigenvalues(Real* eigenvalues) const
+{
+    if (eigenvalues && mSize > 0)
+    {
+        if (mPermutation[0] >= 0)
+        {
+            // Sorting was requested.
+            for (int i = 0; i < mSize; ++i)
+            {
+                int p = mPermutation[i];
+                eigenvalues[i] = mDiagonal[p];
+            }
+        }
+        else
+        {
+            // Sorting was not requested.
+            size_t numBytes = mSize*sizeof(Real);
+            Memcpy(eigenvalues, &mDiagonal[0], numBytes);
+        }
+    }
+}
+
+template <typename Real>
+void SymmetricEigensolver<Real>::GetEigenvectors(Real* eigenvectors) const
+{
+    mEigenvectorMatrixType = -1;
+
+    if (eigenvectors && mSize > 0)
+    {
+        // Start with the identity matrix.
+        std::fill(eigenvectors, eigenvectors + mSize*mSize, (Real)0);
+        for (int d = 0; d < mSize; ++d)
+        {
+            eigenvectors[d + mSize*d] = (Real)1;
+        }
+
+        // Multiply the Householder reflections using backward accumulation.
+        int r, c;
+        for (int i = mSize - 3, rmin = i + 1; i >= 0; --i, --rmin)
+        {
+            // Copy the v vector and 2/Dot(v,v) from the matrix.
+            Real const* column = &mMatrix[i];
+            Real twoinvvdv = column[mSize*(i + 1)];
+            for (r = 0; r < i + 1; ++r)
+            {
+                mVVector[r] = (Real)0;
+            }
+            mVVector[r] = (Real)1;
+            for (++r; r < mSize; ++r)
+            {
+                mVVector[r] = column[mSize*r];
+            }
+
+            // Compute the w vector.
+            for (r = 0; r < mSize; ++r)
+            {
+                mWVector[r] = (Real)0;
+                for (c = rmin; c < mSize; ++c)
+                {
+                    mWVector[r] += mVVector[c] * eigenvectors[r + mSize*c];
+                }
+                mWVector[r] *= twoinvvdv;
+            }
+
+            // Update the matrix, Q <- Q - v*w^T.
+            for (r = rmin; r < mSize; ++r)
+            {
+                for (c = 0; c < mSize; ++c)
+                {
+                    eigenvectors[c + mSize*r] -= mVVector[r] * mWVector[c];
+                }
+            }
+        }
+
+        // Multiply the Givens rotations.
+        for (auto const& givens : mGivens)
+        {
+            for (r = 0; r < mSize; ++r)
+            {
+                int j = givens.index + mSize*r;
+                Real& q0 = eigenvectors[j];
+                Real& q1 = eigenvectors[j + 1];
+                Real prd0 = givens.cs * q0 - givens.sn * q1;
+                Real prd1 = givens.sn * q0 + givens.cs * q1;
+                q0 = prd0;
+                q1 = prd1;
+            }
+        }
+
+        // The number of Householder reflections is H = mSize - 2.  If H is
+        // even, the product of Householder reflections is a rotation;
+        // otherwise, H is odd and the product is a reflection.  The number 
+        // of Givens rotations does not influence the type of the product of
+        // Householder reflections.
+        mEigenvectorMatrixType = 1 - (mSize & 1);
+
+        if (mPermutation[0] >= 0)
+        {
+            // Sorting was requested.
+            std::fill(mVisited.begin(), mVisited.end(), 0);
+            for (int i = 0; i < mSize; ++i)
+            {
+                if (mVisited[i] == 0 && mPermutation[i] != i)
+                {
+                    // The item starts a cycle with 2 or more elements.
+                    int start = i, current = i, j, next;
+                    for (j = 0; j < mSize; ++j)
+                    {
+                        mPVector[j] = eigenvectors[i + mSize*j];
+                    }
+                    while ((next = mPermutation[current]) != start)
+                    {
+                        mEigenvectorMatrixType = 1 - mEigenvectorMatrixType;
+                        mVisited[current] = 1;
+                        for (j = 0; j < mSize; ++j)
+                        {
+                            eigenvectors[current + mSize*j] =
+                                eigenvectors[next + mSize*j];
+                        }
+                        current = next;
+                    }
+                    mVisited[current] = 1;
+                    for (j = 0; j < mSize; ++j)
+                    {
+                        eigenvectors[current + mSize*j] = mPVector[j];
+                    }
+                }
+            }
+        }
+    }
+}
+
+template <typename Real>
+int SymmetricEigensolver<Real>::GetEigenvectorMatrixType() const
+{
+    return mEigenvectorMatrixType;
+}
+
+template <typename Real>
+void SymmetricEigensolver<Real>::GetEigenvector(int c, Real* eigenvector)
+const
+{
+    if (0 <= c && c < mSize)
+    {
+        // y = H*x, then x and y are swapped for the next H
+        Real* x = eigenvector;
+        Real* y = &mPVector[0];
+
+        // Start with the Euclidean basis vector.
+        memset(x, 0, mSize*sizeof(Real));
+        if (mPermutation[0] >= 0)
+        {
+            // Sorting was requested.
+            x[mPermutation[c]] = (Real)1;
+        }
+        else
+        {
+            x[c] = (Real)1;
+        }
+
+        // Apply the Givens rotations.
+        for (auto const& givens : gte::reverse(mGivens))
+        {
+            Real& xr = x[givens.index];
+            Real& xrp1 = x[givens.index + 1];
+            Real tmp0 = givens.cs * xr + givens.sn * xrp1;
+            Real tmp1 = -givens.sn * xr + givens.cs * xrp1;
+            xr = tmp0;
+            xrp1 = tmp1;
+        }
+
+        // Apply the Householder reflections.
+        for (int i = mSize - 3; i >= 0; --i)
+        {
+            // Get the Householder vector v.
+            Real const* column = &mMatrix[i];
+            Real twoinvvdv = column[mSize*(i + 1)];
+            int r;
+            for (r = 0; r < i + 1; ++r)
+            {
+                y[r] = x[r];
+            }
+
+            // Compute s = Dot(x,v) * 2/v^T*v.
+            Real s = x[r];  // r = i+1, v[i+1] = 1
+            for (int j = r + 1; j < mSize; ++j)
+            {
+                s += x[j] * column[mSize*j];
+            }
+            s *= twoinvvdv;
+
+            y[r] = x[r] - s;  // v[i+1] = 1
+
+            // Compute the remaining components of y.
+            for (++r; r < mSize; ++r)
+            {
+                y[r] = x[r] - s * column[mSize*r];
+            }
+
+            std::swap(x, y);
+        }
+
+        // The final product is stored in x.
+        if (x != eigenvector)
+        {
+            size_t numBytes = mSize*sizeof(Real);
+            Memcpy(eigenvector, x, numBytes);
+        }
+    }
+}
+
+template <typename Real>
+Real SymmetricEigensolver<Real>::GetEigenvalue(int c) const
+{
+    if (mSize > 0)
+    {
+        if (mPermutation[0] >= 0)
+        {
+            // Sorting was requested.
+            return mDiagonal[mPermutation[c]];
+        }
+        else
+        {
+            // Sorting was not requested.
+            return mDiagonal[c];
+        }
+    }
+    else
+    {
+        return std::numeric_limits<Real>::max();
+    }
+}
+
+template <typename Real>
+void SymmetricEigensolver<Real>::Tridiagonalize()
+{
+    int r, c;
+    for (int i = 0, ip1 = 1; i < mSize - 2; ++i, ++ip1)
+    {
+        // Compute the Householder vector.  Read the initial vector from the
+        // row of the matrix.
+        Real length = (Real)0;
+        for (r = 0; r < ip1; ++r)
+        {
+            mVVector[r] = (Real)0;
+        }
+        for (r = ip1; r < mSize; ++r)
+        {
+            Real vr = mMatrix[r + mSize*i];
+            mVVector[r] = vr;
+            length += vr * vr;
+        }
+        Real vdv = (Real)1;
+        length = std::sqrt(length);
+        if (length >(Real)0)
+        {
+            Real& v1 = mVVector[ip1];
+            Real sgn = (v1 >= (Real)0 ? (Real)1 : (Real)-1);
+            Real invDenom = ((Real)1) / (v1 + sgn * length);
+            v1 = (Real)1;
+            for (r = ip1 + 1; r < mSize; ++r)
+            {
+                Real& vr = mVVector[r];
+                vr *= invDenom;
+                vdv += vr * vr;
+            }
+        }
+
+        // Compute the rank-1 offsets v*w^T and w*v^T.
+        Real invvdv = (Real)1 / vdv;
+        Real twoinvvdv = invvdv * (Real)2;
+        Real pdvtvdv = (Real)0;
+        for (r = i; r < mSize; ++r)
+        {
+            mPVector[r] = (Real)0;
+            for (c = i; c < r; ++c)
+            {
+                mPVector[r] += mMatrix[r + mSize*c] * mVVector[c];
+            }
+            for (/**/; c < mSize; ++c)
+            {
+                mPVector[r] += mMatrix[c + mSize*r] * mVVector[c];
+            }
+            mPVector[r] *= twoinvvdv;
+            pdvtvdv += mPVector[r] * mVVector[r];
+        }
+
+        pdvtvdv *= invvdv;
+        for (r = i; r < mSize; ++r)
+        {
+            mWVector[r] = mPVector[r] - pdvtvdv * mVVector[r];
+        }
+
+        // Update the input matrix.
+        for (r = i; r < mSize; ++r)
+        {
+            Real vr = mVVector[r];
+            Real wr = mWVector[r];
+            Real offset = vr * wr * (Real)2;
+            mMatrix[r + mSize*r] -= offset;
+            for (c = r + 1; c < mSize; ++c)
+            {
+                offset = vr * mWVector[c] + wr * mVVector[c];
+                mMatrix[c + mSize*r] -= offset;
+            }
+        }
+
+        // Copy the vector to column i of the matrix.  The 0-valued components
+        // at indices 0 through i are not stored.  The 1-valued component at
+        // index i+1 is also not stored; instead, the quantity 2/Dot(v,v) is
+        // stored for use in eigenvector construction. That construction must
+        // take into account the implied components that are not stored.
+        mMatrix[i + mSize*ip1] = twoinvvdv;
+        for (r = ip1 + 1; r < mSize; ++r)
+        {
+            mMatrix[i + mSize*r] = mVVector[r];
+        }
+    }
+
+    // Copy the diagonal and subdiagonal entries for cache coherence in
+    // the QR iterations.
+    int k, ksup = mSize - 1, index = 0, delta = mSize + 1;
+    for (k = 0; k < ksup; ++k, index += delta)
+    {
+        mDiagonal[k] = mMatrix[index];
+        mSuperdiagonal[k] = mMatrix[index + 1];
+    }
+    mDiagonal[k] = mMatrix[index];
+}
+
+template <typename Real>
+void SymmetricEigensolver<Real>::GetSinCos(Real x, Real y, Real& cs, Real& sn)
+{
+    // Solves sn*x + cs*y = 0 robustly.
+    Real tau;
+    if (y != (Real)0)
+    {
+        if (std::abs(y) > std::abs(x))
+        {
+            tau = -x / y;
+            sn = ((Real)1) / std::sqrt(((Real)1) + tau*tau);
+            cs = sn * tau;
+        }
+        else
+        {
+            tau = -y / x;
+            cs = ((Real)1) / std::sqrt(((Real)1) + tau*tau);
+            sn = cs * tau;
+        }
+    }
+    else
+    {
+        cs = (Real)1;
+        sn = (Real)0;
+    }
+}
+
+template <typename Real>
+void SymmetricEigensolver<Real>::DoQRImplicitShift(int imin, int imax)
+{
+    // The implicit shift.  Compute the eigenvalue u of the lower-right 2x2
+    // block that is closer to a11.
+    Real a00 = mDiagonal[imax];
+    Real a01 = mSuperdiagonal[imax];
+    Real a11 = mDiagonal[imax + 1];
+    Real dif = (a00 - a11) * (Real)0.5;
+    Real sgn = (dif >= (Real)0 ? (Real)1 : (Real)-1);
+    Real a01sqr = a01 * a01;
+    Real u = a11 - a01sqr / (dif + sgn * std::sqrt(dif*dif + a01sqr));
+    Real x = mDiagonal[imin] - u;
+    Real y = mSuperdiagonal[imin];
+
+    Real a12, a22, a23, tmp11, tmp12, tmp21, tmp22, cs, sn;
+    Real a02 = (Real)0;
+    int i0 = imin - 1, i1 = imin, i2 = imin + 1;
+    for (/**/; i1 <= imax; ++i0, ++i1, ++i2)
+    {
+        // Compute the Givens rotation and save it for use in computing the
+        // eigenvectors.
+        GetSinCos(x, y, cs, sn);
+        mGivens.push_back(GivensRotation(i1, cs, sn));
+
+        // Update the tridiagonal matrix.  This amounts to updating a 4x4
+        // subblock,
+        //   b00 b01 b02 b03
+        //   b01 b11 b12 b13
+        //   b02 b12 b22 b23
+        //   b03 b13 b23 b33
+        // The four corners (b00, b03, b33) do not change values.  The
+        // The interior block {{b11,b12},{b12,b22}} is updated on each pass.
+        // For the first pass, the b0c values are out of range, so only
+        // the values (b13, b23) change.  For the last pass, the br3 values
+        // are out of range, so only the values (b01, b02) change.  For
+        // passes between first and last, the values (b01, b02, b13, b23)
+        // change.
+        if (i1 > imin)
+        {
+            mSuperdiagonal[i0] = cs*mSuperdiagonal[i0] - sn*a02;
+        }
+
+        a11 = mDiagonal[i1];
+        a12 = mSuperdiagonal[i1];
+        a22 = mDiagonal[i2];
+        tmp11 = cs*a11 - sn*a12;
+        tmp12 = cs*a12 - sn*a22;
+        tmp21 = sn*a11 + cs*a12;
+        tmp22 = sn*a12 + cs*a22;
+        mDiagonal[i1] = cs*tmp11 - sn*tmp12;
+        mSuperdiagonal[i1] = sn*tmp11 + cs*tmp12;
+        mDiagonal[i2] = sn*tmp21 + cs*tmp22;
+
+        if (i1 < imax)
+        {
+            a23 = mSuperdiagonal[i2];
+            a02 = -sn*a23;
+            mSuperdiagonal[i2] = cs*a23;
+
+            // Update the parameters for the next Givens rotation.
+            x = mSuperdiagonal[i1];
+            y = a02;
+        }
+    }
+}
+
+template <typename Real>
+void SymmetricEigensolver<Real>::ComputePermutation(int sortType)
+{
+    // The number of Householder reflections is H = mSize - 2.  If H is
+    // even, the product of Householder reflections is a rotation;
+    // otherwise, H is odd and the product is a reflection.  The number 
+    // of Givens rotations does not influence the type of the product of
+    // Householder reflections.
+    mEigenvectorMatrixType = 1 - (mSize & 1);
+
+    if (sortType == 0)
+    {
+        // Set a flag for GetEigenvalues() and GetEigenvectors() to know
+        // that sorted output was not requested.
+        mPermutation[0] = -1;
+        return;
+    }
+
+    // Compute the permutation induced by sorting.  Initially, we start with
+    // the identity permutation I = (0,1,...,N-1).
+    struct SortItem
+    {
+        Real eigenvalue;
+        int index;
+    };
+
+    std::vector<SortItem> items(mSize);
+    int i;
+    for (i = 0; i < mSize; ++i)
+    {
+        items[i].eigenvalue = mDiagonal[i];
+        items[i].index = i;
+    }
+
+    if (sortType > 0)
+    {
+        std::sort(items.begin(), items.end(),
+            [](SortItem const& item0, SortItem const& item1)
+        {
+            return item0.eigenvalue < item1.eigenvalue;
+        }
+        );
+    }
+    else
+    {
+        std::sort(items.begin(), items.end(),
+            [](SortItem const& item0, SortItem const& item1)
+        {
+            return item0.eigenvalue > item1.eigenvalue;
+        }
+        );
+    }
+
+    i = 0;
+    for (auto const& item : items)
+    {
+        mPermutation[i++] = item.index;
+    }
+
+    // GetEigenvectors() has nontrivial code for computing the orthogonal Q
+    // from the reflections and rotations.  To avoid complicating the code
+    // further when sorting is requested, Q is computed as in the unsorted
+    // case.  We then need to swap columns of Q to be consistent with the
+    // sorting of the eigenvalues.  To minimize copying due to column swaps,
+    // we use permutation P.  The minimum number of transpositions to obtain
+    // P from I is N minus the number of cycles of P.  Each cycle is reordered
+    // with a minimum number of transpositions; that is, the eigenitems are
+    // cyclically swapped, leading to a minimum amount of copying.  For
+    // example, if there is a cycle i0 -> i1 -> i2 -> i3, then the copying is
+    //   save = eigenitem[i0];
+    //   eigenitem[i1] = eigenitem[i2];
+    //   eigenitem[i2] = eigenitem[i3];
+    //   eigenitem[i3] = save;
+}
+
+template <typename Real>
+SymmetricEigensolver<Real>::GivensRotation::GivensRotation()
+{
+    // No default initialization for fast creation of std::vector of objects
+    // of this type.
+}
+
+template <typename Real>
+SymmetricEigensolver<Real>::GivensRotation::GivensRotation(int inIndex,
+    Real inCs, Real inSn)
+    :
+    index(inIndex),
+    cs(inCs),
+    sn(inSn)
+{
+}
+
+
+}
diff --git a/Engine/Plugins/Experimental/GeometryProcessing/Source/GeometryAlgorithms/Private/GTEngine/Mathematics/GteSymmetricEigensolver2x2.h b/Engine/Plugins/Experimental/GeometryProcessing/Source/GeometryAlgorithms/Private/GTEngine/Mathematics/GteSymmetricEigensolver2x2.h
new file mode 100644
--- /dev/null
+++ b/Engine/Plugins/Experimental/GeometryProcessing/Source/GeometryAlgorithms/Private/GTEngine/Mathematics/GteSymmetricEigensolver2x2.h
@@ -0,0 +1,94 @@
+// David Eberly, Geometric Tools, Redmond WA 98052
+// Copyright (c) 1998-2019
+// Distributed under the Boost Software License, Version 1.0.
+// http://www.boost.org/LICENSE_1_0.txt
+// http://www.geometrictools.com/License/Boost/LICENSE_1_0.txt
+// File Version: 3.0.1 (2018/10/05)
+
+#pragma once
+
+#include <GTEngine/GTEngineDEF.h>
+#include <algorithm>
+#include <array>
+#include <cmath>
+
+namespace gte
+{
+
+template <typename Real>
+class SymmetricEigensolver2x2
+{
+public:
+    // The input matrix must be symmetric, so only the unique elements must
+    // be specified: a00, a01, and a11.
+    //
+    // The order of the eigenvalues is specified by sortType: -1 (decreasing),
+    // 0 (no sorting), or +1 (increasing).  When sorted, the eigenvectors are
+    // ordered accordingly, and {evec[0], evec[1]} is guaranteed to be a
+    // right-handed orthonormal set.
+
+    void operator()(Real a00, Real a01, Real a11, int sortType,
+        std::array<Real, 2>& eval, std::array<std::array<Real, 2>, 2>& evec)
+        const;
+};
+
+
+template <typename Real>
+void SymmetricEigensolver2x2<Real>::operator()(Real a00, Real a01, Real a11,
+    int sortType, std::array<Real, 2>& eval,
+    std::array<std::array<Real, 2>, 2>& evec) const
+{
+    // Normalize (c2,s2) robustly, avoiding floating-point overflow in the
+    // sqrt call.
+    Real const zero = (Real)0, one = (Real)1, half = (Real)0.5;
+    Real c2 = half * (a00 - a11), s2 = a01;
+    Real maxAbsComp = std::max(std::abs(c2), std::abs(s2));
+    if (maxAbsComp > zero)
+    {
+        c2 /= maxAbsComp;  // in [-1,1]
+        s2 /= maxAbsComp;  // in [-1,1]
+        Real length = std::sqrt(c2 * c2 + s2 * s2);
+        c2 /= length;
+        s2 /= length;
+        if (c2 > zero)
+        {
+            c2 = -c2;
+            s2 = -s2;
+        }
+    }
+    else
+    {
+        c2 = -one;
+        s2 = zero;
+    }
+
+    Real s = std::sqrt(half * (one - c2));  // >= 1/sqrt(2)
+    Real c = half * s2 / s;
+
+    Real diagonal[2];
+    Real csqr = c * c, ssqr = s * s, mid = s2 * a01;
+    diagonal[0] = csqr * a00 + mid + ssqr * a11;
+    diagonal[1] = csqr * a11 - mid + ssqr * a00;
+
+    if (sortType == 0 || sortType * diagonal[0] <= sortType * diagonal[1])
+    {
+        eval[0] = diagonal[0];
+        eval[1] = diagonal[1];
+        evec[0][0] = c;
+        evec[0][1] = s;
+        evec[1][0] = -s;
+        evec[1][1] = c;
+    }
+    else
+    {
+        eval[0] = diagonal[1];
+        eval[1] = diagonal[0];
+        evec[0][0] = s;
+        evec[0][1] = -c;
+        evec[1][0] = c;
+        evec[1][1] = s;
+    }
+}
+
+
+}
