move GeometryProcessing plugin from NotForLicensees to Experimental

diff --git a/Engine/Plugins/Experimental/GeometryProcessing/Source/GeometryAlgorithms/Private/GTEngine/Mathematics/GteSymmetricEigensolver3x3.h b/Engine/Plugins/Experimental/GeometryProcessing/Source/GeometryAlgorithms/Private/GTEngine/Mathematics/GteSymmetricEigensolver3x3.h
new file mode 100644
--- /dev/null
+++ b/Engine/Plugins/Experimental/GeometryProcessing/Source/GeometryAlgorithms/Private/GTEngine/Mathematics/GteSymmetricEigensolver3x3.h
@@ -0,0 +1,777 @@
+// David Eberly, Geometric Tools, Redmond WA 98052
+// Copyright (c) 1998-2019
+// Distributed under the Boost Software License, Version 1.0.
+// http://www.boost.org/LICENSE_1_0.txt
+// http://www.geometrictools.com/License/Boost/LICENSE_1_0.txt
+// File Version: 3.0.5 (2018/10/09)
+
+#pragma once
+
+#include <algorithm>
+#include <array>
+#include <cmath>
+#include <limits>
+
+// The document
+// http://www.geometrictools.com/Documentation/RobustEigenSymmetric3x3.pdf
+// describes algorithms for solving the eigensystem associated with a 3x3
+// symmetric real-valued matrix.  The iterative algorithm is implemented
+// by class SymmmetricEigensolver3x3.  The noniterative algorithm is
+// implemented by class NISymmetricEigensolver3x3.  The code does not use
+// GTEngine objects.
+
+namespace gte
+{
+
+template <typename Real>
+class SymmetricEigensolver3x3
+{
+public:
+    // The input matrix must be symmetric, so only the unique elements must
+    // be specified: a00, a01, a02, a11, a12, and a22.
+    //
+    // If 'aggressive' is 'true', the iterations occur until a superdiagonal
+    // entry is exactly zero.  If 'aggressive' is 'false', the iterations
+    // occur until a superdiagonal entry is effectively zero compared to the
+    // sum of magnitudes of its diagonal neighbors.  Generally, the
+    // nonaggressive convergence is acceptable.
+    //
+    // The order of the eigenvalues is specified by sortType: -1 (decreasing),
+    // 0 (no sorting), or +1 (increasing).  When sorted, the eigenvectors are
+    // ordered accordingly, and {evec[0], evec[1], evec[2]} is guaranteed to
+    // be a right-handed orthonormal set.  The return value is the number of
+    // iterations used by the algorithm.
+
+    int operator()(Real a00, Real a01, Real a02, Real a11, Real a12, Real a22,
+        bool aggressive, int sortType, std::array<Real, 3>& eval,
+        std::array<std::array<Real, 3>, 3>& evec) const;
+
+private:
+    // Update Q = Q*G in-place using G = {{c,0,-s},{s,0,c},{0,0,1}}.
+    void Update0(Real Q[3][3], Real c, Real s) const;
+
+    // Update Q = Q*G in-place using G = {{0,1,0},{c,0,s},{-s,0,c}}.
+    void Update1(Real Q[3][3], Real c, Real s) const;
+
+    // Update Q = Q*H in-place using H = {{c,s,0},{s,-c,0},{0,0,1}}.
+    void Update2(Real Q[3][3], Real c, Real s) const;
+
+    // Update Q = Q*H in-place using H = {{1,0,0},{0,c,s},{0,s,-c}}.
+    void Update3(Real Q[3][3], Real c, Real s) const;
+
+    // Normalize (u,v) robustly, avoiding floating-point overflow in the sqrt
+    // call.  The normalized pair is (cs,sn) with cs <= 0.  If (u,v) = (0,0),
+    // the function returns (cs,sn) = (-1,0).  When used to generate a
+    // Householder reflection, it does not matter whether (cs,sn) or (-cs,-sn)
+    // is used.  When generating a Givens reflection, cs = cos(2*theta) and
+    // sn = sin(2*theta).  Having a negative cosine for the double-angle
+    // term ensures that the single-angle terms c = cos(theta) and
+    // s = sin(theta) satisfy |c| <= |s|.
+    void GetCosSin(Real u, Real v, Real& cs, Real& sn) const;
+
+    // The convergence test.  When 'aggressive' is 'true', the superdiagonal
+    // test is "bSuper == 0".  When 'aggressive' is 'false', the superdiagonal
+    // test is "|bDiag0| + |bDiag1| + |bSuper| == |bDiag0| + |bDiag1|, which
+    // means bSuper is effectively zero compared to the sizes of the diagonal
+    // entries.
+    bool Converged(bool aggressive, Real bDiag0, Real bDiag1,
+        Real bSuper) const;
+
+    // Support for sorting the eigenvalues and eigenvectors.  The output
+    // (i0,i1,i2) is a permutation of (0,1,2) so that d[i0] <= d[i1] <= d[i2].
+    // The 'bool' return indicates whether the permutation is odd.  If it is
+    // not, the handedness of the Q matrix must be adjusted.
+    bool Sort(std::array<Real, 3> const& d, int& i0, int& i1, int& i2) const;
+};
+
+
+template <typename Real>
+class NISymmetricEigensolver3x3
+{
+public:
+    // The input matrix must be symmetric, so only the unique elements must
+    // be specified: a00, a01, a02, a11, a12, and a22.  The eigenvalues are
+    // sorted in ascending order: eval0 <= eval1 <= eval2.
+
+    void operator()(Real a00, Real a01, Real a02, Real a11, Real a12, Real a22,
+        std::array<Real, 3>& eval, std::array<std::array<Real, 3>, 3>& evec) const;
+
+private:
+    static std::array<Real, 3> Multiply(Real s, std::array<Real, 3> const& U);
+    static std::array<Real, 3> Subtract(std::array<Real, 3> const& U, std::array<Real, 3> const& V);
+    static std::array<Real, 3> Divide(std::array<Real, 3> const& U, Real s);
+    static Real Dot(std::array<Real, 3> const& U, std::array<Real, 3> const& V);
+    static std::array<Real, 3> Cross(std::array<Real, 3> const& U, std::array<Real, 3> const& V);
+
+    void ComputeOrthogonalComplement(std::array<Real, 3> const& W,
+        std::array<Real, 3>& U, std::array<Real, 3>& V) const;
+
+    void ComputeEigenvector0(Real a00, Real a01, Real a02, Real a11, Real a12, Real a22,
+        Real eval0, std::array<Real, 3>& evec0) const;
+
+    void ComputeEigenvector1(Real a00, Real a01, Real a02, Real a11, Real a12, Real a22,
+        std::array<Real, 3> const& evec0, Real eval1, std::array<Real, 3>& evec1) const;
+};
+
+
+template <typename Real>
+int SymmetricEigensolver3x3<Real>::operator()(Real a00, Real a01, Real a02,
+    Real a11, Real a12, Real a22, bool aggressive, int sortType,
+    std::array<Real, 3>& eval, std::array<std::array<Real, 3>, 3>& evec) const
+{
+    // Compute the Householder reflection H and B = H*A*H, where b02 = 0.
+    Real const zero = (Real)0, one = (Real)1, half = (Real)0.5;
+    bool isRotation = false;
+    Real c, s;
+    GetCosSin(a12, -a02, c, s);
+    Real Q[3][3] = { { c, s, zero }, { s, -c, zero }, { zero, zero, one } };
+    Real term0 = c * a00 + s * a01;
+    Real term1 = c * a01 + s * a11;
+    Real b00 = c * term0 + s * term1;
+    Real b01 = s * term0 - c * term1;
+    term0 = s * a00 - c * a01;
+    term1 = s * a01 - c * a11;
+    Real b11 = s * term0 - c * term1;
+    Real b12 = s * a02 - c * a12;
+    Real b22 = a22;
+
+    // Givens reflections, B' = G^T*B*G, preserve tridiagonal matrices.
+    int const maxIteration = 2 * (1 + std::numeric_limits<Real>::digits -
+        std::numeric_limits<Real>::min_exponent);
+    int iteration;
+    Real c2, s2;
+
+    if (std::abs(b12) <= std::abs(b01))
+    {
+        Real saveB00, saveB01, saveB11;
+        for (iteration = 0; iteration < maxIteration; ++iteration)
+        {
+            // Compute the Givens reflection.
+            GetCosSin(half * (b00 - b11), b01, c2, s2);
+            s = std::sqrt(half * (one - c2));  // >= 1/sqrt(2)
+            c = half * s2 / s;
+
+            // Update Q by the Givens reflection.
+            Update0(Q, c, s);
+            isRotation = !isRotation;
+
+            // Update B <- Q^T*B*Q, ensuring that b02 is zero and |b12| has
+            // strictly decreased.
+            saveB00 = b00;
+            saveB01 = b01;
+            saveB11 = b11;
+            term0 = c * saveB00 + s * saveB01;
+            term1 = c * saveB01 + s * saveB11;
+            b00 = c * term0 + s * term1;
+            b11 = b22;
+            term0 = c * saveB01 - s * saveB00;
+            term1 = c * saveB11 - s * saveB01;
+            b22 = c * term1 - s * term0;
+            b01 = s * b12;
+            b12 = c * b12;
+
+            if (Converged(aggressive, b00, b11, b01))
+            {
+                // Compute the Householder reflection.
+                GetCosSin(half * (b00 - b11), b01, c2, s2);
+                s = std::sqrt(half * (one - c2));
+                c = half * s2 / s;  // >= 1/sqrt(2)
+
+                // Update Q by the Householder reflection.
+                Update2(Q, c, s);
+                isRotation = !isRotation;
+
+                // Update D = Q^T*B*Q.
+                saveB00 = b00;
+                saveB01 = b01;
+                saveB11 = b11;
+                term0 = c * saveB00 + s * saveB01;
+                term1 = c * saveB01 + s * saveB11;
+                b00 = c * term0 + s * term1;
+                term0 = s * saveB00 - c * saveB01;
+                term1 = s * saveB01 - c * saveB11;
+                b11 = s * term0 - c * term1;
+                break;
+            }
+        }
+    }
+    else
+    {
+        Real saveB11, saveB12, saveB22;
+        for (iteration = 0; iteration < maxIteration; ++iteration)
+        {
+            // Compute the Givens reflection.
+            GetCosSin(half * (b22 - b11), b12, c2, s2);
+            s = std::sqrt(half * (one - c2));  // >= 1/sqrt(2)
+            c = half * s2 / s;
+
+            // Update Q by the Givens reflection.
+            Update1(Q, c, s);
+            isRotation = !isRotation;
+
+            // Update B <- Q^T*B*Q, ensuring that b02 is zero and |b12| has
+            // strictly decreased.  MODIFY...
+            saveB11 = b11;
+            saveB12 = b12;
+            saveB22 = b22;
+            term0 = c * saveB22 + s * saveB12;
+            term1 = c * saveB12 + s * saveB11;
+            b22 = c * term0 + s * term1;
+            b11 = b00;
+            term0 = c * saveB12 - s * saveB22;
+            term1 = c * saveB11 - s * saveB12;
+            b00 = c * term1 - s * term0;
+            b12 = s * b01;
+            b01 = c * b01;
+
+            if (Converged(aggressive, b11, b22, b12))
+            {
+                // Compute the Householder reflection.
+                GetCosSin(half * (b11 - b22), b12, c2, s2);
+                s = std::sqrt(half * (one - c2));
+                c = half * s2 / s;  // >= 1/sqrt(2)
+
+                // Update Q by the Householder reflection.
+                Update3(Q, c, s);
+                isRotation = !isRotation;
+
+                // Update D = Q^T*B*Q.
+                saveB11 = b11;
+                saveB12 = b12;
+                saveB22 = b22;
+                term0 = c * saveB11 + s * saveB12;
+                term1 = c * saveB12 + s * saveB22;
+                b11 = c * term0 + s * term1;
+                term0 = s * saveB11 - c * saveB12;
+                term1 = s * saveB12 - c * saveB22;
+                b22 = s * term0 - c * term1;
+                break;
+            }
+        }
+    }
+
+    std::array<Real, 3> diagonal = { b00, b11, b22 };
+    int i0, i1, i2;
+    if (sortType >= 1)
+    {
+        // diagonal[i0] <= diagonal[i1] <= diagonal[i2]
+        bool isOdd = Sort(diagonal, i0, i1, i2);
+        if (!isOdd)
+        {
+            isRotation = !isRotation;
+        }
+    }
+    else if (sortType <= -1)
+    {
+        // diagonal[i0] >= diagonal[i1] >= diagonal[i2]
+        bool isOdd = Sort(diagonal, i0, i1, i2);
+        std::swap(i0, i2);  // (i0,i1,i2)->(i2,i1,i0) is odd
+        if (isOdd)
+        {
+            isRotation = !isRotation;
+        }
+    }
+    else
+    {
+        i0 = 0;
+        i1 = 1;
+        i2 = 2;
+    }
+
+    eval[0] = diagonal[i0];
+    eval[1] = diagonal[i1];
+    eval[2] = diagonal[i2];
+    evec[0][0] = Q[0][i0];
+    evec[0][1] = Q[1][i0];
+    evec[0][2] = Q[2][i0];
+    evec[1][0] = Q[0][i1];
+    evec[1][1] = Q[1][i1];
+    evec[1][2] = Q[2][i1];
+    evec[2][0] = Q[0][i2];
+    evec[2][1] = Q[1][i2];
+    evec[2][2] = Q[2][i2];
+
+    // Ensure the columns of Q form a right-handed set.
+    if (!isRotation)
+    {
+        for (int j = 0; j < 3; ++j)
+        {
+            evec[2][j] = -evec[2][j];
+        }
+    }
+
+    return iteration;
+}
+
+template <typename Real>
+void SymmetricEigensolver3x3<Real>::Update0(Real Q[3][3], Real c, Real s)
+const
+{
+    for (int r = 0; r < 3; ++r)
+    {
+        Real tmp0 = c * Q[r][0] + s * Q[r][1];
+        Real tmp1 = Q[r][2];
+        Real tmp2 = c * Q[r][1] - s * Q[r][0];
+        Q[r][0] = tmp0;
+        Q[r][1] = tmp1;
+        Q[r][2] = tmp2;
+    }
+}
+
+template <typename Real>
+void SymmetricEigensolver3x3<Real>::Update1(Real Q[3][3], Real c, Real s)
+const
+{
+    for (int r = 0; r < 3; ++r)
+    {
+        Real tmp0 = c * Q[r][1] - s * Q[r][2];
+        Real tmp1 = Q[r][0];
+        Real tmp2 = c * Q[r][2] + s * Q[r][1];
+        Q[r][0] = tmp0;
+        Q[r][1] = tmp1;
+        Q[r][2] = tmp2;
+    }
+}
+
+template <typename Real>
+void SymmetricEigensolver3x3<Real>::Update2(Real Q[3][3], Real c, Real s)
+const
+{
+    for (int r = 0; r < 3; ++r)
+    {
+        Real tmp0 = c * Q[r][0] + s * Q[r][1];
+        Real tmp1 = s * Q[r][0] - c * Q[r][1];
+        Q[r][0] = tmp0;
+        Q[r][1] = tmp1;
+    }
+}
+
+template <typename Real>
+void SymmetricEigensolver3x3<Real>::Update3(Real Q[3][3], Real c, Real s)
+const
+{
+    for (int r = 0; r < 3; ++r)
+    {
+        Real tmp0 = c * Q[r][1] + s * Q[r][2];
+        Real tmp1 = s * Q[r][1] - c * Q[r][2];
+        Q[r][1] = tmp0;
+        Q[r][2] = tmp1;
+    }
+}
+
+template <typename Real>
+void SymmetricEigensolver3x3<Real>::GetCosSin(Real u, Real v, Real& cs,
+    Real& sn) const
+{
+    Real maxAbsComp = std::max(std::abs(u), std::abs(v));
+    if (maxAbsComp > (Real)0)
+    {
+        u /= maxAbsComp;  // in [-1,1]
+        v /= maxAbsComp;  // in [-1,1]
+        Real length = std::sqrt(u * u + v * v);
+        cs = u / length;
+        sn = v / length;
+        if (cs > (Real)0)
+        {
+            cs = -cs;
+            sn = -sn;
+        }
+    }
+    else
+    {
+        cs = (Real)-1;
+        sn = (Real)0;
+    }
+}
+
+template <typename Real>
+bool SymmetricEigensolver3x3<Real>::Converged(bool aggressive, Real bDiag0,
+    Real bDiag1, Real bSuper) const
+{
+    if (aggressive)
+    {
+        return bSuper == (Real)0;
+    }
+    else
+    {
+        Real sum = std::abs(bDiag0) + std::abs(bDiag1);
+        return sum + std::abs(bSuper) == sum;
+    }
+}
+
+template <typename Real>
+bool SymmetricEigensolver3x3<Real>::Sort(std::array<Real, 3> const& d,
+    int& i0, int& i1, int& i2) const
+{
+    bool odd;
+    if (d[0] < d[1])
+    {
+        if (d[2] < d[0])
+        {
+            i0 = 2; i1 = 0; i2 = 1; odd = true;
+        }
+        else if (d[2] < d[1])
+        {
+            i0 = 0; i1 = 2; i2 = 1; odd = false;
+        }
+        else
+        {
+            i0 = 0; i1 = 1; i2 = 2; odd = true;
+        }
+    }
+    else
+    {
+        if (d[2] < d[1])
+        {
+            i0 = 2; i1 = 1; i2 = 0; odd = false;
+        }
+        else if (d[2] < d[0])
+        {
+            i0 = 1; i1 = 2; i2 = 0; odd = true;
+        }
+        else
+        {
+            i0 = 1; i1 = 0; i2 = 2; odd = false;
+        }
+    }
+    return odd;
+}
+
+
+template <typename Real>
+void NISymmetricEigensolver3x3<Real>::operator() (Real a00, Real a01, Real a02,
+    Real a11, Real a12, Real a22, std::array<Real, 3>& eval,
+    std::array<std::array<Real, 3>, 3>& evec) const
+{
+    // Precondition the matrix by factoring out the maximum absolute value
+    // of the components.  This guards against floating-point overflow when
+    // computing the eigenvalues.
+    Real max0 = std::max(std::abs(a00), std::abs(a01));
+    Real max1 = std::max(std::abs(a02), std::abs(a11));
+    Real max2 = std::max(std::abs(a12), std::abs(a22));
+    Real maxAbsElement = std::max(std::max(max0, max1), max2);
+    if (maxAbsElement == (Real)0)
+    {
+        // A is the zero matrix.
+        eval[0] = (Real)0;
+        eval[1] = (Real)0;
+        eval[2] = (Real)0;
+        evec[0] = { (Real)1, (Real)0, (Real)0 };
+        evec[1] = { (Real)0, (Real)1, (Real)0 };
+        evec[2] = { (Real)0, (Real)0, (Real)1 };
+        return;
+    }
+
+    Real invMaxAbsElement = (Real)1 / maxAbsElement;
+    a00 *= invMaxAbsElement;
+    a01 *= invMaxAbsElement;
+    a02 *= invMaxAbsElement;
+    a11 *= invMaxAbsElement;
+    a12 *= invMaxAbsElement;
+    a22 *= invMaxAbsElement;
+
+    Real norm = a01 * a01 + a02 * a02 + a12 * a12;
+    if (norm > (Real)0)
+    {
+        // Compute the eigenvalues of A.
+
+        // In the PDF mentioned previously, B = (A - q*I)/p, where q = tr(A)/3 
+        // with tr(A) the trace of A (sum of the diagonal entries of A) and where
+        // p = sqrt(tr((A - q*I)^2)/6).
+        Real q = (a00 + a11 + a22) / (Real)3;
+
+        // The matrix A - q*I is represented by the following, where b00, b11 and
+        // b22 are computed after these comments,
+        //   +-           -+
+        //   | b00 a01 a02 |
+        //   | a01 b11 a12 |
+        //   | a02 a12 b22 |
+        //   +-           -+
+        Real b00 = a00 - q;
+        Real b11 = a11 - q;
+        Real b22 = a22 - q;
+
+        // The is the variable p mentioned in the PDF.
+        Real p = std::sqrt((b00 * b00 + b11 * b11 + b22 * b22 + norm * (Real)2) / (Real)6);
+
+        // We need det(B) = det((A - q*I)/p) = det(A - q*I)/p^3.  The value
+        // det(A - q*I) is computed using a cofactor expansion by the first
+        // row of A - q*I.  The cofactors are c00, c01 and c02 and the
+        // determinant is b00*c00 - a01*c01 + a02*c02.  The det(B) is then
+        // computed finally by the division with p^3.
+        Real c00 = b11 * b22 - a12 * a12;
+        Real c01 = a01 * b22 - a12 * a02;
+        Real c02 = a01 * a12 - b11 * a02;
+        Real det = (b00 * c00 - a01 * c01 + a02 * c02) / (p * p * p);
+
+        // The halfDet value is cos(3*theta) mentioned in the PDF. The acos(z)
+        // function requires |z| <= 1, but will fail silently and return NaN
+        // if the input is larger than 1 in magnitude.  To avoid this problem
+        // due to rounding errors, the halfDet/ value is clamped to [-1,1].
+        Real halfDet = det * (Real)0.5;
+        halfDet = std::min(std::max(halfDet, (Real)-1), (Real)1);
+
+        // The eigenvalues of B are ordered as beta0 <= beta1 <= beta2.  The
+        // number of digits in twoThirdsPi is chosen so that, whether float or
+        // double, the floating-point number is the closest to theoretical 2*pi/3.
+        Real angle = std::acos(halfDet) / (Real)3;
+        Real const twoThirdsPi = (Real)2.09439510239319549;
+        Real beta2 = std::cos(angle) * (Real)2;
+        Real beta0 = std::cos(angle + twoThirdsPi) * (Real)2;
+        Real beta1 = -(beta0 + beta2);
+
+        // The eigenvalues of A are ordered as alpha0 <= alpha1 <= alpha2.
+        eval[0] = q + p * beta0;
+        eval[1] = q + p * beta1;
+        eval[2] = q + p * beta2;
+
+        // Compute the eigenvectors so that the set {evec[0], evec[1], evec[2]}
+        // is right handed and orthonormal.
+        if (halfDet >= (Real)0)
+        {
+            ComputeEigenvector0(a00, a01, a02, a11, a12, a22, eval[2], evec[2]);
+            ComputeEigenvector1(a00, a01, a02, a11, a12, a22, evec[2], eval[1], evec[1]);
+            evec[0] = Cross(evec[1], evec[2]);
+        }
+        else
+        {
+            ComputeEigenvector0(a00, a01, a02, a11, a12, a22, eval[0], evec[0]);
+            ComputeEigenvector1(a00, a01, a02, a11, a12, a22, evec[0], eval[1], evec[1]);
+            evec[2] = Cross(evec[0], evec[1]);
+        }
+    }
+    else
+    {
+        // The matrix is diagonal.
+        eval[0] = a00;
+        eval[1] = a11;
+        eval[2] = a22;
+        evec[0] = { (Real)1, (Real)0, (Real)0 };
+        evec[1] = { (Real)0, (Real)1, (Real)0 };
+        evec[2] = { (Real)0, (Real)0, (Real)1 };
+    }
+
+    // The preconditioning scaled the matrix A, which scales the eigenvalues.
+    // Revert the scaling.
+    eval[0] *= maxAbsElement;
+    eval[1] *= maxAbsElement;
+    eval[2] *= maxAbsElement;
+}
+
+template <typename Real>
+std::array<Real, 3> NISymmetricEigensolver3x3<Real>::Multiply(
+    Real s, std::array<Real, 3> const& U)
+{
+    std::array<Real, 3> product = { s * U[0], s * U[1], s * U[2] };
+    return product;
+}
+
+template <typename Real>
+std::array<Real, 3> NISymmetricEigensolver3x3<Real>::Subtract(
+    std::array<Real, 3> const& U, std::array<Real, 3> const& V)
+{
+    std::array<Real, 3> difference = { U[0] - V[0], U[1] - V[1], U[2] - V[2] };
+    return difference;
+}
+
+template <typename Real>
+std::array<Real, 3> NISymmetricEigensolver3x3<Real>::Divide(
+    std::array<Real, 3> const& U, Real s)
+{
+    Real invS = (Real)1 / s;
+    std::array<Real, 3> division = { U[0] * invS, U[1] * invS, U[2] * invS };
+    return division;
+}
+
+template <typename Real>
+Real NISymmetricEigensolver3x3<Real>::Dot(std::array<Real, 3> const& U,
+    std::array<Real, 3> const& V)
+{
+    Real dot = U[0] * V[0] + U[1] * V[1] + U[2] * V[2];
+    return dot;
+}
+
+template <typename Real>
+std::array<Real, 3> NISymmetricEigensolver3x3<Real>::Cross(std::array<Real, 3> const& U,
+    std::array<Real, 3> const& V)
+{
+    std::array<Real, 3> cross =
+    {
+        U[1] * V[2] - U[2] * V[1],
+        U[2] * V[0] - U[0] * V[2],
+        U[0] * V[1] - U[1] * V[0]
+    };
+    return cross;
+}
+
+template <typename Real>
+void NISymmetricEigensolver3x3<Real>::ComputeOrthogonalComplement(
+    std::array<Real, 3> const& W, std::array<Real, 3>& U, std::array<Real, 3>& V) const
+{
+    // Robustly compute a right-handed orthonormal set { U, V, W }.  The
+    // vector W is guaranteed to be unit-length, in which case there is no
+    // need to worry about a division by zero when computing invLength.
+    Real invLength;
+    if (std::abs(W[0]) > std::abs(W[1]))
+    {
+        // The component of maximum absolute value is either W[0] or W[2].
+        invLength = (Real)1 / std::sqrt(W[0] * W[0] + W[2] * W[2]);
+        U = { -W[2] * invLength, (Real)0, +W[0] * invLength };
+    }
+    else
+    {
+        // The component of maximum absolute value is either W[1] or W[2].
+        invLength = (Real)1 / std::sqrt(W[1] * W[1] + W[2] * W[2]);
+        U = { (Real)0, +W[2] * invLength, -W[1] * invLength };
+    }
+    V = Cross(W, U);
+}
+
+template <typename Real>
+void NISymmetricEigensolver3x3<Real>::ComputeEigenvector0(Real a00, Real a01,
+    Real a02, Real a11, Real a12, Real a22, Real eval0, std::array<Real, 3>& evec0) const
+{
+    // Compute a unit-length eigenvector for eigenvalue[i0].  The matrix is
+    // rank 2, so two of the rows are linearly independent.  For a robust
+    // computation of the eigenvector, select the two rows whose cross product
+    // has largest length of all pairs of rows.
+    std::array<Real, 3> row0 = { a00 - eval0, a01, a02 };
+    std::array<Real, 3> row1 = { a01, a11 - eval0, a12 };
+    std::array<Real, 3> row2 = { a02, a12, a22 - eval0 };
+    std::array<Real, 3>  r0xr1 = Cross(row0, row1);
+    std::array<Real, 3>  r0xr2 = Cross(row0, row2);
+    std::array<Real, 3>  r1xr2 = Cross(row1, row2);
+    Real d0 = Dot(r0xr1, r0xr1);
+    Real d1 = Dot(r0xr2, r0xr2);
+    Real d2 = Dot(r1xr2, r1xr2);
+
+    Real dmax = d0;
+    int imax = 0;
+    if (d1 > dmax)
+    {
+        dmax = d1;
+        imax = 1;
+    }
+    if (d2 > dmax)
+    {
+        imax = 2;
+    }
+
+    if (imax == 0)
+    {
+        evec0 = Divide(r0xr1, std::sqrt(d0));
+    }
+    else if (imax == 1)
+    {
+        evec0 = Divide(r0xr2, std::sqrt(d1));
+    }
+    else
+    {
+        evec0 = Divide(r1xr2, std::sqrt(d2));
+    }
+}
+
+template <typename Real>
+void NISymmetricEigensolver3x3<Real>::ComputeEigenvector1(Real a00, Real a01,
+    Real a02, Real a11, Real a12, Real a22, std::array<Real, 3> const& evec0,
+    Real eval1, std::array<Real, 3>& evec1) const
+{
+    // Robustly compute a right-handed orthonormal set { U, V, evec0 }.
+    std::array<Real, 3> U, V;
+    ComputeOrthogonalComplement(evec0, U, V);
+
+    // Let e be eval1 and let E be a corresponding eigenvector which is a
+    // solution to the linear system (A - e*I)*E = 0.  The matrix (A - e*I)
+    // is 3x3, not invertible (so infinitely many solutions), and has rank 2
+    // when eval1 and eval are different.  It has rank 1 when eval1 and eval2
+    // are equal.  Numerically, it is difficult to compute robustly the rank
+    // of a matrix.  Instead, the 3x3 linear system is reduced to a 2x2 system
+    // as follows.  Define the 3x2 matrix J = [U V] whose columns are the U
+    // and V computed previously.  Define the 2x1 vector X = J*E.  The 2x2
+    // system is 0 = M * X = (J^T * (A - e*I) * J) * X where J^T is the
+    // transpose of J and M = J^T * (A - e*I) * J is a 2x2 matrix.  The system
+    // may be written as
+    //     +-                        -++-  -+       +-  -+
+    //     | U^T*A*U - e  U^T*A*V     || x0 | = e * | x0 |
+    //     | V^T*A*U      V^T*A*V - e || x1 |       | x1 |
+    //     +-                        -++   -+       +-  -+
+    // where X has row entries x0 and x1.
+
+    std::array<Real, 3> AU =
+    {
+        a00 * U[0] + a01 * U[1] + a02 * U[2],
+        a01 * U[0] + a11 * U[1] + a12 * U[2],
+        a02 * U[0] + a12 * U[1] + a22 * U[2]
+    };
+
+    std::array<Real, 3> AV =
+    {
+        a00 * V[0] + a01 * V[1] + a02 * V[2],
+        a01 * V[0] + a11 * V[1] + a12 * V[2],
+        a02 * V[0] + a12 * V[1] + a22 * V[2]
+    };
+
+    Real m00 = U[0] * AU[0] + U[1] * AU[1] + U[2] * AU[2] - eval1;
+    Real m01 = U[0] * AV[0] + U[1] * AV[1] + U[2] * AV[2];
+    Real m11 = V[0] * AV[0] + V[1] * AV[1] + V[2] * AV[2] - eval1;
+
+    // For robustness, choose the largest-length row of M to compute the
+    // eigenvector.  The 2-tuple of coefficients of U and V in the
+    // assignments to eigenvector[1] lies on a circle, and U and V are
+    // unit length and perpendicular, so eigenvector[1] is unit length
+    // (within numerical tolerance).
+    Real absM00 = std::abs(m00);
+    Real absM01 = std::abs(m01);
+    Real absM11 = std::abs(m11);
+    Real maxAbsComp;
+    if (absM00 >= absM11)
+    {
+        maxAbsComp = std::max(absM00, absM01);
+        if (maxAbsComp > (Real)0)
+        {
+            if (absM00 >= absM01)
+            {
+                m01 /= m00;
+                m00 = (Real)1 / std::sqrt((Real)1 + m01 * m01);
+                m01 *= m00;
+            }
+            else
+            {
+                m00 /= m01;
+                m01 = (Real)1 / std::sqrt((Real)1 + m00 * m00);
+                m00 *= m01;
+            }
+            evec1 = Subtract(Multiply(m01, U), Multiply(m00, V));
+        }
+        else
+        {
+            evec1 = U;
+        }
+    }
+    else
+    {
+        maxAbsComp = std::max(absM11, absM01);
+        if (maxAbsComp > (Real)0)
+        {
+            if (absM11 >= absM01)
+            {
+                m01 /= m11;
+                m11 = (Real)1 / std::sqrt((Real)1 + m01 * m01);
+                m01 *= m11;
+            }
+            else
+            {
+                m11 /= m01;
+                m01 = (Real)1 / std::sqrt((Real)1 + m11 * m11);
+                m11 *= m01;
+            }
+            evec1 = Subtract(Multiply(m11, U), Multiply(m01, V));
+        }
+        else
+        {
+            evec1 = U;
+        }
+    }
+}
+
+}
diff --git a/Engine/Plugins/Experimental/GeometryProcessing/Source/GeometryAlgorithms/Private/GTEngine/Mathematics/GteTCBSplineCurve.h b/Engine/Plugins/Experimental/GeometryProcessing/Source/GeometryAlgorithms/Private/GTEngine/Mathematics/GteTCBSplineCurve.h
new file mode 100644
--- /dev/null
+++ b/Engine/Plugins/Experimental/GeometryProcessing/Source/GeometryAlgorithms/Private/GTEngine/Mathematics/GteTCBSplineCurve.h
@@ -0,0 +1,249 @@
+// David Eberly, Geometric Tools, Redmond WA 98052
+// Copyright (c) 1998-2019
+// Distributed under the Boost Software License, Version 1.0.
+// http://www.boost.org/LICENSE_1_0.txt
+// http://www.geometrictools.com/License/Boost/LICENSE_1_0.txt
+// File Version: 3.0.0 (2016/06/19)
+
+#pragma once
+
+#include <GTEngine/LowLevel/GteLogger.h>
+#include <GTEngine/Mathematics/GteParametricCurve.h>
+
+namespace gte
+{
+
+template <int N, typename Real>
+class TCBSplineCurve : public ParametricCurve<N, Real>
+{
+public:
+    // Construction and destruction.  The object copies the input arrays.
+    // The number of points must be at least 2.  To validate construction,
+    // create an object as shown:
+    //     TCBSplineCurve<N, Real> curve(parameters);
+    //     if (!curve) { <constructor failed, handle accordingly>; }
+    virtual ~TCBSplineCurve();
+    TCBSplineCurve(int numPoints, Vector<N, Real> const* points,
+        Real const* times, Real const* tension, Real const* continuity,
+        Real const* bias);
+
+    // Member access.
+    inline int GetNumPoints() const;
+    inline Vector<N, Real> const* GetPoints() const;
+    inline Real const* GetTensions() const;
+    inline Real const* GetContinuities() const;
+    inline Real const* GetBiases() const;
+
+    // Evaluation of the curve.  The function supports derivative calculation
+    // through order 3; that is, maxOrder <= 3 is required.  If you want
+    // only the position, pass in maxOrder of 0.  If you want the position and
+    // first derivative, pass in maxOrder of 1, and so on.  The output
+    // 'values' are ordered as: position, first derivative, second derivative,
+    // third derivative.
+    virtual void Evaluate(Real t, unsigned int maxOrder,
+        Vector<N, Real> values[4]) const;
+
+protected:
+    // Support for construction.
+    void ComputePoly(int i0, int i1, int i2, int i3);
+
+    // Determine the index i for which times[i] <= t < times[i+1].
+    void GetKeyInfo(Real t, int& key, Real& dt) const;
+
+    std::vector<Vector<N, Real>> mPoints;
+    std::vector<Real> mTension, mContinuity, mBias;
+
+    // Polynomial coefficients.  mA are the degree 0 coefficients,  mB are
+    // the degree 1 coefficients, mC are the degree 2 coefficients, and mD
+    // are the degree 3 coefficients.
+    std::vector<Vector<N, Real>> mA, mB, mC, mD;
+};
+
+
+template <int N, typename Real>
+TCBSplineCurve<N, Real>::~TCBSplineCurve()
+{
+}
+
+template <int N, typename Real>
+TCBSplineCurve<N, Real>::TCBSplineCurve(int numPoints,
+    Vector<N, Real> const* points, Real const* times, Real const* tension,
+    Real const* continuity, Real const* bias)
+    :
+    ParametricCurve<N, Real>(numPoints - 1, times)
+{
+    if (numPoints < 2 || !points)
+    {
+        LogError("Invalid input.");
+        return;
+    }
+
+    mPoints.resize(numPoints);
+    mTension.resize(numPoints);
+    mContinuity.resize(numPoints);
+    mBias.resize(numPoints);
+    std::copy(points, points + numPoints, mPoints.begin());
+    std::copy(tension, tension + numPoints, mTension.begin());
+    std::copy(continuity, continuity + numPoints, mContinuity.begin());
+    std::copy(bias, bias + numPoints, mBias.begin());
+
+    int numSegments = numPoints - 1;
+    mA.resize(numSegments);
+    mB.resize(numSegments);
+    mC.resize(numSegments);
+    mD.resize(numSegments);
+
+    // For now, treat the first point as if it occurred twice.
+    ComputePoly(0, 0, 1, 2);
+
+    for (int i = 1; i < numSegments - 1; ++i)
+    {
+        ComputePoly(i - 1, i, i + 1, i + 2);
+    }
+
+    // For now, treat the last point as if it occurred twice.
+    ComputePoly(numSegments - 2, numSegments - 1, numSegments, numSegments);
+
+    this->mConstructed = true;
+}
+
+template <int N, typename Real> inline
+int TCBSplineCurve<N, Real>::GetNumPoints() const
+{
+    return static_cast<int>(mPoints.size());
+}
+
+template <int N, typename Real> inline
+Vector<N, Real> const* TCBSplineCurve<N, Real>::GetPoints() const
+{
+    return &mPoints[0];
+}
+
+template <int N, typename Real> inline
+Real const* TCBSplineCurve<N, Real>::GetTensions() const
+{
+    return &mTension[0];
+}
+
+template <int N, typename Real> inline
+Real const* TCBSplineCurve<N, Real>::GetContinuities() const
+{
+    return &mContinuity[0];
+}
+
+template <int N, typename Real> inline
+Real const* TCBSplineCurve<N, Real>::GetBiases() const
+{
+    return &mBias[0];
+}
+
+template <int N, typename Real>
+void TCBSplineCurve<N, Real>::Evaluate(Real t, unsigned int maxOrder,
+    Vector<N, Real> values[4]) const
+{
+    if (!this->mConstructed)
+    {
+        // Errors were already generated during construction.
+        for (unsigned int order = 0; order < 4; ++order)
+        {
+            values[order].MakeZero();
+        }
+        return;
+    }
+
+    int key;
+    Real dt;
+    GetKeyInfo(t, key, dt);
+    dt /= (this->mTime[key + 1] - this->mTime[key]);
+
+    // Compute position.
+    values[0] = mA[key] + dt * (mB[key] + dt * (mC[key] + dt * mD[key]));
+    if (maxOrder >= 1)
+    {
+        // Compute first derivative.
+        values[1] = mB[key] + dt * (mC[key] * ((Real)2) + mD[key] *
+            (((Real)3) * dt));
+        if (maxOrder >= 2)
+        {
+            // Compute second derivative.
+            values[2] = mC[key] * ((Real)2) + mD[key] * (((Real)6) * dt);
+            if (maxOrder == 3)
+            {
+                values[3] = ((Real)6) * mD[key];
+            }
+            else
+            {
+                values[3].MakeZero();
+            }
+        }
+    }
+}
+
+template <int N, typename Real>
+void TCBSplineCurve<N, Real>::ComputePoly(int i0, int i1, int i2, int i3)
+{
+    Vector<N, Real> diff = mPoints[i2] - mPoints[i1];
+    Real dt = this->mTime[i2] - this->mTime[i1];
+
+    // Build multipliers at P1.
+    Real oneMinusT0 = (Real)1 - mTension[i1];
+    Real oneMinusC0 = (Real)1 - mContinuity[i1];
+    Real onePlusC0 = (Real)1 + mContinuity[i1];
+    Real oneMinusB0 = (Real)1 - mBias[i1];
+    Real onePlusB0 = (Real)1 + mBias[i1];
+    Real adj0 = ((Real)2)*dt / (this->mTime[i2] - this->mTime[i0]);
+    Real out0 = ((Real)0.5)*adj0*oneMinusT0*onePlusC0*onePlusB0;
+    Real out1 = ((Real)0.5)*adj0*oneMinusT0*oneMinusC0*oneMinusB0;
+
+    // Build outgoing tangent at P1.
+    Vector<N, Real> tOut = out1*diff + out0*(mPoints[i1] - mPoints[i0]);
+
+    // Build multipliers at point P2.
+    Real oneMinusT1 = (Real)1 - mTension[i2];
+    Real oneMinusC1 = (Real)1 - mContinuity[i2];
+    Real onePlusC1 = (Real)1 + mContinuity[i2];
+    Real oneMinusB1 = (Real)1 - mBias[i2];
+    Real onePlusB1 = (Real)1 + mBias[i2];
+    Real adj1 = ((Real)2)*dt / (this->mTime[i3] - this->mTime[i1]);
+    Real in0 = ((Real)0.5)*adj1*oneMinusT1*oneMinusC1*onePlusB1;
+    Real in1 = ((Real)0.5)*adj1*oneMinusT1*onePlusC1*oneMinusB1;
+
+    // Build incoming tangent at P2.
+    Vector<N, Real> tIn = in1*(mPoints[i3] - mPoints[i2]) + in0*diff;
+
+    mA[i1] = mPoints[i1];
+    mB[i1] = tOut;
+    mC[i1] = ((Real)3)*diff - ((Real)2)*tOut - tIn;
+    mD[i1] = ((Real)-2)*diff + tOut + tIn;
+}
+
+template <int N, typename Real>
+void TCBSplineCurve<N, Real>::GetKeyInfo(Real t, int& key, Real& dt) const
+{
+    int numSegments = static_cast<int>(mA.size());
+    if (t <= this->mTime[0])
+    {
+        key = 0;
+        dt = (Real)0;
+    }
+    else if (t >= this->mTime[numSegments])
+    {
+        key = numSegments - 1;
+        dt = this->mTime[numSegments] - this->mTime[numSegments - 1];
+    }
+    else
+    {
+        for (int i = 0; i < numSegments; ++i)
+        {
+            if (t < this->mTime[i + 1])
+            {
+                key = i;
+                dt = t - this->mTime[i];
+                break;
+            }
+        }
+    }
+}
+
+
+}
diff --git a/Engine/Plugins/Experimental/GeometryProcessing/Source/GeometryAlgorithms/Private/GTEngine/Mathematics/GteTIQuery.h b/Engine/Plugins/Experimental/GeometryProcessing/Source/GeometryAlgorithms/Private/GTEngine/Mathematics/GteTIQuery.h
new file mode 100644
--- /dev/null
+++ b/Engine/Plugins/Experimental/GeometryProcessing/Source/GeometryAlgorithms/Private/GTEngine/Mathematics/GteTIQuery.h
@@ -0,0 +1,33 @@
+// David Eberly, Geometric Tools, Redmond WA 98052
+// Copyright (c) 1998-2019
+// Distributed under the Boost Software License, Version 1.0.
+// http://www.boost.org/LICENSE_1_0.txt
+// http://www.geometrictools.com/License/Boost/LICENSE_1_0.txt
+// File Version: 3.0.2 (2018/10/05)
+
+#pragma once
+
+#include <GTEngine/Mathematics/GteMath.h>
+
+namespace gte
+{
+
+// Test-intersection queries.
+
+template <typename Real, typename Type0, typename Type1>
+class TIQuery
+{
+public:
+    struct Result
+    {
+        // A TIQuery-base class B must define a B::Result struct with member
+        // 'bool intersect'.  A TIQuery-derived class D must also derive a
+        // D::Result from B:Result but may have no members.  The member
+        // 'intersect' is 'true' iff the primitives intersect.  The operator()
+        // is non-const to allow TIQuery to store and modify private state
+        // that supports the query.
+    };
+    Result operator()(Type0 const& primitive0, Type1 const& primitive1);
+};
+
+}
diff --git a/Engine/Plugins/Experimental/GeometryProcessing/Source/GeometryAlgorithms/Private/GTEngine/Mathematics/GteTSManifoldMesh.cpp b/Engine/Plugins/Experimental/GeometryProcessing/Source/GeometryAlgorithms/Private/GTEngine/Mathematics/GteTSManifoldMesh.cpp
new file mode 100644
--- /dev/null
+++ b/Engine/Plugins/Experimental/GeometryProcessing/Source/GeometryAlgorithms/Private/GTEngine/Mathematics/GteTSManifoldMesh.cpp
@@ -0,0 +1,255 @@
+// David Eberly, Geometric Tools, Redmond WA 98052
+// Copyright (c) 1998-2019
+// Distributed under the Boost Software License, Version 1.0.
+// http://www.boost.org/LICENSE_1_0.txt
+// http://www.geometrictools.com/License/Boost/LICENSE_1_0.txt
+// File Version: 3.0.0 (2016/06/19)
+
+#include <GTEngine/Mathematics/GteTSManifoldMesh.h>
+#include <GTEngine/LowLevel/GteLogger.h>
+#include <GTEngine/GTEnginePCH.h>
+
+using namespace gte;
+
+TSManifoldMesh::~TSManifoldMesh()
+{
+}
+
+TSManifoldMesh::TSManifoldMesh(TCreator tCreator, SCreator sCreator)
+    :
+    mTCreator(tCreator ? tCreator : CreateTriangle),
+    mSCreator(sCreator ? sCreator : CreateTetrahedron),
+    mAssertOnNonmanifoldInsertion(true)
+{
+}
+
+TSManifoldMesh::TSManifoldMesh(TSManifoldMesh const& mesh)
+{
+    *this = mesh;
+}
+
+TSManifoldMesh& TSManifoldMesh::operator=(TSManifoldMesh const& mesh)
+{
+    Clear();
+
+    mTCreator = mesh.mTCreator;
+    mSCreator = mesh.mSCreator;
+    mAssertOnNonmanifoldInsertion = mesh.mAssertOnNonmanifoldInsertion;
+    for (auto const& element : mesh.mSMap)
+    {
+        Insert(element.first.V[0], element.first.V[1], element.first.V[2], element.first.V[3]);
+    }
+
+    return *this;
+}
+
+TSManifoldMesh::TMap const& TSManifoldMesh::GetTriangles() const
+{
+    return mTMap;
+}
+
+TSManifoldMesh::SMap const& TSManifoldMesh::GetTetrahedra() const
+{
+    return mSMap;
+}
+
+std::shared_ptr<TSManifoldMesh::Triangle> TSManifoldMesh::CreateTriangle(int v0, int v1, int v2)
+{
+    return std::make_shared<Triangle>(v0, v1, v2);
+}
+
+std::shared_ptr<TSManifoldMesh::Tetrahedron> TSManifoldMesh::CreateTetrahedron(int v0, int v1, int v2, int v3)
+{
+    return std::make_shared<Tetrahedron>(v0, v1, v2, v3);
+}
+
+void TSManifoldMesh::AssertOnNonmanifoldInsertion(bool doAssert)
+{
+    mAssertOnNonmanifoldInsertion = doAssert;
+}
+
+std::shared_ptr<TSManifoldMesh::Tetrahedron> TSManifoldMesh::Insert(int v0, int v1, int v2, int v3)
+{
+    TetrahedronKey<true> skey(v0, v1, v2, v3);
+    if (mSMap.find(skey) != mSMap.end())
+    {
+        // The tetrahedron already exists.  Return a null pointer as a signal
+        // to the caller that the insertion failed.
+        return nullptr;
+    }
+
+    // Add the new tetrahedron.
+    std::shared_ptr<Tetrahedron> tetra = mSCreator(v0, v1, v2, v3);
+    mSMap[skey] = tetra;
+
+    // Add the faces to the mesh if they do not already exist.
+    for (int i = 0; i < 4; ++i)
+    {
+        auto opposite = TetrahedronKey<true>::oppositeFace[i];
+        TriangleKey<false> tkey(tetra->V[opposite[0]], tetra->V[opposite[1]], tetra->V[opposite[2]]);
+        std::shared_ptr<Triangle> face;
+        auto titer = mTMap.find(tkey);
+        if (titer == mTMap.end())
+        {
+            // This is the first time the face is encountered.
+            face = mTCreator(tetra->V[opposite[0]], tetra->V[opposite[1]], tetra->V[opposite[2]]);
+            mTMap[tkey] = face;
+
+            // Update the face and tetrahedron.
+            face->T[0] = tetra;
+            tetra->T[i] = face;
+        }
+        else
+        {
+            // This is the second time the face is encountered.
+            face = titer->second;
+            if (!face)
+            {
+                LogError("Unexpected condition.");
+                return nullptr;
+            }
+
+            // Update the face.
+            if (face->T[1].lock())
+            {
+                if (mAssertOnNonmanifoldInsertion)
+                {
+                    LogInformation("The mesh must be manifold.");
+                }
+                return nullptr;
+            }
+            face->T[1] = tetra;
+
+            // Update the adjacent tetrahedra.
+            auto adjacent = face->T[0].lock();
+            if (!adjacent)
+            {
+                LogError("Unexpected condition.");
+                return nullptr;
+            }
+            for (int j = 0; j < 4; ++j)
+            {
+                if (adjacent->T[j].lock() == face)
+                {
+                    adjacent->S[j] = tetra;
+                    break;
+                }
+            }
+
+            // Update the tetrahedron.
+            tetra->T[i] = face;
+            tetra->S[i] = adjacent;
+        }
+    }
+
+    return tetra;
+}
+
+bool TSManifoldMesh::Remove(int v0, int v1, int v2, int v3)
+{
+    TetrahedronKey<true> skey(v0, v1, v2, v3);
+    auto siter = mSMap.find(skey);
+    if (siter == mSMap.end())
+    {
+        // The tetrahedron does not exist.
+        return false;
+    }
+
+    // Get the tetrahedron.
+    std::shared_ptr<Tetrahedron> tetra = siter->second;
+
+    // Remove the faces and update adjacent tetrahedra if necessary.
+    for (int i = 0; i < 4; ++i)
+    {
+        // Inform the faces the tetrahedron is being deleted.
+        auto face = tetra->T[i].lock();
+        if (!face)
+        {
+            // The triangle edge should be nonnull.
+            LogError("Unexpected condition.");
+            return false;
+        }
+
+        if (face->T[0].lock() == tetra)
+        {
+            // One-tetrahedron faces always have pointer at index zero.
+            face->T[0] = face->T[1];
+            face->T[1].reset();
+        }
+        else if (face->T[1].lock() == tetra)
+        {
+            face->T[1].reset();
+        }
+        else
+        {
+            LogError("Unexpected condition.");
+            return false;
+        }
+
+        // Remove the face if you have the last reference to it.
+        if (!face->T[0].lock() && !face->T[1].lock())
+        {
+            TriangleKey<false> tkey(face->V[0], face->V[1], face->V[2]);
+            mTMap.erase(tkey);
+        }
+
+        // Inform adjacent tetrahedra the tetrahedron is being deleted.
+        auto adjacent = tetra->S[i].lock();
+        if (adjacent)
+        {
+            for (int j = 0; j < 4; ++j)
+            {
+                if (adjacent->S[j].lock() == tetra)
+                {
+                    adjacent->S[j].reset();
+                    break;
+                }
+            }
+        }
+    }
+
+    mSMap.erase(skey);
+    return true;
+}
+
+void TSManifoldMesh::Clear()
+{
+    mTMap.clear();
+    mSMap.clear();
+}
+
+bool TSManifoldMesh::IsClosed() const
+{
+    for (auto const& element : mSMap)
+    {
+        auto tri = element.second;
+        if (!tri->S[0].lock() || !tri->S[1].lock())
+        {
+            return false;
+        }
+    }
+    return true;
+}
+
+TSManifoldMesh::Triangle::~Triangle()
+{
+}
+
+TSManifoldMesh::Triangle::Triangle(int v0, int v1, int v2)
+{
+    V[0] = v0;
+    V[1] = v1;
+    V[2] = v2;
+}
+
+TSManifoldMesh::Tetrahedron::~Tetrahedron()
+{
+}
+
+TSManifoldMesh::Tetrahedron::Tetrahedron(int v0, int v1, int v2, int v3)
+{
+    V[0] = v0;
+    V[1] = v1;
+    V[2] = v2;
+    V[3] = v3;
+}
diff --git a/Engine/Plugins/Experimental/GeometryProcessing/Source/GeometryAlgorithms/Private/GTEngine/Mathematics/GteTSManifoldMesh.h b/Engine/Plugins/Experimental/GeometryProcessing/Source/GeometryAlgorithms/Private/GTEngine/Mathematics/GteTSManifoldMesh.h
new file mode 100644
--- /dev/null
+++ b/Engine/Plugins/Experimental/GeometryProcessing/Source/GeometryAlgorithms/Private/GTEngine/Mathematics/GteTSManifoldMesh.h
@@ -0,0 +1,120 @@
+// David Eberly, Geometric Tools, Redmond WA 98052
+// Copyright (c) 1998-2019
+// Distributed under the Boost Software License, Version 1.0.
+// http://www.boost.org/LICENSE_1_0.txt
+// http://www.geometrictools.com/License/Boost/LICENSE_1_0.txt
+// File Version: 3.0.0 (2016/06/19)
+
+#pragma once
+
+#include <GTEngine/Mathematics/GteTetrahedronKey.h>
+#include <GTEngine/Mathematics/GteTriangleKey.h>
+#include <map>
+#include <memory>
+
+namespace gte
+{
+
+class GTE_IMPEXP TSManifoldMesh
+{
+public:
+    // Triangle data types.
+    class Triangle;
+    typedef std::shared_ptr<Triangle> (*TCreator)(int, int, int);
+    typedef std::map<TriangleKey<false>, std::shared_ptr<Triangle>> TMap;
+
+    // Tetrahedron data types.
+    class Tetrahedron;
+    typedef std::shared_ptr<Tetrahedron> (*SCreator)(int, int, int, int);
+    typedef std::map<TetrahedronKey<true>, std::shared_ptr<Tetrahedron>> SMap;
+
+    // Triangle object.
+    class GTE_IMPEXP Triangle
+    {
+    public:
+        virtual ~Triangle();
+        Triangle(int v0, int v1, int v2);
+
+        // Vertices of the face.
+        int V[3];
+
+        // Tetrahedra sharing the face.
+        std::weak_ptr<Tetrahedron> T[2];
+    };
+
+    // Tetrahedron object.
+    class GTE_IMPEXP Tetrahedron
+    {
+    public:
+        virtual ~Tetrahedron();
+        Tetrahedron(int v0, int v1, int v2, int v3);
+
+        // Vertices, listed in an order so that each face vertices in
+        // counterclockwise order when viewed from outside the tetrahedron.
+        int V[4];
+
+        // Adjacent faces.  T[i] points to the triangle face opposite V[i].
+        //   T[0] points to face (V[1],V[2],V[3])
+        //   T[1] points to face (V[0],V[3],V[2])
+        //   T[2] points to face (V[0],V[1],V[3])
+        //   T[3] points to face (V[0],V[2],V[1])
+        std::weak_ptr<Triangle> T[4];
+
+        // Adjacent tetrahedra.  S[i] points to the adjacent tetrahedron
+        // sharing face T[i].
+        std::weak_ptr<Tetrahedron> S[4];
+    };
+
+
+    // Construction and destruction.
+    virtual ~TSManifoldMesh();
+    TSManifoldMesh(TCreator tCreator = nullptr, SCreator sCreator = nullptr);
+
+    // Support for a deep copy of the mesh.  The mTMap and mSMap objects have
+    // dynamically allocated memory for triangles and tetrahedra.  A shallow
+    // copy of the pointers to this memory is problematic.  Allowing sharing,
+    // say, via std::shared_ptr, is an option but not really the intent of
+    // copying the mesh graph.
+    TSManifoldMesh(TSManifoldMesh const& mesh);
+    TSManifoldMesh& operator=(TSManifoldMesh const& mesh);
+
+    // Member access.
+    TMap const& GetTriangles() const;
+    SMap const& GetTetrahedra() const;
+
+    // If the insertion of a tetrahedron fails because the mesh would become
+    // nonmanifold, the default behavior is to trigger a LogError message.
+    // You can disable this behavior in situations where you want the Logger
+    // system on but you want to continue gracefully without an assertion.
+    void AssertOnNonmanifoldInsertion(bool doAssert);
+
+    // If <v0,v1,v2,v3> is not in the mesh, a Tetrahedron object is created
+    // and returned; otherwise, <v0,v1,v2,v3> is in the mesh and nullptr is
+    // returned.  If the insertion leads to a nonmanifold mesh, the call
+    // fails with a nullptr returned.
+    std::shared_ptr<Tetrahedron> Insert(int v0, int v1, int v2, int v3);
+
+    // If <v0,v1,v2,v3> is in the mesh, it is removed and 'true' is returned;
+    // otherwise, <v0,v1,v2,v3> is not in the mesh and 'false' is returned.
+    bool Remove(int v0, int v1, int v2, int v3);
+
+    // Destroy the triangles and tetrahedra to obtain an empty mesh.
+    virtual void Clear();
+
+    // A manifold mesh is closed if each face is shared twice.
+    bool IsClosed() const;
+
+protected:
+    // The triangle data and default triangle creation.
+    static std::shared_ptr<Triangle> CreateTriangle(int v0, int v1, int v2);
+    TCreator mTCreator;
+    TMap mTMap;
+
+    // The tetrahedron data and default tetrahedron creation.
+    static std::shared_ptr<Tetrahedron> CreateTetrahedron(int v0, int v1, int v2, int v3);
+    SCreator mSCreator;
+    SMap mSMap;
+    bool mAssertOnNonmanifoldInsertion;  // default: true
+};
+
+}
diff --git a/Engine/Plugins/Experimental/GeometryProcessing/Source/GeometryAlgorithms/Private/GTEngine/Mathematics/GteTanEstimate.h b/Engine/Plugins/Experimental/GeometryProcessing/Source/GeometryAlgorithms/Private/GTEngine/Mathematics/GteTanEstimate.h
new file mode 100644
--- /dev/null
+++ b/Engine/Plugins/Experimental/GeometryProcessing/Source/GeometryAlgorithms/Private/GTEngine/Mathematics/GteTanEstimate.h
@@ -0,0 +1,187 @@
+// David Eberly, Geometric Tools, Redmond WA 98052
+// Copyright (c) 1998-2019
+// Distributed under the Boost Software License, Version 1.0.
+// http://www.boost.org/LICENSE_1_0.txt
+// http://www.geometrictools.com/License/Boost/LICENSE_1_0.txt
+// File Version: 3.0.1 (2018/10/05)
+
+#pragma once
+
+#include <GTEngine/Mathematics/GteMath.h>
+
+// Minimax polynomial approximations to tan(x).  The polynomial p(x) of
+// degree D has only odd-power terms, is required to have linear term x,
+// and p(pi/4) = tan(pi/4) = 1.  It minimizes the quantity
+// maximum{|tan(x) - p(x)| : x in [-pi/4,pi/4]} over all polynomials of
+// degree D subject to the constraints mentioned.
+
+namespace gte
+{
+
+template <typename Real>
+class TanEstimate
+{
+public:
+    // The input constraint is x in [-pi/4,pi/4].  For example,
+    //   float x; // in [-pi/4,pi/4]
+    //   float result = TanEstimate<float>::Degree<3>(x);
+    template <int D>
+    inline static Real Degree(Real x);
+
+    // The input x can be any real number.  Range reduction is used to
+    // generate a value y in [-pi/2,pi/2].  If |y| <= pi/4, then the
+    // polynomial is evaluated.  If y in (pi/4,pi/2), set z = y - pi/4
+    // and use the identity
+    //   tan(y) = tan(z + pi/4) = [1 + tan(z)]/[1 - tan(z)]
+    // Be careful when evaluating at y nearly pi/2, because tan(y)
+    // becomes infinite.  For example,
+    //   float x;  // x any real number
+    //   float result = TanEstimate<float>::DegreeRR<3>(x);
+    template <int D>
+    inline static Real DegreeRR(Real x);
+
+private:
+    // Metaprogramming and private implementation to allow specialization of
+    // a template member function.
+    template <int D> struct degree {};
+    inline static Real Evaluate(degree<3>, Real x);
+    inline static Real Evaluate(degree<5>, Real x);
+    inline static Real Evaluate(degree<7>, Real x);
+    inline static Real Evaluate(degree<9>, Real x);
+    inline static Real Evaluate(degree<11>, Real x);
+    inline static Real Evaluate(degree<13>, Real x);
+
+    // Support for range reduction.
+    inline static void Reduce(Real x, Real& y);
+};
+
+
+template <typename Real>
+template <int D>
+inline Real TanEstimate<Real>::Degree(Real x)
+{
+    return Evaluate(degree<D>(), x);
+}
+
+template <typename Real>
+template <int D>
+inline Real TanEstimate<Real>::DegreeRR(Real x)
+{
+    Real y;
+    Reduce(x, y);
+    if (std::abs(y) <= (Real)GTE_C_QUARTER_PI)
+    {
+        return Degree<D>(y);
+    }
+    else if (y > (Real)GTE_C_QUARTER_PI)
+    {
+        Real poly = Degree<D>(y - (Real)GTE_C_QUARTER_PI);
+        return ((Real)1 + poly) / ((Real)1 - poly);
+    }
+    else
+    {
+        Real poly = Degree<D>(y + (Real)GTE_C_QUARTER_PI);
+        return -((Real)1 - poly) / ((Real)1 + poly);
+    }
+}
+
+template <typename Real>
+inline Real TanEstimate<Real>::Evaluate(degree<3>, Real x)
+{
+    Real xsqr = x * x;
+    Real poly;
+    poly = (Real)GTE_C_TAN_DEG3_C1;
+    poly = (Real)GTE_C_TAN_DEG3_C0 + poly * xsqr;
+    poly = poly * x;
+    return poly;
+}
+
+template <typename Real>
+inline Real TanEstimate<Real>::Evaluate(degree<5>, Real x)
+{
+    Real xsqr = x * x;
+    Real poly;
+    poly = (Real)GTE_C_TAN_DEG5_C2;
+    poly = (Real)GTE_C_TAN_DEG5_C1 + poly * xsqr;
+    poly = (Real)GTE_C_TAN_DEG5_C0 + poly * xsqr;
+    poly = poly * x;
+    return poly;
+}
+
+template <typename Real>
+inline Real TanEstimate<Real>::Evaluate(degree<7>, Real x)
+{
+    Real xsqr = x * x;
+    Real poly;
+    poly = (Real)GTE_C_TAN_DEG7_C3;
+    poly = (Real)GTE_C_TAN_DEG7_C2 + poly * xsqr;
+    poly = (Real)GTE_C_TAN_DEG7_C1 + poly * xsqr;
+    poly = (Real)GTE_C_TAN_DEG7_C0 + poly * xsqr;
+    poly = poly * x;
+    return poly;
+}
+
+template <typename Real>
+inline Real TanEstimate<Real>::Evaluate(degree<9>, Real x)
+{
+    Real xsqr = x * x;
+    Real poly;
+    poly = (Real)GTE_C_TAN_DEG9_C4;
+    poly = (Real)GTE_C_TAN_DEG9_C3 + poly * xsqr;
+    poly = (Real)GTE_C_TAN_DEG9_C2 + poly * xsqr;
+    poly = (Real)GTE_C_TAN_DEG9_C1 + poly * xsqr;
+    poly = (Real)GTE_C_TAN_DEG9_C0 + poly * xsqr;
+    poly = poly * x;
+    return poly;
+}
+
+template <typename Real>
+inline Real TanEstimate<Real>::Evaluate(degree<11>, Real x)
+{
+    Real xsqr = x * x;
+    Real poly;
+    poly = (Real)GTE_C_TAN_DEG11_C5;
+    poly = (Real)GTE_C_TAN_DEG11_C4 + poly * xsqr;
+    poly = (Real)GTE_C_TAN_DEG11_C3 + poly * xsqr;
+    poly = (Real)GTE_C_TAN_DEG11_C2 + poly * xsqr;
+    poly = (Real)GTE_C_TAN_DEG11_C1 + poly * xsqr;
+    poly = (Real)GTE_C_TAN_DEG11_C0 + poly * xsqr;
+    poly = poly * x;
+    return poly;
+}
+
+template <typename Real>
+inline Real TanEstimate<Real>::Evaluate(degree<13>, Real x)
+{
+    Real xsqr = x * x;
+    Real poly;
+    poly = (Real)GTE_C_TAN_DEG13_C6;
+    poly = (Real)GTE_C_TAN_DEG13_C5 + poly * xsqr;
+    poly = (Real)GTE_C_TAN_DEG13_C4 + poly * xsqr;
+    poly = (Real)GTE_C_TAN_DEG13_C3 + poly * xsqr;
+    poly = (Real)GTE_C_TAN_DEG13_C2 + poly * xsqr;
+    poly = (Real)GTE_C_TAN_DEG13_C1 + poly * xsqr;
+    poly = (Real)GTE_C_TAN_DEG13_C0 + poly * xsqr;
+    poly = poly * x;
+    return poly;
+}
+
+template <typename Real>
+inline void TanEstimate<Real>::Reduce(Real x, Real& y)
+{
+    // Map x to y in [-pi,pi], x = pi*quotient + remainder.
+    y = std::fmod(x, (Real)GTE_C_PI);
+
+    // Map y to [-pi/2,pi/2] with tan(y) = tan(x).
+    if (y > (Real)GTE_C_HALF_PI)
+    {
+        y -= (Real)GTE_C_PI;
+    }
+    else if (y < (Real)-GTE_C_HALF_PI)
+    {
+        y += (Real)GTE_C_PI;
+    }
+}
+
+
+}
diff --git a/Engine/Plugins/Experimental/GeometryProcessing/Source/GeometryAlgorithms/Private/GTEngine/Mathematics/GteTetrahedron3.h b/Engine/Plugins/Experimental/GeometryProcessing/Source/GeometryAlgorithms/Private/GTEngine/Mathematics/GteTetrahedron3.h
new file mode 100644
--- /dev/null
+++ b/Engine/Plugins/Experimental/GeometryProcessing/Source/GeometryAlgorithms/Private/GTEngine/Mathematics/GteTetrahedron3.h
@@ -0,0 +1,200 @@
+// David Eberly, Geometric Tools, Redmond WA 98052
+// Copyright (c) 1998-2019
+// Distributed under the Boost Software License, Version 1.0.
+// http://www.boost.org/LICENSE_1_0.txt
+// http://www.geometrictools.com/License/Boost/LICENSE_1_0.txt
+// File Version: 3.0.0 (2016/06/19)
+
+#pragma once
+
+#include <GTEngine/Mathematics/GteVector3.h>
+#include <GTEngine/Mathematics/GteHyperplane.h>
+
+// The tetrahedron is represented as an array of four vertices: V0, V1, V2,
+// and V3.  The vertices are ordered so that the triangular faces are
+// counterclockwise-ordered triangles when viewed by an observer outside the
+// tetrahedron:
+//   face 0 = <V[0],V[2],V[1]>
+//   face 1 = <V[0],V[1],V[3]>
+//   face 2 = <V[0],V[3],V[2]>
+//   face 3 = <V[1],V[2],V[3]>
+
+namespace gte
+{
+
+template <typename Real>
+class Tetrahedron3
+{
+public:
+    // Construction and destruction.  The default constructor sets the
+    // vertices to (0,0,0), (1,0,0), (0,1,0), and (0,0,1).
+    Tetrahedron3();
+    Tetrahedron3(Vector3<Real> const& v0, Vector3<Real> const& v1,
+        Vector3<Real> const& v2, Vector3<Real> const& v3);
+    Tetrahedron3(Vector3<Real> const inV[4]);
+
+    // Get the vertex indices for the specified face.  The input 'face' must
+    // be in {0,1,2,3}.
+    void GetFaceIndices(int face, int index[3]) const;
+
+    // Construct the planes of the faces.  The planes have outer pointing
+    // normal vectors.  The plane indexing is the same as the face indexing
+    // mentioned previously.
+    void GetPlanes(Plane3<Real> plane[4]) const;
+
+    // Public member access.
+    Vector3<Real> v[4];
+
+public:
+    // Comparisons to support sorted containers.
+    bool operator==(Tetrahedron3 const& tetrahedron) const;
+    bool operator!=(Tetrahedron3 const& tetrahedron) const;
+    bool operator< (Tetrahedron3 const& tetrahedron) const;
+    bool operator<=(Tetrahedron3 const& tetrahedron) const;
+    bool operator> (Tetrahedron3 const& tetrahedron) const;
+    bool operator>=(Tetrahedron3 const& tetrahedron) const;
+};
+
+
+template <typename Real>
+Tetrahedron3<Real>::Tetrahedron3()
+{
+    v[0] = Vector3<Real>::Zero();
+    v[1] = Vector3<Real>::Unit(0);
+    v[2] = Vector3<Real>::Unit(1);
+    v[3] = Vector3<Real>::Unit(2);
+}
+
+template <typename Real>
+Tetrahedron3<Real>::Tetrahedron3(Vector3<Real> const& v0,
+    Vector3<Real> const& v1, Vector3<Real> const& v2, Vector3<Real> const& v3)
+{
+    v[0] = v0;
+    v[1] = v1;
+    v[2] = v2;
+    v[3] = v3;
+}
+
+template <typename Real>
+Tetrahedron3<Real>::Tetrahedron3(Vector3<Real> const inV[4])
+{
+    for (int i = 0; i < 4; ++i)
+    {
+        v[i] = inV[i];
+    }
+}
+
+template <typename Real>
+void Tetrahedron3<Real>::GetFaceIndices(int face, int index[3])
+const
+{
+    if (face == 0)
+    {
+        index[0] = 0;
+        index[1] = 2;
+        index[2] = 1;
+    }
+    else if (face == 1)
+    {
+        index[0] = 0;
+        index[1] = 1;
+        index[2] = 3;
+    }
+    else if (face == 2)
+    {
+        index[0] = 0;
+        index[1] = 3;
+        index[2] = 2;
+    }
+    else  // face == 3 (no index validation is performed)
+    {
+        index[0] = 1;
+        index[1] = 2;
+        index[2] = 3;
+    }
+}
+
+template <typename Real>
+void Tetrahedron3<Real>::GetPlanes(Plane3<Real> plane[4]) const
+{
+    Vector3<Real> edge10 = v[1] - v[0];
+    Vector3<Real> edge20 = v[2] - v[0];
+    Vector3<Real> edge30 = v[3] - v[0];
+    Vector3<Real> edge21 = v[2] - v[1];
+    Vector3<Real> edge31 = v[3] - v[1];
+
+    plane[0].normal = UnitCross(edge20, edge10);  // <v0,v2,v1>
+    plane[1].normal = UnitCross(edge10, edge30);  // <v0,v1,v3>
+    plane[2].normal = UnitCross(edge30, edge20);  // <v0,v3,v2>
+    plane[3].normal = UnitCross(edge21, edge31);  // <v1,v2,v3>
+
+    Real det = Dot(edge10, plane[3].normal);
+    if (det < (Real)0)
+    {
+        // The normals are inner pointing, reverse their directions.
+        for (int i = 0; i < 4; ++i)
+        {
+            plane[i].normal = -plane[i].normal;
+        }
+    }
+
+    for (int i = 0; i < 4; ++i)
+    {
+        plane[i].constant = Dot(v[i], plane[i].normal);
+    }
+}
+
+template <typename Real>
+bool Tetrahedron3<Real>::operator==(Tetrahedron3 const& tetrahedron) const
+{
+    return v[0] == tetrahedron.v[0]
+        && v[1] == tetrahedron.v[1]
+        && v[2] == tetrahedron.v[2]
+        && v[3] == tetrahedron.v[3];
+}
+
+template <typename Real>
+bool Tetrahedron3<Real>::operator!=(Tetrahedron3 const& tetrahedron) const
+{
+    return !operator==(tetrahedron);
+}
+
+template <typename Real>
+bool Tetrahedron3<Real>::operator<(Tetrahedron3 const& tetrahedron) const
+{
+    for (int i = 0; i < 4; ++i)
+    {
+        if (v[i] < tetrahedron.v[i])
+        {
+            return true;
+        }
+
+        if (v[i] > tetrahedron.v[i])
+        {
+            return false;
+        }
+    }
+
+    return false;
+}
+
+template <typename Real>
+bool Tetrahedron3<Real>::operator<=(Tetrahedron3 const& tetrahedron) const
+{
+    return operator<(tetrahedron) || operator==(tetrahedron);
+}
+
+template <typename Real>
+bool Tetrahedron3<Real>::operator>(Tetrahedron3 const& tetrahedron) const
+{
+    return !operator<=(tetrahedron);
+}
+
+template <typename Real>
+bool Tetrahedron3<Real>::operator>=(Tetrahedron3 const& tetrahedron) const
+{
+    return !operator<(tetrahedron);
+}
+
+
+}
diff --git a/Engine/Plugins/Experimental/GeometryProcessing/Source/GeometryAlgorithms/Private/GTEngine/Mathematics/GteTetrahedronKey.cpp b/Engine/Plugins/Experimental/GeometryProcessing/Source/GeometryAlgorithms/Private/GTEngine/Mathematics/GteTetrahedronKey.cpp
new file mode 100644
--- /dev/null
+++ b/Engine/Plugins/Experimental/GeometryProcessing/Source/GeometryAlgorithms/Private/GTEngine/Mathematics/GteTetrahedronKey.cpp
@@ -0,0 +1,145 @@
+// David Eberly, Geometric Tools, Redmond WA 98052
+// Copyright (c) 1998-2019
+// Distributed under the Boost Software License, Version 1.0.
+// http://www.boost.org/LICENSE_1_0.txt
+// http://www.geometrictools.com/License/Boost/LICENSE_1_0.txt
+// File Version: 3.0.1 (2018/05/22)
+
+#include <GTEngine/Mathematics/GteTetrahedronKey.h>
+#include <GTEngine/GTEnginePCH.h>
+#include <algorithm>
+
+namespace
+{
+    void Permute(int u0, int u1, int u2, int V[4])
+    {
+        // Once V[0] is determined, create a permutation (V[1],V[2],V[3]) so
+        // that (V[0],V[1],V[2],V[3]) is a positive permutation of
+        // (v0,v1,v2,v3).
+        if (u0 < u1)
+        {
+            if (u0 < u2)
+            {
+                // u0 is minimum
+                V[1] = u0;
+                V[2] = u1;
+                V[3] = u2;
+            }
+            else
+            {
+                // u2 is minimum
+                V[1] = u2;
+                V[2] = u0;
+                V[3] = u1;
+            }
+        }
+        else
+        {
+            if (u1 < u2)
+            {
+                // u1 is minimum
+                V[1] = u1;
+                V[2] = u2;
+                V[3] = u0;
+            }
+            else
+            {
+                // u2 is minimum
+                V[1] = u2;
+                V[2] = u0;
+                V[3] = u1;
+            }
+        }
+    }
+}
+
+namespace gte
+{
+    // TetrahedronKey<true>
+
+    template<> std::array<int, 3> const TetrahedronKey<true>::oppositeFace[4] =
+    {
+        { 1, 2, 3 },
+        { 0, 3, 2 },
+        { 0, 1, 3 },
+        { 0, 2, 1 }
+    };
+
+    template<>
+    TetrahedronKey<true>::TetrahedronKey()
+    {
+        V[0] = -1;
+        V[1] = -1;
+        V[2] = -1;
+        V[3] = -1;
+    }
+
+    template<>
+    TetrahedronKey<true>::TetrahedronKey(int v0, int v1, int v2, int v3)
+    {
+        int imin = 0;
+        V[0] = v0;
+        if (v1 < V[0])
+        {
+            V[0] = v1;
+            imin = 1;
+        }
+        if (v2 < V[0])
+        {
+            V[0] = v2;
+            imin = 2;
+        }
+        if (v3 < V[0])
+        {
+            V[0] = v3;
+            imin = 3;
+        }
+
+        if (imin == 0)
+        {
+            Permute(v1, v2, v3, V);
+        }
+        else if (imin == 1)
+        {
+            Permute(v0, v3, v2, V);
+        }
+        else if (imin == 2)
+        {
+            Permute(v0, v1, v3, V);
+        }
+        else  // imin == 3
+        {
+            Permute(v0, v2, v1, V);
+        }
+    }
+
+
+    // TetrahedronKey<false>
+
+    template<> std::array<int, 3> const TetrahedronKey<false>::oppositeFace[4] =
+    {
+        { 1, 2, 3 },
+        { 0, 3, 2 },
+        { 0, 1, 3 },
+        { 0, 2, 1 }
+    };
+
+    template<>
+    TetrahedronKey<false>::TetrahedronKey()
+    {
+        V[0] = -1;
+        V[1] = -1;
+        V[2] = -1;
+        V[3] = -1;
+    }
+
+    template<>
+    TetrahedronKey<false>::TetrahedronKey(int v0, int v1, int v2, int v3)
+    {
+        V[0] = v0;
+        V[1] = v1;
+        V[2] = v2;
+        V[3] = v3;
+        std::sort(&V[0], &V[4]);
+    }
+}
diff --git a/Engine/Plugins/Experimental/GeometryProcessing/Source/GeometryAlgorithms/Private/GTEngine/Mathematics/GteTetrahedronKey.h b/Engine/Plugins/Experimental/GeometryProcessing/Source/GeometryAlgorithms/Private/GTEngine/Mathematics/GteTetrahedronKey.h
new file mode 100644
--- /dev/null
+++ b/Engine/Plugins/Experimental/GeometryProcessing/Source/GeometryAlgorithms/Private/GTEngine/Mathematics/GteTetrahedronKey.h
@@ -0,0 +1,61 @@
+// David Eberly, Geometric Tools, Redmond WA 98052
+// Copyright (c) 1998-2019
+// Distributed under the Boost Software License, Version 1.0.
+// http://www.boost.org/LICENSE_1_0.txt
+// http://www.geometrictools.com/License/Boost/LICENSE_1_0.txt
+// File Version: 3.0.1 (2018/05/22)
+
+#pragma once
+
+#include <GTEngine/Mathematics/GteFeatureKey.h>
+#include <array>
+
+namespace gte
+{
+    template <bool Ordered>
+    class TetrahedronKey : public FeatureKey<4, Ordered>
+    {
+    public:
+        // An ordered tetrahedron has V[0] = min(v0,v1,v2,v3).  Let {u1,u2,u3}
+        // be the set of inputs excluding the one assigned to V[0] and define
+        // V[1] = min(u1,u2,u3).  Choose (V[1],V[2],V[3]) to be a permutation
+        // of (u1,u2,u3) so that the final storage is one of
+        //   (v0,v1,v2,v3), (v0,v2,v3,v1), (v0,v3,v1,v2)
+        //   (v1,v3,v2,v0), (v1,v2,v0,v3), (v1,v0,v3,v2)
+        //   (v2,v3,v0,v1), (v2,v0,v1,v3), (v2,v1,v3,v0)
+        //   (v3,v1,v0,v2), (v3,v0,v2,v1), (v3,v2,v1,v0)
+        // The idea is that if v0 corresponds to (1,0,0,0), v1 corresponds to
+        // (0,1,0,0), v2 corresponds to (0,0,1,0), and v3 corresponds to
+        // (0,0,0,1), the ordering (v0,v1,v2,v3) corresponds to the 4x4
+        // identity matrix I; the rows are the specified 4-tuples.  The
+        // permutation (V[0],V[1],V[2],V[3]) induces a permutation of the rows
+        // of the identity matrix to form a permutation matrix P with det(P)
+        // = 1 = det(I).
+        //
+        // An unordered tetrahedron stores a permutation of (v0,v1,v2,v3) so
+        // that V[0] < V[1] < V[2] < V[3].
+        TetrahedronKey();  // creates key (-1,-1,-1,-1)
+        explicit TetrahedronKey(int v0, int v1, int v2, int v3);
+
+        // Indexing for the vertices of the triangle opposite a vertex.  The
+        // triangle opposite vertex j is
+        //   <oppositeFace[j][0], oppositeFace[j][1], oppositeFace[j][2]>
+        // and is listed in counterclockwise order when viewed from outside
+        // the tetrahedron.
+        static std::array<int, 3> const oppositeFace[4];
+    };
+
+#if !defined(__MSWINDOWS__)
+    // Apple LLVM 6.1.0 (clang-602.0.49) correctly requires these declarations
+    // to occur before the definition in GteTetrahedronKey.cpp.  From the C++
+    // specification:
+    //    An explicit specialization of a static data member of a template is
+    //    a definition if the declaration includes an initializer; otherwise,
+    //    it is a declaration.
+    // If these declarations are exposed for MSVS 2013, the compiler generates
+    // error C2086, claiming that this is a definition (rather than a
+    // declaration) and the cpp file has a redefinition.
+    template<> std::array<int, 3> const TetrahedronKey<false>::oppositeFace[4];
+    template<> std::array<int, 3> const TetrahedronKey<true>::oppositeFace[4];
+#endif
+}
diff --git a/Engine/Plugins/Experimental/GeometryProcessing/Source/GeometryAlgorithms/Private/GTEngine/Mathematics/GteTorus3.h b/Engine/Plugins/Experimental/GeometryProcessing/Source/GeometryAlgorithms/Private/GTEngine/Mathematics/GteTorus3.h
new file mode 100644
--- /dev/null
+++ b/Engine/Plugins/Experimental/GeometryProcessing/Source/GeometryAlgorithms/Private/GTEngine/Mathematics/GteTorus3.h
@@ -0,0 +1,254 @@
+// David Eberly, Geometric Tools, Redmond WA 98052
+// Copyright (c) 1998-2019
+// Distributed under the Boost Software License, Version 1.0.
+// http://www.boost.org/LICENSE_1_0.txt
+// http://www.geometrictools.com/License/Boost/LICENSE_1_0.txt
+// File Version: 3.0.2 (2018/10/05)
+
+#pragma once
+
+#include <GTEngine/Mathematics/GteVector3.h>
+
+// A torus with origin (0,0,0), outer radius r0 and inner radius r1 (with
+// (r0 >= r1) is defined implicitly as follows.  The point P0 = (x,y,z) is on
+// the torus. Its projection onto the xy-plane is P1 = (x,y,0).  The circular
+// cross section of the torus that contains the projection has radius r0 and
+// center P2 = r0*(x,y,0)/sqrt(x^2+y^2).  The points triangle <P0,P1,P2> is a
+// right triangle with right angle at P1.  The hypotenuse <P0,P2> has length
+// r1, leg <P1,P2> has length z and leg <P0,P1> has length |r0 - sqrt(x^2+y^2)|.
+// The Pythagorean theorem says z^2 + |r0 - sqrt(x^2+y^2)|^2 = r1^2.  This can
+// be algebraically manipulated to
+//   (x^2 + y^2 + z^2 + r0^2 - r1^2)^2 - 4 * r0^2 * (x^2 + y^2) = 0
+//
+// A parametric form is
+//   x = (r0 + r1 * cos(v)) * cos(u)
+//   y = (r0 + r1 * cos(v)) * sin(u)
+//   z = r1 * sin(v)
+// for u in [0,2*pi) and v in [0,2*pi).
+//
+// Generally, let the torus center be C with plane of symmetry containing C
+// and having directions D0 and D1.  The axis of symmetry is the line
+// containing C and having direction N (the plane normal).  The radius from
+// the center of the torus is r0 and the radius of the tube of the torus is
+// r1.  A point P may be written as P = C + x*D0 + y*D1 + z*N, where matrix
+// [D0 D1 N] is orthonormal and has determinant 1.  Thus, x = Dot(D0,P-C),
+// y = Dot(D1,P-C) and z = Dot(N,P-C).  The implicit form is
+//      [|P-C|^2 + r0^2 - r1^2]^2 - 4*r0^2*[|P-C|^2 - (Dot(N,P-C))^2] = 0
+// Observe that D0 and D1 are not present in the equation, which is to be
+// expected by the symmetry.  The parametric form is
+//      P(u,v) = C + (r0 + r1*cos(v))*(cos(u)*D0 + sin(u)*D1) + r1*sin(v)*N
+// for u in [0,2*pi) and v in [0,2*pi).
+//
+// In the class Torus3, the members are 'center' C, 'direction0' D0,
+// 'direction1' D1, 'normal' N, 'radius0' r0 and 'radius1' r1.
+
+namespace gte
+{
+
+template <typename Real>
+class Torus3
+{
+public:
+    // Construction and destruction.  The default constructor sets center to
+    // (0,0,0), direction0 to (1,0,0), direction1 to (0,1,0), normal to
+    // (0,0,1), radius0 to 2, and radius1 to 1.
+    Torus3();
+    Torus3(Vector3<Real> const& inCenter, Vector3<Real> const& inDirection0,
+        Vector3<Real> const& inDirection1, Vector3<Real> const& inNormal,
+        Real inRadius0, Real inRadius1);
+
+    // Evaluation of the surface.  The function supports derivative
+    // calculation through order 2; that is, maxOrder <= 2 is required.  If
+    // you want only the position, pass in maxOrder of 0.  If you want the
+    // position and first-order derivatives, pass in maxOrder of 1, and so on.
+    // The output 'values' are ordered as: position X; first-order derivatives
+    // dX/du, dX/dv; second-order derivatives d2X/du2, d2X/dudv, d2X/dv2.
+    void Evaluate(Real u, Real v, unsigned int maxOrder, Vector3<Real> values[6]) const;
+
+    // Reverse lookup of parameters from position.
+    void GetParameters(Vector3<Real> const& X, Real& u, Real& v) const;
+
+    Vector3<float> center, direction0, direction1, normal;
+    Real radius0, radius1;
+
+public:
+    // Comparisons to support sorted containers.
+    bool operator==(Torus3 const& torus) const;
+    bool operator!=(Torus3 const& torus) const;
+    bool operator< (Torus3 const& torus) const;
+    bool operator<=(Torus3 const& torus) const;
+    bool operator> (Torus3 const& torus) const;
+    bool operator>=(Torus3 const& torus) const;
+};
+
+
+template <typename Real>
+Torus3<Real>::Torus3()
+    :
+    center(Vector3<Real>::Zero()),
+    direction0(Vector3<Real>::Unit(0)),
+    direction1(Vector3<Real>::Unit(1)),
+    normal(Vector3<Real>::Unit(2)),
+    radius0((Real)2),
+    radius1((Real)1)
+{
+}
+
+template <typename Real>
+Torus3<Real>::Torus3(Vector3<Real> const& inCenter,
+    Vector3<Real> const& inDirection0, Vector3<Real> const& inDirection1,
+    Vector3<Real> const& inNormal, Real inRadius0, Real inRadius1)
+    :
+    center(inCenter),
+    direction0(inDirection0),
+    direction1(inDirection1),
+    normal(inNormal),
+    radius0(inRadius0),
+    radius1(inRadius1)
+{
+}
+
+template <typename Real>
+void Torus3<Real>::Evaluate(Real u, Real v, unsigned int maxOrder,
+    Vector3<Real> values[6]) const
+{
+    // Compute position.
+    Real csu = std::cos(u), snu = std::sin(u), csv = std::cos(v), snv = std::sin(v);
+    Real r1csv = radius1 * csv;
+    Real r1snv = radius1 * snv;
+    Real r0pr1csv = radius0 + r1csv;
+    Vector3<Real> combo0 = csu * direction0 + snu * direction1;
+    Vector3<Real> r0pr1csvcombo0 = r0pr1csv * combo0;
+    Vector3<Real> r1snvnormal = r1snv * normal;
+    values[0] = center + r0pr1csvcombo0 + r1snvnormal;
+
+    if (maxOrder >= 1)
+    {
+        // Compute first-order derivatives.
+        Vector3<Real> combo1 = -snu * direction0 + csu * direction1;
+        values[1] = r0pr1csv * combo1;
+        values[2] = -r1snv * combo0 + r1csv * normal;
+
+        if (maxOrder >= 2)
+        {
+            // Compute second-order derivatives.
+            values[3] = -r0pr1csvcombo0;
+            values[4] = -r1snv * combo1;
+            values[5] = -r1csv * combo0 - r1snvnormal;
+
+            if (maxOrder >= 3)
+            {
+                // These orders are not supported.
+                for (int i = 0; i < 6; ++i)
+                {
+                    values[i] = Vector3<float>::Zero();
+                }
+            }
+        }
+    }
+}
+
+template <typename Real>
+void Torus3<Real>::GetParameters(Vector3<Real> const& X, Real& u, Real& v) const
+{
+    Vector3<Real> delta = X - center;
+    Real dot0 = Dot(direction0, delta);  // (r0 + r1*cos(v))*cos(u)
+    Real dot1 = Dot(direction1, delta);  // (r0 + r1*cos(v))*sin(u)
+    Real dot2 = Dot(normal, delta);      // r1*sin(v)
+    Real r1csv = std::sqrt(dot0 * dot0 + dot1 * dot1) - radius0;  // r1*cos(v)
+    u = std::atan2(dot1, dot0);
+    v = std::atan2(dot2, r1csv);
+}
+
+template <typename Real>
+bool Torus3<Real>::operator==(Torus3 const& torus) const
+{
+    return center == torus.center
+        && direction0 == torus.direction0
+        && direction1 == torus.direction1
+        && normal == torus.normal
+        && radius0 == torus.radius0
+        && radius1 == torus.radius1;
+}
+
+template <typename Real>
+bool Torus3<Real>::operator!=(Torus3 const& torus) const
+{
+    return !operator==(torus);
+}
+
+template <typename Real>
+bool Torus3<Real>::operator<(Torus3 const& torus) const
+{
+    if (center < torus.center)
+    {
+        return true;
+    }
+
+    if (center > torus.center)
+    {
+        return false;
+    }
+
+    if (direction0 < torus.direction0)
+    {
+        return true;
+    }
+
+    if (direction0 > torus.direction0)
+    {
+        return false;
+    }
+
+    if (direction1 < torus.direction1)
+    {
+        return true;
+    }
+
+    if (direction1 > torus.direction1)
+    {
+        return false;
+    }
+
+    if (normal < torus.normal)
+    {
+        return true;
+    }
+
+    if (normal > torus.normal)
+    {
+        return false;
+    }
+
+    if (radius0 < torus.radius0)
+    {
+        return true;
+    }
+
+    if (radius0 > torus.radius0)
+    {
+        return false;
+    }
+
+    return radius1 < torus.radius1;
+}
+
+template <typename Real>
+bool Torus3<Real>::operator<=(Torus3 const& torus) const
+{
+    return operator<(torus) || operator==(torus);
+}
+
+template <typename Real>
+bool Torus3<Real>::operator>(Torus3 const& torus) const
+{
+    return !operator<=(torus);
+}
+
+template <typename Real>
+bool Torus3<Real>::operator>=(Torus3 const& torus) const
+{
+    return !operator<(torus);
+}
+
+}
diff --git a/Engine/Plugins/Experimental/GeometryProcessing/Source/GeometryAlgorithms/Private/GTEngine/Mathematics/GteTriangle.h b/Engine/Plugins/Experimental/GeometryProcessing/Source/GeometryAlgorithms/Private/GTEngine/Mathematics/GteTriangle.h
new file mode 100644
--- /dev/null
+++ b/Engine/Plugins/Experimental/GeometryProcessing/Source/GeometryAlgorithms/Private/GTEngine/Mathematics/GteTriangle.h
@@ -0,0 +1,111 @@
+// David Eberly, Geometric Tools, Redmond WA 98052
+// Copyright (c) 1998-2019
+// Distributed under the Boost Software License, Version 1.0.
+// http://www.boost.org/LICENSE_1_0.txt
+// http://www.geometrictools.com/License/Boost/LICENSE_1_0.txt
+// File Version: 3.0.0 (2016/06/19)
+
+#pragma once
+
+#include <GTEngine/Mathematics/GteVector.h>
+
+// The triangle is represented as an array of three vertices.  The dimension
+// N must be 2 or larger.
+
+namespace gte
+{
+
+template <int N, typename Real>
+class Triangle
+{
+public:
+    // Construction and destruction.  The default constructor sets the
+    // vertices to (0,..,0), (1,0,...,0), and (0,1,0,...,0).
+    Triangle();
+    Triangle(Vector<N, Real> const& v0, Vector<N, Real> const& v1,
+        Vector<N, Real> const& v2);
+    Triangle(std::array<Vector<N, Real>, 3> const& inV);
+
+    // Public member access.
+    std::array<Vector<N, Real>, 3> v;
+
+public:
+    // Comparisons to support sorted containers.
+    bool operator==(Triangle const& triangle) const;
+    bool operator!=(Triangle const& triangle) const;
+    bool operator< (Triangle const& triangle) const;
+    bool operator<=(Triangle const& triangle) const;
+    bool operator> (Triangle const& triangle) const;
+    bool operator>=(Triangle const& triangle) const;
+};
+
+// Template aliases for convenience.
+template <typename Real>
+using Triangle2 = Triangle<2, Real>;
+
+template <typename Real>
+using Triangle3 = Triangle<3, Real>;
+
+
+template <int N, typename Real>
+Triangle<N, Real>::Triangle()
+{
+    v[0].MakeZero();
+    v[1].MakeUnit(0);
+    v[2].MakeUnit(1);
+}
+
+template <int N, typename Real>
+Triangle<N, Real>::Triangle(Vector<N, Real> const& v0,
+    Vector<N, Real> const& v1, Vector<N, Real> const& v2)
+{
+    v[0] = v0;
+    v[1] = v1;
+    v[2] = v2;
+}
+
+template <int N, typename Real>
+Triangle<N, Real>::Triangle(std::array<Vector<N, Real>, 3> const& inV)
+    :
+    v(inV)
+{
+}
+
+template <int N, typename Real>
+bool Triangle<N, Real>::operator==(Triangle const& triangle) const
+{
+    return v == triangle.v;
+}
+
+template <int N, typename Real>
+bool Triangle<N, Real>::operator!=(Triangle const& triangle) const
+{
+    return !operator==(triangle);
+}
+
+template <int N, typename Real>
+bool Triangle<N, Real>::operator<(Triangle const& triangle) const
+{
+    return v < triangle.v;
+}
+
+template <int N, typename Real>
+bool Triangle<N, Real>::operator<=(Triangle const& triangle) const
+{
+    return operator<(triangle) || operator==(triangle);
+}
+
+template <int N, typename Real>
+bool Triangle<N, Real>::operator>(Triangle const& triangle) const
+{
+    return !operator<=(triangle);
+}
+
+template <int N, typename Real>
+bool Triangle<N, Real>::operator>=(Triangle const& triangle) const
+{
+    return !operator<(triangle);
+}
+
+
+}
diff --git a/Engine/Plugins/Experimental/GeometryProcessing/Source/GeometryAlgorithms/Private/GTEngine/Mathematics/GteTriangleKey.h b/Engine/Plugins/Experimental/GeometryProcessing/Source/GeometryAlgorithms/Private/GTEngine/Mathematics/GteTriangleKey.h
new file mode 100644
--- /dev/null
+++ b/Engine/Plugins/Experimental/GeometryProcessing/Source/GeometryAlgorithms/Private/GTEngine/Mathematics/GteTriangleKey.h
@@ -0,0 +1,134 @@
+// David Eberly, Geometric Tools, Redmond WA 98052
+// Copyright (c) 1998-2019
+// Distributed under the Boost Software License, Version 1.0.
+// http://www.boost.org/LICENSE_1_0.txt
+// http://www.geometrictools.com/License/Boost/LICENSE_1_0.txt
+// File Version: 3.0.1 (2018/05/22)
+
+#pragma once
+
+#include <GTEngine/Mathematics/GteFeatureKey.h>
+
+#include <algorithm>
+
+namespace gte
+{
+    template <bool Ordered>
+    class TriangleKey : public FeatureKey<3, Ordered>
+    {
+    public:
+        // An ordered triangle has V[0] = min(v0,v1,v2).  Choose
+        // (V[0],V[1],V[2]) to be a permutation of (v0,v1,v2) so that the
+        // final storage is one of
+        //   (v0,v1,v2), (v1,v2,v0), (v2,v0,v1)
+        // The idea is that if v0 corresponds to (1,0,0), v1 corresponds to
+        // (0,1,0), and v2 corresponds to (0,0,1), the ordering (v0,v1,v2)
+        // corresponds to the 3x3 identity matrix I; the rows are the
+        // specified 3-tuples.  The permutation (V[0],V[1],V[2]) induces a
+        // permutation of the rows of the identity matrix to form a
+        // permutation matrix P with det(P) = 1 = det(I).
+        //
+        // An unordered triangle stores a permutation of (v0,v1,v2) so that
+        // V[0] < V[1] < V[2].
+        TriangleKey();  // creates key (-1,-1,-1)
+        explicit TriangleKey(int v0, int v1, int v2);
+    };
+
+	template<>
+	inline
+	TriangleKey<true>::TriangleKey()
+	{
+		V[0] = -1;
+		V[1] = -1;
+		V[2] = -1;
+	}
+
+	template<>
+	inline
+	TriangleKey<true>::TriangleKey(int v0, int v1, int v2)
+	{
+		if (v0 < v1)
+		{
+			if (v0 < v2)
+			{
+				// v0 is minimum
+				V[0] = v0;
+				V[1] = v1;
+				V[2] = v2;
+			}
+			else
+			{
+				// v2 is minimum
+				V[0] = v2;
+				V[1] = v0;
+				V[2] = v1;
+			}
+		}
+		else
+		{
+			if (v1 < v2)
+			{
+				// v1 is minimum
+				V[0] = v1;
+				V[1] = v2;
+				V[2] = v0;
+			}
+			else
+			{
+				// v2 is minimum
+				V[0] = v2;
+				V[1] = v0;
+				V[2] = v1;
+			}
+		}
+	}
+
+	template<>
+	inline
+	TriangleKey<false>::TriangleKey()
+	{
+		V[0] = -1;
+		V[1] = -1;
+		V[2] = -1;
+	}
+
+	template<>
+	inline
+	TriangleKey<false>::TriangleKey(int v0, int v1, int v2)
+	{
+		if (v0 < v1)
+		{
+			if (v0 < v2)
+			{
+				// v0 is minimum
+				V[0] = v0;
+				V[1] = std::min(v1, v2);
+				V[2] = std::max(v1, v2);
+			}
+			else
+			{
+				// v2 is minimum
+				V[0] = v2;
+				V[1] = std::min(v0, v1);
+				V[2] = std::max(v0, v1);
+			}
+		}
+		else
+		{
+			if (v1 < v2)
+			{
+				// v1 is minimum
+				V[0] = v1;
+				V[1] = std::min(v2, v0);
+				V[2] = std::max(v2, v0);
+			}
+			else
+			{
+				// v2 is minimum
+				V[0] = v2;
+				V[1] = std::min(v0, v1);
+				V[2] = std::max(v0, v1);
+			}
+		}
+	}
+}
