move GeometryProcessing plugin from NotForLicensees to Experimental

diff --git a/Engine/Plugins/Experimental/GeometryProcessing/Source/GeometryAlgorithms/Private/GTEngine/Mathematics/GteTriangulateCDT.h b/Engine/Plugins/Experimental/GeometryProcessing/Source/GeometryAlgorithms/Private/GTEngine/Mathematics/GteTriangulateCDT.h
new file mode 100644
--- /dev/null
+++ b/Engine/Plugins/Experimental/GeometryProcessing/Source/GeometryAlgorithms/Private/GTEngine/Mathematics/GteTriangulateCDT.h
@@ -0,0 +1,576 @@
+// David Eberly, Geometric Tools, Redmond WA 98052
+// Copyright (c) 1998-2019
+// Distributed under the Boost Software License, Version 1.0.
+// http://www.boost.org/LICENSE_1_0.txt
+// http://www.geometrictools.com/License/Boost/LICENSE_1_0.txt
+// File Version: 3.0.0 (2016/06/19)
+
+#pragma once
+
+#include <GTEngine/Mathematics/GteConstrainedDelaunay2.h>
+#include <GTEngine/Mathematics/GteContPointInPolygon2.h>
+#include <algorithm>
+#include <memory>
+#include <queue>
+
+// The triangulation is based on constrained Delaunay triangulations (CDT),
+// which does not use divisions, so ComputeType may be chosen using BSNumber.
+// The input constraints are relaxed compared to TriangulateEC; specifically,
+// the inner polygons are allowed to share vertices with the outer polygons.
+// The CDT produces a triangulation of the convex hull of the input, which
+// includes triangles outside the top-level outer polygon and inside the
+// inner polygons.  Only the triangles relevant to the input are returned
+// via the 'std::vector<int>& triangles', but the other triangles are
+// retained so that you can perform linear walks in search of points inside
+// the original polygon (nested polygon, tree of nested polygons).  This is
+// useful, for example, when subsampling the polygon triangles for
+// interpolation of function data specified at the vertices.  A linear walk
+// does not work for a mesh consisting only of the polygon triangles, but
+// with the additional triangles, the walk can navigate through holes in
+// the polygon to find the containing triangle of the specified point.
+
+namespace gte
+{
+
+template <typename InputType, typename ComputeType>
+class TriangulateCDT
+{
+public:
+    // The class is a functor to support triangulating multiple polygons that
+    // share vertices in a collection of points.  The interpretation of
+    // 'numPoints' and 'points' is described before each operator() function.
+    // Preconditions are numPoints >= 3 and points is a nonnull pointer to an
+    // array of at least numPoints elements.  If the preconditions are
+    // satisfied, then operator() functions will return 'true'; otherwise,
+    // they return 'false'.
+    TriangulateCDT(int numPoints, Vector2<InputType> const* points);
+    TriangulateCDT(std::vector<Vector2<InputType>> const& points);
+
+    // The triangles of the polygon triangulation.
+    inline std::vector<std::array<int, 3>> const& GetTriangles() const;
+
+    // The triangles inside the convex hull of the points but outside the
+    // triangulation.
+    inline std::vector<std::array<int, 3>> const& GetOutsideTriangles() const;
+
+    // The triangles of the convex hull of the inputs to the constructor.
+    inline std::vector<std::array<int, 3>> const& GetAllTriangles() const;
+
+    // The classification of whether a triangle is part of the triangulation
+    // or outside the triangulation.  These may be used in conjunction with
+    // the array returned by GetAllTriangles().
+    inline std::vector<bool> const& GetIsInside() const;
+    inline bool IsInside(int triIndex) const;
+    inline bool IsOutside(int triIndex) const;
+
+    // The outer polygons have counterclockwise ordered vertices.  The inner
+    // polygons have clockwise ordered vertices.
+    typedef std::vector<int> Polygon;
+
+    // The input 'points' represents an array of vertices for a simple
+    // polygon. The vertices are points[0] through points[numPoints-1] and
+    // are listed in counterclockwise order.
+    bool operator()();
+
+    // The input 'points' represents an array of vertices that contains the
+    // vertices of a simple polygon.
+    bool operator()(Polygon const& polygon);
+
+    // The input 'points' is a shared array of vertices that contains the
+    // vertices for two simple polygons, an outer polygon and an inner
+    // polygon.  The inner polygon must be strictly inside the outer polygon.
+    bool operator()(Polygon const& outer, Polygon const& inner);
+
+    // The input 'points' is a shared array of vertices that contains the
+    // vertices for multiple simple polygons, an outer polygon and one or more
+    // inner polygons.  The inner polygons must be nonoverlapping and strictly
+    // inside the outer polygon.
+    bool operator()(Polygon const& outer, std::vector<Polygon> const& inners);
+
+    // A tree of nested polygons.  The root node corresponds to an outer
+    // polygon.  The children of the root correspond to inner polygons, which
+    // are nonoverlapping polygons strictly contained in the outer polygon.
+    // Each inner polygon may itself contain an outer polygon, thus leading
+    // to a hierarchy of polygons.  The outer polygons have vertices listed
+    // in counterclockwise order.  The inner polygons have vertices listed in
+    // clockwise order.
+    class Tree
+    {
+    public:
+        Polygon polygon;
+        std::vector<std::shared_ptr<Tree>> child;
+    };
+
+    // The input 'positions' is a shared array of vertices that contains the
+    // vertices for multiple simple polygons in a tree of polygons.
+    bool operator()(std::shared_ptr<Tree> const& tree);
+
+private:
+    // Triangulate the points referenced by an operator(...) query.  The
+    // mAllTriangles and mIsInside are populated by this function, but the
+    // indices of mAllTriangles are relative to the packed 'points'.  After
+    // the call, the indices need to be mapped back to the original set
+    // provided by the input arrays to operator(...).  The mTriangles and
+    // mOutsideTriangles are generated after the call by examining
+    // mAllTriangles and mIsInside.
+    bool TriangulatePacked(int numPoints, Vector2<InputType> const* points,
+        std::shared_ptr<Tree> const& tree,
+        std::map<std::shared_ptr<Tree>, int> const& offsets);
+
+    int GetNumPointsAndOffsets(std::shared_ptr<Tree> const& tree,
+        std::map<std::shared_ptr<Tree>, int>& offsets) const;
+
+    void PackPoints(std::shared_ptr<Tree> const& tree,
+        std::vector<Vector2<InputType>>& points);
+
+    bool InsertEdges(std::shared_ptr<Tree> const& tree);
+
+    void LookupIndex(std::shared_ptr<Tree> const& tree, int& v,
+        std::map<std::shared_ptr<Tree>, int> const& offsets) const;
+
+    bool IsInside(std::shared_ptr<Tree> const& tree, Vector2<ComputeType> const* points,
+        Vector2<ComputeType> const& test,
+        std::map<std::shared_ptr<Tree>, int> const& offsets) const;
+
+    // The input polygon.
+    int mNumPoints;
+    Vector2<InputType> const* mPoints;
+
+    // The output triangulation and those triangle inside the hull of the
+    // points but outside the triangulation.
+    std::vector<std::array<int, 3>> mTriangles;
+    std::vector<std::array<int, 3>> mOutsideTriangles;
+    std::vector<std::array<int, 3>> mAllTriangles;
+    std::vector<bool> mIsInside;
+
+    ConstrainedDelaunay2<InputType, ComputeType> mConstrainedDelaunay;
+};
+
+
+template <typename InputType, typename ComputeType>
+TriangulateCDT<InputType, ComputeType>::TriangulateCDT(int numPoints, Vector2<InputType> const* points)
+    :
+    mNumPoints(numPoints),
+    mPoints(points)
+{
+    if (mNumPoints < 3 || !mPoints)
+    {
+        LogError("Invalid input.");
+        mNumPoints = 0;
+        mPoints = nullptr;
+        // The operator() functions will triangulate only when mPoints is
+        // not null.  The test mNumPoints >= 3 is redundant because of the
+        // logic in the constructor.
+    }
+}
+
+template <typename InputType, typename ComputeType>
+TriangulateCDT<InputType, ComputeType>::TriangulateCDT(std::vector<Vector2<InputType>> const& points)
+    :
+    mNumPoints(static_cast<int>(points.size())),
+    mPoints(points.data())
+{
+    if (mNumPoints < 3 || !mPoints)
+    {
+        LogError("Invalid input.");
+        mNumPoints = 0;
+        mPoints = nullptr;
+        // The operator() functions will triangulate only when mPoints is
+        // not null.  The test mNumPoints >= 3 is redundant because of the
+        // logic in the constructor.
+    }
+}
+
+template <typename InputType, typename ComputeType> inline
+std::vector<std::array<int, 3>> const& TriangulateCDT<InputType, ComputeType>::GetTriangles() const
+{
+    return mTriangles;
+}
+
+template <typename InputType, typename ComputeType> inline
+std::vector<std::array<int, 3>> const& TriangulateCDT<InputType, ComputeType>::GetOutsideTriangles() const
+{
+    return mOutsideTriangles;
+}
+
+template <typename InputType, typename ComputeType> inline
+std::vector<std::array<int, 3>> const& TriangulateCDT<InputType, ComputeType>::GetAllTriangles() const
+{
+    return mAllTriangles;
+}
+
+template <typename InputType, typename ComputeType> inline
+std::vector<bool> const& TriangulateCDT<InputType, ComputeType>::GetIsInside() const
+{
+    return mIsInside;
+}
+
+template <typename InputType, typename ComputeType> inline
+bool TriangulateCDT<InputType, ComputeType>::IsInside(int triIndex) const
+{
+    if (0 <= triIndex && triIndex < static_cast<int>(mIsInside.size()))
+    {
+        return mIsInside[triIndex];
+    }
+    else
+    {
+        return false;
+    }
+}
+
+template <typename InputType, typename ComputeType> inline
+bool TriangulateCDT<InputType, ComputeType>::IsOutside(int triIndex) const
+{
+    if (0 <= triIndex && triIndex < static_cast<int>(mIsInside.size()))
+    {
+        return !mIsInside[triIndex];
+    }
+    else
+    {
+        return false;
+    }
+}
+
+template <typename InputType, typename ComputeType>
+bool TriangulateCDT<InputType, ComputeType>::operator()()
+{
+    if (mPoints)
+    {
+        std::shared_ptr<Tree> tree = std::make_shared<Tree>();
+        tree->polygon.resize(mNumPoints);
+        for (int i = 0; i < mNumPoints; ++i)
+        {
+            tree->polygon[i] = i;
+        }
+
+        return operator()(tree);
+    }
+    return false;
+}
+
+template <typename InputType, typename ComputeType>
+bool TriangulateCDT<InputType, ComputeType>::operator()(Polygon const& polygon)
+{
+    if (mPoints)
+    {
+        std::shared_ptr<Tree> tree = std::make_shared<Tree>();
+        tree->polygon = polygon;
+
+        return operator()(tree);
+    }
+    return false;
+}
+
+template <typename InputType, typename ComputeType>
+bool TriangulateCDT<InputType, ComputeType>::operator()(Polygon const& outer, Polygon const& inner)
+{
+    if (mPoints)
+    {
+        std::shared_ptr<Tree> otree = std::make_shared<Tree>();
+        otree->polygon = outer;
+        otree->child.resize(1);
+
+        std::shared_ptr<Tree> itree = std::make_shared<Tree>();
+        itree->polygon = inner;
+        otree->child[0] = itree;
+
+        return operator()(otree);
+    }
+    return false;
+}
+
+template <typename InputType, typename ComputeType>
+bool TriangulateCDT<InputType, ComputeType>::operator()(Polygon const& outer, std::vector<Polygon> const& inners)
+{
+    if (mPoints)
+    {
+        std::shared_ptr<Tree> otree = std::make_shared<Tree>();
+        otree->polygon = outer;
+        otree->child.resize(inners.size());
+
+        std::vector<std::shared_ptr<Tree>> itree(inners.size());
+        for (size_t i = 0; i < inners.size(); ++i)
+        {
+            itree[i] = std::make_shared<Tree>();
+            itree[i]->polygon = inners[i];
+            otree->child[i] = itree[i];
+        }
+
+        return operator()(otree);
+    }
+    return false;
+}
+
+template <typename InputType, typename ComputeType>
+bool TriangulateCDT<InputType, ComputeType>::operator()(std::shared_ptr<Tree> const& tree)
+{
+    if (mPoints)
+    {
+        std::map<std::shared_ptr<Tree>, int> offsets;
+        int numPoints = GetNumPointsAndOffsets(tree, offsets);
+        std::vector<Vector2<InputType>> points(numPoints);
+        PackPoints(tree, points);
+
+        if (TriangulatePacked(numPoints, &points[0], tree, offsets))
+        {
+            int numTriangles = static_cast<int>(mAllTriangles.size());
+            for (int t = 0; t < numTriangles; ++t)
+            {
+                auto& tri = mAllTriangles[t];
+                for (int j = 0; j < 3; ++j)
+                {
+                    LookupIndex(tree, tri[j], offsets);
+                }
+
+                if (mIsInside[t])
+                {
+                    mTriangles.push_back(tri);
+                }
+                else
+                {
+                    mOutsideTriangles.push_back(tri);
+                }
+            }
+            return true;
+        }
+    }
+
+    return false;
+}
+
+template <typename InputType, typename ComputeType>
+bool TriangulateCDT<InputType, ComputeType>::TriangulatePacked(int numPoints,
+    Vector2<InputType> const* points, std::shared_ptr<Tree> const& tree,
+    std::map<std::shared_ptr<Tree>, int> const& offsets)
+{
+    mTriangles.clear();
+    mOutsideTriangles.clear();
+    mAllTriangles.clear();
+    mIsInside.clear();
+
+    mConstrainedDelaunay(numPoints, points, static_cast<InputType>(0));
+    InsertEdges(tree);
+
+    ComputeType half = static_cast<ComputeType>(0.5);
+    ComputeType fourth = static_cast<ComputeType>(0.25);
+    auto const& query = mConstrainedDelaunay.GetQuery();
+    auto const* ctPoints = query.GetVertices();
+    int numTriangles = mConstrainedDelaunay.GetNumTriangles();
+    int const* indices = &mConstrainedDelaunay.GetIndices()[0];
+    mIsInside.resize(numTriangles);
+    for (int t = 0; t < numTriangles; ++t)
+    {
+        int v0 = *indices++;
+        int v1 = *indices++;
+        int v2 = *indices++;
+        auto ctInside = fourth * ctPoints[v0] + half * ctPoints[v1] + fourth *ctPoints[v2];
+        mIsInside[t] = IsInside(tree, ctPoints, ctInside, offsets);
+        mAllTriangles.push_back({ { v0, v1, v2 } });
+    }
+    return true;
+}
+
+template <typename InputType, typename ComputeType>
+int TriangulateCDT<InputType, ComputeType>::GetNumPointsAndOffsets(
+    std::shared_ptr<Tree> const& tree, std::map<std::shared_ptr<Tree>, int>& offsets) const
+{
+    int numPoints = 0;
+    std::queue<std::shared_ptr<Tree>> treeQueue;
+    treeQueue.push(tree);
+    while (treeQueue.size() > 0)
+    {
+        std::shared_ptr<Tree> outer = treeQueue.front();
+        treeQueue.pop();
+        offsets.insert(std::make_pair(outer, numPoints));
+        numPoints += static_cast<int>(outer->polygon.size());
+
+        int numChildren = static_cast<int>(outer->child.size());
+        for (int c = 0; c < numChildren; ++c)
+        {
+            std::shared_ptr<Tree> inner = outer->child[c];
+            offsets.insert(std::make_pair(inner, numPoints));
+            numPoints += static_cast<int>(inner->polygon.size());
+
+            int numGrandChildren = static_cast<int>(inner->child.size());
+            for (int g = 0; g < numGrandChildren; ++g)
+            {
+                treeQueue.push(inner->child[g]);
+            }
+        }
+    }
+    return numPoints;
+}
+
+template <typename InputType, typename ComputeType>
+void TriangulateCDT<InputType, ComputeType>::PackPoints(std::shared_ptr<Tree> const& tree,
+    std::vector<Vector2<InputType>>& points)
+{
+    int numPoints = 0;
+    std::queue<std::shared_ptr<Tree>> treeQueue;
+    treeQueue.push(tree);
+    while (treeQueue.size() > 0)
+    {
+        std::shared_ptr<Tree> outer = treeQueue.front();
+        treeQueue.pop();
+        int const numOuterIndices = static_cast<int>(outer->polygon.size());
+        int const* outerIndices = outer->polygon.data();
+        for (int i = 0; i < numOuterIndices; ++i)
+        {
+            points[numPoints++] = mPoints[outerIndices[i]];
+        }
+
+        int numChildren = static_cast<int>(outer->child.size());
+        for (int c = 0; c < numChildren; ++c)
+        {
+            std::shared_ptr<Tree> inner = outer->child[c];
+            int const numInnerIndices = static_cast<int>(inner->polygon.size());
+            int const* innerIndices = inner->polygon.data();
+            for (int i = 0; i < numInnerIndices; ++i)
+            {
+                points[numPoints++] = mPoints[innerIndices[i]];
+            }
+
+            int numGrandChildren = static_cast<int>(inner->child.size());
+            for (int g = 0; g < numGrandChildren; ++g)
+            {
+                treeQueue.push(inner->child[g]);
+            }
+        }
+    }
+}
+
+template <typename InputType, typename ComputeType>
+bool TriangulateCDT<InputType, ComputeType>::InsertEdges(std::shared_ptr<Tree> const& tree)
+{
+    int numPoints = 0;
+    std::array<int, 2> edge;
+    std::vector<int> ignoreOutEdge;
+    std::queue<std::shared_ptr<Tree>> treeQueue;
+    treeQueue.push(tree);
+    while (treeQueue.size() > 0)
+    {
+        std::shared_ptr<Tree> outer = treeQueue.front();
+        treeQueue.pop();
+        int numOuter = static_cast<int>(outer->polygon.size());
+        for (int i0 = numOuter - 1, i1 = 0; i1 < numOuter; i0 = i1++)
+        {
+            edge[0] = numPoints + i0;
+            edge[1] = numPoints + i1;
+            if (!mConstrainedDelaunay.Insert(edge, ignoreOutEdge))
+            {
+                return false;
+            }
+        }
+        numPoints += numOuter;
+
+        int numChildren = static_cast<int>(outer->child.size());
+        for (int c = 0; c < numChildren; ++c)
+        {
+            std::shared_ptr<Tree> inner = outer->child[c];
+            int numInner = static_cast<int>(inner->polygon.size());
+            for (int i0 = numInner - 1, i1 = 0; i1 < numInner; i0 = i1++)
+            {
+                edge[0] = numPoints + i0;
+                edge[1] = numPoints + i1;
+                if (!mConstrainedDelaunay.Insert(edge, ignoreOutEdge))
+                {
+                    return false;
+                }
+            }
+            numPoints += numInner;
+
+            int numGrandChildren = static_cast<int>(inner->child.size());
+            for (int g = 0; g < numGrandChildren; ++g)
+            {
+                treeQueue.push(inner->child[g]);
+            }
+        }
+    }
+    return true;
+}
+
+template <typename InputType, typename ComputeType>
+void TriangulateCDT<InputType, ComputeType>::LookupIndex(std::shared_ptr<Tree> const& tree,
+    int& v, std::map<std::shared_ptr<Tree>, int> const& offsets) const
+{
+    std::queue<std::shared_ptr<Tree>> treeQueue;
+    treeQueue.push(tree);
+    while (treeQueue.size() > 0)
+    {
+        std::shared_ptr<Tree> outer = treeQueue.front();
+        treeQueue.pop();
+        int const numOuterIndices = static_cast<int>(outer->polygon.size());
+        int const* outerIndices = outer->polygon.data();
+        int offset = offsets.find(outer)->second;
+        if (v < offset + numOuterIndices)
+        {
+            v = outerIndices[v - offset];
+            return;
+        }
+
+        int numChildren = static_cast<int>(outer->child.size());
+        for (int c = 0; c < numChildren; ++c)
+        {
+            std::shared_ptr<Tree> inner = outer->child[c];
+            int const numInnerIndices = static_cast<int>(inner->polygon.size());
+            int const* innerIndices = inner->polygon.data();
+            offset = offsets.find(inner)->second;
+            if (v < offset + numInnerIndices)
+            {
+                v = innerIndices[v - offset];
+                return;
+            }
+
+            int numGrandChildren = static_cast<int>(inner->child.size());
+            for (int g = 0; g < numGrandChildren; ++g)
+            {
+                treeQueue.push(inner->child[g]);
+            }
+        }
+    }
+}
+
+template <typename InputType, typename ComputeType>
+bool TriangulateCDT<InputType, ComputeType>::IsInside(std::shared_ptr<Tree> const& tree,
+    Vector2<ComputeType> const* points, Vector2<ComputeType> const& test,
+    std::map<std::shared_ptr<Tree>, int> const& offsets) const
+{
+    std::queue<std::shared_ptr<Tree>> treeQueue;
+    treeQueue.push(tree);
+    while (treeQueue.size() > 0)
+    {
+        std::shared_ptr<Tree> outer = treeQueue.front();
+        treeQueue.pop();
+        int const numOuterIndices = static_cast<int>(outer->polygon.size());
+        int offset = offsets.find(outer)->second;
+        PointInPolygon2<ComputeType> piOuter(numOuterIndices, points + offset);
+        if (piOuter.Contains(test))
+        {
+            int numChildren = static_cast<int>(outer->child.size());
+            int c;
+            for (c = 0; c < numChildren; ++c)
+            {
+                std::shared_ptr<Tree> inner = outer->child[c];
+                int const numInnerIndices = static_cast<int>(inner->polygon.size());
+                offset = offsets.find(inner)->second;
+                PointInPolygon2<ComputeType> piInner(numInnerIndices, points + offset);
+                if (piInner.Contains(test))
+                {
+                    int numGrandChildren = static_cast<int>(inner->child.size());
+                    for (int g = 0; g < numGrandChildren; ++g)
+                    {
+                        treeQueue.push(inner->child[g]);
+                    }
+                    break;
+                }
+            }
+            if (c == numChildren)
+            {
+                return true;
+            }
+        }
+    }
+    return false;
+}
+
+}
diff --git a/Engine/Plugins/Experimental/GeometryProcessing/Source/GeometryAlgorithms/Private/GTEngine/Mathematics/GteTriangulateEC.h b/Engine/Plugins/Experimental/GeometryProcessing/Source/GeometryAlgorithms/Private/GTEngine/Mathematics/GteTriangulateEC.h
new file mode 100644
--- /dev/null
+++ b/Engine/Plugins/Experimental/GeometryProcessing/Source/GeometryAlgorithms/Private/GTEngine/Mathematics/GteTriangulateEC.h
@@ -0,0 +1,1319 @@
+// David Eberly, Geometric Tools, Redmond WA 98052
+// Copyright (c) 1998-2019
+// Distributed under the Boost Software License, Version 1.0.
+// http://www.boost.org/LICENSE_1_0.txt
+// http://www.geometrictools.com/License/Boost/LICENSE_1_0.txt
+// File Version: 3.0.1 (2018/02/17)
+
+#pragma once
+
+#include <GTEngine/LowLevel/GteLogger.h>
+#include <GTEngine/Mathematics/GtePrimalQuery2.h>
+#include <algorithm>
+#include <limits>
+#include <memory>
+#include <map>
+#include <queue>
+#include <vector>
+
+// The algorithm for processing nested polygons involves a division, so the
+// ComputeType must be rational-based, say, BSRational.  If you process only
+// triangles that are simple, you may use BSNumber for the ComputeType.
+
+namespace gte
+{
+
+template <typename InputType, typename ComputeType>
+class TriangulateEC
+{
+public:
+    // The class is a functor to support triangulating multiple polygons that
+    // share vertices in a collection of points.  The interpretation of
+    // 'numPoints' and 'points' is described before each operator() function.
+    // Preconditions are numPoints >= 3 and points is a nonnull pointer to an
+    // array of at least numPoints elements.  If the preconditions are
+    // satisfied, then operator() functions will return 'true'; otherwise,
+    // they return 'false'.
+    TriangulateEC(int numPoints, Vector2<InputType> const* points);
+    TriangulateEC(std::vector<Vector2<InputType>> const& points);
+
+    // Access the triangulation after each operator() call.
+    inline std::vector<std::array<int, 3>> const& GetTriangles() const;
+
+    // The outer polygons have counterclockwise ordered vertices.  The inner
+    // polygons have clockwise ordered vertices.
+    typedef std::vector<int> Polygon;
+
+    // The input 'points' represents an array of vertices for a simple
+    // polygon. The vertices are points[0] through points[numPoints-1] and
+    // are listed in counterclockwise order.
+    bool operator()();
+
+    // The input 'points' represents an array of vertices that contains the
+    // vertices of a simple polygon.
+    bool operator()(Polygon const& polygon);
+
+    // The input 'points' is a shared array of vertices that contains the
+    // vertices for two simple polygons, an outer polygon and an inner
+    // polygon.  The inner polygon must be strictly inside the outer polygon.
+    bool operator()(Polygon const& outer, Polygon const& inner);
+
+    // The input 'points' is a shared array of vertices that contains the
+    // vertices for multiple simple polygons, an outer polygon and one or more
+    // inner polygons.  The inner polygons must be nonoverlapping and strictly
+    // inside the outer polygon.
+    bool operator()(Polygon const& outer, std::vector<Polygon> const& inners);
+
+    // A tree of nested polygons.  The root node corresponds to an outer
+    // polygon.  The children of the root correspond to inner polygons, which
+    // are nonoverlapping polygons strictly contained in the outer polygon.
+    // Each inner polygon may itself contain an outer polygon, thus leading
+    // to a hierarchy of polygons.  The outer polygons have vertices listed
+    // in counterclockwise order.  The inner polygons have vertices listed in
+    // clockwise order.
+    class Tree
+    {
+    public:
+        Polygon polygon;
+        std::vector<std::shared_ptr<Tree>> child;
+    };
+
+    // The input 'positions' is a shared array of vertices that contains the
+    // vertices for multiple simple polygons in a tree of polygons.
+    bool operator()(std::shared_ptr<Tree> const& tree);
+
+private:
+    // Create the vertex objects that store the various lists required by the
+    // ear-clipping algorithm.
+    void InitializeVertices(int numVertices, int const* indices);
+
+    // Apply ear clipping to the input polygon.  Polygons with holes are
+    // preprocessed to obtain an index array that is nearly a simple polygon.
+    // This outer polygon has a pair of coincident edges per inner polygon.
+    void DoEarClipping(int numVertices, int const* indices);
+
+    // Given an outer polygon that contains an inner polygon, this function
+    // determines a pair of visible vertices and inserts two coincident edges
+    // to generate a nearly simple polygon.
+    bool CombinePolygons(int nextElement, Polygon const& outer,
+        Polygon const& inner, std::map<int, int>& indexMap,
+        std::vector<int>& combined);
+
+    // Given an outer polygon that contains a set of nonoverlapping inner
+    // polygons, this function determines pairs of visible vertices and
+    // inserts coincident edges to generate a nearly simple polygon.  It
+    // repeatedly calls CombinePolygons for each inner polygon of the outer
+    // polygon.
+    bool ProcessOuterAndInners(int& nextElement, Polygon const& outer,
+        std::vector<Polygon> const& inners, std::map<int, int>& indexMap,
+        std::vector<int>& combined);
+
+    // The insertion of coincident edges to obtain a nearly simple polygon
+    // requires duplication of vertices in order that the ear-clipping
+    // algorithm work correctly.  After the triangulation, the indices of
+    // the duplicated vertices are converted to the original indices.
+    void RemapIndices(std::map<int, int> const& indexMap);
+
+    // Two extra elements are needed in the position array per outer-inners
+    // polygon.  This function computes the total number of extra elements
+    // needed for the input tree and it converts InputType vertices to
+    // ComputeType values.
+    int InitializeFromTree(std::shared_ptr<Tree> const& tree);
+
+    // The input polygon.
+    int mNumPoints;
+    Vector2<InputType> const* mPoints;
+
+    // The output triangulation.
+    std::vector<std::array<int, 3>> mTriangles;
+
+    // The array of points used for geometric queries.  If you want to be
+    // certain of a correct result, choose ComputeType to be BSNumber.  The
+    // InputType points are convertex to ComputeType points on demand; the
+    // mIsConverted array keeps track of which input points have been
+    // converted.
+    std::vector<Vector2<ComputeType>> mComputePoints;
+    std::vector<bool> mIsConverted;
+    PrimalQuery2<ComputeType> mQuery;
+
+    // Doubly linked lists for storing specially tagged vertices.
+    class Vertex
+    {
+    public:
+        Vertex();
+
+        int index;          // index of vertex in position array
+        int vPrev, vNext;   // vertex links for polygon
+        int sPrev, sNext;   // convex/reflex vertex links (disjoint lists)
+        int ePrev, eNext;   // ear links
+        bool isConvex, isEar;
+    };
+
+    inline Vertex& V(int i);
+    bool IsConvex(int i);
+    bool IsEar(int i);
+    void InsertAfterC(int i);   // insert convex vertex
+    void InsertAfterR(int i);   // insert reflex vertesx
+    void InsertEndE(int i);     // insert ear at end of list
+    void InsertAfterE(int i);   // insert ear after efirst
+    void InsertBeforeE(int i);  // insert ear before efirst
+    void RemoveV(int i);        // remove vertex
+    int  RemoveE(int i);        // remove ear at i
+    void RemoveR(int i);        // remove reflex vertex
+
+    // The doubly linked list.
+    std::vector<Vertex> mVertices;
+    int mCFirst, mCLast;  // linear list of convex vertices
+    int mRFirst, mRLast;  // linear list of reflex vertices
+    int mEFirst, mELast;  // cyclical list of ears
+};
+
+
+template <typename InputType, typename ComputeType>
+TriangulateEC<InputType, ComputeType>::TriangulateEC(int numPoints, Vector2<InputType> const* points)
+    :
+    mNumPoints(numPoints),
+    mPoints(points)
+{
+    if (mNumPoints >= 3 && mPoints)
+    {
+        mComputePoints.resize(mNumPoints);
+        mIsConverted.resize(mNumPoints);
+        std::fill(mIsConverted.begin(), mIsConverted.end(), false);
+        mQuery.Set(mNumPoints, &mComputePoints[0]);
+    }
+    else
+    {
+        LogError("Invalid input.");
+        mNumPoints = 0;
+        mPoints = nullptr;
+        // The operator() functions will triangulate only when mPoints is
+        // not null.  The test mNumPoints >= 3 is redundant because of the
+        // logic in the constructor.
+    }
+}
+
+template <typename InputType, typename ComputeType>
+TriangulateEC<InputType, ComputeType>::TriangulateEC(std::vector<Vector2<InputType>> const& points)
+    :
+    mNumPoints(static_cast<int>(points.size())),
+    mPoints(points.data())
+{
+    if (mNumPoints >= 3 && mPoints)
+    {
+        mComputePoints.resize(mNumPoints);
+        mIsConverted.resize(mNumPoints);
+        std::fill(mIsConverted.begin(), mIsConverted.end(), false);
+        mQuery.Set(mNumPoints, &mComputePoints[0]);
+    }
+    else
+    {
+        LogError("Invalid input.");
+        mNumPoints = 0;
+        mPoints = nullptr;
+        // The operator() functions will triangulate only when mPoints is
+        // not null.  The test mNumPoints >= 3 is redundant because of the
+        // logic in the constructor.
+    }
+}
+
+template <typename InputType, typename ComputeType> inline
+std::vector<std::array<int, 3>> const& TriangulateEC<InputType, ComputeType>::GetTriangles() const
+{
+    return mTriangles;
+}
+
+template <typename InputType, typename ComputeType>
+bool TriangulateEC<InputType, ComputeType>::operator()()
+{
+    mTriangles.clear();
+    if (mPoints)
+    {
+        // Compute the points for the queries.
+        for (int i = 0; i < mNumPoints; ++i)
+        {
+            if (!mIsConverted[i])
+            {
+                mIsConverted[i] = true;
+                for (int j = 0; j < 2; ++j)
+                {
+                    mComputePoints[i][j] = mPoints[i][j];
+                }
+            }
+        }
+
+        // Triangulate the unindexed polygon.
+        InitializeVertices(mNumPoints, nullptr);
+        DoEarClipping(mNumPoints, nullptr);
+        return true;
+    }
+    else
+    {
+        return false;
+    }
+}
+
+template <typename InputType, typename ComputeType>
+bool TriangulateEC<InputType, ComputeType>::operator()(Polygon const& polygon)
+{
+    mTriangles.clear();
+    if (mPoints)
+    {
+        // Compute the points for the queries.
+        int const numIndices = static_cast<int>(polygon.size());
+        int const* indices = polygon.data();
+        for (int i = 0; i < numIndices; ++i)
+        {
+            int index = indices[i];
+            if (!mIsConverted[index])
+            {
+                mIsConverted[index] = true;
+                for (int j = 0; j < 2; ++j)
+                {
+                    mComputePoints[index][j] = mPoints[index][j];
+                }
+            }
+        }
+
+        // Triangulate the indexed polygon.
+        InitializeVertices(numIndices, indices);
+        DoEarClipping(numIndices, indices);
+        return true;
+    }
+    else
+    {
+        return false;
+    }
+}
+
+template <typename InputType, typename ComputeType>
+bool TriangulateEC<InputType, ComputeType>::operator()(Polygon const& outer, Polygon const& inner)
+{
+    mTriangles.clear();
+    if (mPoints)
+    {
+        // Two extra elements are needed to duplicate the endpoints of the
+        // edge introduced to combine outer and inner polygons.
+        int numPointsPlusExtras = mNumPoints + 2;
+        if (numPointsPlusExtras > static_cast<int>(mComputePoints.size()))
+        {
+            mComputePoints.resize(numPointsPlusExtras);
+            mIsConverted.resize(numPointsPlusExtras);
+            mIsConverted[mNumPoints] = false;
+            mIsConverted[mNumPoints + 1] = false;
+            mQuery.Set(numPointsPlusExtras, &mComputePoints[0]);
+        }
+
+        // Convert any points that have not been encountered in other
+        // triangulation calls.
+        int const numOuterIndices = static_cast<int>(outer.size());
+        int const* outerIndices = outer.data();
+        for (int i = 0; i < numOuterIndices; ++i)
+        {
+            int index = outerIndices[i];
+            if (!mIsConverted[index])
+            {
+                mIsConverted[index] = true;
+                for (int j = 0; j < 2; ++j)
+                {
+                    mComputePoints[index][j] = mPoints[index][j];
+                }
+            }
+        }
+
+        int const numInnerIndices = static_cast<int>(inner.size());
+        int const* innerIndices = inner.data();
+        for (int i = 0; i < numInnerIndices; ++i)
+        {
+            int index = innerIndices[i];
+            if (!mIsConverted[index])
+            {
+                mIsConverted[index] = true;
+                for (int j = 0; j < 2; ++j)
+                {
+                    mComputePoints[index][j] = mPoints[index][j];
+                }
+            }
+        }
+
+        // Combine the outer polygon and the inner polygon into a simple
+        // polygon by inserting two edges connecting mutually visible
+        // vertices, one from the outer polygon and one from the inner
+        // polygon.
+        int nextElement = mNumPoints;  // The next available element.
+        std::map<int, int> indexMap;
+        std::vector<int> combined;
+        if (!CombinePolygons(nextElement, outer, inner, indexMap, combined))
+        {
+            // An unexpected condition was encountered.
+            return false;
+        }
+
+        // The combined polygon is now in the format of a simple polygon,
+        // albeit one with coincident edges.
+        int numVertices = static_cast<int>(combined.size());
+        int* const indices = &combined[0];
+        InitializeVertices(numVertices, indices);
+        DoEarClipping(numVertices, indices);
+
+        // Map the duplicate indices back to the original indices.
+        RemapIndices(indexMap);
+        return true;
+    }
+    else
+    {
+        return false;
+    }
+}
+
+template <typename InputType, typename ComputeType>
+bool TriangulateEC<InputType, ComputeType>::operator()(Polygon const& outer, std::vector<Polygon> const& inners)
+{
+    mTriangles.clear();
+    if (mPoints)
+    {
+        // Two extra elements per inner polygon are needed to duplicate the
+        // endpoints of the edges introduced to combine outer and inner
+        // polygons.
+        int numPointsPlusExtras = mNumPoints + 2 * (int)inners.size();
+        if (numPointsPlusExtras > static_cast<int>(mComputePoints.size()))
+        {
+            mComputePoints.resize(numPointsPlusExtras);
+            mIsConverted.resize(numPointsPlusExtras);
+            for (int i = mNumPoints; i < numPointsPlusExtras; ++i)
+            {
+                mIsConverted[i] = false;
+            }
+            mQuery.Set(numPointsPlusExtras, &mComputePoints[0]);
+        }
+
+        // Convert any points that have not been encountered in other
+        // triangulation calls.
+        int const numOuterIndices = static_cast<int>(outer.size());
+        int const* outerIndices = outer.data();
+        for (int i = 0; i < numOuterIndices; ++i)
+        {
+            int index = outerIndices[i];
+            if (!mIsConverted[index])
+            {
+                mIsConverted[index] = true;
+                for (int j = 0; j < 2; ++j)
+                {
+                    mComputePoints[index][j] = mPoints[index][j];
+                }
+            }
+        }
+
+        for (auto const& inner : inners)
+        {
+            int const numInnerIndices = static_cast<int>(inner.size());
+            int const* innerIndices = inner.data();
+            for (int i = 0; i < numInnerIndices; ++i)
+            {
+                int index = innerIndices[i];
+                if (!mIsConverted[index])
+                {
+                    mIsConverted[index] = true;
+                    for (int j = 0; j < 2; ++j)
+                    {
+                        mComputePoints[index][j] = mPoints[index][j];
+                    }
+                }
+            }
+        }
+
+        // Combine the outer polygon and the inner polygons into a simple
+        // polygon by inserting two edges per inner polygon connecting
+        // mutually visible vertices.
+        int nextElement = mNumPoints;  // The next available element.
+        std::map<int, int> indexMap;
+        std::vector<int> combined;
+        if (!ProcessOuterAndInners(nextElement, outer, inners, indexMap, combined))
+        {
+            // An unexpected condition was encountered.
+            return false;
+        }
+
+        // The combined polygon is now in the format of a simple polygon, albeit
+        // with coincident edges.
+        int numVertices = static_cast<int>(combined.size());
+        int* const indices = &combined[0];
+        InitializeVertices(numVertices, indices);
+        DoEarClipping(numVertices, indices);
+
+        // Map the duplicate indices back to the original indices.
+        RemapIndices(indexMap);
+        return true;
+    }
+    else
+    {
+        return false;
+    }
+}
+
+template <typename InputType, typename ComputeType>
+bool TriangulateEC<InputType, ComputeType>::operator()(std::shared_ptr<Tree> const& tree)
+{
+    mTriangles.clear();
+    if (mPoints)
+    {
+        // Two extra elements per inner polygon are needed to duplicate the
+        // endpoints of the edges introduced to combine outer and inner
+        // polygons.
+        int numPointsPlusExtras = mNumPoints + InitializeFromTree(tree);
+        if (numPointsPlusExtras > static_cast<int>(mComputePoints.size()))
+        {
+            mComputePoints.resize(numPointsPlusExtras);
+            mIsConverted.resize(numPointsPlusExtras);
+            for (int i = mNumPoints; i < numPointsPlusExtras; ++i)
+            {
+                mIsConverted[i] = false;
+            }
+            mQuery.Set(numPointsPlusExtras, &mComputePoints[0]);
+        }
+
+        int nextElement = mNumPoints;
+        std::map<int, int> indexMap;
+
+        std::queue<std::shared_ptr<Tree>> treeQueue;
+        treeQueue.push(tree);
+        while (treeQueue.size() > 0)
+        {
+            std::shared_ptr<Tree> outer = treeQueue.front();
+            treeQueue.pop();
+
+            int numChildren = static_cast<int>(outer->child.size());
+            int numVertices;
+            int const* indices;
+
+            if (numChildren == 0)
+            {
+                // The outer polygon is a simple polygon (no nested inner
+                // polygons).  Triangulate the simple polygon.
+                numVertices = static_cast<int>(outer->polygon.size());
+                indices = outer->polygon.data();
+                InitializeVertices(numVertices, indices);
+                DoEarClipping(numVertices, indices);
+            }
+            else
+            {
+                // Place the next level of outer polygon nodes on the queue for
+                // triangulation.
+                std::vector<Polygon> inners(numChildren);
+                for (int c = 0; c < numChildren; ++c)
+                {
+                    std::shared_ptr<Tree> inner = outer->child[c];
+                    inners[c] = inner->polygon;
+                    int numGrandChildren = static_cast<int>(inner->child.size());
+                    for (int g = 0; g < numGrandChildren; ++g)
+                    {
+                        treeQueue.push(inner->child[g]);
+                    }
+                }
+
+                // Combine the outer polygon and the inner polygons into a
+                // simple polygon by inserting two edges per inner polygon
+                // connecting mutually visible vertices.
+                std::vector<int> combined;
+                ProcessOuterAndInners(nextElement, outer->polygon, inners, indexMap, combined);
+
+                // The combined polygon is now in the format of a simple
+                // polygon, albeit with coincident edges.
+                numVertices = static_cast<int>(combined.size());
+                indices = &combined[0];
+                InitializeVertices(numVertices, indices);
+                DoEarClipping(numVertices, indices);
+            }
+        }
+
+        // Map the duplicate indices back to the original indices.
+        RemapIndices(indexMap);
+        return true;
+    }
+    else
+    {
+        return false;
+    }
+}
+
+template <typename InputType, typename ComputeType>
+void TriangulateEC<InputType, ComputeType>::InitializeVertices(int numVertices, int const* indices)
+{
+    mVertices.clear();
+    mVertices.resize(numVertices);
+    mCFirst = -1;
+    mCLast = -1;
+    mRFirst = -1;
+    mRLast = -1;
+    mEFirst = -1;
+    mELast = -1;
+
+    // Create a circular list of the polygon vertices for dynamic removal of
+    // vertices.
+    int numVerticesM1 = numVertices - 1;
+    for (int i = 0; i <= numVerticesM1; ++i)
+    {
+        Vertex& vertex = V(i);
+        vertex.index = (indices ? indices[i] : i);
+        vertex.vPrev = (i > 0 ? i - 1 : numVerticesM1);
+        vertex.vNext = (i < numVerticesM1 ? i + 1 : 0);
+    }
+
+    // Create a circular list of the polygon vertices for dynamic removal of
+    // vertices.  Keep track of two linear sublists, one for the convex
+    // vertices and one for the reflex vertices.  This is an O(N) process
+    // where N is the number of polygon vertices.
+    for (int i = 0; i <= numVerticesM1; ++i)
+    {
+        if (IsConvex(i))
+        {
+            InsertAfterC(i);
+        }
+        else
+        {
+            InsertAfterR(i);
+        }
+    }
+}
+
+template <typename InputType, typename ComputeType>
+void TriangulateEC<InputType, ComputeType>::DoEarClipping(int numVertices, int const* indices)
+{
+    // If the polygon is convex, just create a triangle fan.
+    if (mRFirst == -1)
+    {
+        int numVerticesM1 = numVertices - 1;
+        if (indices)
+        {
+            for (int i = 1; i < numVerticesM1; ++i)
+            {
+                mTriangles.push_back({ { indices[0], indices[i], indices[i + 1] } });
+            }
+        }
+        else
+        {
+            for (int i = 1; i < numVerticesM1; ++i)
+            {
+                mTriangles.push_back({ { 0, i, i + 1 } });
+            }
+        }
+        return;
+    }
+
+    // Identify the ears and build a circular list of them.  Let V0, V1, and
+    // V2 be consecutive vertices forming a triangle T.  The vertex V1 is an
+    // ear if no other vertices of the polygon lie inside T.  Although it is
+    // enough to show that V1 is not an ear by finding at least one other
+    // vertex inside T, it is sufficient to search only the reflex vertices.
+    // This is an O(C*R) process, where C is the number of convex vertices and
+    // R is the number of reflex vertices with N = C+R.  The order is O(N^2),
+    // for example when C = R = N/2.
+    for (int i = mCFirst; i != -1; i = V(i).sNext)
+    {
+        if (IsEar(i))
+        {
+            InsertEndE(i);
+        }
+    }
+    V(mEFirst).ePrev = mELast;
+    V(mELast).eNext = mEFirst;
+
+    // Remove the ears, one at a time.
+    bool bRemoveAnEar = true;
+    while (bRemoveAnEar)
+    {
+        // Add the triangle with the ear to the output list of triangles.
+        int iVPrev = V(mEFirst).vPrev;
+        int iVNext = V(mEFirst).vNext;
+        mTriangles.push_back({ { V(iVPrev).index, V(mEFirst).index, V(iVNext).index } });
+
+        // Remove the vertex corresponding to the ear.
+        RemoveV(mEFirst);
+        if (--numVertices == 3)
+        {
+            // Only one triangle remains, just remove the ear and copy it.
+            mEFirst = RemoveE(mEFirst);
+            iVPrev = V(mEFirst).vPrev;
+            iVNext = V(mEFirst).vNext;
+            mTriangles.push_back({ { V(iVPrev).index, V(mEFirst).index, V(iVNext).index } });
+            bRemoveAnEar = false;
+            continue;
+        }
+
+        // Removal of the ear can cause an adjacent vertex to become an ear
+        // or to stop being an ear.
+        Vertex& vPrev = V(iVPrev);
+        if (vPrev.isEar)
+        {
+            if (!IsEar(iVPrev))
+            {
+                RemoveE(iVPrev);
+            }
+        }
+        else
+        {
+            bool wasReflex = !vPrev.isConvex;
+            if (IsConvex(iVPrev))
+            {
+                if (wasReflex)
+                {
+                    RemoveR(iVPrev);
+                }
+
+                if (IsEar(iVPrev))
+                {
+                    InsertBeforeE(iVPrev);
+                }
+            }
+        }
+
+        Vertex& vNext = V(iVNext);
+        if (vNext.isEar)
+        {
+            if (!IsEar(iVNext))
+            {
+                RemoveE(iVNext);
+            }
+        }
+        else
+        {
+            bool wasReflex = !vNext.isConvex;
+            if (IsConvex(iVNext))
+            {
+                if (wasReflex)
+                {
+                    RemoveR(iVNext);
+                }
+
+                if (IsEar(iVNext))
+                {
+                    InsertAfterE(iVNext);
+                }
+            }
+        }
+
+        // Remove the ear.
+        mEFirst = RemoveE(mEFirst);
+    }
+}
+
+template <typename InputType, typename ComputeType>
+bool TriangulateEC<InputType, ComputeType>::CombinePolygons(int nextElement,
+    Polygon const& outer, Polygon const& inner, std::map<int, int>& indexMap,
+    std::vector<int>& combined)
+{
+    int const numOuterIndices = static_cast<int>(outer.size());
+    int const* outerIndices = outer.data();
+    int const numInnerIndices = static_cast<int>(inner.size());
+    int const* innerIndices = inner.data();
+
+    // Locate the inner-polygon vertex of maximum x-value, call this vertex M.
+    ComputeType xmax = mComputePoints[innerIndices[0]][0];
+    int xmaxIndex = 0;
+    for (int i = 1; i < numInnerIndices; ++i)
+    {
+        ComputeType x = mComputePoints[innerIndices[i]][0];
+        if (x > xmax)
+        {
+            xmax = x;
+            xmaxIndex = i;
+        }
+    }
+    Vector2<ComputeType> M = mComputePoints[innerIndices[xmaxIndex]];
+
+    // Find the edge whose intersection Intr with the ray M+t*(1,0) minimizes
+    // the ray parameter t >= 0.
+    ComputeType const cmax = static_cast<ComputeType>(std::numeric_limits<InputType>::max());
+    ComputeType const zero = static_cast<ComputeType>(0);
+    Vector2<ComputeType> intr{ cmax, M[1] };
+    int v0min = -1, v1min = -1, endMin = -1;
+    int i0, i1;
+    ComputeType s = cmax;
+    ComputeType t = cmax;
+    for (i0 = numOuterIndices - 1, i1 = 0; i1 < numOuterIndices; i0 = i1++)
+    {
+        // Consider only edges for which the first vertex is below (or on) the
+        // ray and the second vertex is above (or on) the ray.
+        Vector2<ComputeType> diff0 = mComputePoints[outerIndices[i0]] - M;
+        if (diff0[1] > zero)
+        {
+            continue;
+        }
+        Vector2<ComputeType> diff1 = mComputePoints[outerIndices[i1]] - M;
+        if (diff1[1] < zero)
+        {
+            continue;
+        }
+
+        // At this time, diff0.y <= 0 and diff1.y >= 0.
+        int currentEndMin = -1;
+        if (diff0[1] < zero)
+        {
+            if (diff1[1] > zero)
+            {
+                // The intersection of the edge and ray occurs at an interior
+                // edge point.
+                s = diff0[1] / (diff0[1] - diff1[1]);
+                t = diff0[0] + s * (diff1[0] - diff0[0]);
+            }
+            else  // diff1.y == 0
+            {
+                // The vertex Outer[i1] is the intersection of the edge and
+                // the ray.
+                t = diff1[0];
+                currentEndMin = i1;
+            }
+        }
+        else  // diff0.y == 0
+        {
+            if (diff1[1] > zero)
+            {
+                // The vertex Outer[i0] is the intersection of the edge and
+                // the ray;
+                t = diff0[0];
+                currentEndMin = i0;
+            }
+            else  // diff1.y == 0
+            {
+                if (diff0[0] < diff1[0])
+                {
+                    t = diff0[0];
+                    currentEndMin = i0;
+                }
+                else
+                {
+                    t = diff1[0];
+                    currentEndMin = i1;
+                }
+            }
+        }
+
+        if (zero <= t && t < intr[0])
+        {
+            intr[0] = t;
+            v0min = i0;
+            v1min = i1;
+            if (currentEndMin == -1)
+            {
+                // The current closest point is an edge-interior point.
+                endMin = -1;
+            }
+            else
+            {
+                // The current closest point is a vertex.
+                endMin = currentEndMin;
+            }
+        }
+        else if (t == intr[0])
+        {
+            // The current closest point is a vertex shared by multiple edges;
+            // thus, endMin and currentMin refer to the same point.
+            if (endMin == -1 || currentEndMin == -1)
+            {
+                LogError("Unexpected condition.");
+                return false;
+            }
+
+            // We need to select the edge closest to M.  The previous closest
+            // edge is <outer[v0min],outer[v1min]>.  The current candidate is
+            // <outer[i0],outer[i1]>.
+            Vector2<ComputeType> shared = mComputePoints[outerIndices[i1]];
+
+            // For the previous closest edge, endMin refers to a vertex of
+            // the edge.  Get the index of the other vertex.
+            int other = (endMin == v0min ? v1min : v0min);
+
+            // The new edge is closer if the other vertex of the old edge is
+            // left-of the new edge.
+            diff0 = mComputePoints[outerIndices[i0]] - shared;
+            diff1 = mComputePoints[outerIndices[other]] - shared;
+            ComputeType dotperp = DotPerp(diff0, diff1);
+            if (dotperp > zero)
+            {
+                // The new edge is closer to M.
+                v0min = i0;
+                v1min = i1;
+                endMin = currentEndMin;
+            }
+        }
+    }
+
+    // The intersection intr[0] stored only the t-value of the ray.  The
+    // actual point is (mx,my)+t*(1,0), so intr[0] must be adjusted.
+    intr[0] += M[0];
+
+    int maxCosIndex;
+    if (endMin == -1)
+    {
+        // If you reach this assert, there is a good chance that you have two
+        // inner polygons that share a vertex or an edge.
+        if (v0min < 0 || v1min < 0)
+        {
+            LogError("Is this an invalid nested polygon?");
+            return false;
+        }
+
+        // Select one of Outer[v0min] and Outer[v1min] that has an x-value
+        // larger than M.x, call this vertex P.  The triangle <M,I,P> must
+        // contain an outer-polygon vertex that is visible to M, which is
+        // possibly P itself.
+        Vector2<ComputeType> sTriangle[3];  // <P,M,I> or <P,I,M>
+        int pIndex;
+        if (mComputePoints[outerIndices[v0min]][0] > mComputePoints[outerIndices[v1min]][0])
+        {
+            sTriangle[0] = mComputePoints[outerIndices[v0min]];
+            sTriangle[1] = intr;
+            sTriangle[2] = M;
+            pIndex = v0min;
+        }
+        else
+        {
+            sTriangle[0] = mComputePoints[outerIndices[v1min]];
+            sTriangle[1] = M;
+            sTriangle[2] = intr;
+            pIndex = v1min;
+        }
+
+        // If any outer-polygon vertices other than P are inside the triangle
+        // <M,I,P>, then at least one of these vertices must be a reflex
+        // vertex.  It is sufficient to locate the reflex vertex R (if any)
+        // in <M,I,P> that minimizes the angle between R-M and (1,0).  The
+        // data member mQuery is used for the reflex query.
+        Vector2<ComputeType> diff = sTriangle[0] - M;
+        ComputeType maxSqrLen = Dot(diff, diff);
+        ComputeType maxCos = diff[0] * diff[0] / maxSqrLen;
+        PrimalQuery2<ComputeType> localQuery(3, sTriangle);
+        maxCosIndex = pIndex;
+        for (int i = 0; i < numOuterIndices; ++i)
+        {
+            if (i == pIndex)
+            {
+                continue;
+            }
+
+            int curr = outerIndices[i];
+            int prev = outerIndices[(i + numOuterIndices - 1) % numOuterIndices];
+            int next = outerIndices[(i + 1) % numOuterIndices];
+            if (mQuery.ToLine(curr, prev, next) <= 0
+                && localQuery.ToTriangle(mComputePoints[curr], 0, 1, 2) <= 0)
+            {
+                // The vertex is reflex and inside the <M,I,P> triangle.
+                diff = mComputePoints[curr] - M;
+                ComputeType sqrLen = Dot(diff, diff);
+                ComputeType cs = diff[0] * diff[0] / sqrLen;
+                if (cs > maxCos)
+                {
+                    // The reflex vertex forms a smaller angle with the
+                    // positive x-axis, so it becomes the new visible
+                    // candidate.
+                    maxSqrLen = sqrLen;
+                    maxCos = cs;
+                    maxCosIndex = i;
+                }
+                else if (cs == maxCos && sqrLen < maxSqrLen)
+                {
+                    // The reflex vertex has angle equal to the current
+                    // minimum but the length is smaller, so it becomes the
+                    // new visible candidate.
+                    maxSqrLen = sqrLen;
+                    maxCosIndex = i;
+                }
+            }
+        }
+    }
+    else
+    {
+        maxCosIndex = endMin;
+    }
+
+    // The visible vertices are Position[Inner[xmaxIndex]] and
+    // Position[Outer[maxCosIndex]].  Two coincident edges with these
+    // endpoints are inserted to connect the outer and inner polygons into a
+    // simple polygon.  Each of the two Position[] values must be duplicated,
+    // because the original might be convex (or reflex) and the duplicate is
+    // reflex (or convex).  The ear-clipping algorithm needs to distinguish
+    // between them.
+    combined.resize(numOuterIndices + numInnerIndices + 2);
+    int cIndex = 0;
+    for (int i = 0; i <= maxCosIndex; ++i, ++cIndex)
+    {
+        combined[cIndex] = outerIndices[i];
+    }
+
+    for (int i = 0; i < numInnerIndices; ++i, ++cIndex)
+    {
+        int j = (xmaxIndex + i) % numInnerIndices;
+        combined[cIndex] = innerIndices[j];
+    }
+
+    int innerIndex = innerIndices[xmaxIndex];
+    mComputePoints[nextElement] = mComputePoints[innerIndex];
+    combined[cIndex] = nextElement;
+    auto iter = indexMap.find(innerIndex);
+    if (iter != indexMap.end())
+    {
+        innerIndex = iter->second;
+    }
+    indexMap[nextElement] = innerIndex;
+    ++cIndex;
+    ++nextElement;
+
+    int outerIndex = outerIndices[maxCosIndex];
+    mComputePoints[nextElement] = mComputePoints[outerIndex];
+    combined[cIndex] = nextElement;
+    iter = indexMap.find(outerIndex);
+    if (iter != indexMap.end())
+    {
+        outerIndex = iter->second;
+    }
+    indexMap[nextElement] = outerIndex;
+    ++cIndex;
+    ++nextElement;
+
+    for (int i = maxCosIndex + 1; i < numOuterIndices; ++i, ++cIndex)
+    {
+        combined[cIndex] = outerIndices[i];
+    }
+    return true;
+}
+
+template <typename InputType, typename ComputeType>
+bool TriangulateEC<InputType, ComputeType>::ProcessOuterAndInners(int& nextElement,
+    Polygon const& outer, std::vector<Polygon> const& inners,
+    std::map<int, int>& indexMap, std::vector<int>& combined)
+{
+    // Sort the inner polygons based on maximum x-values.
+    int numInners = static_cast<int>(inners.size());
+    std::vector<std::pair<ComputeType, int>> pairs(numInners);
+    for (int p = 0; p < numInners; ++p)
+    {
+        int numIndices = static_cast<int>(inners[p].size());
+        int const* indices = inners[p].data();
+        ComputeType xmax = mComputePoints[indices[0]][0];
+        for (int j = 1; j < numIndices; ++j)
+        {
+            ComputeType x = mComputePoints[indices[j]][0];
+            if (x > xmax)
+            {
+                xmax = x;
+            }
+        }
+        pairs[p].first = xmax;
+        pairs[p].second = p;
+    }
+    std::sort(pairs.begin(), pairs.end());
+
+    // Merge the inner polygons with the outer polygon.
+    Polygon currentPolygon = outer;
+    for (int p = numInners - 1; p >= 0; --p)
+    {
+        Polygon const& polygon = inners[pairs[p].second];
+        Polygon currentCombined;
+        if (!CombinePolygons(nextElement, currentPolygon, polygon, indexMap, currentCombined))
+        {
+            return false;
+        }
+        currentPolygon = std::move(currentCombined);
+        nextElement += 2;
+    }
+
+    for (auto index : currentPolygon)
+    {
+        combined.push_back(index);
+    }
+    return true;
+}
+
+template <typename InputType, typename ComputeType>
+void TriangulateEC<InputType, ComputeType>::RemapIndices(std::map<int, int> const& indexMap)
+{
+    // The triangulation includes indices to the duplicated outer and inner
+    // vertices.  These indices must be mapped back to the original ones.
+    for (auto& tri : mTriangles)
+    {
+        for (int i = 0; i < 3; ++i)
+        {
+            auto iter = indexMap.find(tri[i]);
+            if (iter != indexMap.end())
+            {
+                tri[i] = iter->second;
+            }
+        }
+    }
+}
+
+template <typename InputType, typename ComputeType>
+int TriangulateEC<InputType, ComputeType>::InitializeFromTree(std::shared_ptr<Tree> const& tree)
+{
+    // Use a breadth-first search to process the outer-inners pairs of the
+    // tree of nested polygons.
+    int numExtraPoints = 0;
+
+    std::queue<std::shared_ptr<Tree>> treeQueue;
+    treeQueue.push(tree);
+    while (treeQueue.size() > 0)
+    {
+        // The 'root' is an outer polygon.
+        std::shared_ptr<Tree> outer = treeQueue.front();
+        treeQueue.pop();
+
+        // Count number of extra points for this outer-inners pair.
+        int numChildren = static_cast<int>(outer->child.size());
+        numExtraPoints += 2 * numChildren;
+
+        // Convert outer points from InputType to ComputeType.
+        int const numOuterIndices = static_cast<int>(outer->polygon.size());
+        int const* outerIndices = outer->polygon.data();
+        for (int i = 0; i < numOuterIndices; ++i)
+        {
+            int index = outerIndices[i];
+            if (!mIsConverted[index])
+            {
+                mIsConverted[index] = true;
+                for (int j = 0; j < 2; ++j)
+                {
+                    mComputePoints[index][j] = mPoints[index][j];
+                }
+            }
+        }
+
+        // The grandchildren of the outer polygon are also outer polygons.
+        // Insert them into the queue for processing.
+        for (int c = 0; c < numChildren; ++c)
+        {
+            // The 'child' is an inner polygon.
+            std::shared_ptr<Tree> inner = outer->child[c];
+
+            // Convert inner points from InputType to ComputeType.
+            int const numInnerIndices = static_cast<int>(inner->polygon.size());
+            int const* innerIndices = inner->polygon.data();
+            for (int i = 0; i < numInnerIndices; ++i)
+            {
+                int index = innerIndices[i];
+                if (!mIsConverted[index])
+                {
+                    mIsConverted[index] = true;
+                    for (int j = 0; j < 2; ++j)
+                    {
+                        mComputePoints[index][j] = mPoints[index][j];
+                    }
+                }
+            }
+
+            int numGrandChildren = static_cast<int>(inner->child.size());
+            for (int g = 0; g < numGrandChildren; ++g)
+            {
+                treeQueue.push(inner->child[g]);
+            }
+        }
+    }
+
+    return numExtraPoints;
+}
+
+template <typename InputType, typename ComputeType>
+TriangulateEC<InputType, ComputeType>::Vertex::Vertex()
+    :
+    index(-1),
+    vPrev(-1),
+    vNext(-1),
+    sPrev(-1),
+    sNext(-1),
+    ePrev(-1),
+    eNext(-1),
+    isConvex(false),
+    isEar(false)
+{
+}
+
+template <typename InputType, typename ComputeType> inline
+typename TriangulateEC<InputType, ComputeType>::Vertex& TriangulateEC<InputType, ComputeType>::V(int i)
+{
+    return mVertices[i];
+}
+
+template <typename InputType, typename ComputeType>
+bool TriangulateEC<InputType, ComputeType>::IsConvex(int i)
+{
+    Vertex& vertex = V(i);
+    int curr = vertex.index;
+    int prev = V(vertex.vPrev).index;
+    int next = V(vertex.vNext).index;
+    vertex.isConvex = (mQuery.ToLine(curr, prev, next) > 0);
+    return vertex.isConvex;
+}
+
+template <typename InputType, typename ComputeType>
+bool TriangulateEC<InputType, ComputeType>::IsEar(int i)
+{
+    Vertex& vertex = V(i);
+
+    if (mRFirst == -1)
+    {
+        // The remaining polygon is convex.
+        vertex.isEar = true;
+        return true;
+    }
+
+    // Search the reflex vertices and test if any are in the triangle
+    // <V[prev],V[curr],V[next]>.
+    int prev = V(vertex.vPrev).index;
+    int curr = vertex.index;
+    int next = V(vertex.vNext).index;
+    vertex.isEar = true;
+    for (int j = mRFirst; j != -1; j = V(j).sNext)
+    {
+        // Check if the test vertex is already one of the triangle vertices.
+        if (j == vertex.vPrev || j == i || j == vertex.vNext)
+        {
+            continue;
+        }
+
+        // V[j] has been ruled out as one of the original vertices of the
+        // triangle <V[prev],V[curr],V[next]>.  When triangulating polygons
+        // with holes, V[j] might be a duplicated vertex, in which case it
+        // does not affect the earness of V[curr].
+        int test = V(j).index;
+        if (mComputePoints[test] == mComputePoints[prev]
+            || mComputePoints[test] == mComputePoints[curr]
+            || mComputePoints[test] == mComputePoints[next])
+        {
+            continue;
+        }
+
+        // Test if the vertex is inside or on the triangle.  When it is, it
+        // causes V[curr] not to be an ear.
+        if (mQuery.ToTriangle(test, prev, curr, next) <= 0)
+        {
+            vertex.isEar = false;
+            break;
+        }
+    }
+
+    return vertex.isEar;
+}
+
+template <typename InputType, typename ComputeType>
+void TriangulateEC<InputType, ComputeType>::InsertAfterC(int i)
+{
+    if (mCFirst == -1)
+    {
+        // Add the first convex vertex.
+        mCFirst = i;
+    }
+    else
+    {
+        V(mCLast).sNext = i;
+        V(i).sPrev = mCLast;
+    }
+    mCLast = i;
+}
+
+template <typename InputType, typename ComputeType>
+void TriangulateEC<InputType, ComputeType>::InsertAfterR(int i)
+{
+    if (mRFirst == -1)
+    {
+        // Add the first reflex vertex.
+        mRFirst = i;
+    }
+    else
+    {
+        V(mRLast).sNext = i;
+        V(i).sPrev = mRLast;
+    }
+    mRLast = i;
+}
+
+template <typename InputType, typename ComputeType>
+void TriangulateEC<InputType, ComputeType>::InsertEndE(int i)
+{
+    if (mEFirst == -1)
+    {
+        // Add the first ear.
+        mEFirst = i;
+        mELast = i;
+    }
+    V(mELast).eNext = i;
+    V(i).ePrev = mELast;
+    mELast = i;
+}
+
+template <typename InputType, typename ComputeType>
+void TriangulateEC<InputType, ComputeType>::InsertAfterE(int i)
+{
+    Vertex& first = V(mEFirst);
+    int currENext = first.eNext;
+    Vertex& vertex = V(i);
+    vertex.ePrev = mEFirst;
+    vertex.eNext = currENext;
+    first.eNext = i;
+    V(currENext).ePrev = i;
+}
+
+template <typename InputType, typename ComputeType>
+void TriangulateEC<InputType, ComputeType>::InsertBeforeE(int i)
+{
+    Vertex& first = V(mEFirst);
+    int currEPrev = first.ePrev;
+    Vertex& vertex = V(i);
+    vertex.ePrev = currEPrev;
+    vertex.eNext = mEFirst;
+    first.ePrev = i;
+    V(currEPrev).eNext = i;
+}
+
+template <typename InputType, typename ComputeType>
+void TriangulateEC<InputType, ComputeType>::RemoveV(int i)
+{
+    int currVPrev = V(i).vPrev;
+    int currVNext = V(i).vNext;
+    V(currVPrev).vNext = currVNext;
+    V(currVNext).vPrev = currVPrev;
+}
+
+template <typename InputType, typename ComputeType>
+int TriangulateEC<InputType, ComputeType>::RemoveE(int i)
+{
+    int currEPrev = V(i).ePrev;
+    int currENext = V(i).eNext;
+    V(currEPrev).eNext = currENext;
+    V(currENext).ePrev = currEPrev;
+    return currENext;
+}
+
+template <typename InputType, typename ComputeType>
+void TriangulateEC<InputType, ComputeType>::RemoveR(int i)
+{
+    LogAssert(mRFirst != -1 && mRLast != -1, "Reflex vertices must exist.");
+
+    if (i == mRFirst)
+    {
+        mRFirst = V(i).sNext;
+        if (mRFirst != -1)
+        {
+            V(mRFirst).sPrev = -1;
+        }
+        V(i).sNext = -1;
+    }
+    else if (i == mRLast)
+    {
+        mRLast = V(i).sPrev;
+        if (mRLast != -1)
+        {
+            V(mRLast).sNext = -1;
+        }
+        V(i).sPrev = -1;
+    }
+    else
+    {
+        int currSPrev = V(i).sPrev;
+        int currSNext = V(i).sNext;
+        V(currSPrev).sNext = currSNext;
+        V(currSNext).sPrev = currSPrev;
+        V(i).sNext = -1;
+        V(i).sPrev = -1;
+    }
+}
+
+}
diff --git a/Engine/Plugins/Experimental/GeometryProcessing/Source/GeometryAlgorithms/Private/GTEngine/Mathematics/GteTubeMesh.h b/Engine/Plugins/Experimental/GeometryProcessing/Source/GeometryAlgorithms/Private/GTEngine/Mathematics/GteTubeMesh.h
new file mode 100644
--- /dev/null
+++ b/Engine/Plugins/Experimental/GeometryProcessing/Source/GeometryAlgorithms/Private/GTEngine/Mathematics/GteTubeMesh.h
@@ -0,0 +1,259 @@
+// David Eberly, Geometric Tools, Redmond WA 98052
+// Copyright (c) 1998-2019
+// Distributed under the Boost Software License, Version 1.0.
+// http://www.boost.org/LICENSE_1_0.txt
+// http://www.geometrictools.com/License/Boost/LICENSE_1_0.txt
+// File Version: 3.0.6 (2018/10/05)
+
+#pragma once
+
+#include <GTEngine/Mathematics/GteMesh.h>
+#include <GTEngine/Mathematics/GteFrenetFrame.h>
+#include <functional>
+#include <memory>
+
+namespace gte
+{
+
+template <typename Real>
+class TubeMesh : public Mesh<Real>
+{
+public:
+    // Create a mesh (x(u,v),y(u,v),z(u,v)) defined by the specified medial
+    // curve and radial function.  The mesh has torus topology when 'closed'
+    // is true and has cylinder topology when 'closed' is false.  The client
+    // is responsible for setting the topology correctly in the 'description'
+    // input.  The rows correspond to medial samples and the columns
+    // correspond to radial samples. The medial curve is sampled according
+    // to its natural t-parameter when 'sampleByArcLength' is false;
+    // otherwise, it is sampled uniformly in arclength.
+    TubeMesh(MeshDescription const& description,
+        std::shared_ptr<ParametricCurve<3, Real>> const& medial,
+        std::function<Real(Real)> const& radial, bool closed,
+        bool sampleByArcLength, Vector3<Real> upVector);
+
+    // Member access.
+    inline std::shared_ptr<ParametricCurve<3, Real>> const& GetMedial() const;
+    inline std::function<Real(Real)> const& GetRadial() const;
+    inline bool IsClosed() const;
+    inline bool IsSampleByArcLength() const;
+    inline Vector3<Real> const& GetUpVector() const;
+
+private:
+    void InitializeTCoords();
+    virtual void UpdatePositions() override;
+
+    std::shared_ptr<ParametricCurve<3, Real>> mMedial;
+    std::function<Real(Real)> mRadial;
+    bool mClosed, mSampleByArcLength;
+    Vector3<Real> mUpVector;
+    std::vector<Real> mCosAngle, mSinAngle;
+    std::function<Real(unsigned int)> mTSampler;
+    std::function<std::array<Vector3<Real>, 4>(Real)> mFSampler;
+    std::unique_ptr<FrenetFrame3<Real>> mFrenet;
+
+    // If the client does not request texture coordinates, they will be
+    // computed internally for use in evaluation of the surface geometry.
+    std::vector<Vector2<Real>> mDefaultTCoords;
+};
+
+
+template <typename Real>
+TubeMesh<Real>::TubeMesh(MeshDescription const& description,
+    std::shared_ptr<ParametricCurve<3, Real>> const& medial,
+    std::function<Real(Real)> const& radial, bool closed, bool sampleByArcLength,
+    Vector3<Real> upVector)
+    :
+    Mesh<Real>(description, { MeshTopology::CYLINDER }),  // TODO: Allow TORUS and remove the 'closed' input
+    mMedial(medial),
+    mRadial(radial),
+    mClosed(closed),
+    mSampleByArcLength(sampleByArcLength),
+    mUpVector(upVector)
+{
+    if (!this->mDescription.constructed)
+    {
+        // The logger system will report these errors in the Mesh constructor.
+        mMedial = nullptr;
+        return;
+    }
+
+    if (!mMedial)
+    {
+        LogWarning("A nonnull medial curve is required.");
+        this->mDescription.constructed = false;
+        return;
+    }
+
+    mCosAngle.resize(this->mDescription.numCols);
+    mSinAngle.resize(this->mDescription.numCols);
+    Real invRadialSamples = (Real)1 / (Real)(this->mDescription.numCols - 1);
+    for (unsigned int i = 0; i < this->mDescription.numCols - 1; ++i)
+    {
+        Real angle = i * invRadialSamples * (Real)GTE_C_TWO_PI;
+        mCosAngle[i] = std::cos(angle);
+        mSinAngle[i] = std::sin(angle);
+    }
+    mCosAngle[this->mDescription.numCols - 1] = mCosAngle[0];
+    mSinAngle[this->mDescription.numCols - 1] = mSinAngle[0];
+
+    Real invDenom;
+    if (mClosed)
+    {
+        invDenom = ((Real)1) / (Real)this->mDescription.numRows;
+    }
+    else
+    {
+        invDenom = ((Real)1) / (Real)(this->mDescription.numRows - 1);
+    }
+
+    Real factor;
+    if (mSampleByArcLength)
+    {
+        factor = mMedial->GetTotalLength() * invDenom;
+        mTSampler = [this, factor](unsigned int row)
+        {
+            return mMedial->GetTime(row * factor);
+        };
+    }
+    else
+    {
+        factor = (mMedial->GetTMax() - mMedial->GetTMin()) * invDenom;
+        mTSampler = [this, factor](unsigned int row)
+        {
+            return mMedial->GetTMin() + row * factor;
+        };
+    }
+
+    if (mUpVector != Vector3<Real>::Zero())
+    {
+        mFSampler = [this](Real t)
+        {
+            std::array<Vector3<Real>, 4> frame;
+            frame[0] = mMedial->GetPosition(t);
+            frame[1] = mMedial->GetTangent(t);
+            frame[3] = UnitCross(frame[1], mUpVector);
+            frame[2] = UnitCross(frame[3], frame[1]);
+            return frame;
+        };
+    }
+    else
+    {
+        mFrenet = std::make_unique<FrenetFrame3<Real>>(mMedial);
+        mFSampler = [this](Real t)
+        {
+            std::array<Vector3<Real>, 4> frame;
+            (*mFrenet)(t, frame[0], frame[1], frame[2], frame[3]);
+            return frame;
+        };
+    }
+
+    if (!this->mTCoords)
+    {
+        mDefaultTCoords.resize(this->mDescription.numVertices);
+        this->mTCoords = mDefaultTCoords.data();
+        this->mTCoordStride = sizeof(Vector2<Real>);
+
+        this->mDescription.allowUpdateFrame = this->mDescription.wantDynamicTangentSpaceUpdate;
+        if (this->mDescription.allowUpdateFrame)
+        {
+            if (!this->mDescription.hasTangentSpaceVectors)
+            {
+                this->mDescription.allowUpdateFrame = false;
+            }
+
+            if (!this->mNormals)
+            {
+                this->mDescription.allowUpdateFrame = false;
+            }
+        }
+    }
+
+    this->ComputeIndices();
+    InitializeTCoords();
+    UpdatePositions();
+    if (this->mDescription.allowUpdateFrame)
+    {
+        this->UpdateFrame();
+    }
+    else if (this->mNormals)
+    {
+        this->UpdateNormals();
+    }
+}
+
+template <typename Real>
+inline std::shared_ptr<ParametricCurve<3, Real>> const& TubeMesh<Real>::GetMedial() const
+{
+    return mMedial;
+}
+
+template <typename Real>
+inline std::function<Real(Real)> const& TubeMesh<Real>::GetRadial() const
+{
+    return mRadial;
+}
+
+template <typename Real>
+inline bool TubeMesh<Real>::IsClosed() const
+{
+    return mClosed;
+}
+
+template <typename Real>
+inline bool TubeMesh<Real>::IsSampleByArcLength() const
+{
+    return mSampleByArcLength;
+}
+
+template <typename Real>
+inline Vector3<Real> const& TubeMesh<Real>::GetUpVector() const
+{
+    return mUpVector;
+}
+
+template <typename Real>
+void TubeMesh<Real>::InitializeTCoords()
+{
+    Vector2<Real>tcoord;
+    for (unsigned int r = 0, i = 0; r < this->mDescription.numRows; ++r)
+    {
+        tcoord[1] = (Real)r / (Real)this->mDescription.rMax;
+        for (unsigned int c = 0; c <= this->mDescription.numCols; ++c, ++i)
+        {
+            tcoord[0] = (Real)c / (Real)this->mDescription.numCols;
+            this->TCoord(i) = tcoord;
+        }
+    }
+}
+
+template <typename Real>
+void TubeMesh<Real>::UpdatePositions()
+{
+    uint32_t row, col, v, save;
+    for (row = 0, v = 0; row < this->mDescription.numRows; ++row, ++v)
+    {
+        Real t = mTSampler(row);
+        Real radius = mRadial(t);
+        // frame = (position, tangent, normal, binormal)
+        std::array<Vector3<Real>, 4> frame = mFSampler(t);
+        for (col = 0, save = v; col < this->mDescription.numCols; ++col, ++v)
+        {
+            this->Position(v) = frame[0] + radius * (mCosAngle[col] * frame[2] +
+                mSinAngle[col] * frame[3]);
+        }
+        this->Position(v) = this->Position(save);
+    }
+
+    if (mClosed)
+    {
+        for (col = 0; col < this->mDescription.numCols; ++col)
+        {
+            uint32_t i0 = col;
+            uint32_t i1 = col + this->mDescription.numCols * (this->mDescription.numRows - 1);
+            this->Position(i1) = this->Position(i0);
+        }
+    }
+}
+
+}
diff --git a/Engine/Plugins/Experimental/GeometryProcessing/Source/GeometryAlgorithms/Private/GTEngine/Mathematics/GteUIntegerALU32.h b/Engine/Plugins/Experimental/GeometryProcessing/Source/GeometryAlgorithms/Private/GTEngine/Mathematics/GteUIntegerALU32.h
new file mode 100644
--- /dev/null
+++ b/Engine/Plugins/Experimental/GeometryProcessing/Source/GeometryAlgorithms/Private/GTEngine/Mathematics/GteUIntegerALU32.h
@@ -0,0 +1,554 @@
+// David Eberly, Geometric Tools, Redmond WA 98052
+// Copyright (c) 1998-2019
+// Distributed under the Boost Software License, Version 1.0.
+// http://www.boost.org/LICENSE_1_0.txt
+// http://www.geometrictools.com/License/Boost/LICENSE_1_0.txt
+// File Version: 3.0.1 (2017/07/04)
+
+#pragma once
+
+#include <GTEngine/Mathematics/GteBitHacks.h>
+#include <algorithm>
+
+// Support for unsigned integer arithmetic in BSNumber and BSRational.  The
+// Curiously Recurring Template Paradigm is used to allow the UInteger
+// types to share code without introducing virtual functions.
+
+namespace gte
+{
+
+template <typename UInteger>
+class UIntegerALU32
+{
+public:
+    // Comparisons.  These are not generic.  They rely on their being caled
+    // when the two BSNumber arguments to BSNumber::operatorX() are of the
+    // form 1.u*2^p and 1.v*2^p.  The comparisons apply to 1.u and 1.v as
+    // unsigned integers with their leading 1-bits aligned.
+    bool operator==(UInteger const& number) const;
+    bool operator!=(UInteger const& number) const;
+    bool operator< (UInteger const& number) const;
+    bool operator<=(UInteger const& number) const;
+    bool operator> (UInteger const& number) const;
+    bool operator>=(UInteger const& number) const;
+
+    // Arithmetic operations.  These are performed in-place; that is, the
+    // result is stored in 'this' object.  The goal is to reduce the number of
+    // object copies, much like the goal is for std::move.  The Sub function
+    // requires the inputs to satisfy n0 > n1.
+    void Add(UInteger const& n0, UInteger const& n1);
+    void Sub(UInteger const& n0, UInteger const& n1);
+    void Mul(UInteger const& n0, UInteger const& n1);
+
+    // The shift is performed in-place; that is, the result is stored in
+    // 'this' object.
+    void ShiftLeft(UInteger const& number, int32_t shift);
+
+    // The 'number' is even and positive.  It is shifted right to become an
+    // odd number and the return value is the amount shifted.  The operation
+    // is performed in-place; that is, the result is stored in 'this' object.
+    int32_t ShiftRightToOdd(UInteger const& number);
+
+    // Get a block of numRequested bits starting with the leading 1-bit of the
+    // nonzero number.  The returned number has the prefix stored in the
+    // high-order bits.  Additional bits are copied and used by the caller for
+    // rounding.  This function supports conversions from 'float' and 'double'.
+    // The input 'numRequested' is smaller than 64.
+    uint64_t GetPrefix(int32_t numRequested) const;
+};
+
+
+template <typename UInteger>
+bool UIntegerALU32<UInteger>::operator==(UInteger const& number) const
+{
+    UInteger const& self = *(UInteger const*)this;
+    int32_t numBits = self.GetNumBits();
+    if (numBits != number.GetNumBits())
+    {
+        return false;
+    }
+
+    if (numBits > 0)
+    {
+        auto const& bits = self.GetBits();
+        auto const& nBits = number.GetBits();
+        int32_t const last = self.GetSize() - 1;
+        for (int32_t i = last; i >= 0; --i)
+        {
+            if (bits[i] != nBits[i])
+            {
+                return false;
+            }
+        }
+    }
+    return true;
+}
+
+template <typename UInteger>
+bool UIntegerALU32<UInteger>::operator!=(UInteger const& number) const
+{
+    return !operator==(number);
+}
+
+template <typename UInteger>
+bool UIntegerALU32<UInteger>::operator< (UInteger const& number) const
+{
+    UInteger const& self = *(UInteger const*)this;
+    int32_t nNumBits = number.GetNumBits();
+    auto const& nBits = number.GetBits();
+
+    int32_t numBits = self.GetNumBits();
+    if (numBits > 0 && nNumBits > 0)
+    {
+        // The numbers must be compared as if they are left-aligned with
+        // each other.  We got here because we had self = 1.u * 2^p and
+        // number = 1.v * 2^p.  Although they have the same exponent, it is
+        // possible that 'self < number' but 'numBits(1u) > numBits(1v)'.
+        // Compare the bits one 32-bit block at a time.
+        auto const& bits = self.GetBits();
+        int bitIndex0 = numBits - 1;
+        int bitIndex1 = nNumBits - 1;
+        int block0 = bitIndex0 / 32;
+        int block1 = bitIndex1 / 32;
+        int numBlockBits0 = 1 + (bitIndex0 % 32);
+        int numBlockBits1 = 1 + (bitIndex1 % 32);
+        uint64_t n0shift = bits[block0];
+        uint64_t n1shift = nBits[block1];
+        while (block0 >= 0 && block1 >= 0)
+        {
+            // Shift the bits in the leading blocks to the high-order bit.
+            uint32_t value0 = (uint32_t)((n0shift << (32 - numBlockBits0)) & 0x00000000FFFFFFFFull);
+            uint32_t value1 = (uint32_t)((n1shift << (32 - numBlockBits1)) & 0x00000000FFFFFFFFull);
+
+            // Shift bits in the next block (if any) to fill the current
+            // block.
+            if (--block0 >= 0)
+            {
+                n0shift = bits[block0];
+                value0 |= (uint32_t)((n0shift >> numBlockBits0) & 0x00000000FFFFFFFFull);
+            }
+            if (--block1 >= 0)
+            {
+                n1shift = nBits[block1];
+                value1 |= (uint32_t)((n1shift >> numBlockBits1) & 0x00000000FFFFFFFFull);
+            }
+            if (value0 < value1)
+            {
+                return true;
+            }
+            if (value0 > value1)
+            {
+                return false;
+            }
+        }
+        return block0 < block1;
+    }
+    else
+    {
+        // One or both numbers are negative.  The only time 'less than' is
+        // 'true' is when 'number' is positive.
+        return (nNumBits > 0);
+    }
+}
+
+template <typename UInteger>
+bool UIntegerALU32<UInteger>::operator<=(UInteger const& number) const
+{
+    return operator<(number) || operator==(number);
+}
+
+template <typename UInteger>
+bool UIntegerALU32<UInteger>::operator> (UInteger const& number) const
+{
+    return !operator<=(number);
+}
+
+template <typename UInteger>
+bool UIntegerALU32<UInteger>::operator>=(UInteger const& number) const
+{
+    return !operator<(number);
+}
+
+template <typename UInteger>
+void UIntegerALU32<UInteger>::Add(UInteger const& n0, UInteger const& n1)
+{
+    UInteger& self = *(UInteger*)this;
+    int32_t n0NumBits = n0.GetNumBits();
+    int32_t n1NumBits = n1.GetNumBits();
+
+    // Add the numbers considered as positive integers.  Set the last block to
+    // zero in case no carry-out occurs.
+    int numBits = std::max(n0NumBits, n1NumBits) + 1;
+    self.SetNumBits(numBits);
+    self.SetBack(0);
+
+    // Get the input array sizes.
+    int32_t numElements0 = n0.GetSize();
+    int32_t numElements1 = n1.GetSize();
+
+    // Order the inputs so that the first has the most blocks.
+    auto const& u0 =
+        (numElements0 >= numElements1 ? n0.GetBits() : n1.GetBits());
+    auto const& u1 =
+        (numElements0 >= numElements1 ? n1.GetBits() : n0.GetBits());
+    auto numElements = std::minmax(numElements0, numElements1);
+
+    // Add the u1-blocks to u0-blocks.
+    auto& bits = self.GetBits();
+    uint64_t carry = 0, sum;
+    int32_t i;
+    for (i = 0; i < numElements.first; ++i)
+    {
+        sum = u0[i] + (u1[i] + carry);
+        bits[i] = (uint32_t)(sum & 0x00000000FFFFFFFFull);
+        carry = (sum >> 32);
+    }
+
+    // We have no more u1-blocks. Propagate the carry-out, if there is one, or
+    // copy the remaining blocks if there is not.
+    if (carry > 0)
+    {
+        for (/**/; i < numElements.second; ++i)
+        {
+            sum = u0[i] + carry;
+            bits[i] = (uint32_t)(sum & 0x00000000FFFFFFFFull);
+            carry = (sum >> 32);
+        }
+        if (carry > 0)
+        {
+            bits[i] = (uint32_t)(carry & 0x00000000FFFFFFFFull);
+        }
+    }
+    else
+    {
+        for (/**/; i < numElements.second; ++i)
+        {
+            bits[i] = u0[i];
+        }
+    }
+
+    // Reduce the number of bits if there was not a carry-out.
+    uint32_t firstBitIndex = (numBits - 1) % 32;
+    uint32_t mask = (1 << firstBitIndex);
+    if ((mask & self.GetBack()) == 0)
+    {
+        self.SetNumBits(--numBits);
+    }
+}
+
+template <typename UInteger>
+void UIntegerALU32<UInteger>::Sub(UInteger const& n0, UInteger const& n1)
+{
+    UInteger& self = *(UInteger*)this;
+    int32_t n0NumBits = n0.GetNumBits();
+    auto const& n0Bits = n0.GetBits();
+    auto const& n1Bits = n1.GetBits();
+
+    // Subtract the numbers considered as positive integers.  We know that
+    // n0 > n1, so create a number n2 that has the same number of bits as
+    // n0 and use two's-complement to generate -n2, and then add n0 and -n2.
+    // The result is nonnegative, so we do not need to apply two's complement
+    // to a negative result to extract the sign and absolute value.
+
+    // Get the input array sizes.  We know numElements0 >= numElements1.
+    int32_t numElements0 = n0.GetSize();
+    int32_t numElements1 = n1.GetSize();
+
+    // Create the two's-complement number n2.  We know n2.GetNumElements() is
+    // the same as numElements0.
+    UInteger n2(n0NumBits);
+    auto& n2Bits = n2.GetBits();
+    int32_t i;
+    for (i = 0; i < numElements1; ++i)
+    {
+        n2Bits[i] = ~n1Bits[i];
+    }
+    for (/**/; i < numElements0; ++i)
+    {
+        n2Bits[i] = ~0u;
+    }
+
+    // Now add 1 to the bit-negated result to obtain -n1.
+    uint64_t carry = 1, sum;
+    for (i = 0; i < numElements0; ++i)
+    {
+        sum = n2Bits[i] + carry;
+        n2Bits[i] = (uint32_t)(sum & 0x00000000FFFFFFFFull);
+        carry = (sum >> 32);
+    }
+
+    // Add the numbers as positive integers.  Set the last block to zero in
+    // case no carry-out occurs.
+    self.SetNumBits(n0NumBits + 1);
+    self.SetBack(0);
+
+    // Add the n0-blocks to n2-blocks.
+    auto& bits = self.GetBits();
+    for (i = 0, carry = 0; i < numElements0; ++i)
+    {
+        sum = n2Bits[i] + (n0Bits[i] + carry);
+        bits[i] = (uint32_t)(sum & 0x00000000FFFFFFFFull);
+        carry = (sum >> 32);
+    }
+
+    // Strip off the bits introduced by two's-complement.
+    int32_t block;
+    for (block = numElements0 - 1; block >= 0; --block)
+    {
+        if (bits[block] > 0)
+        {
+            break;
+        }
+    }
+
+    self.SetNumBits(32 * block + GetLeadingBit(bits[block]) + 1);
+}
+
+template <typename UInteger>
+void UIntegerALU32<UInteger>::Mul(UInteger const& n0, UInteger const& n1)
+{
+    UInteger& self = *(UInteger*)this;
+    int32_t n0NumBits = n0.GetNumBits();
+    int32_t n1NumBits = n1.GetNumBits();
+    auto const& n0Bits = n0.GetBits();
+    auto const& n1Bits = n1.GetBits();
+
+    // The number of bits is at most this, possibly one bit smaller.
+    int numBits = n0NumBits + n1NumBits;
+    self.SetNumBits(numBits);
+    auto& bits = self.GetBits();
+
+    // Product of a single-block number with a multiple-block number.
+    UInteger product(numBits);
+    auto& pBits = product.GetBits();
+
+    // Get the array sizes.
+    int32_t const numElements0 = n0.GetSize();
+    int32_t const numElements1 = n1.GetSize();
+    int32_t const numElements = self.GetSize();
+
+    // Compute the product v = u0*u1.
+    int32_t i0, i1, i2;
+    uint64_t term, sum;
+
+    // The case i0 == 0 is handled separately to initialize the accumulator
+    // with u0[0]*v.  This avoids having to fill the bytes of 'bits' with
+    // zeros outside the double loop, something that can be a performance
+    // issue when 'numBits' is large.
+    uint64_t block0 = n0Bits[0];
+    uint64_t carry = 0;
+    for (i1 = 0; i1 < numElements1; ++i1)
+    {
+        term = block0 * n1Bits[i1] + carry;
+        bits[i1] = (uint32_t)(term & 0x00000000FFFFFFFFull);
+        carry = (term >> 32);
+    }
+    if (i1 < numElements)
+    {
+        bits[i1] = (uint32_t)(carry & 0x00000000FFFFFFFFull);
+    }
+
+    for (i0 = 1; i0 < numElements0; ++i0)
+    {
+        // Compute the product p = u0[i0]*u1.
+        block0 = n0Bits[i0];
+        carry = 0;
+        for (i1 = 0, i2 = i0; i1 < numElements1; ++i1, ++i2)
+        {
+            term = block0 * n1Bits[i1] + carry;
+            pBits[i2] = (uint32_t)(term & 0x00000000FFFFFFFFull);
+            carry = (term >> 32);
+        }
+        if (i2 < numElements)
+        {
+            pBits[i2] = (uint32_t)(carry & 0x00000000FFFFFFFFull);
+        }
+
+        // Add p to the accumulator v.
+        carry = 0;
+        for (i1 = 0, i2 = i0; i1 < numElements1; ++i1, ++i2)
+        {
+            sum = pBits[i2] + (bits[i2] + carry);
+            bits[i2] = (uint32_t)(sum & 0x00000000FFFFFFFFull);
+            carry = (sum >> 32);
+        }
+        if (i2 < numElements)
+        {
+            sum = pBits[i2] + carry;
+            bits[i2] = (uint32_t)(sum & 0x00000000FFFFFFFFull);
+        }
+    }
+
+    // Reduce the number of bits if there was not a carry-out.
+    uint32_t firstBitIndex = (numBits - 1) % 32;
+    uint32_t mask = (1 << firstBitIndex);
+    if ((mask & self.GetBack()) == 0)
+    {
+        self.SetNumBits(--numBits);
+    }
+}
+
+template <typename UInteger>
+void UIntegerALU32<UInteger>::ShiftLeft(UInteger const& number, int32_t shift)
+{
+    UInteger& self = *(UInteger*)this;
+    int32_t nNumBits = number.GetNumBits();
+    auto const& nBits = number.GetBits();
+
+    // Shift the 'number' considered as an odd positive integer.
+    self.SetNumBits(nNumBits + shift);
+
+    // Set the low-order bits to zero.
+    auto& bits = self.GetBits();
+    int32_t const shiftBlock = shift / 32;
+    for (int32_t i = 0; i < shiftBlock; ++i)
+    {
+        bits[i] = 0;
+    }
+
+    // Get the location of the low-order 1-bit within the result.
+    int32_t const numInElements = number.GetSize();
+    int32_t const lshift = shift % 32;
+    int32_t i, j;
+    if (lshift > 0)
+    {
+        // The trailing 1-bits for source and target are at different
+        // relative indices.  Each shifted source block straddles a boundary
+        // between two target blocks, so we must extract the subblocks and
+        // copy accordingly.
+        int32_t const rshift = 32 - lshift;
+        uint32_t prev = 0, curr;
+        for (i = shiftBlock, j = 0; j < numInElements; ++i, ++j)
+        {
+            curr = nBits[j];
+            bits[i] = (curr << lshift) | (prev >> rshift);
+            prev = curr;
+        }
+        if (i < self.GetSize())
+        {
+            // The leading 1-bit of the source is at a relative index such
+            // that when you add the shift amount, that bit occurs in a new
+            // block.
+            bits[i] = (prev >> rshift);
+        }
+    }
+    else
+    {
+        // The trailing 1-bits for source and target are at the same relative
+        // index.  The shift reduces to a block copy.
+        for (i = shiftBlock, j = 0; j < numInElements; ++i, ++j)
+        {
+            bits[i] = nBits[j];
+        }
+    }
+}
+
+template <typename UInteger>
+int32_t UIntegerALU32<UInteger>::ShiftRightToOdd(UInteger const& number)
+{
+    UInteger& self = *(UInteger*)this;
+    auto const& nBits = number.GetBits();
+
+    // Get the leading 1-bit.
+    int32_t const numElements = number.GetSize();
+    int32_t const numM1 = numElements - 1;
+    int32_t firstBitIndex = 32 * numM1 + GetLeadingBit(nBits[numM1]);
+
+    // Get the trailing 1-bit.
+    int32_t lastBitIndex = -1;
+    for (int32_t block = 0; block < numElements; ++block)
+    {
+        uint32_t value = nBits[block];
+        if (value > 0)
+        {
+            lastBitIndex = 32 * block + GetTrailingBit(value);
+            break;
+        }
+    }
+
+    // The right-shifted result.
+    self.SetNumBits(firstBitIndex - lastBitIndex + 1);
+    auto& bits = self.GetBits();
+    int32_t const numBlocks = self.GetSize();
+
+    // Get the location of the low-order 1-bit within the result.
+    int32_t const shiftBlock = lastBitIndex / 32;
+    int32_t rshift = lastBitIndex % 32;
+    if (rshift > 0)
+    {
+        int32_t const lshift = 32 - rshift;
+        int32_t i, j = shiftBlock;
+        uint32_t curr = nBits[j++];
+        for (i = 0; j < numElements; ++i, ++j)
+        {
+            uint32_t next = nBits[j];
+            bits[i] = (curr >> rshift) | (next << lshift);
+            curr = next;
+        }
+        if (i < numBlocks)
+        {
+            bits[i] = (curr >> rshift);
+        }
+    }
+    else
+    {
+        for (int32_t i = 0, j = shiftBlock; i < numBlocks; ++i, ++j)
+        {
+            bits[i] = nBits[j];
+        }
+    }
+
+    return rshift + 32 * shiftBlock;
+}
+
+template <typename UInteger>
+uint64_t UIntegerALU32<UInteger>::GetPrefix(int32_t numRequested) const
+{
+    UInteger const& self = *(UInteger const*)this;
+    auto const& bits = self.GetBits();
+
+    // Copy to 'prefix' the leading 32-bit block that is nonzero.
+    int32_t bitIndex = self.GetNumBits() - 1;
+    int32_t blockIndex = bitIndex / 32;
+    uint64_t prefix = bits[blockIndex];
+
+    // Get the number of bits in the block starting with the leading 1-bit.
+    int32_t firstBitIndex = bitIndex % 32;
+    int32_t numBlockBits = firstBitIndex + 1;
+
+    // Shift the leading 1-bit to bit-63 of prefix.  We have consumed
+    // numBlockBits, which might not be the entire budget.
+    int32_t targetIndex = 63;
+    prefix <<= targetIndex - firstBitIndex;
+    numRequested -= numBlockBits;
+    targetIndex -= numBlockBits;
+
+    if (numRequested > 0 && --blockIndex >= 0)
+    {
+        // More bits are available.  Copy and shift the entire 32-bit next
+        // block and OR it into the 'prefix'.  For 'float', we will have
+        // consumed the entire budget.  For 'double', we might have to get
+        // bits from a third block.
+        uint64_t nextBlock = bits[blockIndex];
+        nextBlock <<= targetIndex - 31;  // Shift amount is positive.
+        prefix |= nextBlock;
+        numRequested -= 32;
+        targetIndex -= 32;
+
+        if (numRequested > 0 && --blockIndex >= 0)
+        {
+            // We know that targetIndex > 0; only 'double' allows us to get
+            // here, so numRequested is at most 53.  We also know that
+            // targetIndex < 32 because we started with 63 and subtracted
+            // at least 32 from it.  Thus, the shift amount is positive.
+            nextBlock = bits[blockIndex];
+            nextBlock >>= 31 - targetIndex;
+            prefix |= nextBlock;
+        }
+    }
+
+    return prefix;
+}
+
+
+}
diff --git a/Engine/Plugins/Experimental/GeometryProcessing/Source/GeometryAlgorithms/Private/GTEngine/Mathematics/GteUIntegerAP32.cpp b/Engine/Plugins/Experimental/GeometryProcessing/Source/GeometryAlgorithms/Private/GTEngine/Mathematics/GteUIntegerAP32.cpp
new file mode 100644
--- /dev/null
+++ b/Engine/Plugins/Experimental/GeometryProcessing/Source/GeometryAlgorithms/Private/GTEngine/Mathematics/GteUIntegerAP32.cpp
@@ -0,0 +1,154 @@
+// David Eberly, Geometric Tools, Redmond WA 98052
+// Copyright (c) 1998-2019
+// Distributed under the Boost Software License, Version 1.0.
+// http://www.boost.org/LICENSE_1_0.txt
+// http://www.geometrictools.com/License/Boost/LICENSE_1_0.txt
+// File Version: 3.0.1 (2017/07/04)
+
+#include <GTEngine/Mathematics/GteUIntegerAP32.h>
+#include <GTEngine/Mathematics/GteBitHacks.h>
+#include <GTEngine/GTEnginePCH.h>
+#include <algorithm>
+using namespace gte;
+
+#if defined(GTE_COLLECT_UINTEGERAP32_STATISTICS)
+std::atomic<size_t> UIntegerAP32::msMaxSize;
+#endif
+
+
+UIntegerAP32::UIntegerAP32()
+    :
+    mNumBits(0)
+{
+}
+
+UIntegerAP32::UIntegerAP32(UIntegerAP32 const& number)
+{
+    *this = number;
+}
+
+UIntegerAP32::UIntegerAP32(uint32_t number)
+{
+    if (number > 0)
+    {
+        int32_t first = GetLeadingBit(number);
+        int32_t last = GetTrailingBit(number);
+        mNumBits = first - last + 1;
+        mBits.resize(1);
+        mBits[0] = (number >> last);
+    }
+    else
+    {
+        mNumBits = 0;
+    }
+
+#if defined(GTE_COLLECT_UINTEGERAP32_STATISTICS)
+    AtomicMax(msMaxSize, mBits.size());
+#endif
+}
+
+UIntegerAP32::UIntegerAP32(uint64_t number)
+{
+    if (number > 0)
+    {
+        int32_t first = GetLeadingBit(number);
+        int32_t last = GetTrailingBit(number);
+        number >>= last;
+        mNumBits = first - last + 1;
+        mBits.resize(1 + (mNumBits - 1) / 32);
+        mBits[0] = (uint32_t)(number & 0x00000000FFFFFFFFull);
+        if (mBits.size() > 1)
+        {
+            mBits[1] = (uint32_t)((number >> 32) & 0x00000000FFFFFFFFull);
+        }
+    }
+    else
+    {
+        mNumBits = 0;
+    }
+
+#if defined(GTE_COLLECT_UINTEGERAP32_STATISTICS)
+    AtomicMax(msMaxSize, mBits.size());
+#endif
+}
+
+UIntegerAP32::UIntegerAP32(int numBits)
+    :
+    mNumBits(numBits),
+    mBits(1 + (numBits - 1) / 32)
+{
+#if defined(GTE_COLLECT_UINTEGERAP32_STATISTICS)
+    AtomicMax(msMaxSize, mBits.size());
+#endif
+}
+
+UIntegerAP32& UIntegerAP32::operator=(UIntegerAP32 const& number)
+{
+    mNumBits = number.mNumBits;
+    mBits = number.mBits;
+    return *this;
+}
+
+UIntegerAP32::UIntegerAP32(UIntegerAP32&& number)
+{
+    *this = std::move(number);
+}
+
+UIntegerAP32& UIntegerAP32::operator=(UIntegerAP32&& number)
+{
+    mNumBits = number.mNumBits;
+    mBits = std::move(number.mBits);
+    number.mNumBits = 0;
+    return *this;
+}
+
+void UIntegerAP32::SetNumBits(uint32_t numBits)
+{
+    mNumBits = numBits;
+    if (mNumBits > 0)
+    {
+        mBits.resize(1 + (numBits - 1) / 32);
+    }
+    else
+    {
+        mBits.clear();
+    }
+
+#if defined(GTE_COLLECT_UINTEGERAP32_STATISTICS)
+    AtomicMax(msMaxSize, mBits.size());
+#endif
+}
+
+bool UIntegerAP32::Write(std::ofstream& output) const
+{
+    if (output.write((char const*)&mNumBits, sizeof(mNumBits)).bad())
+    {
+        return false;
+    }
+
+    std::size_t size = mBits.size();
+    if (output.write((char const*)&size, sizeof(size)).bad())
+    {
+        return false;
+    }
+
+    return output.write((char const*)&mBits[0], size*sizeof(mBits[0])).good();
+}
+
+bool UIntegerAP32::Read(std::ifstream& input)
+{
+    if (input.read((char*)&mNumBits, sizeof(mNumBits)).bad())
+    {
+        return false;
+    }
+
+    std::size_t size;
+    if (input.read((char*)&size, sizeof(size)).bad())
+    {
+        return false;
+    }
+
+    mBits.resize(size);
+    return input.read((char*)&mBits[0], size*sizeof(mBits[0])).good();
+}
+
