move GeometryProcessing plugin from NotForLicensees to Experimental

diff --git a/Engine/Plugins/Experimental/GeometryProcessing/Source/GeometryAlgorithms/Private/GTEngine/Mathematics/GteUIntegerAP32.h b/Engine/Plugins/Experimental/GeometryProcessing/Source/GeometryAlgorithms/Private/GTEngine/Mathematics/GteUIntegerAP32.h
new file mode 100644
--- /dev/null
+++ b/Engine/Plugins/Experimental/GeometryProcessing/Source/GeometryAlgorithms/Private/GTEngine/Mathematics/GteUIntegerAP32.h
@@ -0,0 +1,118 @@
+// David Eberly, Geometric Tools, Redmond WA 98052
+// Copyright (c) 1998-2019
+// Distributed under the Boost Software License, Version 1.0.
+// http://www.boost.org/LICENSE_1_0.txt
+// http://www.geometrictools.com/License/Boost/LICENSE_1_0.txt
+// File Version: 3.0.0 (2016/06/19)
+
+#pragma once
+
+#include <GTEngine/Mathematics/GteUIntegerALU32.h>
+#include <fstream>
+#include <vector>
+
+// Class UIntegerAP32 is designed to support arbitrary precision arithmetic
+// using BSNumber and BSRational.  It is not a general-purpose class for
+// arithmetic of unsigned integers.
+
+// Uncomment this to collect statistics on how large the UIntegerAP32 storage
+// becomes when using it for the UIntegerType of BSNumber.  After a sequence
+// of BSNumber operations,  look at UIntegerAP32::msMaxSize in the debugger
+// watch window.  If the number is not too large, you might be safe in
+// replacing UIntegerAP32 by UIntegerFP32<N>, where N is the value of
+// UIntegerAP32::msMaxSize.  This leads to much faster code because you no
+// you no longer have dynamic memory allocations and deallocations that occur
+// regularly with std::vector<uint32_t> during BSNumber operations.  A safer
+// choice is to argue mathematically that the maximum size is bounded by N.
+// This requires an analysis of how many bits of precision you need for the
+// types of computation you perform.  See class BSPrecision for code that
+// allows you to compute maximum N.
+//
+//#define GTE_COLLECT_UINTEGERAP32_STATISTICS
+
+#if defined(GTE_COLLECT_UINTEGERAP32_STATISTICS)
+#include <GTEngine/LowLevel/GteAtomicMinMax.h>
+#endif
+
+namespace gte
+{
+
+class UIntegerAP32 : public UIntegerALU32<UIntegerAP32>
+{
+public:
+    // Construction.
+    UIntegerAP32();
+    UIntegerAP32(UIntegerAP32 const& number);
+    UIntegerAP32(uint32_t number);
+    UIntegerAP32(uint64_t number);
+    UIntegerAP32(int numBits);
+
+    // Assignment.
+    UIntegerAP32& operator=(UIntegerAP32 const& number);
+
+    // Support for std::move.
+    UIntegerAP32(UIntegerAP32&& number);
+    UIntegerAP32& operator=(UIntegerAP32&& number);
+
+    // Member access.
+    void SetNumBits(uint32_t numBits);
+    inline int32_t GetNumBits() const;
+    inline std::vector<uint32_t> const& GetBits() const;
+    inline std::vector<uint32_t>& GetBits();
+    inline void SetBack(uint32_t value);
+    inline uint32_t GetBack() const;
+    inline int32_t GetSize() const;
+
+    // Disk input/output.  The fstream objects should be created using
+    // std::ios::binary.  The return value is 'true' iff the operation
+    // was successful.
+    bool Write(std::ofstream& output) const;
+    bool Read(std::ifstream& input);
+
+private:
+    int32_t mNumBits;
+    std::vector<uint32_t> mBits;
+
+    friend class UnitTestBSNumber;
+
+#if defined(GTE_COLLECT_UINTEGERAP32_STATISTICS)
+    static std::atomic<size_t> msMaxSize;
+public:
+    static void SetMaxSizeToZero() { msMaxSize = 0; }
+    static size_t GetMaxSize() { return msMaxSize; }
+#endif
+};
+
+
+inline int32_t UIntegerAP32::GetNumBits() const
+{
+    return mNumBits;
+}
+
+inline std::vector<uint32_t> const& UIntegerAP32::GetBits() const
+{
+    return mBits;
+}
+
+inline std::vector<uint32_t>& UIntegerAP32::GetBits()
+{
+    return mBits;
+}
+
+inline void UIntegerAP32::SetBack(uint32_t value)
+{
+    mBits.back() = value;
+}
+
+inline uint32_t UIntegerAP32::GetBack() const
+{
+    return mBits.back();
+}
+
+inline int32_t UIntegerAP32::GetSize() const
+{
+    return static_cast<int32_t>(mBits.size());
+}
+
+
+}
diff --git a/Engine/Plugins/Experimental/GeometryProcessing/Source/GeometryAlgorithms/Private/GTEngine/Mathematics/GteUIntegerFP32.h b/Engine/Plugins/Experimental/GeometryProcessing/Source/GeometryAlgorithms/Private/GTEngine/Mathematics/GteUIntegerFP32.h
new file mode 100644
--- /dev/null
+++ b/Engine/Plugins/Experimental/GeometryProcessing/Source/GeometryAlgorithms/Private/GTEngine/Mathematics/GteUIntegerFP32.h
@@ -0,0 +1,301 @@
+// David Eberly, Geometric Tools, Redmond WA 98052
+// Copyright (c) 1998-2019
+// Distributed under the Boost Software License, Version 1.0.
+// http://www.boost.org/LICENSE_1_0.txt
+// http://www.geometrictools.com/License/Boost/LICENSE_1_0.txt
+// File Version: 3.0.1 (2017/07/04)
+
+#pragma once
+
+#include <GTEngine/LowLevel/GteLogger.h>
+#include <GTEngine/Mathematics/GteUIntegerALU32.h>
+#include <array>
+#include <fstream>
+
+// Class UIntegerFP32 is designed to support fixed precision arithmetic
+// using BSNumber and BSRational.  It is not a general-purpose class for
+// arithmetic of unsigned integers.  The template parameter N is the
+// number of 32-bit words required to store the precision for the desired
+// computations (maximum number of bits is 32*N).
+
+// Uncomment this to trap when an attempt is made to create storage with
+// more than N uint32_t items.
+//
+//#define GTE_ASSERT_ON_UINTEGERFP32_OUT_OF_RANGE
+
+// Uncomment this to collect statistics on how large the UIntegerFP32 storage
+// becomes when using it for the UIntegerType of BSNumber.  After a sequence
+// of BSNumber operations,  look at UIntegerFP32::msMaxSize in the debugger
+// watch window.  If the number is not too large, you might be safe in
+// replacing the template parameter N by a smaller number.  See class
+// BSPrecision for code that allows you to compute maximum N.
+//
+// NOTE:  Because UIntegerFP32::msMaxSize is a static member of a template
+// class, if you expose this define, you must also declare in your code
+// 'std::atomic<int32_t> gte::UIntegerFP32<N>::msMaxSize;'
+//#define GTE_COLLECT_UINTEGERFP32_STATISTICS
+
+#if defined(GTE_COLLECT_UINTEGERFP32_STATISTICS)
+#include <GTEngine/LowLevel/GteAtomicMinMax.h>
+#endif
+
+namespace gte
+{
+
+template <int N>
+class UIntegerFP32 : public UIntegerALU32<UIntegerFP32<N>>
+{
+public:
+    // Construction.
+    UIntegerFP32();
+    UIntegerFP32(UIntegerFP32 const& number);
+    UIntegerFP32(uint32_t number);
+    UIntegerFP32(uint64_t number);
+    UIntegerFP32(int numBits);
+
+    // Assignment.  Only mSize elements are copied.
+    UIntegerFP32& operator=(UIntegerFP32 const& number);
+
+    // Support for std::move.  The interface is required by BSNumber, but the
+    // std::move of std::array is a copy (no pointer stealing).  Moreover,
+    // a std::array object in this class typically uses smaller than N
+    // elements, the actual size stored in mSize, so we do not want to move
+    // everything.  Therefore, the move operator only copies the bits BUT
+    // BUT 'number' is modified as if you have stolen the data (mNumBits and
+    // mSize set to zero).
+    UIntegerFP32(UIntegerFP32&& number);
+    UIntegerFP32& operator=(UIntegerFP32&& number);
+
+    // Member access.
+    void SetNumBits(uint32_t numBits);
+    inline int32_t GetNumBits() const;
+    inline std::array<uint32_t, N> const& GetBits() const;
+    inline std::array<uint32_t, N>& GetBits();
+    inline void SetBack(uint32_t value);
+    inline uint32_t GetBack() const;
+    inline int32_t GetSize() const;
+
+    // Disk input/output.  The fstream objects should be created using
+    // std::ios::binary.  The return value is 'true' iff the operation
+    // was successful.
+    bool Write(std::ofstream& output) const;
+    bool Read(std::ifstream& input);
+
+private:
+    int32_t mNumBits, mSize;
+    std::array<uint32_t, N> mBits;
+
+    friend class UnitTestBSNumber;
+
+#if defined(GTE_COLLECT_UINTEGERFP32_STATISTICS)
+    static std::atomic<int32_t> msMaxSize;
+public:
+    static void SetMaxSizeToZero() { msMaxSize = 0; }
+    static int32_t GetMaxSize() { return msMaxSize; }
+#endif
+};
+
+
+template <int N>
+UIntegerFP32<N>::UIntegerFP32()
+    :
+    mNumBits(0),
+    mSize(0)
+{
+    static_assert(N >= 1, "Invalid size N.");
+}
+
+template <int N>
+UIntegerFP32<N>::UIntegerFP32(UIntegerFP32 const& number)
+{
+    static_assert(N >= 1, "Invalid size N.");
+
+    *this = number;
+}
+
+template <int N>
+UIntegerFP32<N>::UIntegerFP32(uint32_t number)
+{
+    static_assert(N >= 1, "Invalid size N.");
+
+    if (number > 0)
+    {
+        int32_t first = GetLeadingBit(number);
+        int32_t last = GetTrailingBit(number);
+        mNumBits = first - last + 1;
+        mSize = 1;
+        mBits[0] = (number >> last);
+    }
+    else
+    {
+        mNumBits = 0;
+        mSize = 0;
+    }
+
+#if defined(GTE_COLLECT_UINTEGERFP32_STATISTICS)
+    AtomicMax(msMaxSize, mSize);
+#endif
+}
+
+template <int N>
+UIntegerFP32<N>::UIntegerFP32(uint64_t number)
+{
+    static_assert(N >= 2, "N not large enough to store 64-bit integers.");
+
+    if (number > 0)
+    {
+        int32_t first = GetLeadingBit(number);
+        int32_t last = GetTrailingBit(number);
+        number >>= last;
+        mNumBits = first - last + 1;
+        mSize = 1 + (mNumBits - 1) / 32;
+        mBits[0] = (uint32_t)(number & 0x00000000FFFFFFFFull);
+        if (mSize > 1)
+        {
+            mBits[1] = (uint32_t)((number >> 32) & 0x00000000FFFFFFFFull);
+        }
+    }
+    else
+    {
+        mNumBits = 0;
+        mSize = 0;
+    }
+
+#if defined(GTE_COLLECT_UINTEGERFP32_STATISTICS)
+    AtomicMax(msMaxSize, mSize);
+#endif
+}
+
+template <int N>
+UIntegerFP32<N>::UIntegerFP32(int numBits)
+    :
+    mNumBits(numBits),
+    mSize(1 + (numBits - 1) / 32)
+{
+    static_assert(N >= 1, "Invalid size N.");
+
+#if defined(GTE_ASSERT_ON_UINTEGERFP32_OUT_OF_RANGE)
+    LogAssert(mSize <= N, "N not large enough to store number of bits.");
+#endif
+
+#if defined(GTE_COLLECT_UINTEGERFP32_STATISTICS)
+    AtomicMax(msMaxSize, mSize);
+#endif
+}
+
+template <int N>
+UIntegerFP32<N>& UIntegerFP32<N>::operator=(UIntegerFP32 const& number)
+{
+    static_assert(N >= 1, "Invalid size N.");
+
+    mNumBits = number.mNumBits;
+    mSize = number.mSize;
+    std::copy(number.mBits.begin(), number.mBits.begin() + mSize,
+        mBits.begin());
+    return *this;
+}
+
+template <int N>
+UIntegerFP32<N>::UIntegerFP32(UIntegerFP32&& number)
+{
+    *this = std::move(number);
+}
+
+template <int N>
+UIntegerFP32<N>& UIntegerFP32<N>::operator=(UIntegerFP32&& number)
+{
+    mNumBits = number.mNumBits;
+    mSize = number.mSize;
+    std::copy(number.mBits.begin(), number.mBits.begin() + mSize,
+        mBits.begin());
+    number.mNumBits = 0;
+    number.mSize = 0;
+    return *this;
+}
+
+template <int N>
+void UIntegerFP32<N>::SetNumBits(uint32_t numBits)
+{
+    mNumBits = numBits;
+    mSize = 1 + (numBits - 1) / 32;
+
+#if defined(GTE_ASSERT_ON_UINTEGERFP32_OUT_OF_RANGE)
+    LogAssert(mSize <= N, "N not large enough to store number of bits.");
+#endif
+
+#if defined(GTE_COLLECT_UINTEGERFP32_STATISTICS)
+    AtomicMax(msMaxSize, mSize);
+#endif
+}
+
+template <int N> inline
+int32_t UIntegerFP32<N>::GetNumBits() const
+{
+    return mNumBits;
+}
+
+template <int N> inline
+std::array<uint32_t, N> const& UIntegerFP32<N>::GetBits() const
+{
+    return mBits;
+}
+
+template <int N> inline
+std::array<uint32_t, N>& UIntegerFP32<N>::GetBits()
+{
+    return mBits;
+}
+
+template <int N> inline
+void UIntegerFP32<N>::SetBack(uint32_t value)
+{
+    mBits[mSize - 1] = value;
+}
+
+template <int N> inline
+uint32_t UIntegerFP32<N>::GetBack() const
+{
+    return mBits[mSize - 1];
+}
+
+template <int N> inline
+int32_t UIntegerFP32<N>::GetSize() const
+{
+    return mSize;
+}
+
+template <int N>
+bool UIntegerFP32<N>::Write(std::ofstream& output) const
+{
+    if (output.write((char const*)&mNumBits, sizeof(mNumBits)).bad())
+    {
+        return false;
+    }
+
+    if (output.write((char const*)&mSize, sizeof(mSize)).bad())
+    {
+        return false;
+    }
+
+    return output.write((char const*)&mBits[0],
+        mSize*sizeof(mBits[0])).good();
+}
+
+template <int N>
+bool UIntegerFP32<N>::Read(std::ifstream& input)
+{
+    if (input.read((char*)&mNumBits, sizeof(mNumBits)).bad())
+    {
+        return false;
+    }
+
+    if (input.read((char*)&mSize, sizeof(mSize)).bad())
+    {
+        return false;
+    }
+
+    return input.read((char*)&mBits[0], mSize*sizeof(mBits[0])).good();
+}
+
+
+}
diff --git a/Engine/Plugins/Experimental/GeometryProcessing/Source/GeometryAlgorithms/Private/GTEngine/Mathematics/GteUniqueVerticesTriangles.h b/Engine/Plugins/Experimental/GeometryProcessing/Source/GeometryAlgorithms/Private/GTEngine/Mathematics/GteUniqueVerticesTriangles.h
new file mode 100644
--- /dev/null
+++ b/Engine/Plugins/Experimental/GeometryProcessing/Source/GeometryAlgorithms/Private/GTEngine/Mathematics/GteUniqueVerticesTriangles.h
@@ -0,0 +1,154 @@
+// David Eberly, Geometric Tools, Redmond WA 98052
+// Copyright (c) 1998-2019
+// Distributed under the Boost Software License, Version 1.0.
+// http://www.boost.org/LICENSE_1_0.txt
+// http://www.geometrictools.com/License/Boost/LICENSE_1_0.txt
+// File Version: 3.0.0 (2016/06/19)
+
+#pragma once
+
+#include <GTEngine/GTEngineDEF.h>
+#include <algorithm>
+#include <map>
+#include <vector>
+
+// The VertexType must have an operator< member because it is used as the
+// key in a std::map<VertexType,int>.  For example, if VertexType is
+// Vector3<float>, the comparison operator already exist.  Another example is
+// that you have a vertex type used for vertex coloring, say,
+//   struct VertexType
+//   {
+//       Vector3<float> position;
+//       Vector4<float> color;
+//       bool operator< (VertexType const& v) const
+//       {
+//           return position < v.position;
+//       }
+//   }
+// The comparision will guarantee unique vertex positions, although if you
+// have two VertexType objects with the same position but different colors,
+// there is no guarantee which color will occur in the final result.
+
+namespace gte
+{
+
+template <typename VertexType>
+class UniqueVerticesTriangles
+{
+public:
+    // Triangle soup.  The input vertex array consists of triples of vertices,
+    // each triple representing a triangle.  The array 'inVertices' must have
+    // a multiple of 3 elements.  An array 'outVertices' of unique vertices
+    // and an array 'outIndices' of 'inVertices.size()/3' unique index triples
+    // are computed.  The indices are relative to the array of unique vertices
+    // and each index triple represents a triangle.
+    UniqueVerticesTriangles(
+        std::vector<VertexType> const& inVertices,
+        std::vector<VertexType>& outVertices,
+        std::vector<int>& outIndices);
+
+    // Indexed triangles.  The input vertex array consists of all vertices
+    // referenced by the input index array.  The array 'inIndices' must have a
+    // multiple of 3 elements.  An array 'outVertices' of unique vertices and
+    // an array 'outIndices' of 'indices.size()/3' unique index triples are
+    // computed.  The indices are relative to the array of unique vertices and
+    // each index triple represents a triangle.
+    UniqueVerticesTriangles(
+        std::vector<VertexType> const& inVertices,
+        std::vector<int> const& inIndices,
+        std::vector<VertexType>& outVertices,
+        std::vector<int>& outIndices);
+
+    // The input vertices have indices 0 <= i < VInNum.  The output vertices
+    // have indices 0 <= j < VOutNum.  The construction leads to a mapping of
+    // input indices i to output indices j.  Duplicate vertices have different
+    // input indices but the same output index.  The following function gives
+    // you access to the mapping.  If the input index is invalid (i < 0 or
+    // i >= VINum), the return value is -1.
+    inline int GetOutputIndexFor(int index) const;
+
+private:
+    void ConstructUniqueVertices(std::vector<VertexType> const& inVertices,
+        std::vector<VertexType>& outVertices);
+
+    int mNumInVertices, mNumOutVertices;
+    std::vector<int> mInToOutMapping;
+};
+
+
+template <typename VertexType>
+UniqueVerticesTriangles<VertexType>::UniqueVerticesTriangles(
+    std::vector<VertexType> const& inVertices,
+    std::vector<VertexType>& outVertices, std::vector<int>& outIndices)
+{
+    ConstructUniqueVertices(inVertices, outVertices);
+
+    // The input index array is implicitly {<0,1,2>,<3,4,5>,...,<n-3,n-2,n-1>}
+    // where n is the number of vertices.  The output index array is the same
+    // as the mapping array.
+    outIndices.resize(inVertices.size());
+    std::copy(mInToOutMapping.begin(), mInToOutMapping.end(),
+        outIndices.begin());
+}
+
+template <typename VertexType>
+UniqueVerticesTriangles<VertexType>::UniqueVerticesTriangles(
+    std::vector<VertexType> const& inVertices,
+    std::vector<int> const& inIndices, std::vector<VertexType>& outVertices,
+    std::vector<int>& outIndices)
+{
+    ConstructUniqueVertices(inVertices, outVertices);
+
+    // The input index array needs it indices mapped to the unique vertex
+    // indices.
+    outIndices.resize(inIndices.size());
+    for (size_t i = 0; i < inIndices.size(); ++i)
+    {
+        outIndices[i] = mInToOutMapping[inIndices[i]];
+    }
+}
+
+template <typename VertexType> inline
+int UniqueVerticesTriangles<VertexType>::GetOutputIndexFor(int index) const
+{
+    return mInToOutMapping[index];
+}
+
+template <typename VertexType>
+void UniqueVerticesTriangles<VertexType>::ConstructUniqueVertices(
+    std::vector<VertexType> const& inVertices,
+    std::vector<VertexType>& outVertices)
+{
+    // Construct the unique vertices.
+    mNumInVertices = (int)inVertices.size();
+    mInToOutMapping.resize(mNumInVertices);
+    std::map<VertexType, int> table;
+    mNumOutVertices = 0;
+    for (int i = 0; i < mNumInVertices; ++i)
+    {
+        auto const iter = table.find(inVertices[i]);
+        if (iter != table.end())
+        {
+            // Vertex i is a duplicate of one inserted earlier into the
+            // table.  Map vertex i to the first-found copy.
+            mInToOutMapping[i] = iter->second;
+        }
+        else
+        {
+            // Vertex i is the first occurrence of such a point.
+            table.insert(std::make_pair(inVertices[i], mNumOutVertices));
+            mInToOutMapping[i] = mNumOutVertices;
+            ++mNumOutVertices;
+        }
+    }
+
+    // Pack the unique vertices into an array in the correct order.
+    outVertices.resize(mNumOutVertices);
+    for (auto const& element : table)
+    {
+        outVertices[element.second] = element.first;
+    }
+}
+
+
+}
diff --git a/Engine/Plugins/Experimental/GeometryProcessing/Source/GeometryAlgorithms/Private/GTEngine/Mathematics/GteUnsymmetricEigenvalues.h b/Engine/Plugins/Experimental/GeometryProcessing/Source/GeometryAlgorithms/Private/GTEngine/Mathematics/GteUnsymmetricEigenvalues.h
new file mode 100644
--- /dev/null
+++ b/Engine/Plugins/Experimental/GeometryProcessing/Source/GeometryAlgorithms/Private/GTEngine/Mathematics/GteUnsymmetricEigenvalues.h
@@ -0,0 +1,403 @@
+// David Eberly, Geometric Tools, Redmond WA 98052
+// Copyright (c) 1998-2019
+// Distributed under the Boost Software License, Version 1.0.
+// http://www.boost.org/LICENSE_1_0.txt
+// http://www.geometrictools.com/License/Boost/LICENSE_1_0.txt
+// File Version: 3.0.1 (2018/10/05)
+
+#pragma once
+
+#include <GTEngine/LowLevel/GteWrapper.h>
+#include <algorithm>
+#include <array>
+#include <cmath>
+#include <cstdint>
+#include <functional>
+#include <vector>
+
+// An implementation of the QR algorithm described in "Matrix Computations,
+// 2nd edition" by G. H. Golub and C. F. Van Loan, The Johns Hopkins
+// University Press, Baltimore MD, Fourth Printing 1993.  In particular,
+// the implementation is based on Chapter 7 (The Unsymmetric Eigenvalue
+// Problem), Section 7.5 (The Practical QR Algorithm).
+
+namespace gte
+{
+
+template <typename Real>
+class UnsymmetricEigenvalues
+{
+public:
+    // The solver processes NxN matrices (not necessarily symmetric), where
+    // N >= 3 ('size' is N) and the matrix is stored in row-major order.  The
+    // maximum number of iterations ('maxIterations') must be specified for
+    // reducing an upper Hessenberg matrix to an upper quasi-triangular
+    // matrix (upper triangular matrix of blocks where the diagonal blocks
+    // are 1x1 or 2x2).  The goal is to compute the real-valued eigenvalues.
+    UnsymmetricEigenvalues(int32_t size, uint32_t maxIterations);
+
+    // A copy of the NxN input is made internally.  The order of the
+    // eigenvalues is specified by sortType: -1 (decreasing), 0 (no
+    // sorting), or +1 (increasing).  When sorted, the eigenvectors are
+    // ordered accordingly.  The return value is the number of iterations
+    // consumed when convergence occurred, 0xFFFFFFFF when convergence did
+    // not occur, or 0 when N <= 1 was passed to the constructor.
+    uint32_t Solve(Real const* input, int32_t sortType);
+
+    // Get the real-valued eigenvalues of the matrix passed to Solve(...).
+    // The input 'eigenvalues' must have at least N elements.
+    void GetEigenvalues(uint32_t& numEigenvalues, Real* eigenvalues) const;
+
+private:
+    // 2D accessors to elements of mMatrix[].
+    inline Real const& A(int r, int c) const;
+    inline Real& A(int r, int c);
+
+    // Compute the Householder vector for (X[rmin],...,x[rmax]).  The
+    // input vector is stored in mX in the index range [rmin,rmax].  The
+    // output vector V is stored in mV in the index range [rmin,rmax].  The
+    // scaled vector is S = (-2/Dot(V,V))*V and is stored in mScaledV in
+    // the index range [rmin,rmax].
+    void House(int rmin, int rmax);
+
+    // Support for replacing matrix A by P^T*A*P, where P is a Householder
+    // reflection computed using House(...).
+    void RowHouse(int rmin, int rmax, int cmin, int cmax);
+    void ColHouse(int rmin, int rmax, int cmin, int cmax);
+
+    void ReduceToUpperHessenberg();
+    void FrancisQRStep(int rmin, int rmax);
+    bool GetBlock(std::array<int, 2>& block);
+
+    // The number N of rows and columns of the matrices to be processed.
+    int32_t mSize, mSizeM1;
+
+    // The maximum number of iterations for reducing the tridiagonal mtarix
+    // to a diagonal matrix.
+    uint32_t mMaxIterations;
+
+    // The internal copy of a matrix passed to the solver.
+    std::vector<Real> mMatrix;  // NxN elements
+
+    // Temporary storage to compute Householder reflections.
+    std::vector<Real> mX, mV, mScaledV, mW;  // N elements
+
+    // Flags about the zeroness of the subdiagonal entries.  This is used
+    // to detect uncoupled submatrices and apply the QR algorithm to the
+    // corresponding subproblems.  The storage is padded on both ends with
+    // zeros to avoid additional code logic when packing the eigenvalues
+    // for access by the caller.
+    std::vector<int> mFlagStorage;
+    int* mSubdiagonalFlag;
+
+    int mNumEigenvalues;
+    std::vector<Real> mEigenvalues;
+};
+
+
+template <typename Real>
+UnsymmetricEigenvalues<Real>::UnsymmetricEigenvalues(int32_t size, uint32_t maxIterations)
+    :
+    mSize(0),
+    mSizeM1(0),
+    mMaxIterations(0),
+    mNumEigenvalues(0)
+{
+    if (size >= 3 && maxIterations > 0)
+    {
+        mSize = size;
+        mSizeM1 = size - 1;
+        mMaxIterations = maxIterations;
+        mMatrix.resize(size * size);
+        mX.resize(size);
+        mV.resize(size);
+        mScaledV.resize(size);
+        mW.resize(size);
+        mFlagStorage.resize(size + 1);
+        std::fill(mFlagStorage.begin(), mFlagStorage.end(), 0);
+        mSubdiagonalFlag = &mFlagStorage[1];
+        mEigenvalues.resize(mSize);
+    }
+}
+
+template <typename Real>
+uint32_t UnsymmetricEigenvalues<Real>::Solve(Real const* input, int32_t sortType)
+{
+    if (mSize > 0)
+    {
+        std::copy(input, input + mSize * mSize, mMatrix.begin());
+        ReduceToUpperHessenberg();
+
+        std::array<int, 2> block;
+        bool found = GetBlock(block);
+        uint32_t numIterations;
+        for (numIterations = 0; numIterations < mMaxIterations; ++numIterations)
+        {
+            if (found)
+            {
+                // Solve the current subproblem.
+                FrancisQRStep(block[0], block[1] + 1);
+
+                // Find another subproblem (if any).
+                found = GetBlock(block);
+            }
+            else
+            {
+                break;
+            }
+        }
+
+        // The matrix is fully uncoupled, upper Hessenberg with 1x1 or
+        // 2x2 diagonal blocks. Golub and Van Loan call this "upper
+        // quasi-triangular".
+        mNumEigenvalues = 0;
+        std::fill(mEigenvalues.begin(), mEigenvalues.end(), (Real)0);
+        for (int i = 0; i < mSizeM1; ++i)
+        {
+            if (mSubdiagonalFlag[i] == 0)
+            {
+                if (mSubdiagonalFlag[i - 1] == 0)
+                {
+                    // We have a 1x1 block with a real eigenvalue.
+                    mEigenvalues[mNumEigenvalues++] = A(i, i);
+                }
+            }
+            else
+            {
+                if (mSubdiagonalFlag[i - 1] == 0 && mSubdiagonalFlag[i + 1] == 0)
+                {
+                    // We have a 2x2 block that might have real eigenvalues.
+                    Real a00 = A(i, i);
+                    Real a01 = A(i, i + 1);
+                    Real a10 = A(i + 1, i);
+                    Real a11 = A(i + 1, i + 1);
+                    Real tr = a00 + a11;
+                    Real det = a00 * a11 - a01 * a10;
+                    Real halfTr = tr * (Real)0.5;
+                    Real discr = halfTr * halfTr - det;
+                    if (discr >= (Real)0)
+                    {
+                        Real rootDiscr = std::sqrt(discr);
+                        mEigenvalues[mNumEigenvalues++] = halfTr - rootDiscr;
+                        mEigenvalues[mNumEigenvalues++] = halfTr + rootDiscr;
+                    }
+                }
+                // else:
+                // The QR iteration failed to converge at this block.  It
+                // must also be the case that numIterations == mMaxIterations.
+                // TODO: The caller will be aware of this when testing the
+                // returned numIterations.  Is there a remedy for such a
+                // case?  (This happened with root finding using the
+                // companion matrix of a polynomial.)
+            }
+        }
+
+        if (sortType != 0 && mNumEigenvalues > 1)
+        {
+            if (sortType > 0)
+            {
+                std::sort(mEigenvalues.begin(), mEigenvalues.begin() + mNumEigenvalues,
+                    std::less<Real>());
+            }
+            else
+            {
+                std::sort(mEigenvalues.begin(), mEigenvalues.begin() + mNumEigenvalues,
+                    std::greater<Real>());
+            }
+        }
+
+        return numIterations;
+    }
+    return 0;
+}
+
+template <typename Real>
+void UnsymmetricEigenvalues<Real>::GetEigenvalues(uint32_t& numEigenvalues, Real* eigenvalues) const
+{
+    if (mSize > 0)
+    {
+        numEigenvalues = mNumEigenvalues;
+        Memcpy(eigenvalues, mEigenvalues.data(), numEigenvalues * sizeof(Real));
+    }
+    else
+    {
+        numEigenvalues = 0;
+    }
+}
+
+template <typename Real>
+inline Real const& UnsymmetricEigenvalues<Real>::A(int r, int c) const
+{
+    return mMatrix[c + r * mSize];
+}
+
+template <typename Real>
+inline Real& UnsymmetricEigenvalues<Real>::A(int r, int c)
+{
+    return mMatrix[c + r * mSize];
+}
+
+template <typename Real>
+void UnsymmetricEigenvalues<Real>::House(int rmin, int rmax)
+{
+    Real length = (Real)0;
+    for (int r = rmin; r <= rmax; ++r)
+    {
+        length += mX[r] * mX[r];
+    }
+    length = std::sqrt(length);
+    if (length != (Real)0)
+    {
+        Real sign = (mX[rmin] >= (Real)0 ? (Real)1 : (Real)-1);
+        Real invDenom = ((Real)1) / (mX[rmin] + sign * length);
+        for (int r = rmin + 1; r <= rmax; ++r)
+        {
+            mV[r] = mX[r] * invDenom;
+        }
+    }
+    mV[rmin] = (Real)1;
+
+    Real dot = (Real)1;
+    for (int r = rmin + 1; r <= rmax; ++r)
+    {
+        dot += mV[r] * mV[r];
+    }
+    Real scale = ((Real)-2) / dot;
+    for (int r = rmin; r <= rmax; ++r)
+    {
+        mScaledV[r] = scale * mV[r];
+    }
+}
+
+template <typename Real>
+void UnsymmetricEigenvalues<Real>::RowHouse(int rmin, int rmax, int cmin, int cmax)
+{
+    for (int c = cmin; c <= cmax; ++c)
+    {
+        mW[c] = (Real)0;
+        for (int r = rmin; r <= rmax; ++r)
+        {
+            mW[c] += mScaledV[r] * A(r, c);
+        }
+    }
+
+    for (int r = rmin; r <= rmax; ++r)
+    {
+        for (int c = cmin; c <= cmax; ++c)
+        {
+            A(r, c) += mV[r] * mW[c];
+        }
+    }
+}
+
+template <typename Real>
+void UnsymmetricEigenvalues<Real>::ColHouse(int rmin, int rmax, int cmin, int cmax)
+{
+    for (int r = rmin; r <= rmax; ++r)
+    {
+        mW[r] = (Real)0;
+        for (int c = cmin; c <= cmax; ++c)
+        {
+            mW[r] += mScaledV[c] * A(r, c);
+        }
+    }
+
+    for (int r = rmin; r <= rmax; ++r)
+    {
+        for (int c = cmin; c <= cmax; ++c)
+        {
+            A(r, c) += mW[r] * mV[c];
+        }
+    }
+}
+
+template <typename Real>
+void UnsymmetricEigenvalues<Real>::ReduceToUpperHessenberg()
+{
+    for (int c = 0, cp1 = 1; c <= mSize - 3; ++c, ++cp1)
+    {
+        for (int r = cp1; r <= mSizeM1; ++r)
+        {
+            mX[r] = A(r, c);
+        }
+
+        House(cp1, mSizeM1);
+        RowHouse(cp1, mSizeM1, c, mSizeM1);
+        ColHouse(0, mSizeM1, cp1, mSizeM1);
+    }
+}
+
+template <typename Real>
+void UnsymmetricEigenvalues<Real>::FrancisQRStep(int rmin, int rmax)
+{
+    // Apply the double implicit shift step.
+    int const i0 = rmax - 1, i1 = rmax;
+    Real a00 = A(i0, i0);
+    Real a01 = A(i0, i1);
+    Real a10 = A(i1, i0);
+    Real a11 = A(i1, i1);
+    Real tr = a00 + a11;
+    Real det = a00 * a11 - a01 * a10;
+
+    int const j0 = rmin, j1 = j0 + 1, j2 = j1 + 1;
+    Real b00 = A(j0, j0);
+    Real b01 = A(j0, j1);
+    Real b10 = A(j1, j0);
+    Real b11 = A(j1, j1);
+    Real b21 = A(j2, j1);
+    mX[rmin] = b00 * (b00 - tr) + b01 * b10 + det;
+    mX[rmin + 1] = b10 * (b00 + b11 - tr);
+    mX[rmin + 2] = b10 * b21;
+
+    House(rmin, rmin + 2);
+    RowHouse(rmin, rmin + 2, rmin, rmax);
+    ColHouse(rmin, std::min(rmax, rmin + 3), rmin, rmin + 2);
+
+    // Apply Householder reflections to restore the matrix to upper
+    // Hessenberg form.
+    for (int c = 0, cp1 = 1; c <= mSize - 3; ++c, ++cp1)
+    {
+        int kmax = std::min(cp1 + 2, mSizeM1);
+        for (int r = cp1; r <= kmax; ++r)
+        {
+            mX[r] = A(r, c);
+        }
+
+        House(cp1, kmax);
+        RowHouse(cp1, kmax, c, mSizeM1);
+        ColHouse(0, mSizeM1, cp1, kmax);
+    }
+}
+
+template <typename Real>
+bool UnsymmetricEigenvalues<Real>::GetBlock(std::array<int, 2>& block)
+{
+    for (int i = 0; i < mSizeM1; ++i)
+    {
+        Real a00 = A(i, i);
+        Real a11 = A(i + 1, i + 1);
+        Real a21 = A(i + 1, i);
+        Real sum0 = a00 + a11;
+        Real sum1 = sum0 + a21;
+        mSubdiagonalFlag[i] = (sum1 != sum0 ? 1 : 0);
+    }
+
+    for (int i = 0; i < mSizeM1; ++i)
+    {
+        if (mSubdiagonalFlag[i] == 1)
+        {
+            block = { i, -1 };
+            while (i < mSizeM1 && mSubdiagonalFlag[i] == 1)
+            {
+                block[1] = i++;
+            }
+            if (block[1] != block[0])
+            {
+                return true;
+            }
+        }
+    }
+    return false;
+}
+
+}
diff --git a/Engine/Plugins/Experimental/GeometryProcessing/Source/GeometryAlgorithms/Private/GTEngine/Mathematics/GteVEManifoldMesh.cpp b/Engine/Plugins/Experimental/GeometryProcessing/Source/GeometryAlgorithms/Private/GTEngine/Mathematics/GteVEManifoldMesh.cpp
new file mode 100644
--- /dev/null
+++ b/Engine/Plugins/Experimental/GeometryProcessing/Source/GeometryAlgorithms/Private/GTEngine/Mathematics/GteVEManifoldMesh.cpp
@@ -0,0 +1,228 @@
+// David Eberly, Geometric Tools, Redmond WA 98052
+// Copyright (c) 1998-2019
+// Distributed under the Boost Software License, Version 1.0.
+// http://www.boost.org/LICENSE_1_0.txt
+// http://www.geometrictools.com/License/Boost/LICENSE_1_0.txt
+// File Version: 3.0.0 (2016/06/19)
+
+#include <GTEngine/Mathematics/GteVEManifoldMesh.h>
+#include <GTEngine/LowLevel/GteLogger.h>
+#include <GTEngine/GTEnginePCH.h>
+using namespace gte;
+
+VEManifoldMesh::~VEManifoldMesh()
+{
+}
+
+VEManifoldMesh::VEManifoldMesh(VCreator vCreator, ECreator eCreator)
+    :
+    mVCreator(vCreator ? vCreator : CreateVertex),
+    mECreator(eCreator ? eCreator : CreateEdge),
+    mAssertOnNonmanifoldInsertion(true)
+{
+}
+
+VEManifoldMesh::VMap const& VEManifoldMesh::GetVertices() const
+{
+    return mVMap;
+}
+
+VEManifoldMesh::EMap const& VEManifoldMesh::GetEdges() const
+{
+    return mEMap;
+}
+
+std::shared_ptr<VEManifoldMesh::Vertex> VEManifoldMesh::CreateVertex(int v)
+{
+    return std::make_shared<Vertex>(v);
+}
+
+std::shared_ptr<VEManifoldMesh::Edge> VEManifoldMesh::CreateEdge(int v0, int v1)
+{
+    return std::make_shared<Edge>(v0, v1);
+}
+
+void VEManifoldMesh::AssertOnNonmanifoldInsertion(bool doAssert)
+{
+    mAssertOnNonmanifoldInsertion = doAssert;
+}
+
+std::shared_ptr<VEManifoldMesh::Edge> VEManifoldMesh::Insert(int v0, int v1)
+{
+    std::pair<int,int> ekey(v0, v1);
+    if (mEMap.find(ekey) != mEMap.end())
+    {
+        // The edge already exists.  Return a null pointer as a signal to
+        // the caller that the insertion failed.
+        return nullptr;
+    }
+
+    // Add the new edge.
+    std::shared_ptr<Edge> edge = mECreator(v0, v1);
+    mEMap[ekey] = edge;
+
+    // Add the vertices if they do not already exist.
+    for (int i = 0; i < 2; ++i)
+    {
+        int v = edge->V[i];
+        std::shared_ptr<Vertex> vertex;
+        auto viter = mVMap.find(v);
+        if (viter == mVMap.end())
+        {
+            // This is the first time the vertex is encountered.
+            vertex = mVCreator(v);
+            mVMap[v] = vertex;
+
+            // Update the vertex.
+            vertex->E[0] = edge;
+        }
+        else
+        {
+            // This is the second time the vertex is encountered.
+            vertex = viter->second;
+            if (!vertex)
+            {
+                LogError("Unexpected condition.");
+                return nullptr;
+            }
+
+            // Update the vertex.
+            if (vertex->E[1].lock())
+            {
+                if (mAssertOnNonmanifoldInsertion)
+                {
+                    LogInformation("The mesh must be manifold.");
+                }
+                return nullptr;
+            }
+            vertex->E[1] = edge;
+
+            // Update the adjacent edge.
+            auto adjacent = vertex->E[0].lock();
+            if (!adjacent)
+            {
+                LogError("Unexpected condition.");
+                return nullptr;
+            }
+            for (int j = 0; j < 2; ++j)
+            {
+                if (adjacent->V[j] == v)
+                {
+                    adjacent->E[j] = edge;
+                    break;
+                }
+            }
+
+            // Update the edge.
+            edge->E[i] = adjacent;
+        }
+    }
+
+    return edge;
+}
+
+bool VEManifoldMesh::Remove(int v0, int v1)
+{
+    std::pair<int,int> ekey(v0, v1);
+    auto eiter = mEMap.find(ekey);
+    if (eiter == mEMap.end())
+    {
+        // The edge does not exist.
+        return false;
+    }
+
+    // Get the edge.
+    std::shared_ptr<Edge> edge = eiter->second;
+
+    // Remove the vertices if necessary (when they are not shared).
+    for (int i = 0; i < 2; ++i)
+    {
+        // Inform the vertices the edge is being deleted.
+        auto viter = mVMap.find(edge->V[i]);
+        if (viter == mVMap.end())
+        {
+            // The edge vertices should be in the map.
+            LogError("Unexpected condition.");
+            return false;
+        }
+
+        std::shared_ptr<Vertex> vertex = viter->second;
+        if (!vertex)
+        {
+            // Any <vkey,vertex> in the map must have a dynamically allocated
+            // Vertex object.
+            LogError("Unexpected condition.");
+            return false;
+        }
+        if (vertex->E[0].lock() == edge)
+        {
+            // One-edge vertices always have pointer at index zero.
+            vertex->E[0] = vertex->E[1];
+            vertex->E[1].reset();
+        }
+        else if (vertex->E[1].lock() == edge)
+        {
+            vertex->E[1].reset();
+        }
+        else
+        {
+            LogError("Unexpected condition.");
+            return false;
+        }
+
+        // Remove the vertex if you have the last reference to it.
+        if (!vertex->E[0].lock() && !vertex->E[1].lock())
+        {
+            mVMap.erase(vertex->V);
+        }
+
+        // Inform adjacent edges the edge is being deleted.
+        auto adjacent = edge->E[i].lock();
+        if (adjacent)
+        {
+            for (int j = 0; j < 2; ++j)
+            {
+                if (adjacent->E[j].lock() == edge)
+                {
+                    adjacent->E[j].reset();
+                    break;
+                }
+            }
+        }
+    }
+
+    mEMap.erase(ekey);
+    return true;
+}
+
+bool VEManifoldMesh::IsClosed() const
+{
+    for (auto const& element : mVMap)
+    {
+        auto vertex = element.second;
+        if (!vertex->E[0].lock() || !vertex->E[1].lock())
+        {
+            return false;
+        }
+    }
+    return true;
+}
+
+VEManifoldMesh::Vertex::~Vertex()
+{
+}
+
+VEManifoldMesh::Vertex::Vertex(int v)
+{
+    V = v;
+}
+
+VEManifoldMesh::Edge::~Edge()
+{
+}
+
+VEManifoldMesh::Edge::Edge(int v0, int v1)
+{
+    V[0] = v0;
+    V[1] = v1;
+}
diff --git a/Engine/Plugins/Experimental/GeometryProcessing/Source/GeometryAlgorithms/Private/GTEngine/Mathematics/GteVEManifoldMesh.h b/Engine/Plugins/Experimental/GeometryProcessing/Source/GeometryAlgorithms/Private/GTEngine/Mathematics/GteVEManifoldMesh.h
new file mode 100644
--- /dev/null
+++ b/Engine/Plugins/Experimental/GeometryProcessing/Source/GeometryAlgorithms/Private/GTEngine/Mathematics/GteVEManifoldMesh.h
@@ -0,0 +1,99 @@
+// David Eberly, Geometric Tools, Redmond WA 98052
+// Copyright (c) 1998-2019
+// Distributed under the Boost Software License, Version 1.0.
+// http://www.boost.org/LICENSE_1_0.txt
+// http://www.geometrictools.com/License/Boost/LICENSE_1_0.txt
+// File Version: 3.0.0 (2016/06/19)
+
+#pragma once
+
+#include <GTEngine/GTEngineDEF.h>
+#include <map>
+#include <memory>
+
+namespace gte
+{
+
+class GTE_IMPEXP VEManifoldMesh
+{
+public:
+    // Vertex data types.
+    class Vertex;
+    typedef std::shared_ptr<Vertex> (*VCreator)(int);
+    typedef std::map<int, std::shared_ptr<Vertex>> VMap;
+
+    // Edge data types.
+    class Edge;
+    typedef std::shared_ptr<Edge> (*ECreator)(int, int);
+    typedef std::map<std::pair<int, int>, std::shared_ptr<Edge>> EMap;
+
+    // Vertex object.
+    class GTE_IMPEXP Vertex
+    {
+    public:
+        virtual ~Vertex();
+        Vertex(int v);
+
+        // The unique vertex index.
+        int V;
+
+        // The edges (if any) sharing the vertex.
+        std::weak_ptr<Edge> E[2];
+    };
+
+    // Edge object.
+    class GTE_IMPEXP Edge
+    {
+    public:
+        virtual ~Edge();
+        Edge(int v0, int v1);
+
+        // Vertices, listed as a directed edge <V[0],V[1]>.
+        int V[2];
+
+        // Adjacent edges.  E[i] points to edge sharing V[i].
+        std::weak_ptr<Edge> E[2];
+    };
+
+
+    // Construction and destruction.
+    virtual ~VEManifoldMesh();
+    VEManifoldMesh(VCreator vCreator = nullptr, ECreator eCreator = nullptr);
+
+    // Member access.
+    VMap const& GetVertices() const;
+    EMap const& GetEdges() const;
+
+    // If the insertion of an edge fails because the mesh would become
+    // nonmanifold, the default behavior is to trigger a LogError message.
+    // You can disable this behavior in situations where you want the Logger
+    // system on but you want to continue gracefully without an assertion.
+    void AssertOnNonmanifoldInsertion(bool doAssert);
+
+    // If <v0,v1> is not in the mesh, an Edge object is created and returned;
+    // otherwise, <v0,v1> is in the mesh and nullptr is returned.  If the
+    // insertion leads to a nonmanifold mesh, the call fails with a nullptr
+    // returned.
+    std::shared_ptr<Edge> Insert(int v0, int v1);
+
+    // If <v0,v1> is in the mesh, it is removed and 'true' is returned;
+    // otherwise, <v0,v1> is not in the mesh and 'false' is returned.
+    bool Remove(int v0, int v1);
+
+    // A manifold mesh is closed if each vertex is shared twice.
+    bool IsClosed() const;
+
+protected:
+    // The vertex data and default vertex creation.
+    static std::shared_ptr<Vertex> CreateVertex(int v0);
+    VCreator mVCreator;
+    VMap mVMap;
+
+    // The edge data and default edge creation.
+    static std::shared_ptr<Edge> CreateEdge(int v0, int v1);
+    ECreator mECreator;
+    EMap mEMap;
+    bool mAssertOnNonmanifoldInsertion;  // default: true
+};
+
+}
diff --git a/Engine/Plugins/Experimental/GeometryProcessing/Source/GeometryAlgorithms/Private/GTEngine/Mathematics/GteVETManifoldMesh.cpp b/Engine/Plugins/Experimental/GeometryProcessing/Source/GeometryAlgorithms/Private/GTEngine/Mathematics/GteVETManifoldMesh.cpp
new file mode 100644
--- /dev/null
+++ b/Engine/Plugins/Experimental/GeometryProcessing/Source/GeometryAlgorithms/Private/GTEngine/Mathematics/GteVETManifoldMesh.cpp
@@ -0,0 +1,175 @@
+// David Eberly, Geometric Tools, Redmond WA 98052
+// Copyright (c) 1998-2019
+// Distributed under the Boost Software License, Version 1.0.
+// http://www.boost.org/LICENSE_1_0.txt
+// http://www.geometrictools.com/License/Boost/LICENSE_1_0.txt
+// File Version: 3.0.0 (2016/06/19)
+
+#include <GTEngine/Mathematics/GteVETManifoldMesh.h>
+#include <GTEngine/LowLevel/GteLogger.h>
+#include <GTEngine/GTEnginePCH.h>
+using namespace gte;
+
+VETManifoldMesh::~VETManifoldMesh()
+{
+}
+
+VETManifoldMesh::VETManifoldMesh(VCreator vCreator, ECreator eCreator, TCreator tCreator)
+    :
+    ETManifoldMesh(eCreator, tCreator),
+    mVCreator(vCreator ? vCreator : CreateVertex)
+{
+}
+
+VETManifoldMesh::VETManifoldMesh(VETManifoldMesh const& mesh)
+{
+    *this = mesh;
+}
+
+VETManifoldMesh& VETManifoldMesh::operator=(VETManifoldMesh const& mesh)
+{
+    Clear();
+    mVCreator = mesh.mVCreator;
+    ETManifoldMesh::operator=(mesh);
+    return *this;
+}
+
+VETManifoldMesh::VMap const& VETManifoldMesh::GetVertices() const
+{
+    return mVMap;
+}
+
+std::shared_ptr<VETManifoldMesh::Triangle> VETManifoldMesh::Insert(int v0, int v1, int v2)
+{
+    std::shared_ptr<Triangle> tri = ETManifoldMesh::Insert(v0, v1, v2);
+    if (!tri)
+    {
+        return nullptr;
+    }
+
+    for (int i = 0; i < 3; ++i)
+    {
+        int vIndex = tri->V[i];
+        auto vItem = mVMap.find(vIndex);
+        std::shared_ptr<Vertex> vertex;
+        if (vItem == mVMap.end())
+        {
+            vertex = mVCreator(vIndex);
+            mVMap[vIndex] = vertex;
+        }
+        else
+        {
+            vertex = vItem->second;
+        }
+
+        vertex->TAdjacent.insert(tri);
+
+        for (int j = 0; j < 3; ++j)
+        {
+            auto edge = tri->E[j].lock();
+            if (edge)
+            {
+                if (edge->V[0] == vIndex)
+                {
+                    vertex->VAdjacent.insert(edge->V[1]);
+                    vertex->EAdjacent.insert(edge);
+                }
+                else if (edge->V[1] == vIndex)
+                {
+                    vertex->VAdjacent.insert(edge->V[0]);
+                    vertex->EAdjacent.insert(edge);
+                }
+            }
+            else
+            {
+                LogError("Malformed mesh: Triangle edges must not be null.");
+                return nullptr;
+            }
+        }
+    }
+
+    return tri;
+}
+
+bool VETManifoldMesh::Remove(int v0, int v1, int v2)
+{
+    auto tItem = mTMap.find(TriangleKey<true>(v0, v1, v2));
+    if (tItem == mTMap.end())
+    {
+        return false;
+    }
+
+    std::shared_ptr<Triangle> tri = tItem->second;
+    for (int i = 0; i < 3; ++i)
+    {
+        int vIndex = tri->V[i];
+        auto vItem = mVMap.find(vIndex);
+        if (vItem != mVMap.end())
+        {
+            std::shared_ptr<Vertex> vertex = vItem->second;
+            for (int j = 0; j < 3; ++j)
+            {
+                auto edge = tri->E[j].lock();
+                if (edge)
+                {
+                    if (edge->T[0].lock() && !edge->T[1].lock())
+                    {
+                        if (edge->V[0] == vIndex)
+                        {
+                            vertex->VAdjacent.erase(edge->V[1]);
+                            vertex->EAdjacent.erase(edge);
+                        }
+                        else if (edge->V[1] == vIndex)
+                        {
+                            vertex->VAdjacent.erase(edge->V[0]);
+                            vertex->EAdjacent.erase(edge);
+                        }
+                    }
+                }
+                else
+                {
+                    LogError("Malformed mesh: Triangle edges must not be null.");
+                    return false;
+                }
+            }
+
+            vertex->TAdjacent.erase(tri);
+
+            if (vertex->TAdjacent.size() == 0)
+            {
+                LogAssert(vertex->VAdjacent.size() == 0 && vertex->EAdjacent.size() == 0,
+                    "Malformed mesh: Inconsistent vertex adjacency information.");
+
+                mVMap.erase(vItem);
+            }
+        }
+        else
+        {
+            LogError("Malformed mesh: Vertex must exist in the mesh.");
+            return false;
+        }
+    }
+
+    return ETManifoldMesh::Remove(v0, v1, v2);
+}
+
+void VETManifoldMesh::Clear()
+{
+    mVMap.clear();
+    ETManifoldMesh::Clear();
+}
+
+std::shared_ptr<VETManifoldMesh::Vertex> VETManifoldMesh::CreateVertex(int vIndex)
+{
+    return std::make_shared<Vertex>(vIndex);
+}
+
+VETManifoldMesh::Vertex::~Vertex()
+{
+}
+
+VETManifoldMesh::Vertex::Vertex(int vIndex)
+    :
+    V(vIndex)
+{
+}
diff --git a/Engine/Plugins/Experimental/GeometryProcessing/Source/GeometryAlgorithms/Private/GTEngine/Mathematics/GteVETManifoldMesh.h b/Engine/Plugins/Experimental/GeometryProcessing/Source/GeometryAlgorithms/Private/GTEngine/Mathematics/GteVETManifoldMesh.h
new file mode 100644
--- /dev/null
+++ b/Engine/Plugins/Experimental/GeometryProcessing/Source/GeometryAlgorithms/Private/GTEngine/Mathematics/GteVETManifoldMesh.h
@@ -0,0 +1,79 @@
+// David Eberly, Geometric Tools, Redmond WA 98052
+// Copyright (c) 1998-2019
+// Distributed under the Boost Software License, Version 1.0.
+// http://www.boost.org/LICENSE_1_0.txt
+// http://www.geometrictools.com/License/Boost/LICENSE_1_0.txt
+// File Version: 3.0.0 (2016/06/19)
+
+#pragma once
+
+#include <GTEngine/Mathematics/GteETManifoldMesh.h>
+#include <set>
+
+// The VETManifoldMesh class represents an edge-triangle manifold mesh
+// but additionally stores vertex adjacency information.
+
+namespace gte
+{
+
+class GTE_IMPEXP VETManifoldMesh : public ETManifoldMesh
+{
+public:
+    // Vertex data types.
+    class Vertex;
+    typedef std::shared_ptr<Vertex> (*VCreator)(int);
+    typedef std::map<int, std::shared_ptr<Vertex>> VMap;
+
+    // Vertex object.
+    class GTE_IMPEXP Vertex
+    {
+    public:
+        virtual ~Vertex();
+        Vertex(int vIndex);
+
+        // The index into the vertex pool of the mesh.
+        int V;
+
+        // Adjacent objects.
+        std::set<int> VAdjacent;
+        std::set<std::shared_ptr<Edge>> EAdjacent;
+        std::set<std::shared_ptr<Triangle>> TAdjacent;
+    };
+
+
+    // Construction and destruction.
+    virtual ~VETManifoldMesh();
+    VETManifoldMesh(VCreator vCreator = nullptr, ECreator eCreator = nullptr, TCreator tCreator = nullptr);
+
+    // Support for a deep copy of the mesh.  The mVMap, mEMap, and mTMap
+    // objects have dynamically allocated memory for vertices, edges, and
+    // triangles.  A shallow copy of the pointers to this memory is
+    // problematic.  Allowing sharing, say, via std::shared_ptr, is an
+    // option but not really the intent of copying the mesh graph.
+    VETManifoldMesh(VETManifoldMesh const& mesh);
+    VETManifoldMesh& operator=(VETManifoldMesh const& mesh);
+
+    // Member access.
+    VMap const& GetVertices() const;
+
+    // If <v0,v1,v2> is not in the mesh, a Triangle object is created and
+    // returned; otherwise, <v0,v1,v2> is in the mesh and nullptr is returned.
+    // If the insertion leads to a nonmanifold mesh, the call fails with a
+    // nullptr returned.
+    virtual std::shared_ptr<Triangle> Insert(int v0, int v1, int v2) override;
+
+    // If <v0,v1,v2> is in the mesh, it is removed and 'true' is returned;
+    // otherwise, <v0,v1,v2> is not in the mesh and 'false' is returned.
+    virtual bool Remove(int v0, int v1, int v2) override;
+
+    // Destroy the vertices, edges, and triangles to obtain an empty mesh.
+    virtual void Clear() override;
+
+protected:
+    // The vertex data and default vertex creation.
+    static std::shared_ptr<Vertex> CreateVertex(int vIndex);
+    VCreator mVCreator;
+    VMap mVMap;
+};
+
+}
diff --git a/Engine/Plugins/Experimental/GeometryProcessing/Source/GeometryAlgorithms/Private/GTEngine/Mathematics/GteVETNonmanifoldMesh.cpp b/Engine/Plugins/Experimental/GeometryProcessing/Source/GeometryAlgorithms/Private/GTEngine/Mathematics/GteVETNonmanifoldMesh.cpp
new file mode 100644
--- /dev/null
+++ b/Engine/Plugins/Experimental/GeometryProcessing/Source/GeometryAlgorithms/Private/GTEngine/Mathematics/GteVETNonmanifoldMesh.cpp
@@ -0,0 +1,182 @@
+// David Eberly, Geometric Tools, Redmond WA 98052
+// Copyright (c) 1998-2018
+// Distributed under the Boost Software License, Version 1.0.
+// http://www.boost.org/LICENSE_1_0.txt
+// http://www.geometrictools.com/License/Boost/LICENSE_1_0.txt
+// File Version: 3.20.0 (2019/01/08)
+
+#include <GTEngine/Mathematics/GteVETNonmanifoldMesh.h>
+#include <GTEngine/LowLevel/GteLogger.h>
+#include <GTEngine/GTEnginePCH.h>
+using namespace gte;
+
+VETNonmanifoldMesh::~VETNonmanifoldMesh()
+{
+}
+
+VETNonmanifoldMesh::VETNonmanifoldMesh(VCreator vCreator, ECreator eCreator, TCreator tCreator)
+    :
+    ETNonmanifoldMesh(eCreator, tCreator),
+    mVCreator(vCreator ? vCreator : CreateVertex)
+{
+}
+
+VETNonmanifoldMesh::VETNonmanifoldMesh(VETNonmanifoldMesh const& mesh)
+{
+    *this = mesh;
+}
+
+VETNonmanifoldMesh& VETNonmanifoldMesh::operator=(VETNonmanifoldMesh const& mesh)
+{
+    Clear();
+    mVCreator = mesh.mVCreator;
+    ETNonmanifoldMesh::operator=(mesh);
+    return *this;
+}
+
+VETNonmanifoldMesh::VMap const& VETNonmanifoldMesh::GetVertices() const
+{
+    return mVMap;
+}
+
+std::shared_ptr<VETNonmanifoldMesh::Triangle> VETNonmanifoldMesh::Insert(int v0, int v1, int v2)
+{
+    std::shared_ptr<Triangle> tri = ETNonmanifoldMesh::Insert(v0, v1, v2);
+    if (!tri)
+    {
+        return nullptr;
+    }
+
+    for (int i = 0; i < 3; ++i)
+    {
+        int vIndex = tri->V[i];
+        auto vItem = mVMap.find(vIndex);
+        std::shared_ptr<Vertex> vertex;
+        if (vItem == mVMap.end())
+        {
+            vertex = mVCreator(vIndex);
+            mVMap[vIndex] = vertex;
+        }
+        else
+        {
+            vertex = vItem->second;
+        }
+
+        vertex->TAdjacent.insert(tri);
+
+        for (int j = 0; j < 3; ++j)
+        {
+            auto edge = tri->E[j].lock();
+            if (!edge)
+            {
+                LogError("Unexpected condition.");
+                return nullptr;
+            }
+
+            if (edge->V[0] == vIndex)
+            {
+                vertex->VAdjacent.insert(edge->V[1]);
+                vertex->EAdjacent.insert(edge);
+            }
+            else if (edge->V[1] == vIndex)
+            {
+                vertex->VAdjacent.insert(edge->V[0]);
+                vertex->EAdjacent.insert(edge);
+            }
+        }
+    }
+
+    return tri;
+}
+
+bool VETNonmanifoldMesh::Remove(int v0, int v1, int v2)
+{
+    auto tItem = mTMap.find(TriangleKey<true>(v0, v1, v2));
+    if (tItem == mTMap.end())
+    {
+        return false;
+    }
+
+    std::shared_ptr<Triangle> tri = tItem->second;
+    for (int i = 0; i < 3; ++i)
+    {
+        int vIndex = tri->V[i];
+        auto vItem = mVMap.find(vIndex);
+        if (vItem == mVMap.end())
+        {
+            LogError("Unexpected condition.");
+            return false;
+        }
+
+        std::shared_ptr<Vertex> vertex = vItem->second;
+        for (int j = 0; j < 3; ++j)
+        {
+            auto edge = tri->E[j].lock();
+            if (!edge)
+            {
+                LogError("Unexpected condition.");
+                return false;
+            }
+
+            // If the edge will be removed by ETNonmanifoldMesh::Remove,
+            // remove the vertex references to it.
+            if (edge->T.size() == 1)
+            {
+                for (auto const& adjw : edge->T)
+                {
+                    auto adj = adjw.lock();
+                    if (!adj)
+                    {
+                        LogError("Unexpected condition.");
+                        return false;
+                    }
+
+                    if (edge->V[0] == vIndex)
+                    {
+                        vertex->VAdjacent.erase(edge->V[1]);
+                        vertex->EAdjacent.erase(edge);
+                    }
+                    else if (edge->V[1] == vIndex)
+                    {
+                        vertex->VAdjacent.erase(edge->V[0]);
+                        vertex->EAdjacent.erase(edge);
+                    }
+                }
+            }
+        }
+
+        vertex->TAdjacent.erase(tri);
+
+        // If the vertex is no longer shared by any triangle, remove it.
+        if (vertex->TAdjacent.size() == 0)
+        {
+            LogAssert(vertex->VAdjacent.size() == 0 && vertex->EAdjacent.size() == 0,
+                "Malformed mesh: Inconsistent vertex adjacency information.");
+
+            mVMap.erase(vItem);
+        }
+    }
+
+    return ETNonmanifoldMesh::Remove(v0, v1, v2);
+}
+
+void VETNonmanifoldMesh::Clear()
+{
+    mVMap.clear();
+    ETNonmanifoldMesh::Clear();
+}
+
+std::shared_ptr<VETNonmanifoldMesh::Vertex> VETNonmanifoldMesh::CreateVertex(int vIndex)
+{
+    return std::make_shared<Vertex>(vIndex);
+}
+
+VETNonmanifoldMesh::Vertex::~Vertex()
+{
+}
+
+VETNonmanifoldMesh::Vertex::Vertex(int vIndex)
+    :
+    V(vIndex)
+{
+}
diff --git a/Engine/Plugins/Experimental/GeometryProcessing/Source/GeometryAlgorithms/Private/GTEngine/Mathematics/GteVETNonmanifoldMesh.h b/Engine/Plugins/Experimental/GeometryProcessing/Source/GeometryAlgorithms/Private/GTEngine/Mathematics/GteVETNonmanifoldMesh.h
new file mode 100644
--- /dev/null
+++ b/Engine/Plugins/Experimental/GeometryProcessing/Source/GeometryAlgorithms/Private/GTEngine/Mathematics/GteVETNonmanifoldMesh.h
@@ -0,0 +1,80 @@
+// David Eberly, Geometric Tools, Redmond WA 98052
+// Copyright (c) 1998-2018
+// Distributed under the Boost Software License, Version 1.0.
+// http://www.boost.org/LICENSE_1_0.txt
+// http://www.geometrictools.com/License/Boost/LICENSE_1_0.txt
+// File Version: 3.20.0 (2019/01/09)
+
+#pragma once
+
+#include <GTEngine/LowLevel/GteSharedPtrCompare.h>
+#include <GTEngine/Mathematics/GteETNonmanifoldMesh.h>
+
+// The VETNonmanifoldMesh class represents an edge-triangle nonmanifold mesh
+// but additionally stores vertex adjacency information.
+
+namespace gte
+{
+    class GTE_IMPEXP VETNonmanifoldMesh : public ETNonmanifoldMesh
+    {
+    public:
+        // Vertex data types.
+        class Vertex;
+        typedef std::shared_ptr<Vertex>(*VCreator)(int);
+        typedef std::map<int, std::shared_ptr<Vertex>> VMap;
+
+        // Vertex object.
+        class GTE_IMPEXP Vertex
+        {
+        public:
+            virtual ~Vertex();
+            Vertex(int vIndex);
+
+            // The index into the vertex pool of the mesh.
+            int V;
+
+            bool operator<(Vertex const& other) const
+            {
+                return V < other.V;
+            }
+
+            // Adjacent objects.
+            std::set<int> VAdjacent;
+            std::set<std::shared_ptr<Edge>, SharedPtrLT<Edge>> EAdjacent;
+            std::set<std::shared_ptr<Triangle>, SharedPtrLT<Triangle>> TAdjacent;
+        };
+
+
+        // Construction and destruction.
+        virtual ~VETNonmanifoldMesh();
+        VETNonmanifoldMesh(VCreator vCreator = nullptr, ECreator eCreator = nullptr, TCreator tCreator = nullptr);
+
+        // Support for a deep copy of the mesh.  The mVMap, mEMap, and mTMap
+        // objects have dynamically allocated memory for vertices, edges, and
+        // triangles.  A shallow copy of the pointers to this memory is
+        // problematic.  Allowing sharing, say, via std::shared_ptr, is an
+        // option but not really the intent of copying the mesh graph.
+        VETNonmanifoldMesh(VETNonmanifoldMesh const& mesh);
+        VETNonmanifoldMesh& operator=(VETNonmanifoldMesh const& mesh);
+
+        // Member access.
+        VMap const& GetVertices() const;
+
+        // If <v0,v1,v2> is not in the mesh, a Triangle object is created and
+        // returned; otherwise, <v0,v1,v2> is in the mesh and nullptr is returned.
+        virtual std::shared_ptr<Triangle> Insert(int v0, int v1, int v2) override;
+
+        // If <v0,v1,v2> is in the mesh, it is removed and 'true' is returned;
+        // otherwise, <v0,v1,v2> is not in the mesh and 'false' is returned.
+        virtual bool Remove(int v0, int v1, int v2) override;
+
+        // Destroy the vertices, edges, and triangles to obtain an empty mesh.
+        virtual void Clear() override;
+
+    protected:
+        // The vertex data and default vertex creation.
+        static std::shared_ptr<Vertex> CreateVertex(int vIndex);
+        VCreator mVCreator;
+        VMap mVMap;
+    };
+}
diff --git a/Engine/Plugins/Experimental/GeometryProcessing/Source/GeometryAlgorithms/Private/GTEngine/Mathematics/GteVector.h b/Engine/Plugins/Experimental/GeometryProcessing/Source/GeometryAlgorithms/Private/GTEngine/Mathematics/GteVector.h
new file mode 100644
--- /dev/null
+++ b/Engine/Plugins/Experimental/GeometryProcessing/Source/GeometryAlgorithms/Private/GTEngine/Mathematics/GteVector.h
@@ -0,0 +1,692 @@
+// David Eberly, Geometric Tools, Redmond WA 98052
+// Copyright (c) 1998-2019
+// Distributed under the Boost Software License, Version 1.0.
+// http://www.boost.org/LICENSE_1_0.txt
+// http://www.geometrictools.com/License/Boost/LICENSE_1_0.txt
+// File Version: 3.0.3 (2018/10/05)
+
+#pragma once
+
+#include <GTEngine/GTEngineDEF.h>
+#include <algorithm>
+#include <array>
+#include <cmath>
+#include <initializer_list>
+
+namespace gte
+{
+
+template <int N, typename Real>
+class Vector
+{
+public:
+    // The tuple is uninitialized.
+    Vector();
+
+    // The tuple is fully initialized by the inputs.
+    Vector(std::array<Real, N> const& values);
+
+    // At most N elements are copied from the initializer list, setting any
+    // remaining elements to zero.  Create the zero vector using the syntax
+    //   Vector<N,Real> zero{(Real)0};
+    // WARNING:  The C++ 11 specification states that
+    //   Vector<N,Real> zero{};
+    // will lead to a call of the default constructor, not the initializer
+    // constructor!
+    Vector(std::initializer_list<Real> values);
+
+    // For 0 <= d < N, element d is 1 and all others are 0.  If d is invalid,
+    // the zero vector is created.  This is a convenience for creating the
+    // standard Euclidean basis vectors; see also MakeUnit(int) and Unit(int).
+    Vector(int d);
+
+    // The copy constructor, destructor, and assignment operator are generated
+    // by the compiler.
+
+    // Member access.  The first operator[] returns a const reference rather
+    // than a Real value.  This supports writing via standard file operations
+    // that require a const pointer to data.
+    inline int GetSize() const;  // returns N
+    inline Real const& operator[](int i) const;
+    inline Real& operator[](int i);
+
+    // Comparisons for sorted containers and geometric ordering.
+    inline bool operator==(Vector const& vec) const;
+    inline bool operator!=(Vector const& vec) const;
+    inline bool operator< (Vector const& vec) const;
+    inline bool operator<=(Vector const& vec) const;
+    inline bool operator> (Vector const& vec) const;
+    inline bool operator>=(Vector const& vec) const;
+
+    // Special vectors.
+    void MakeZero();  // All components are 0.
+    void MakeUnit(int d);  // Component d is 1, all others are zero.
+    static Vector Zero();
+    static Vector Unit(int d);
+
+protected:
+    // This data structure takes advantage of the built-in operator[],
+    // range checking, and visualizers in MSVS.
+    std::array<Real, N> mTuple;
+};
+
+// Unary operations.
+template <int N, typename Real>
+Vector<N, Real> operator+(Vector<N, Real> const& v);
+
+template <int N, typename Real>
+Vector<N, Real> operator-(Vector<N, Real> const& v);
+
+// Linear-algebraic operations.
+template <int N, typename Real>
+Vector<N, Real>
+operator+(Vector<N, Real> const& v0, Vector<N, Real> const& v1);
+
+template <int N, typename Real>
+Vector<N, Real>
+operator-(Vector<N, Real> const& v0, Vector<N, Real> const& v1);
+
+template <int N, typename Real>
+Vector<N, Real> operator*(Vector<N, Real> const& v, Real scalar);
+
+template <int N, typename Real>
+Vector<N, Real> operator*(Real scalar, Vector<N, Real> const& v);
+
+template <int N, typename Real>
+Vector<N, Real> operator/(Vector<N, Real> const& v, Real scalar);
+
+template <int N, typename Real>
+Vector<N, Real>& operator+=(Vector<N, Real>& v0, Vector<N, Real> const& v1);
+
+template <int N, typename Real>
+Vector<N, Real>& operator-=(Vector<N, Real>& v0, Vector<N, Real> const& v1);
+
+template <int N, typename Real>
+Vector<N, Real>& operator*=(Vector<N, Real>& v, Real scalar);
+
+template <int N, typename Real>
+Vector<N, Real>& operator/=(Vector<N, Real>& v, Real scalar);
+
+// Componentwise algebraic operations.
+template <int N, typename Real>
+Vector<N, Real>
+operator*(Vector<N, Real> const& v0, Vector<N, Real> const& v1);
+
+template <int N, typename Real>
+Vector<N, Real>
+operator/(Vector<N, Real> const& v0, Vector<N, Real> const& v1);
+
+template <int N, typename Real>
+Vector<N, Real>& operator*=(Vector<N, Real>& v0, Vector<N, Real> const& v1);
+
+template <int N, typename Real>
+Vector<N, Real>& operator/=(Vector<N, Real>& v0, Vector<N, Real> const& v1);
+
+// Geometric operations.  The functions with 'robust' set to 'false' use the
+// standard algorithm for normalizing a vector by computing the length as a
+// square root of the squared length and dividing by it.  The results can be
+// infinite (or NaN) if the length is zero.  When 'robust' is set to 'true',
+// the algorithm is designed to avoid floating-point overflow and sets the
+// normalized vector to zero when the length is zero.
+template <int N, typename Real>
+Real Dot(Vector<N, Real> const& v0, Vector<N, Real> const& v1);
+
+template <int N, typename Real>
+Real Length(Vector<N, Real> const& v, bool robust = false);
+
+template <int N, typename Real>
+Real Normalize(Vector<N, Real>& v, bool robust = false);
+
+// Gram-Schmidt orthonormalization to generate orthonormal vectors from the
+// linearly independent inputs.  The function returns the smallest length of
+// the unnormalized vectors computed during the process.  If this value is
+// nearly zero, it is possible that the inputs are linearly dependent (within
+// numerical round-off errors).  On input, 1 <= numElements <= N and v[0]
+// through v[numElements-1] must be initialized.  On output, the vectors
+// v[0] through v[numElements-1] form an orthonormal set.
+template <int N, typename Real>
+Real Orthonormalize(int numElements, Vector<N, Real>* v, bool robust = false);
+
+// Construct a single vector orthogonal to the nonzero input vector.  If the
+// maximum absolute component occurs at index i, then the orthogonal vector
+// U has u[i] = v[i+1], u[i+1] = -v[i], and all other components zero.  The
+// index addition i+1 is computed modulo N.
+template <int N, typename Real>
+Vector<N, Real> GetOrthogonal(Vector<N, Real> const& v, bool unitLength);
+
+// Compute the axis-aligned bounding box of the vectors.  The return value is
+// 'true' iff the inputs are valid, in which case vmin and vmax have valid
+// values.
+template <int N, typename Real>
+bool ComputeExtremes(int numVectors, Vector<N, Real> const* v,
+    Vector<N, Real>& vmin, Vector<N, Real>& vmax);
+
+// Lift n-tuple v to homogeneous (n+1)-tuple (v,last).
+template <int N, typename Real>
+Vector<N + 1, Real> HLift(Vector<N, Real> const& v, Real last);
+
+// Project homogeneous n-tuple v = (u,v[n-1]) to (n-1)-tuple u.
+template <int N, typename Real>
+Vector<N - 1, Real> HProject(Vector<N, Real> const& v);
+
+// Lift n-tuple v = (w0,w1) to (n+1)-tuple u = (w0,u[inject],w1).  By
+// inference, w0 is a (inject)-tuple [nonexistent when inject=0] and w1 is a
+// (n-inject)-tuple [nonexistent when inject=n].
+template <int N, typename Real>
+Vector<N + 1, Real> Lift(Vector<N, Real> const& v, int inject, Real value);
+
+// Project n-tuple v = (w0,v[reject],w1) to (n-1)-tuple u = (w0,w1).  By
+// inference, w0 is a (reject)-tuple [nonexistent when reject=0] and w1 is a
+// (n-1-reject)-tuple [nonexistent when reject=n-1].
+template <int N, typename Real>
+Vector<N - 1, Real> Project(Vector<N, Real> const& v, int reject);
+
+
+template <int N, typename Real>
+Vector<N, Real>::Vector()
+{
+    // Uninitialized.
+}
+
+template <int N, typename Real>
+Vector<N, Real>::Vector(std::array<Real, N> const& values)
+    :
+    mTuple(values)
+{
+}
+template <int N, typename Real>
+Vector<N, Real>::Vector(std::initializer_list<Real> values)
+{
+    int const numValues = static_cast<int>(values.size());
+    if (N == numValues)
+    {
+        std::copy(values.begin(), values.end(), mTuple.begin());
+    }
+    else if (N > numValues)
+    {
+        std::copy(values.begin(), values.end(), mTuple.begin());
+        std::fill(mTuple.begin() + numValues, mTuple.end(), (Real)0);
+    }
+    else // N < numValues
+    {
+        std::copy(values.begin(), values.begin() + N, mTuple.begin());
+    }
+}
+
+template <int N, typename Real>
+Vector<N, Real>::Vector(int d)
+{
+    MakeUnit(d);
+}
+
+template <int N, typename Real> inline
+int Vector<N, Real>::GetSize() const
+{
+    return N;
+}
+
+template <int N, typename Real> inline
+Real const& Vector<N, Real>::operator[](int i) const
+{
+    return mTuple[i];
+}
+
+template <int N, typename Real> inline
+Real& Vector<N, Real>::operator[](int i)
+{
+    return mTuple[i];
+}
+
+template <int N, typename Real> inline
+bool Vector<N, Real>::operator==(Vector const& vec) const
+{
+    return mTuple == vec.mTuple;
+}
+
+template <int N, typename Real> inline
+bool Vector<N, Real>::operator!=(Vector const& vec) const
+{
+    return mTuple != vec.mTuple;
+}
+
+template <int N, typename Real> inline
+bool Vector<N, Real>::operator<(Vector const& vec) const
+{
+    return mTuple < vec.mTuple;
+}
+
+template <int N, typename Real> inline
+bool Vector<N, Real>::operator<=(Vector const& vec) const
+{
+    return mTuple <= vec.mTuple;
+}
+
+template <int N, typename Real> inline
+bool Vector<N, Real>::operator>(Vector const& vec) const
+{
+    return mTuple > vec.mTuple;
+}
+
+template <int N, typename Real> inline
+bool Vector<N, Real>::operator>=(Vector const& vec) const
+{
+    return mTuple >= vec.mTuple;
+}
+
+template <int N, typename Real>
+void Vector<N, Real>::MakeZero()
+{
+    std::fill(mTuple.begin(), mTuple.end(), (Real)0);
+}
+
+template <int N, typename Real>
+void Vector<N, Real>::MakeUnit(int d)
+{
+    std::fill(mTuple.begin(), mTuple.end(), (Real)0);
+    if (0 <= d && d < N)
+    {
+        mTuple[d] = (Real)1;
+    }
+}
+
+template <int N, typename Real>
+Vector<N, Real> Vector<N, Real>::Zero()
+{
+    Vector<N, Real> v;
+    v.MakeZero();
+    return v;
+}
+
+template <int N, typename Real>
+Vector<N, Real> Vector<N, Real>::Unit(int d)
+{
+    Vector<N, Real> v;
+    v.MakeUnit(d);
+    return v;
+}
+
+
+
+template <int N, typename Real>
+Vector<N, Real> operator+(Vector<N, Real> const& v)
+{
+    return v;
+}
+
+template <int N, typename Real>
+Vector<N, Real> operator-(Vector<N, Real> const& v)
+{
+    Vector<N, Real> result;
+    for (int i = 0; i < N; ++i)
+    {
+        result[i] = -v[i];
+    }
+    return result;
+}
+
+template <int N, typename Real>
+Vector<N, Real> operator+(Vector<N, Real> const& v0,
+    Vector<N, Real> const& v1)
+{
+    Vector<N, Real> result = v0;
+    return result += v1;
+}
+
+template <int N, typename Real>
+Vector<N, Real> operator-(Vector<N, Real> const& v0,
+    Vector<N, Real> const& v1)
+{
+    Vector<N, Real> result = v0;
+    return result -= v1;
+}
+
+template <int N, typename Real>
+Vector<N, Real> operator*(Vector<N, Real> const& v, Real scalar)
+{
+    Vector<N, Real> result = v;
+    return result *= scalar;
+}
+
+template <int N, typename Real>
+Vector<N, Real> operator*(Real scalar, Vector<N, Real> const& v)
+{
+    Vector<N, Real> result = v;
+    return result *= scalar;
+}
+
+template <int N, typename Real>
+Vector<N, Real> operator/(Vector<N, Real> const& v, Real scalar)
+{
+    Vector<N, Real> result = v;
+    return result /= scalar;
+}
+
+template <int N, typename Real>
+Vector<N, Real>& operator+=(Vector<N, Real>& v0, Vector<N, Real> const& v1)
+{
+    for (int i = 0; i < N; ++i)
+    {
+        v0[i] += v1[i];
+    }
+    return v0;
+}
+
+template <int N, typename Real>
+Vector<N, Real>& operator-=(Vector<N, Real>& v0, Vector<N, Real> const& v1)
+{
+    for (int i = 0; i < N; ++i)
+    {
+        v0[i] -= v1[i];
+    }
+    return v0;
+}
+
+template <int N, typename Real>
+Vector<N, Real>& operator*=(Vector<N, Real>& v, Real scalar)
+{
+    for (int i = 0; i < N; ++i)
+    {
+        v[i] *= scalar;
+    }
+    return v;
+}
+
+template <int N, typename Real>
+Vector<N, Real>& operator/=(Vector<N, Real>& v, Real scalar)
+{
+    if (scalar != (Real)0)
+    {
+        Real invScalar = ((Real)1) / scalar;
+        for (int i = 0; i < N; ++i)
+        {
+            v[i] *= invScalar;
+        }
+    }
+    else
+    {
+        for (int i = 0; i < N; ++i)
+        {
+            v[i] = (Real)0;
+        }
+    }
+    return v;
+}
+
+template <int N, typename Real>
+Vector<N, Real> operator*(Vector<N, Real> const& v0,
+    Vector<N, Real> const& v1)
+{
+    Vector<N, Real> result = v0;
+    return result *= v1;
+}
+
+template <int N, typename Real>
+Vector<N, Real> operator/(Vector<N, Real> const& v0,
+    Vector<N, Real> const& v1)
+{
+    Vector<N, Real> result = v0;
+    return result /= v1;
+}
+
+template <int N, typename Real>
+Vector<N, Real>& operator*=(Vector<N, Real>& v0, Vector<N, Real> const& v1)
+{
+    for (int i = 0; i < N; ++i)
+    {
+        v0[i] *= v1[i];
+    }
+    return v0;
+}
+
+template <int N, typename Real>
+Vector<N, Real>& operator/=(Vector<N, Real>& v0, Vector<N, Real> const& v1)
+{
+    for (int i = 0; i < N; ++i)
+    {
+        v0[i] /= v1[i];
+    }
+    return v0;
+}
+
+template <int N, typename Real>
+Real Dot(Vector<N, Real> const& v0, Vector<N, Real> const& v1)
+{
+    Real dot = v0[0] * v1[0];
+    for (int i = 1; i < N; ++i)
+    {
+        dot += v0[i] * v1[i];
+    }
+    return dot;
+}
+
+template <int N, typename Real>
+Real Length(Vector<N, Real> const& v, bool robust)
+{
+    if (robust)
+    {
+        Real maxAbsComp = std::abs(v[0]);
+        for (int i = 1; i < N; ++i)
+        {
+            Real absComp = std::abs(v[i]);
+            if (absComp > maxAbsComp)
+            {
+                maxAbsComp = absComp;
+            }
+        }
+
+        Real length;
+        if (maxAbsComp > (Real)0)
+        {
+            Vector<N, Real> scaled = v / maxAbsComp;
+            length = maxAbsComp * std::sqrt(Dot(scaled, scaled));
+        }
+        else
+        {
+            length = (Real)0;
+        }
+        return length;
+    }
+    else
+    {
+        return std::sqrt(Dot(v, v));
+    }
+}
+
+template <int N, typename Real>
+Real Normalize(Vector<N, Real>& v, bool robust)
+{
+    if (robust)
+    {
+        Real maxAbsComp = std::abs(v[0]);
+        for (int i = 1; i < N; ++i)
+        {
+            Real absComp = std::abs(v[i]);
+            if (absComp > maxAbsComp)
+            {
+                maxAbsComp = absComp;
+            }
+        }
+
+        Real length;
+        if (maxAbsComp > (Real)0)
+        {
+            v /= maxAbsComp;
+            length = std::sqrt(Dot(v, v));
+            v /= length;
+            length *= maxAbsComp;
+        }
+        else
+        {
+            length = (Real)0;
+            for (int i = 0; i < N; ++i)
+            {
+                v[i] = (Real)0;
+            }
+        }
+        return length;
+    }
+    else
+    {
+        Real length = std::sqrt(Dot(v, v));
+        if (length > (Real)0)
+        {
+            v /= length;
+        }
+        else
+        {
+            for (int i = 0; i < N; ++i)
+            {
+                v[i] = (Real)0;
+            }
+        }
+        return length;
+    }
+}
+
+template <int N, typename Real>
+Real Orthonormalize(int numInputs, Vector<N, Real>* v, bool robust)
+{
+    if (v && 1 <= numInputs && numInputs <= N)
+    {
+        Real minLength = Normalize(v[0], robust);
+        for (int i = 1; i < numInputs; ++i)
+        {
+            for (int j = 0; j < i; ++j)
+            {
+                Real dot = Dot(v[i], v[j]);
+                v[i] -= v[j] * dot;
+            }
+            Real length = Normalize(v[i], robust);
+            if (length < minLength)
+            {
+                minLength = length;
+            }
+        }
+        return minLength;
+    }
+
+    return (Real)0;
+}
+
+template <int N, typename Real>
+Vector<N, Real> GetOrthogonal(Vector<N, Real> const& v, bool unitLength)
+{
+    Real cmax = std::abs(v[0]);
+    int imax = 0;
+    for (int i = 1; i < N; ++i)
+    {
+        Real c = std::abs(v[i]);
+        if (c > cmax)
+        {
+            cmax = c;
+            imax = i;
+        }
+    }
+
+    Vector<N, Real> result;
+    result.MakeZero();
+    int inext = imax + 1;
+    if (inext == N)
+    {
+        inext = 0;
+    }
+    result[imax] = v[inext];
+    result[inext] = -v[imax];
+    if (unitLength)
+    {
+        Real sqrDistance = result[imax] * result[imax] + result[inext] * result[inext];
+        Real invLength = ((Real)1) / std::sqrt(sqrDistance);
+        result[imax] *= invLength;
+        result[inext] *= invLength;
+    }
+    return result;
+}
+
+template <int N, typename Real>
+bool ComputeExtremes(int numVectors, Vector<N, Real> const* v,
+    Vector<N, Real>& vmin, Vector<N, Real>& vmax)
+{
+    if (v && numVectors > 0)
+    {
+        vmin = v[0];
+        vmax = vmin;
+        for (int j = 1; j < numVectors; ++j)
+        {
+            Vector<N, Real> const& vec = v[j];
+            for (int i = 0; i < N; ++i)
+            {
+                if (vec[i] < vmin[i])
+                {
+                    vmin[i] = vec[i];
+                }
+                else if (vec[i] > vmax[i])
+                {
+                    vmax[i] = vec[i];
+                }
+            }
+        }
+        return true;
+    }
+
+    return false;
+}
+
+template <int N, typename Real>
+Vector<N + 1, Real> HLift(Vector<N, Real> const& v, Real last)
+{
+    Vector<N + 1, Real> result;
+    for (int i = 0; i < N; ++i)
+    {
+        result[i] = v[i];
+    }
+    result[N] = last;
+    return result;
+}
+
+template <int N, typename Real>
+Vector<N - 1, Real> HProject(Vector<N, Real> const& v)
+{
+    static_assert(N >= 2, "Invalid dimension.");
+    Vector<N - 1, Real> result;
+    for (int i = 0; i < N - 1; ++i)
+    {
+        result[i] = v[i];
+    }
+    return result;
+}
+
+template <int N, typename Real>
+Vector<N + 1, Real> Lift(Vector<N, Real> const& v, int inject, Real value)
+{
+    Vector<N + 1, Real> result;
+    int i;
+    for (i = 0; i < inject; ++i)
+    {
+        result[i] = v[i];
+    }
+    result[i] = value;
+    int j = i;
+    for (++j; i < N; ++i, ++j)
+    {
+        result[j] = v[i];
+    }
+    return result;
+}
+
+template <int N, typename Real>
+Vector<N - 1, Real> Project(Vector<N, Real> const& v, int reject)
+{
+    static_assert(N >= 2, "Invalid dimension.");
+    Vector<N - 1, Real> result;
+    for (int i = 0, j = 0; i < N - 1; ++i, ++j)
+    {
+        if (j == reject)
+        {
+            ++j;
+        }
+        result[i] = v[j];
+    }
+    return result;
+}
+
+}
diff --git a/Engine/Plugins/Experimental/GeometryProcessing/Source/GeometryAlgorithms/Private/GTEngine/Mathematics/GteVector2.h b/Engine/Plugins/Experimental/GeometryProcessing/Source/GeometryAlgorithms/Private/GTEngine/Mathematics/GteVector2.h
new file mode 100644
--- /dev/null
+++ b/Engine/Plugins/Experimental/GeometryProcessing/Source/GeometryAlgorithms/Private/GTEngine/Mathematics/GteVector2.h
@@ -0,0 +1,276 @@
+// David Eberly, Geometric Tools, Redmond WA 98052
+// Copyright (c) 1998-2019
+// Distributed under the Boost Software License, Version 1.0.
+// http://www.boost.org/LICENSE_1_0.txt
+// http://www.geometrictools.com/License/Boost/LICENSE_1_0.txt
+// File Version: 3.0.2 (2018/10/04)
+
+#pragma once
+
+#include <GTEngine/Mathematics/GteVector.h>
+
+namespace gte
+{
+
+// Template alias for convenience.
+template <typename Real>
+using Vector2 = Vector<2, Real>;
+
+// Compute the perpendicular using the formal determinant,
+//   perp = det{{e0,e1},{x0,x1}} = (x1,-x0)
+// where e0 = (1,0), e1 = (0,1), and v = (x0,x1).
+template <typename Real>
+Vector2<Real> Perp(Vector2<Real> const& v);
+
+// Compute the normalized perpendicular.
+template <typename Real>
+Vector2<Real> UnitPerp(Vector2<Real> const& v, bool robust = false);
+
+// Compute Dot((x0,x1),Perp(y0,y1)) = x0*y1 - x1*y0, where v0 = (x0,x1)
+// and v1 = (y0,y1).
+template <typename Real>
+Real DotPerp(Vector2<Real> const& v0, Vector2<Real> const& v1);
+
+// Compute a right-handed orthonormal basis for the orthogonal complement
+// of the input vectors.  The function returns the smallest length of the
+// unnormalized vectors computed during the process.  If this value is nearly
+// zero, it is possible that the inputs are linearly dependent (within
+// numerical round-off errors).  On input, numInputs must be 1 and v[0]
+// must be initialized.  On output, the vectors v[0] and v[1] form an
+// orthonormal set.
+template <typename Real>
+Real ComputeOrthogonalComplement(int numInputs, Vector2<Real>* v, bool robust = false);
+
+// Compute the barycentric coordinates of the point P with respect to the
+// triangle <V0,V1,V2>, P = b0*V0 + b1*V1 + b2*V2, where b0 + b1 + b2 = 1.
+// The return value is 'true' iff {V0,V1,V2} is a linearly independent set.
+// Numerically, this is measured by |det[V0 V1 V2]| <= epsilon.  The values
+// bary[] are valid only when the return value is 'true' but set to zero when
+// the return value is 'false'.
+template <typename Real>
+bool ComputeBarycentrics(Vector2<Real> const& p, Vector2<Real> const& v0,
+    Vector2<Real> const& v1, Vector2<Real> const& v2, Real bary[3],
+    Real epsilon = (Real)0);
+
+// Get intrinsic information about the input array of vectors.  The return
+// value is 'true' iff the inputs are valid (numVectors > 0, v is not null,
+// and epsilon >= 0), in which case the class members are valid.
+template <typename Real>
+class IntrinsicsVector2
+{
+public:
+    // The constructor sets the class members based on the input set.
+    IntrinsicsVector2(int numVectors, Vector2<Real> const* v, Real inEpsilon);
+
+    // A nonnegative tolerance that is used to determine the intrinsic
+    // dimension of the set.
+    Real epsilon;
+
+    // The intrinsic dimension of the input set, computed based on the
+    // nonnegative tolerance mEpsilon.
+    int dimension;
+
+    // Axis-aligned bounding box of the input set.  The maximum range is
+    // the larger of max[0]-min[0] and max[1]-min[1].
+    Real min[2], max[2];
+    Real maxRange;
+
+    // Coordinate system.  The origin is valid for any dimension d.  The
+    // unit-length direction vector is valid only for 0 <= i < d.  The
+    // extreme index is relative to the array of input points, and is also
+    // valid only for 0 <= i < d.  If d = 0, all points are effectively
+    // the same, but the use of an epsilon may lead to an extreme index
+    // that is not zero.  If d = 1, all points effectively lie on a line
+    // segment.  If d = 2, the points are not collinear.
+    Vector2<Real> origin;
+    Vector2<Real> direction[2];
+
+    // The indices that define the maximum dimensional extents.  The
+    // values extreme[0] and extreme[1] are the indices for the points
+    // that define the largest extent in one of the coordinate axis
+    // directions.  If the dimension is 2, then extreme[2] is the index
+    // for the point that generates the largest extent in the direction
+    // perpendicular to the line through the points corresponding to
+    // extreme[0] and extreme[1].  The triangle formed by the points
+    // V[extreme[0]], V[extreme[1]], and V[extreme[2]] is clockwise or
+    // counterclockwise, the condition stored in extremeCCW.
+    int extreme[3];
+    bool extremeCCW;
+};
+
+
+template <typename Real>
+Vector2<Real> Perp(Vector2<Real> const& v)
+{
+    return Vector2<Real>{ v[1], -v[0] };
+}
+
+template <typename Real>
+Vector2<Real> UnitPerp(Vector2<Real> const& v, bool robust)
+{
+    Vector2<Real> unitPerp{ v[1], -v[0] };
+    Normalize(unitPerp, robust);
+    return unitPerp;
+}
+
+template <typename Real>
+Real DotPerp(Vector2<Real> const& v0, Vector2<Real> const& v1)
+{
+    return Dot(v0, Perp(v1));
+}
+
+template <typename Real>
+Real ComputeOrthogonalComplement(int numInputs, Vector2<Real>* v, bool robust)
+{
+    if (numInputs == 1)
+    {
+        v[1] = -Perp(v[0]);
+        return Orthonormalize<2, Real>(2, v, robust);
+    }
+
+    return (Real)0;
+}
+
+template <typename Real>
+bool ComputeBarycentrics(Vector2<Real> const& p, Vector2<Real> const& v0,
+    Vector2<Real> const& v1, Vector2<Real> const& v2, Real bary[3],
+    Real epsilon)
+{
+    // Compute the vectors relative to V2 of the triangle.
+    Vector2<Real> diff[3] = { v0 - v2, v1 - v2, p - v2 };
+
+    Real det = DotPerp(diff[0], diff[1]);
+    if (det < -epsilon || det > epsilon)
+    {
+        Real invDet = ((Real)1) / det;
+        bary[0] = DotPerp(diff[2], diff[1])*invDet;
+        bary[1] = DotPerp(diff[0], diff[2])*invDet;
+        bary[2] = (Real)1 - bary[0] - bary[1];
+        return true;
+    }
+
+    for (int i = 0; i < 3; ++i)
+    {
+        bary[i] = (Real)0;
+    }
+    return false;
+}
+
+template <typename Real>
+IntrinsicsVector2<Real>::IntrinsicsVector2(int numVectors,
+    Vector2<Real> const* v, Real inEpsilon)
+    :
+    epsilon(inEpsilon),
+    dimension(0),
+    maxRange((Real)0),
+    origin({ (Real)0, (Real)0 }),
+    extremeCCW(false)
+{
+    min[0] = (Real)0;
+    min[1] = (Real)0;
+    direction[0] = { (Real)0, (Real)0 };
+    direction[1] = { (Real)0, (Real)0 };
+    extreme[0] = 0;
+    extreme[1] = 0;
+    extreme[2] = 0;
+
+    if (numVectors > 0 && v && epsilon >= (Real)0)
+    {
+        // Compute the axis-aligned bounding box for the input vectors.  Keep
+        // track of the indices into 'vectors' for the current min and max.
+        int j, indexMin[2], indexMax[2];
+        for (j = 0; j < 2; ++j)
+        {
+            min[j] = v[0][j];
+            max[j] = min[j];
+            indexMin[j] = 0;
+            indexMax[j] = 0;
+        }
+
+        int i;
+        for (i = 1; i < numVectors; ++i)
+        {
+            for (j = 0; j < 2; ++j)
+            {
+                if (v[i][j] < min[j])
+                {
+                    min[j] = v[i][j];
+                    indexMin[j] = i;
+                }
+                else if (v[i][j] > max[j])
+                {
+                    max[j] = v[i][j];
+                    indexMax[j] = i;
+                }
+            }
+        }
+
+        // Determine the maximum range for the bounding box.
+        maxRange = max[0] - min[0];
+        extreme[0] = indexMin[0];
+        extreme[1] = indexMax[0];
+        Real range = max[1] - min[1];
+        if (range > maxRange)
+        {
+            maxRange = range;
+            extreme[0] = indexMin[1];
+            extreme[1] = indexMax[1];
+        }
+
+        // The origin is either the vector of minimum x0-value or vector of
+        // minimum x1-value.
+        origin = v[extreme[0]];
+
+        // Test whether the vector set is (nearly) a vector.
+        if (maxRange <= epsilon)
+        {
+            dimension = 0;
+            for (j = 0; j < 2; ++j)
+            {
+                extreme[j + 1] = extreme[0];
+            }
+            return;
+        }
+
+        // Test whether the vector set is (nearly) a line segment.  We need
+        // direction[1] to span the orthogonal complement of direction[0].
+        direction[0] = v[extreme[1]] - origin;
+        Normalize(direction[0], false);
+        direction[1] = -Perp(direction[0]);
+
+        // Compute the maximum distance of the points from the line
+        // origin+t*direction[0].
+        Real maxDistance = (Real)0;
+        Real maxSign = (Real)0;
+        extreme[2] = extreme[0];
+        for (i = 0; i < numVectors; ++i)
+        {
+            Vector2<Real> diff = v[i] - origin;
+            Real distance = Dot(direction[1], diff);
+            Real sign = (distance >(Real)0 ? (Real)1 :
+                (distance < (Real)0 ? (Real)-1 : (Real)0));
+            distance = std::abs(distance);
+            if (distance > maxDistance)
+            {
+                maxDistance = distance;
+                maxSign = sign;
+                extreme[2] = i;
+            }
+        }
+
+        if (maxDistance <= epsilon * maxRange)
+        {
+            // The points are (nearly) on the line origin+t*direction[0].
+            dimension = 1;
+            extreme[2] = extreme[1];
+            return;
+        }
+
+        dimension = 2;
+        extremeCCW = (maxSign > (Real)0);
+        return;
+    }
+}
+
+
+}
