move GeometryProcessing plugin from NotForLicensees to Experimental

diff --git a/Engine/Plugins/Experimental/GeometryProcessing/Source/GeometryAlgorithms/Private/GTEngine/Mathematics/GteVector3.h b/Engine/Plugins/Experimental/GeometryProcessing/Source/GeometryAlgorithms/Private/GTEngine/Mathematics/GteVector3.h
new file mode 100644
--- /dev/null
+++ b/Engine/Plugins/Experimental/GeometryProcessing/Source/GeometryAlgorithms/Private/GTEngine/Mathematics/GteVector3.h
@@ -0,0 +1,382 @@
+// David Eberly, Geometric Tools, Redmond WA 98052
+// Copyright (c) 1998-2019
+// Distributed under the Boost Software License, Version 1.0.
+// http://www.boost.org/LICENSE_1_0.txt
+// http://www.geometrictools.com/License/Boost/LICENSE_1_0.txt
+// File Version: 3.0.2 (2018/10/04)
+
+#pragma once
+
+#include <GTEngine/Mathematics/GteVector.h>
+
+namespace gte
+{
+
+// Template alias for convenience.
+template <typename Real>
+using Vector3 = Vector<3, Real>;
+
+// Cross, UnitCross, and DotCross have a template parameter N that should
+// be 3 or 4.  The latter case supports affine vectors in 4D (last component
+// w = 0) when you want to use 4-tuples and 4x4 matrices for affine algebra.
+
+// Compute the cross product using the formal determinant:
+//   cross = det{{e0,e1,e2},{x0,x1,x2},{y0,y1,y2}}
+//         = (x1*y2-x2*y1, x2*y0-x0*y2, x0*y1-x1*y0)
+// where e0 = (1,0,0), e1 = (0,1,0), e2 = (0,0,1), v0 = (x0,x1,x2), and
+// v1 = (y0,y1,y2).
+template <int N, typename Real>
+Vector<N,Real> Cross(Vector<N,Real> const& v0, Vector<N,Real> const& v1);
+
+// Compute the normalized cross product.
+template <int N, typename Real>
+Vector<N, Real> UnitCross(Vector<N,Real> const& v0, Vector<N,Real> const& v1,
+    bool robust = false);
+
+// Compute Dot((x0,x1,x2),Cross((y0,y1,y2),(z0,z1,z2)), the triple scalar
+// product of three vectors, where v0 = (x0,x1,x2), v1 = (y0,y1,y2), and
+// v2 is (z0,z1,z2).
+template <int N, typename Real>
+Real DotCross(Vector<N,Real> const& v0, Vector<N,Real> const& v1,
+    Vector<N,Real> const& v2);
+
+// Compute a right-handed orthonormal basis for the orthogonal complement
+// of the input vectors.  The function returns the smallest length of the
+// unnormalized vectors computed during the process.  If this value is nearly
+// zero, it is possible that the inputs are linearly dependent (within
+// numerical round-off errors).  On input, numInputs must be 1 or 2 and
+// v[0] through v[numInputs-1] must be initialized.  On output, the
+// vectors v[0] through v[2] form an orthonormal set.
+template <typename Real>
+Real ComputeOrthogonalComplement(int numInputs, Vector3<Real>* v, bool robust = false);
+
+// Compute the barycentric coordinates of the point P with respect to the
+// tetrahedron <V0,V1,V2,V3>, P = b0*V0 + b1*V1 + b2*V2 + b3*V3, where
+// b0 + b1 + b2 + b3 = 1.  The return value is 'true' iff {V0,V1,V2,V3} is
+// a linearly independent set.  Numerically, this is measured by
+// |det[V0 V1 V2 V3]| <= epsilon.  The values bary[] are valid only when the
+// return value is 'true' but set to zero when the return value is 'false'.
+template <typename Real>
+bool ComputeBarycentrics(Vector3<Real> const& p, Vector3<Real> const& v0,
+    Vector3<Real> const& v1, Vector3<Real> const& v2, Vector3<Real> const& v3,
+    Real bary[4], Real epsilon = (Real)0);
+
+// Get intrinsic information about the input array of vectors.  The return
+// value is 'true' iff the inputs are valid (numVectors > 0, v is not null,
+// and epsilon >= 0), in which case the class members are valid.
+template <typename Real>
+class IntrinsicsVector3
+{
+public:
+    // The constructor sets the class members based on the input set.
+    IntrinsicsVector3(int numVectors, Vector3<Real> const* v, Real inEpsilon);
+
+    // A nonnegative tolerance that is used to determine the intrinsic
+    // dimension of the set.
+    Real epsilon;
+
+    // The intrinsic dimension of the input set, computed based on the
+    // nonnegative tolerance mEpsilon.
+    int dimension;
+
+    // Axis-aligned bounding box of the input set.  The maximum range is
+    // the larger of max[0]-min[0], max[1]-min[1], and max[2]-min[2].
+    Real min[3], max[3];
+    Real maxRange;
+
+    // Coordinate system.  The origin is valid for any dimension d.  The
+    // unit-length direction vector is valid only for 0 <= i < d.  The
+    // extreme index is relative to the array of input points, and is also
+    // valid only for 0 <= i < d.  If d = 0, all points are effectively
+    // the same, but the use of an epsilon may lead to an extreme index
+    // that is not zero.  If d = 1, all points effectively lie on a line
+    // segment.  If d = 2, all points effectively line on a plane.  If
+    // d = 3, the points are not coplanar.
+    Vector3<Real> origin;
+    Vector3<Real> direction[3];
+
+    // The indices that define the maximum dimensional extents.  The
+    // values extreme[0] and extreme[1] are the indices for the points
+    // that define the largest extent in one of the coordinate axis
+    // directions.  If the dimension is 2, then extreme[2] is the index
+    // for the point that generates the largest extent in the direction
+    // perpendicular to the line through the points corresponding to
+    // extreme[0] and extreme[1].  Furthermore, if the dimension is 3,
+    // then extreme[3] is the index for the point that generates the
+    // largest extent in the direction perpendicular to the triangle
+    // defined by the other extreme points.  The tetrahedron formed by the
+    // points V[extreme[0]], V[extreme[1]], V[extreme[2]], and
+    // V[extreme[3]] is clockwise or counterclockwise, the condition
+    // stored in extremeCCW.
+    int extreme[4];
+    bool extremeCCW;
+};
+
+
+template <int N, typename Real>
+Vector<N, Real> Cross(Vector<N, Real> const& v0, Vector<N, Real> const& v1)
+{
+    static_assert(N == 3 || N == 4, "Dimension must be 3 or 4.");
+
+    Vector<N, Real> result;
+    result.MakeZero();
+    result[0] = v0[1] * v1[2] - v0[2] * v1[1];
+    result[1] = v0[2] * v1[0] - v0[0] * v1[2];
+    result[2] = v0[0] * v1[1] - v0[1] * v1[0];
+    return result;
+}
+
+template <int N, typename Real>
+Vector<N, Real> UnitCross(Vector<N, Real> const& v0, Vector<N, Real> const& v1, bool robust)
+{
+    static_assert(N == 3 || N == 4, "Dimension must be 3 or 4.");
+
+    Vector<N, Real> unitCross = Cross(v0, v1);
+    Normalize(unitCross, robust);
+    return unitCross;
+}
+
+template <int N, typename Real>
+Real DotCross(Vector<N, Real> const& v0, Vector<N, Real> const& v1,
+    Vector<N, Real> const& v2)
+{
+    static_assert(N == 3 || N == 4, "Dimension must be 3 or 4.");
+
+    return Dot(v0, Cross(v1, v2));
+}
+
+template <typename Real>
+Real ComputeOrthogonalComplement(int numInputs, Vector3<Real>* v, bool robust)
+{
+    if (numInputs == 1)
+    {
+        if (std::abs(v[0][0]) > std::abs(v[0][1]))
+        {
+            v[1] = { -v[0][2], (Real)0, +v[0][0] };
+        }
+        else
+        {
+            v[1] = { (Real)0, +v[0][2], -v[0][1] };
+        }
+        numInputs = 2;
+    }
+
+    if (numInputs == 2)
+    {
+        v[2] = Cross(v[0], v[1]);
+        return Orthonormalize<3, Real>(3, v, robust);
+    }
+
+    return (Real)0;
+}
+
+template <typename Real>
+bool ComputeBarycentrics(Vector3<Real> const& p, Vector3<Real> const& v0,
+    Vector3<Real> const& v1, Vector3<Real> const& v2, Vector3<Real> const& v3,
+    Real bary[4], Real epsilon)
+{
+    // Compute the vectors relative to V3 of the tetrahedron.
+    Vector3<Real> diff[4] = { v0 - v3, v1 - v3, v2 - v3, p - v3 };
+
+    Real det = DotCross(diff[0], diff[1], diff[2]);
+    if (det < -epsilon || det > epsilon)
+    {
+        Real invDet = ((Real)1) / det;
+        bary[0] = DotCross(diff[3], diff[1], diff[2]) * invDet;
+        bary[1] = DotCross(diff[3], diff[2], diff[0]) * invDet;
+        bary[2] = DotCross(diff[3], diff[0], diff[1]) * invDet;
+        bary[3] = (Real)1 - bary[0] - bary[1] - bary[2];
+        return true;
+    }
+
+    for (int i = 0; i < 4; ++i)
+    {
+        bary[i] = (Real)0;
+    }
+    return false;
+}
+
+template <typename Real>
+IntrinsicsVector3<Real>::IntrinsicsVector3(int numVectors,
+    Vector3<Real> const* v, Real inEpsilon)
+    :
+    epsilon(inEpsilon),
+    dimension(0),
+    maxRange((Real)0),
+    origin({ (Real)0, (Real)0, (Real)0 }),
+    extremeCCW(false)
+{
+    min[0] = (Real)0;
+    min[1] = (Real)0;
+    min[2] = (Real)0;
+    direction[0] = { (Real)0, (Real)0, (Real)0 };
+    direction[1] = { (Real)0, (Real)0, (Real)0 };
+    direction[2] = { (Real)0, (Real)0, (Real)0 };
+    extreme[0] = 0;
+    extreme[1] = 0;
+    extreme[2] = 0;
+    extreme[3] = 0;
+
+    if (numVectors > 0 && v && epsilon >= (Real)0)
+    {
+        // Compute the axis-aligned bounding box for the input vectors.  Keep
+        // track of the indices into 'vectors' for the current min and max.
+        int j, indexMin[3], indexMax[3];
+        for (j = 0; j < 3; ++j)
+        {
+            min[j] = v[0][j];
+            max[j] = min[j];
+            indexMin[j] = 0;
+            indexMax[j] = 0;
+        }
+
+        int i;
+        for (i = 1; i < numVectors; ++i)
+        {
+            for (j = 0; j < 3; ++j)
+            {
+                if (v[i][j] < min[j])
+                {
+                    min[j] = v[i][j];
+                    indexMin[j] = i;
+                }
+                else if (v[i][j] > max[j])
+                {
+                    max[j] = v[i][j];
+                    indexMax[j] = i;
+                }
+            }
+        }
+
+        // Determine the maximum range for the bounding box.
+        maxRange = max[0] - min[0];
+        extreme[0] = indexMin[0];
+        extreme[1] = indexMax[0];
+        Real range = max[1] - min[1];
+        if (range > maxRange)
+        {
+            maxRange = range;
+            extreme[0] = indexMin[1];
+            extreme[1] = indexMax[1];
+        }
+        range = max[2] - min[2];
+        if (range > maxRange)
+        {
+            maxRange = range;
+            extreme[0] = indexMin[2];
+            extreme[1] = indexMax[2];
+        }
+
+        // The origin is either the vector of minimum x0-value, vector of
+        // minimum x1-value, or vector of minimum x2-value.
+        origin = v[extreme[0]];
+
+        // Test whether the vector set is (nearly) a vector.
+        if (maxRange <= epsilon)
+        {
+            dimension = 0;
+            for (j = 0; j < 3; ++j)
+            {
+                extreme[j + 1] = extreme[0];
+            }
+            return;
+        }
+
+        // Test whether the vector set is (nearly) a line segment.  We need
+        // {direction[2],direction[3]} to span the orthogonal complement of
+        // direction[0].
+        direction[0] = v[extreme[1]] - origin;
+        Normalize(direction[0], false);
+        if (std::abs(direction[0][0]) > std::abs(direction[0][1]))
+        {
+            direction[1][0] = -direction[0][2];
+            direction[1][1] = (Real)0;
+            direction[1][2] = +direction[0][0];
+        }
+        else
+        {
+            direction[1][0] = (Real)0;
+            direction[1][1] = +direction[0][2];
+            direction[1][2] = -direction[0][1];
+        }
+        Normalize(direction[1], false);
+        direction[2] = Cross(direction[0], direction[1]);
+
+        // Compute the maximum distance of the points from the line
+        // origin+t*direction[0].
+        Real maxDistance = (Real)0;
+        Real distance, dot;
+        extreme[2] = extreme[0];
+        for (i = 0; i < numVectors; ++i)
+        {
+            Vector3<Real> diff = v[i] - origin;
+            dot = Dot(direction[0], diff);
+            Vector3<Real> proj = diff - dot * direction[0];
+            distance = Length(proj, false);
+            if (distance > maxDistance)
+            {
+                maxDistance = distance;
+                extreme[2] = i;
+            }
+        }
+
+        if (maxDistance <= epsilon*maxRange)
+        {
+            // The points are (nearly) on the line origin+t*direction[0].
+            dimension = 1;
+            extreme[2] = extreme[1];
+            extreme[3] = extreme[1];
+            return;
+        }
+
+        // Test whether the vector set is (nearly) a planar polygon.  The
+        // point v[extreme[2]] is farthest from the line: origin + 
+        // t*direction[0].  The vector v[extreme[2]]-origin is not
+        // necessarily perpendicular to direction[0], so project out the
+        // direction[0] component so that the result is perpendicular to
+        // direction[0].
+        direction[1] = v[extreme[2]] - origin;
+        dot = Dot(direction[0], direction[1]);
+        direction[1] -= dot * direction[0];
+        Normalize(direction[1], false);
+
+        // We need direction[2] to span the orthogonal complement of
+        // {direction[0],direction[1]}.
+        direction[2] = Cross(direction[0], direction[1]);
+
+        // Compute the maximum distance of the points from the plane
+        // origin+t0*direction[0]+t1*direction[1].
+        maxDistance = (Real)0;
+        Real maxSign = (Real)0;
+        extreme[3] = extreme[0];
+        for (i = 0; i < numVectors; ++i)
+        {
+            Vector3<Real> diff = v[i] - origin;
+            distance = Dot(direction[2], diff);
+            Real sign = (distance >(Real)0 ? (Real)1 :
+                (distance < (Real)0 ? (Real)-1 : (Real)0));
+            distance = std::abs(distance);
+            if (distance > maxDistance)
+            {
+                maxDistance = distance;
+                maxSign = sign;
+                extreme[3] = i;
+            }
+        }
+
+        if (maxDistance <= epsilon * maxRange)
+        {
+            // The points are (nearly) on the plane origin+t0*direction[0]
+            // +t1*direction[1].
+            dimension = 2;
+            extreme[3] = extreme[2];
+            return;
+        }
+
+        dimension = 3;
+        extremeCCW = (maxSign > (Real)0);
+        return;
+    }
+}
+
+}
diff --git a/Engine/Plugins/Experimental/GeometryProcessing/Source/GeometryAlgorithms/Private/GTEngine/Mathematics/GteVector4.h b/Engine/Plugins/Experimental/GeometryProcessing/Source/GeometryAlgorithms/Private/GTEngine/Mathematics/GteVector4.h
new file mode 100644
--- /dev/null
+++ b/Engine/Plugins/Experimental/GeometryProcessing/Source/GeometryAlgorithms/Private/GTEngine/Mathematics/GteVector4.h
@@ -0,0 +1,186 @@
+// David Eberly, Geometric Tools, Redmond WA 98052
+// Copyright (c) 1998-2019
+// Distributed under the Boost Software License, Version 1.0.
+// http://www.boost.org/LICENSE_1_0.txt
+// http://www.geometrictools.com/License/Boost/LICENSE_1_0.txt
+// File Version: 3.0.2 (2018/10/04)
+
+#pragma once
+
+#include <GTEngine/Mathematics/GteVector3.h>
+
+namespace gte
+{
+
+// Template alias for convenience.
+template <typename Real>
+using Vector4 = Vector<4, Real>;
+
+// In GteVector3.h, the Vector3 Cross, UnitCross, and DotCross have a template
+// parameter N that should be 3 or 4.  The latter case supports affine vectors
+// in 4D (last component w = 0) when you want to use 4-tuples and 4x4 matrices
+// for affine algebra.  Thus, you may use those template functions for Vector4.
+
+// Compute the hypercross product using the formal determinant:
+//   hcross = det{{e0,e1,e2,e3},{x0,x1,x2,x3},{y0,y1,y2,y3},{z0,z1,z2,z3}}
+// where e0 = (1,0,0,0), e1 = (0,1,0,0), e2 = (0,0,1,0), e3 = (0,0,0,1),
+// v0 = (x0,x1,x2,x3), v1 = (y0,y1,y2,y3), and v2 = (z0,z1,z2,z3).
+template <typename Real>
+Vector4<Real> HyperCross(Vector4<Real> const& v0, Vector4<Real> const& v1,
+    Vector4<Real> const& v2);
+
+// Compute the normalized hypercross product.
+template <typename Real>
+Vector4<Real> UnitHyperCross(Vector4<Real> const& v0,
+    Vector4<Real> const& v1, Vector4<Real> const& v2, bool robust = false);
+
+// Compute Dot(HyperCross((x0,x1,x2,x3),(y0,y1,y2,y3),(z0,z1,z2,z3)),
+// (w0,w1,w2,w3)), where v0 = (x0,x1,x2,x3), v1 = (y0,y1,y2,y3),
+// v2 = (z0,z1,z2,z3), and v3 = (w0,w1,w2,w3).
+template <typename Real>
+Real DotHyperCross(Vector4<Real> const& v0, Vector4<Real> const& v1,
+    Vector4<Real> const& v2, Vector4<Real> const& v3);
+
+// Compute a right-handed orthonormal basis for the orthogonal complement
+// of the input vectors.  The function returns the smallest length of the
+// unnormalized vectors computed during the process.  If this value is nearly
+// zero, it is possible that the inputs are linearly dependent (within
+// numerical round-off errors).  On input, numInputs must be 1, 2, or 3 and
+// v[0] through v[numInputs-1] must be initialized.  On output, the vectors
+// v[0] through v[3] form an orthonormal set.
+template <typename Real>
+Real ComputeOrthogonalComplement(int numInputs, Vector4<Real>* v,
+    bool robust = false);
+
+
+template <typename Real>
+Vector4<Real> HyperCross(Vector4<Real> const& v0, Vector4<Real> const& v1,
+    Vector4<Real> const& v2)
+{
+    Real m01 = v0[0] * v1[1] - v0[1] * v1[0];  // x0*y1 - y0*x1
+    Real m02 = v0[0] * v1[2] - v0[2] * v1[0];  // x0*z1 - z0*x1
+    Real m03 = v0[0] * v1[3] - v0[3] * v1[0];  // x0*w1 - w0*x1
+    Real m12 = v0[1] * v1[2] - v0[2] * v1[1];  // y0*z1 - z0*y1
+    Real m13 = v0[1] * v1[3] - v0[3] * v1[1];  // y0*w1 - w0*y1
+    Real m23 = v0[2] * v1[3] - v0[3] * v1[2];  // z0*w1 - w0*z1
+    return Vector4<Real>
+    {
+        +m23*v2[1] - m13*v2[2] + m12*v2[3],  // +m23*y2 - m13*z2 + m12*w2
+            -m23*v2[0] + m03*v2[2] - m02*v2[3],  // -m23*x2 + m03*z2 - m02*w2
+            +m13*v2[0] - m03*v2[1] + m01*v2[3],  // +m13*x2 - m03*y2 + m01*w2
+            -m12*v2[0] + m02*v2[1] - m01*v2[2]   // -m12*x2 + m02*y2 - m01*z2
+    };
+}
+
+template <typename Real>
+Vector4<Real> UnitHyperCross(Vector4<Real> const& v0,
+    Vector4<Real> const& v1, Vector4<Real> const& v2, bool robust)
+{
+    Vector4<Real> unitHyperCross = HyperCross(v0, v1, v2);
+    Normalize(unitHyperCross, robust);
+    return unitHyperCross;
+}
+
+template <typename Real>
+Real DotHyperCross(Vector4<Real> const& v0, Vector4<Real> const& v1,
+    Vector4<Real> const& v2, Vector4<Real> const& v3)
+{
+    return Dot(HyperCross(v0, v1, v2), v3);
+}
+
+template <typename Real>
+Real ComputeOrthogonalComplement(int numInputs, Vector4<Real>* v, bool robust)
+{
+    if (numInputs == 1)
+    {
+        int maxIndex = 0;
+        Real maxAbsValue = std::abs(v[0][0]);
+        for (int i = 1; i < 4; ++i)
+        {
+            Real absValue = std::abs(v[0][i]);
+            if (absValue > maxAbsValue)
+            {
+                maxIndex = i;
+                maxAbsValue = absValue;
+            }
+        }
+
+        if (maxIndex < 2)
+        {
+            v[1][0] = -v[0][1];
+            v[1][1] = +v[0][0];
+            v[1][2] = (Real)0;
+            v[1][3] = (Real)0;
+        }
+        else if (maxIndex == 3)
+        {
+            // Generally, you can skip this clause and swap the last two
+            // components.  However, by swapping 2 and 3 in this case, we
+            // allow the function to work properly when the inputs are 3D
+            // vectors represented as 4D affine vectors (w = 0).
+            v[1][0] = (Real)0;
+            v[1][1] = +v[0][2];
+            v[1][2] = -v[0][1];
+            v[1][3] = (Real)0;
+        }
+        else
+        {
+            v[1][0] = (Real)0;
+            v[1][1] = (Real)0;
+            v[1][2] = -v[0][3];
+            v[1][3] = +v[0][2];
+        }
+
+        numInputs = 2;
+    }
+
+    if (numInputs == 2)
+    {
+        Real det[6] =
+        {
+            v[0][0] * v[1][1] - v[1][0] * v[0][1],
+            v[0][0] * v[1][2] - v[1][0] * v[0][2],
+            v[0][0] * v[1][3] - v[1][0] * v[0][3],
+            v[0][1] * v[1][2] - v[1][1] * v[0][2],
+            v[0][1] * v[1][3] - v[1][1] * v[0][3],
+            v[0][2] * v[1][3] - v[1][2] * v[0][3]
+        };
+
+        int maxIndex = 0;
+        Real maxAbsValue = std::abs(det[0]);
+        for (int i = 1; i < 6; ++i)
+        {
+            Real absValue = std::abs(det[i]);
+            if (absValue > maxAbsValue)
+            {
+                maxIndex = i;
+                maxAbsValue = absValue;
+            }
+        }
+
+        if (maxIndex == 0)
+        {
+            v[2] = { -det[4], +det[2], (Real)0, -det[0] };
+        }
+        else if (maxIndex <= 2)
+        {
+            v[2] = { +det[5], (Real)0, -det[2], +det[1] };
+        }
+        else
+        {
+            v[2] = { (Real)0, -det[5], +det[4], -det[3] };
+        }
+
+        numInputs = 3;
+    }
+
+    if (numInputs == 3)
+    {
+        v[3] = HyperCross(v[0], v[1], v[2]);
+        return Orthonormalize<4, Real>(4, v, robust);
+    }
+
+    return (Real)0;
+}
+
+}
diff --git a/Engine/Plugins/Experimental/GeometryProcessing/Source/GeometryAlgorithms/Private/GTEngine/Mathematics/GteVertexAttribute.h b/Engine/Plugins/Experimental/GeometryProcessing/Source/GeometryAlgorithms/Private/GTEngine/Mathematics/GteVertexAttribute.h
new file mode 100644
--- /dev/null
+++ b/Engine/Plugins/Experimental/GeometryProcessing/Source/GeometryAlgorithms/Private/GTEngine/Mathematics/GteVertexAttribute.h
@@ -0,0 +1,56 @@
+// David Eberly, Geometric Tools, Redmond WA 98052
+// Copyright (c) 1998-2019
+// Distributed under the Boost Software License, Version 1.0.
+// http://www.boost.org/LICENSE_1_0.txt
+// http://www.geometrictools.com/License/Boost/LICENSE_1_0.txt
+// File Version: 3.3.0 (2016/08/29)
+
+#pragma once
+
+#include <GTEngine/GTEngineDEF.h>
+#include <string>
+
+namespace gte
+{
+
+struct VertexAttribute
+{
+    inline VertexAttribute();
+    inline VertexAttribute(std::string const& inSemantic, void* inSource, size_t inStride);
+
+    // The 'semantic' string allows you to query for a specific vertex
+    // attribute and use the 'source' and 'stride' to access the data
+    // of the attribute.  For example, you might use the semantics
+    // "position" (px,py,pz), "normal" (nx,ny,nz), "tcoord" (texture
+    // coordinates (u,v)), "dpdu" (derivative of position with respect
+    // to u), or "dpdv" (derivative of position with respect to v) for
+    // mesh vertices.
+    //
+    // The source pointer must be 4-byte aligned.  The stride must be
+    // positive and a multiple of 4.  The pointer alignment constraint is
+    // guaranteed on 32-bit and 64-bit architectures.  The stride constraint
+    // is reasonable given that (usually) geometric attributes are usually
+    // arrays of 'float' or 'double'.
+
+    std::string semantic;
+    void* source;
+    size_t stride;
+};
+
+inline VertexAttribute::VertexAttribute()
+    :
+    semantic(""),
+    source(nullptr),
+    stride(0)
+{
+}
+
+inline VertexAttribute::VertexAttribute(std::string const& inSemantic, void* inSource, size_t inStride)
+    :
+    semantic(inSemantic),
+    source(inSource),
+    stride(inStride)
+{
+}
+
+}
diff --git a/Engine/Plugins/Experimental/GeometryProcessing/Source/GeometryAlgorithms/Private/GTEngine/Mathematics/GteVertexCollapseMesh.h b/Engine/Plugins/Experimental/GeometryProcessing/Source/GeometryAlgorithms/Private/GTEngine/Mathematics/GteVertexCollapseMesh.h
new file mode 100644
--- /dev/null
+++ b/Engine/Plugins/Experimental/GeometryProcessing/Source/GeometryAlgorithms/Private/GTEngine/Mathematics/GteVertexCollapseMesh.h
@@ -0,0 +1,546 @@
+// David Eberly, Geometric Tools, Redmond WA 98052
+// Copyright (c) 1998-2019
+// Distributed under the Boost Software License, Version 1.0.
+// http://www.boost.org/LICENSE_1_0.txt
+// http://www.geometrictools.com/License/Boost/LICENSE_1_0.txt
+// File Version: 3.0.1 (2018/10/04)
+
+#pragma once
+
+#include <GTEngine/Mathematics/GteVector3.h>
+#include <GTEngine/Mathematics/GteVETManifoldMesh.h>
+#include <GTEngine/Mathematics/GtePolygon2.h>
+#include <GTEngine/Mathematics/GteTriangulateEC.h>
+#include <GTEngine/LowLevel/GteLogger.h>
+#include <GTEngine/LowLevel/GteMinHeap.h>
+#include <set>
+
+namespace gte
+{
+
+template <typename Real>
+class VertexCollapseMesh
+{
+public:
+    // Construction.
+    VertexCollapseMesh(int numPositions, Vector3<Real> const* positions,
+        int numIndices, int const* indices);
+
+    // Decimate the mesh using vertex collapses
+    struct Record
+    {
+        // The index of the interior vertex that is removed from the mesh.
+        // The triangles adjacent to the vertex are 'removed' from the mesh.
+        // The polygon boundary of the adjacent triangles is triangulated and
+        // the new triangles are 'inserted' into the mesh.
+        int vertex;
+        std::vector<TriangleKey<true>> removed;
+        std::vector<TriangleKey<true>> inserted;
+    };
+
+    // Return 'true' when a vertex collapse occurs.  Once the function returns
+    // 'false', no more vertex collapses are allowed so you may then stop
+    // calling the function.  The implementation has several consistency tests
+    // that should not fail with a theoretically correct implementation.  If a
+    // test fails, the function returns 'false' and the record.vertex is set to
+    // the invalid integer 0x80000000.  When the Logger system is enabled, the
+    // failed tests are reported to any Logger listeners.
+    bool DoCollapse(Record& record);
+
+    // Access the current state of the mesh, whether the original built in the
+    // constructor or a decimated mesh during DoCollapse calls.
+    inline ETManifoldMesh const& GetMesh() const;
+
+private:
+    struct VCVertex : public VETManifoldMesh::Vertex
+    {
+        VCVertex(int v);
+        static std::shared_ptr<Vertex> Create(int v);
+
+        // The weight depends on the area of the triangles sharing the vertex
+        // and the lengths of the projections of the adjacent vertices onto
+        // the vertex normal line.  A side effect of the call is that the
+        // vertex normal is computed and stored.
+        Real ComputeWeight(Vector3<Real> const* positions);
+
+        Vector3<Real> normal;
+        bool isBoundary;
+    };
+
+    // The functions TriangulateLink and Collapsed return one of the
+    // enumerates described next.
+    //
+    // VCM_NO_MORE_ALLOWED:
+    //     Either the mesh has no more interior vertices or a collapse
+    //     will lead to a mesh fold-over or to a nonmanifold mesh.  The
+    //     returned value 'v' is invalid (0x80000000) and 'removed' and
+    //     'inserted' are empty.
+    //
+    // VCM_ALLOWED:
+    //     An interior vertex v has been removed.  This is allowed using the
+    //     following algorithm.  The vertex normal is the weighted average
+    //     of non-unit-length normals of triangles sharing v.  The weights
+    //     are the triangle areas.  The adjacent vertices are projected onto
+    //     a plane containing v and having normal equal to the vertex normal.
+    //     If the projection is a simple polygon in the plane, the collapse
+    //     is allowed.  The triangles sharing v are 'removed', the polygon is
+    //     triangulated, and the new triangles are 'inserted' into the mesh.
+    //
+    // VCM_DEFERRED:
+    //     If the projection polygon described in the previous case is not
+    //     simple (at least one pair of edges overlaps at some edge-interior
+    //     point), the collapse would produce a fold-over in the mesh.  We
+    //     do not collapse in this case.  It is possible that such a vertex
+    //     occurs in a later collapse as its neighbors are adjusted by
+    //     collapses.  When this case occurs, v is valid (even though the
+    //     collapse was not allowed) but 'removed' and 'inserted' are empty.
+    //
+    // VCM_UNEXPECTED_ERROR:
+    //     The code has several tests for conditions that are not expected
+    //     to occur for a theoretically correct implementation.  If you
+    //     receive this error, file a bug report and provide a data set
+    //     that caused the error.
+    enum
+    {
+        VCM_NO_MORE_ALLOWED,
+        VCM_ALLOWED,
+        VCM_DEFERRED,
+        VCM_UNEXPECTED_ERROR
+    };
+
+    int TriangulateLink(std::shared_ptr<VCVertex> const& vertex, std::vector<TriangleKey<true>>& removed,
+        std::vector<TriangleKey<true>>& inserted, std::vector<int>& linkVertices) const;
+
+    int Collapsed(std::vector<TriangleKey<true>> const& removed,
+        std::vector<TriangleKey<true>> const& inserted, std::vector<int> const& linkVertices);
+
+    int mNumPositions;
+    Vector3<Real> const* mPositions;
+    VETManifoldMesh mMesh;
+
+    MinHeap<int, Real> mMinHeap;
+    std::map<int, typename MinHeap<int, Real>::Record*> mHeapRecords;
+};
+
+
+template <typename Real>
+VertexCollapseMesh<Real>::VertexCollapseMesh(int numPositions,
+    Vector3<Real> const* positions, int numIndices, int const* indices)
+    :
+    mNumPositions(numPositions),
+    mPositions(positions),
+    mMesh(VCVertex::Create)
+{
+    if (numPositions <= 0 || !positions || numIndices < 3 || !indices)
+    {
+        mNumPositions = 0;
+        mPositions = nullptr;
+        return;
+    }
+
+    // Build the manifold mesh from the inputs.
+    int numTriangles = numIndices / 3;
+    int const* current = indices;
+    for (int t = 0; t < numTriangles; ++t)
+    {
+        int v0 = *current++;
+        int v1 = *current++;
+        int v2 = *current++;
+        mMesh.Insert(v0, v1, v2);
+    }
+
+    // Locate the vertices (if any) on the mesh boundary.
+    auto const& vmap = mMesh.GetVertices();
+    for (auto const& eelement : mMesh.GetEdges())
+    {
+        auto edge = eelement.second;
+        if (!edge->T[1].lock())
+        {
+            for (int i = 0; i < 2; ++i)
+            {
+                auto velement = vmap.find(edge->V[i]);
+                auto vertex = std::static_pointer_cast<VCVertex>(velement->second);
+                vertex->isBoundary = true;
+            }
+        }
+    }
+
+    // Build the priority queue of weights for the interior vertices.
+    mMinHeap.Reset((int)vmap.size());
+    for (auto const& velement : vmap)
+    {
+        auto vertex = std::static_pointer_cast<VCVertex>(velement.second);
+
+        Real weight;
+        if (vertex->isBoundary)
+        {
+            weight = std::numeric_limits<Real>::max();
+        }
+        else
+        {
+            weight = vertex->ComputeWeight(mPositions);
+        }
+
+        auto record = mMinHeap.Insert(velement.first, weight);
+        mHeapRecords.insert(std::make_pair(velement.first, record));
+    }
+}
+
+template <typename Real>
+bool VertexCollapseMesh<Real>::DoCollapse(Record& record)
+{
+    record.vertex = 0x80000000;
+    record.removed.clear();
+    record.inserted.clear();
+
+    if (mNumPositions == 0)
+    {
+        // The constructor failed, so there is nothing to collapse.
+        return false;
+    }
+
+    while (mMinHeap.GetNumElements() > 0)
+    {
+        int v = -1;
+        Real weight = std::numeric_limits<Real>::max();
+        mMinHeap.GetMinimum(v, weight);
+        if (weight == std::numeric_limits<Real>::max())
+        {
+            // There are no more interior vertices to collapse.
+            return false;
+        }
+
+        auto const& vmap = mMesh.GetVertices();
+        auto velement = vmap.find(v);
+        if (velement == vmap.end())
+        {
+            LogError("Unexpected condition.");
+            return false;
+        }
+
+        auto vertex = std::static_pointer_cast<VCVertex>(velement->second);
+        std::vector<TriangleKey<true>> removed, inserted;
+        std::vector<int> linkVertices;
+        int result = TriangulateLink(vertex, removed, inserted, linkVertices);
+        if (result == VCM_UNEXPECTED_ERROR)
+        {
+            return false;
+        }
+
+        if (result == VCM_ALLOWED)
+        {
+            result = Collapsed(removed, inserted, linkVertices);
+            if (result == VCM_UNEXPECTED_ERROR)
+            {
+                return false;
+            }
+
+            if (result == VCM_ALLOWED)
+            {
+                // Remove the vertex and associated weight.
+                mMinHeap.Remove(v, weight);
+                mHeapRecords.erase(v);
+
+                // Update the weights of the link vertices.
+                for (auto vlink : linkVertices)
+                {
+                    velement = vmap.find(vlink);
+                    if (velement == vmap.end())
+                    {
+                        LogError("Unexpected condition.");
+                        return false;
+                    }
+
+                    vertex = std::static_pointer_cast<VCVertex>(velement->second);
+                    if (!vertex->isBoundary)
+                    {
+                        auto iter = mHeapRecords.find(vlink);
+                        if (iter == mHeapRecords.end())
+                        {
+                            LogError("Unexpected condition.");
+                            return false;
+                        }
+
+                        weight = vertex->ComputeWeight(mPositions);
+                        mMinHeap.Update(iter->second, weight);
+                    }
+                }
+
+                record.vertex = v;
+                record.removed = std::move(removed);
+                record.inserted = std::move(inserted);
+                return true;
+            }
+            // else:  result == VCM_DEFERRED
+        }
+
+        // To get here, result must be VCM_DEFERRED.  The vertex collapse
+        // would cause mesh fold-over.  Temporarily set the edge weight to
+        // infinity.  After removal of other triangles, the vertex weight
+        // will be updated to a finite value and the vertex possibly can be
+        // removed at that time.
+        auto iter = mHeapRecords.find(v);
+        if (iter == mHeapRecords.end())
+        {
+            LogError("Unexpected condition.");
+            return false;
+        }
+        mMinHeap.Update(iter->second, std::numeric_limits<Real>::max());
+    }
+
+    // We do not expect to reach this line of code, even for a closed mesh.
+    // However, the compiler does not know this, yet requires a return value.
+    return false;
+}
+
+template <typename Real> inline
+ETManifoldMesh const& VertexCollapseMesh<Real>::GetMesh() const
+{
+    return mMesh;
+}
+
+template <typename Real>
+int VertexCollapseMesh<Real>::TriangulateLink(std::shared_ptr<VCVertex> const& vertex,
+    std::vector<TriangleKey<true>>& removed, std::vector<TriangleKey<true>>& inserted,
+    std::vector<int>& linkVertices) const
+{
+    // Create the (CCW) polygon boundary of the link of the vertex.  The
+    // incoming vertex is interior, so the number of triangles sharing the
+    // vertex is equal to the number of vertices of the polygon.  A
+    // precondition of the function call is that the vertex normal has already
+    // been computed.
+
+    // Get the edges of the link that are opposite the incoming vertex.
+    int const numVertices = static_cast<int>(vertex->TAdjacent.size());
+    removed.resize(numVertices);
+    int j = 0;
+    std::map<int, int> edgeMap;
+    for (auto tri : vertex->TAdjacent)
+    {
+        for (int i = 0; i < 3; ++i)
+        {
+            if (tri->V[i] == vertex->V)
+            {
+                edgeMap.insert(std::make_pair(tri->V[(i + 1) % 3], tri->V[(i + 2) % 3]));
+                break;
+            }
+        }
+        removed[j++] = TriangleKey<true>(tri->V[0], tri->V[1], tri->V[2]);
+    }
+    if (edgeMap.size() != vertex->TAdjacent.size())
+    {
+        LogError("Unexpected condition.");
+        return VCM_UNEXPECTED_ERROR;
+    }
+
+    // Connect the edges into a polygon.
+    linkVertices.resize(numVertices);
+    auto iter = edgeMap.begin();
+    for (int i = 0; i < numVertices; ++i)
+    {
+        linkVertices[i] = iter->first;
+        iter = edgeMap.find(iter->second);
+        if (iter == edgeMap.end())
+        {
+            LogError("Unexpected condition.");
+            return VCM_UNEXPECTED_ERROR;
+        }
+    }
+    if (iter->first != linkVertices[0])
+    {
+        LogError("Unexpected condition.");
+        return VCM_UNEXPECTED_ERROR;
+    }
+
+    // Project the polygon onto the plane containing the incoming vertex and
+    // having the vertex normal.  The projected polygon is computed so that
+    // the incoming vertex is projected to (0,0).
+    Vector3<Real> center = mPositions[vertex->V];
+    Vector3<Real> basis[3];
+    basis[0] = vertex->normal;
+    ComputeOrthogonalComplement(1, basis);
+    std::vector<Vector2<Real>> projected(numVertices);
+    std::vector<int> indices(numVertices);
+    for (int i = 0; i < numVertices; ++i)
+    {
+        Vector3<Real> diff = mPositions[linkVertices[i]] - center;
+        projected[i][0] = Dot(basis[1], diff);
+        projected[i][1] = Dot(basis[2], diff);
+        indices[i] = i;
+    }
+
+    // The polygon must be simple in order to triangulate it.
+    Polygon2<Real> polygon(projected.data(), numVertices, indices.data(), true);
+    if (polygon.IsSimple())
+    {
+        TriangulateEC<Real, Real> triangulator(numVertices, projected.data());
+        triangulator();
+        auto const& triangles = triangulator.GetTriangles();
+        if (triangles.size() == 0)
+        {
+            LogError("Unexpected condition.");
+            return VCM_UNEXPECTED_ERROR;
+        }
+
+        int const numTriangles = static_cast<int>(triangles.size());
+        inserted.resize(numTriangles);
+        for (int t = 0; t < numTriangles; ++t)
+        {
+            inserted[t] = TriangleKey<true>(
+                linkVertices[triangles[t][0]],
+                linkVertices[triangles[t][1]],
+                linkVertices[triangles[t][2]]);
+        }
+        return VCM_ALLOWED;
+    }
+    else
+    {
+        return VCM_DEFERRED;
+    }
+}
+
+template <typename Real>
+int VertexCollapseMesh<Real>::Collapsed(std::vector<TriangleKey<true>> const& removed,
+    std::vector<TriangleKey<true>> const& inserted, std::vector<int> const& linkVertices)
+{
+    // The triangles that were disconnected from the link edges are guaranteed
+    // to allow manifold reconnection to 'inserted' triangles.  On the
+    // insertion, each diagonal of the link becomes a mesh edge and shares two
+    // (link) triangles.  It is possible that the mesh already contains the
+    // (diagonal) edge, which will lead to a nonmanifold connection, which we
+    // cannot allow.  The following code traps this condition and restores the
+    // mesh to its state before the 'Remove(...)' call.
+    bool isCollapsible = true;
+    auto const& emap = mMesh.GetEdges();
+    std::set<EdgeKey<false>> edges;
+    for (auto const& tri : inserted)
+    {
+        for (int k0 = 2, k1 = 0; k1 < 3; k0 = k1++)
+        {
+            EdgeKey<false> edge(tri.V[k0], tri.V[k1]);
+            if (edges.find(edge) == edges.end())
+            {
+                edges.insert(edge);
+            }
+            else
+            {
+                // The edge has been visited twice, so it is a diagonal of
+                // the link.
+
+                auto eelement = emap.find(edge);
+                if (eelement != emap.end())
+                {
+                    if (eelement->second->T[1].lock())
+                    {
+                        // The edge will not allow a manifold connection.
+                        isCollapsible = false;
+                        break;
+                    }
+                }
+
+                edges.erase(edge);
+            }
+        };
+
+        if (!isCollapsible)
+        {
+            return VCM_DEFERRED;
+        }
+    }
+
+    // Remove the old triangle neighborhood, which will lead to the vertex
+    // itself being removed from the mesh.
+    for (auto tri : removed)
+    {
+        mMesh.Remove(tri.V[0], tri.V[1], tri.V[2]);
+    }
+
+    // Insert the new triangulation.
+    for (auto const& tri : inserted)
+    {
+        mMesh.Insert(tri.V[0], tri.V[1], tri.V[2]);
+    }
+
+    // If the Remove(...) calls remove a boundary vertex that is in the link
+    // vertices, the Insert(...) calls will insert the boundary vertex again.
+    // We must re-tag those boundary vertices.
+    auto const& vmap = mMesh.GetVertices();
+    size_t const numVertices = linkVertices.size();
+    for (size_t i0 = numVertices - 1, i1 = 0; i1 < numVertices; i0 = i1++)
+    {
+        EdgeKey<false> ekey(linkVertices[i0], linkVertices[i1]);
+        auto eelement = emap.find(ekey);
+        if (eelement == emap.end())
+        {
+            LogError("Unexpected condition.");
+            return VCM_UNEXPECTED_ERROR;
+        }
+
+        auto edge = eelement->second;
+        if (!edge)
+        {
+            LogError("Unexpected condition.");
+            return VCM_UNEXPECTED_ERROR;
+        }
+
+        if (edge->T[0].lock() && !edge->T[1].lock())
+        {
+            for (int k = 0; k < 2; ++k)
+            {
+                auto velement = vmap.find(edge->V[k]);
+                if (velement == vmap.end())
+                {
+                    LogError("Unexpected condition.");
+                    return VCM_UNEXPECTED_ERROR;
+                }
+
+                auto vertex = std::static_pointer_cast<VCVertex>(velement->second);
+                vertex->isBoundary = true;
+            }
+        }
+    }
+
+    return VCM_ALLOWED;
+}
+
+template <typename Real>
+VertexCollapseMesh<Real>::VCVertex::VCVertex(int v)
+    :
+    VETManifoldMesh::Vertex(v),
+    normal(Vector3<Real>::Zero()),
+    isBoundary(false)
+{
+}
+
+template <typename Real>
+std::shared_ptr<VETManifoldMesh::Vertex> VertexCollapseMesh<Real>::VCVertex::Create(int v)
+{
+    return std::make_shared<VCVertex>(v);
+}
+
+template <typename Real>
+Real VertexCollapseMesh<Real>::VCVertex::ComputeWeight(Vector3<Real> const* positions)
+{
+    Real weight = (Real)0;
+
+    normal = { (Real)0, (Real)0, (Real)0 };
+    for (auto const& tri : TAdjacent)
+    {
+        Vector3<Real> E0 = positions[tri->V[1]] - positions[tri->V[0]];
+        Vector3<Real> E1 = positions[tri->V[2]] - positions[tri->V[0]];
+        Vector3<Real> N = Cross(E0, E1);
+        normal += N;
+        weight += Length(N);
+    }
+    Normalize(normal);
+
+    for (int index : VAdjacent)
+    {
+        Vector3<Real> diff = positions[index] - positions[V];
+        weight += std::abs(Dot(normal, diff));
+    }
+
+    return weight;
+}
+
+}
diff --git a/Engine/Plugins/Experimental/GeometryProcessing/Source/MeshConversion/MeshConversion.Build.cs b/Engine/Plugins/Experimental/GeometryProcessing/Source/MeshConversion/MeshConversion.Build.cs
new file mode 100644
--- /dev/null
+++ b/Engine/Plugins/Experimental/GeometryProcessing/Source/MeshConversion/MeshConversion.Build.cs
@@ -0,0 +1,28 @@
+// Copyright 1998-2019 Epic Games, Inc. All Rights Reserved.
+
+using UnrealBuildTool;
+
+public class MeshConversion : ModuleRules
+{	
+	public MeshConversion(ReadOnlyTargetRules Target) : base(Target)
+	{
+        PCHUsage = ModuleRules.PCHUsageMode.UseExplicitOrSharedPCHs;
+
+        PublicDependencyModuleNames.AddRange(
+			new string[] {
+                "Core",
+                "MeshDescription",
+				"GeometricObjects",
+				"DynamicMesh"
+            }
+		);
+
+        PrivateDependencyModuleNames.AddRange(
+            new string[]
+            {
+                "CoreUObject",
+                "Engine"
+			}
+		);
+    }
+}
diff --git a/Engine/Plugins/Experimental/GeometryProcessing/Source/MeshConversion/Private/DynamicMeshToMeshDescription.cpp b/Engine/Plugins/Experimental/GeometryProcessing/Source/MeshConversion/Private/DynamicMeshToMeshDescription.cpp
new file mode 100644
--- /dev/null
+++ b/Engine/Plugins/Experimental/GeometryProcessing/Source/MeshConversion/Private/DynamicMeshToMeshDescription.cpp
@@ -0,0 +1,220 @@
+// Copyright 1998-2019 Epic Games, Inc. All Rights Reserved. 
+
+#include "DynamicMeshToMeshDescription.h"
+#include "MeshAttributes.h"
+#include "DynamicMeshAttributeSet.h"
+#include "DynamicMeshOverlay.h"
+#include "MeshDescriptionBuilder.h"
+#include "MeshNormals.h"
+
+
+void FDynamicMeshToMeshDescription::Convert(const FDynamicMesh3* MeshIn, FMeshDescription& MeshOut)
+{
+	if (MeshIn->HasAttributes())
+	{
+		Convert_SharedInstances(MeshIn, MeshOut);
+	}
+	else
+	{
+		Convert_NoAttributes(MeshIn, MeshOut);
+	}
+}
+
+
+
+
+
+void FDynamicMeshToMeshDescription::Convert_NoAttributes(const FDynamicMesh3* MeshIn, FMeshDescription& MeshOut)
+{
+	MeshOut.Empty();
+
+	FMeshDescriptionBuilder Builder;
+	Builder.SetMeshDescription(&MeshOut);
+
+	bool bCopyGroupToPolyGroup = false;
+	if (bSetPolyGroups && MeshIn->HasTriangleGroups())
+	{
+		Builder.EnablePolyGroups();
+		bCopyGroupToPolyGroup = true;
+	}
+
+	// create vertices
+	TArray<FVertexID> MapV; 
+	MapV.SetNum(MeshIn->MaxVertexID());
+	for (int VertID : MeshIn->VertexIndicesItr())
+	{
+		MapV[VertID] = Builder.AppendVertex(MeshIn->GetVertex(VertID));
+	}
+
+	FMeshNormals VertexNormals(MeshIn);
+	VertexNormals.ComputeVertexNormals();
+
+	FPolygonGroupID AllGroupID = Builder.AppendPolygonGroup();
+
+	// create new instances when seen
+	TMap<FIndex3i, FVertexInstanceID> InstanceList;
+	for (int TriID : MeshIn->TriangleIndicesItr())
+	{
+		FIndex3i Triangle = MeshIn->GetTriangle(TriID);
+		FIndex3i UVTriangle(-1, -1, -1);
+		FIndex3i NormalTriangle = Triangle;
+		FVertexInstanceID InstanceTri[3];
+		for (int j = 0; j < 3; ++j)
+		{
+			FIndex3i InstanceElem(Triangle[j], UVTriangle[j], NormalTriangle[j]);
+			if (InstanceList.Contains(InstanceElem) == false)
+			{
+				FVertexInstanceID NewInstanceID = Builder.AppendInstance(MapV[Triangle[j]]);
+				InstanceList.Add(InstanceElem, NewInstanceID);
+				Builder.SetInstance(NewInstanceID, FVector2f::Zero(), VertexNormals[Triangle[j]]);
+			}
+			InstanceTri[j] = InstanceList[InstanceElem];
+		}
+
+		FPolygonID NewPolygonID = Builder.AppendTriangle(InstanceTri[0], InstanceTri[1], InstanceTri[2], AllGroupID);
+		if (bCopyGroupToPolyGroup)
+		{
+			Builder.SetPolyGroupID(NewPolygonID, MeshIn->GetTriangleGroup(TriID));
+		}
+	}
+
+	MeshOut.ComputeTangentsAndNormals(EComputeNTBsOptions::Tangents);
+}
+
+
+
+
+
+
+
+void FDynamicMeshToMeshDescription::Convert_SharedInstances(const FDynamicMesh3* MeshIn, FMeshDescription& MeshOut)
+{
+	const FDynamicMeshUVOverlay* UVOverlay = MeshIn->HasAttributes() ? MeshIn->Attributes()->PrimaryUV() : nullptr;
+	const FDynamicMeshNormalOverlay* NormalOverlay = MeshIn->HasAttributes() ? MeshIn->Attributes()->PrimaryNormals() : nullptr;
+
+	MeshOut.Empty();
+
+	FMeshDescriptionBuilder Builder;
+	Builder.SetMeshDescription(&MeshOut);
+
+	bool bCopyGroupToPolyGroup = false;
+	if (bSetPolyGroups && MeshIn->HasTriangleGroups())
+	{
+		Builder.EnablePolyGroups();
+		bCopyGroupToPolyGroup = true;
+	}
+
+	// create vertices
+	TArray<FVertexID> MapV; MapV.SetNum(MeshIn->MaxVertexID());
+	for (int VertID : MeshIn->VertexIndicesItr())
+	{
+		MapV[VertID] = Builder.AppendVertex(MeshIn->GetVertex(VertID));
+	}
+
+
+	FPolygonGroupID AllGroupID = Builder.AppendPolygonGroup();
+
+
+	// create new instances when seen
+	TMap<FIndex3i, FVertexInstanceID> InstanceList;
+	for (int TriID : MeshIn->TriangleIndicesItr())
+	{
+		// @todo support additional overlays (requires IndexNi...)
+		FIndex3i Triangle = MeshIn->GetTriangle(TriID);
+		FIndex3i UVTriangle = (UVOverlay != nullptr) ? UVOverlay->GetTriangle(TriID) : FIndex3i(-1, -1, -1);
+		FIndex3i NormalTriangle = (NormalOverlay != nullptr) ? NormalOverlay->GetTriangle(TriID) : FIndex3i(-1, -1, -1);
+		FVertexInstanceID InstanceTri[3];
+		for (int j = 0; j < 3; ++j)
+		{
+			FIndex3i InstanceElem(Triangle[j], UVTriangle[j], NormalTriangle[j]);
+			if (InstanceList.Contains(InstanceElem) == false)
+			{
+				FVertexInstanceID NewInstanceID = Builder.AppendInstance(MapV[Triangle[j]]);
+				InstanceList.Add(InstanceElem, NewInstanceID);
+				Builder.SetInstance(NewInstanceID,
+					(UVTriangle[j] == -1 || UVOverlay == nullptr) ? FVector2f::Zero() : UVOverlay->GetElement(UVTriangle[j]),
+					(NormalTriangle[j] == -1 || NormalOverlay == nullptr) ? FVector3f::UnitY() : NormalOverlay->GetElement(NormalTriangle[j]));
+			}
+			InstanceTri[j] = InstanceList[InstanceElem];
+		}
+
+		FPolygonID NewPolygonID = Builder.AppendTriangle(InstanceTri[0], InstanceTri[1], InstanceTri[2], AllGroupID);
+
+		if (bCopyGroupToPolyGroup)
+		{
+			Builder.SetPolyGroupID(NewPolygonID, MeshIn->GetTriangleGroup(TriID));
+		}
+	}
+}
+
+
+
+
+
+void FDynamicMeshToMeshDescription::Convert_NoSharedInstances(const FDynamicMesh3* MeshIn, FMeshDescription& MeshOut)
+{
+	const FDynamicMeshUVOverlay* UVOverlay = MeshIn->HasAttributes() ? MeshIn->Attributes()->PrimaryUV() : nullptr;
+	const FDynamicMeshNormalOverlay* NormalOverlay = MeshIn->HasAttributes() ? MeshIn->Attributes()->PrimaryNormals() : nullptr;
+
+	MeshOut.Empty();
+
+	FMeshDescriptionBuilder Builder;
+	Builder.SetMeshDescription(&MeshOut);
+
+	bool bCopyGroupToPolyGroup = false;
+	if (bSetPolyGroups && MeshIn->HasTriangleGroups())
+	{
+		Builder.EnablePolyGroups();
+		bCopyGroupToPolyGroup = true;
+	}
+
+	TArray<FVertexID> MapV; MapV.SetNum(MeshIn->MaxVertexID());
+	for (int VertID : MeshIn->VertexIndicesItr())
+	{
+		MapV[VertID] = Builder.AppendVertex(MeshIn->GetVertex(VertID));
+	}
+
+	FPolygonGroupID AllGroupID = Builder.AppendPolygonGroup();
+
+	FVertexID TriVertices[3];
+	FVector2D TriUVs[3];
+	FVector TriNormals[3];
+	for (int TriID : MeshIn->TriangleIndicesItr())
+	{
+		FIndex3i Triangle = MeshIn->GetTriangle(TriID);
+
+		FVector2D* UseUVs = nullptr;
+		if (UVOverlay != nullptr)
+		{
+			FIndex3i UVTriangle = UVOverlay->GetTriangle(TriID);
+			TriUVs[0] = UVOverlay->GetElement(UVTriangle[0]);
+			TriUVs[1] = UVOverlay->GetElement(UVTriangle[1]);
+			TriUVs[2] = UVOverlay->GetElement(UVTriangle[2]);
+			UseUVs = TriUVs;
+		}
+
+		FVector* UseNormals = nullptr;
+		if (NormalOverlay != nullptr)
+		{
+			FIndex3i NormalTriangle = NormalOverlay->GetTriangle(TriID);
+			TriNormals[0] = NormalOverlay->GetElement(NormalTriangle[0]);
+			TriNormals[1] = NormalOverlay->GetElement(NormalTriangle[1]);
+			TriNormals[2] = NormalOverlay->GetElement(NormalTriangle[2]);
+			UseNormals = TriNormals;
+		}
+
+		TriVertices[0] = MapV[Triangle[0]];
+		TriVertices[1] = MapV[Triangle[1]];
+		TriVertices[2] = MapV[Triangle[2]];
+
+		FPolygonID NewPolygonID = Builder.AppendTriangle(TriVertices, AllGroupID, UseUVs, UseNormals);
+
+		if (bCopyGroupToPolyGroup)
+		{
+			Builder.SetPolyGroupID(NewPolygonID, MeshIn->GetTriangleGroup(TriID));
+		}
+	}
+
+	// set to hard edges
+	//PolyMeshIn->SetAllEdgesHardness(true);
+}
diff --git a/Engine/Plugins/Experimental/GeometryProcessing/Source/MeshConversion/Private/MeshDescriptionBuilder.cpp b/Engine/Plugins/Experimental/GeometryProcessing/Source/MeshConversion/Private/MeshDescriptionBuilder.cpp
new file mode 100644
--- /dev/null
+++ b/Engine/Plugins/Experimental/GeometryProcessing/Source/MeshConversion/Private/MeshDescriptionBuilder.cpp
@@ -0,0 +1,280 @@
+// Copyright 1998-2019 Epic Games, Inc. All Rights Reserved.
+
+#include "MeshDescriptionBuilder.h"
+#include "MeshAttributes.h"
+#include "VectorTypes.h"
+#include "BoxTypes.h"
+
+#include "DynamicMesh3.h"
+#include "DynamicMeshAttributeSet.h"
+
+namespace ExtendedMeshAttribute
+{
+	const FName PolyTriGroups("PolyTriGroups");
+}
+
+
+
+void FMeshDescriptionBuilder::SetMeshDescription(FMeshDescription* Description)
+{
+	this->MeshDescription = Description;
+	this->VertexPositions = 
+		MeshDescription->VertexAttributes().GetAttributesRef<FVector>(MeshAttribute::Vertex::Position);
+	this->InstanceUVs =
+		MeshDescription->VertexInstanceAttributes().GetAttributesRef<FVector2D>(MeshAttribute::VertexInstance::TextureCoordinate);
+	this->InstanceNormals =
+		MeshDescription->VertexInstanceAttributes().GetAttributesRef<FVector>(MeshAttribute::VertexInstance::Normal);
+}
+
+
+void FMeshDescriptionBuilder::EnablePolyGroups()
+{
+	PolyGroups =
+		MeshDescription->PolygonAttributes().GetAttributesRef<int>(ExtendedMeshAttribute::PolyTriGroups);
+	if (PolyGroups.IsValid() == false)
+	{
+		MeshDescription->PolygonAttributes().RegisterAttribute<int>(
+			ExtendedMeshAttribute::PolyTriGroups, 1, 0, EMeshAttributeFlags::AutoGenerated);
+		PolyGroups =
+			MeshDescription->PolygonAttributes().GetAttributesRef<int>(ExtendedMeshAttribute::PolyTriGroups);
+		check(PolyGroups.IsValid());
+	}
+}
+
+
+FVertexID FMeshDescriptionBuilder::AppendVertex(const FVector& Position)
+{
+	FVertexID VertexID = MeshDescription->CreateVertex();
+	VertexPositions.Set(VertexID, Position);
+	return VertexID;
+}
+
+
+FPolygonGroupID FMeshDescriptionBuilder::AppendPolygonGroup()
+{
+	return MeshDescription->CreatePolygonGroup();
+}
+
+
+
+FPolygonID FMeshDescriptionBuilder::AppendTriangle(const FVertexID& Vertex0, const FVertexID& Vertex1, const FVertexID& Vertex2, const FPolygonGroupID& PolygonGroup)
+{
+	TempBuffer.SetNum(3, false);
+	TempBuffer[0] = Vertex0;
+	TempBuffer[1] = Vertex1;
+	TempBuffer[2] = Vertex2;
+	return AppendPolygon(TempBuffer, PolygonGroup);
+}
+
+
+FVertexInstanceID FMeshDescriptionBuilder::AppendInstance(const FVertexID& VertexID)
+{
+	return MeshDescription->CreateVertexInstance(VertexID);
+}
+
+
+
+void FMeshDescriptionBuilder::SetInstance(const FVertexInstanceID& InstanceID, const FVector2D& InstanceUV, const FVector& InstanceNormal)
+{
+	if (InstanceUVs.IsValid())
+	{
+		InstanceUVs.Set(InstanceID, InstanceUV); 
+	}
+	if (InstanceNormals.IsValid())
+	{
+		InstanceNormals.Set(InstanceID, InstanceNormal);
+	}
+}
+
+
+
+
+FPolygonID FMeshDescriptionBuilder::AppendTriangle(const FVertexID* Triangle, const FPolygonGroupID& PolygonGroup, 
+	const FVector2D * VertexUVs, const FVector* VertexNormals)
+{
+	TempBuffer.SetNum(3, false);
+	TempBuffer[0] = Triangle[0];
+	TempBuffer[1] = Triangle[1];
+	TempBuffer[2] = Triangle[2];
+
+	const TArray<FVector2D> * UseUVs = nullptr;
+	if (VertexUVs != nullptr) 
+	{
+		UVBuffer.SetNum(3, false);
+		UVBuffer[0] = VertexUVs[0];
+		UVBuffer[1] = VertexUVs[1];
+		UVBuffer[2] = VertexUVs[2];
+		UseUVs = &UVBuffer;
+	}
+
+	const TArray<FVector> * UseNormals = nullptr;
+	if (VertexNormals != nullptr)
+	{
+		NormalBuffer.SetNum(3, false);
+		NormalBuffer[0] = VertexNormals[0];
+		NormalBuffer[1] = VertexNormals[1];
+		NormalBuffer[2] = VertexNormals[2];
+		UseNormals = &NormalBuffer;
+	}
+
+	return AppendPolygon(TempBuffer, PolygonGroup, UseUVs, UseNormals);
+}
+
+
+FPolygonID FMeshDescriptionBuilder::AppendPolygon(const TArray<FVertexID>& Vertices, const FPolygonGroupID& PolygonGroup, 
+	const TArray<FVector2D> * VertexUVs, const TArray<FVector>* VertexNormals)
+{
+	int NumVertices = Vertices.Num();
+	TArray<FVertexInstanceID> Polygon; 
+	Polygon.Reserve(NumVertices);
+	for (int j = 0; j < NumVertices; ++j) 
+	{
+		FVertexInstanceID VertexInstance = MeshDescription->CreateVertexInstance(Vertices[j]);
+		Polygon.Add(VertexInstance);
+
+		if (VertexUVs != nullptr)
+		{
+			InstanceUVs.Set(VertexInstance, (*VertexUVs)[j]);
+		}
+		if (VertexNormals != nullptr)
+		{
+			InstanceNormals.Set(VertexInstance, (*VertexNormals)[j]);
+		}
+
+	}
+
+	const FPolygonID NewPolygonID = MeshDescription->CreatePolygon(PolygonGroup, Polygon);
+
+	// compute triangulation
+	FMeshPolygon& NewPolygon = MeshDescription->GetPolygon(NewPolygonID);
+	MeshDescription->ComputePolygonTriangulation(NewPolygonID, NewPolygon.Triangles);
+
+	return NewPolygonID;
+}
+
+
+
+
+FPolygonID FMeshDescriptionBuilder::AppendTriangle(const FVertexInstanceID& Instance0, const FVertexInstanceID& Instance1, const FVertexInstanceID& Instance2, const FPolygonGroupID& PolygonGroup)
+{
+	TArray<FVertexInstanceID> Polygon;
+	Polygon.Add(Instance0);
+	Polygon.Add(Instance1);
+	Polygon.Add(Instance2);
+
+	const FPolygonID NewPolygonID = MeshDescription->CreatePolygon(PolygonGroup, Polygon);
+
+	// compute triangulation
+	FMeshPolygon& NewPolygon = MeshDescription->GetPolygon(NewPolygonID);
+	MeshDescription->ComputePolygonTriangulation(NewPolygonID, NewPolygon.Triangles);
+
+	return NewPolygonID;
+}
+
+
+
+
+
+void FMeshDescriptionBuilder::SetPolyGroupID(const FPolygonID& PolygonID, int GroupID)
+{
+	PolyGroups.Set(PolygonID, 0, GroupID);
+}
+
+
+
+
+
+void FMeshDescriptionBuilder::AppendMesh(const FDynamicMesh3* Mesh, bool bSetPolyGroups)
+{
+	// create vertices
+	TArray<FVertexID> MapV;
+	MapV.SetNum(Mesh->MaxVertexID());
+	for (int VertID : Mesh->VertexIndicesItr())
+	{
+		FVector3d v = Mesh->GetVertex(VertID);
+		MapV[VertID] = AppendVertex(v);
+	}
+
+	bool bDoTransferPolyGroups = false;
+	if (bSetPolyGroups)
+	{
+		EnablePolyGroups();
+		bDoTransferPolyGroups = Mesh->HasTriangleGroups();
+	}
+
+	FPolygonGroupID AllGroupID = AppendPolygonGroup();
+
+	const FDynamicMeshUVOverlay* UVOverlay = Mesh->HasAttributes() ? Mesh->Attributes()->PrimaryUV() : nullptr;
+	const FDynamicMeshNormalOverlay* NormalOverlay = Mesh->HasAttributes() ? Mesh->Attributes()->PrimaryNormals() : nullptr;
+	//FDynamicMeshUVOverlay* UVOverlay = nullptr;
+	//FDynamicMeshNormalOverlay* NormalOverlay = nullptr;
+
+
+	// create new instances when seen
+	TMap<FIndex3i, FVertexInstanceID> InstanceList;
+	for (int TriID : Mesh->TriangleIndicesItr())
+	{
+		// @todo support additional overlays (requires IndexNi...)
+		FIndex3i Triangle = Mesh->GetTriangle(TriID);
+		FIndex3i UVTriangle = (UVOverlay != nullptr) ? UVOverlay->GetTriangle(TriID) : FIndex3i(-1, -1, -1);
+		FIndex3i NormalTriangle = (NormalOverlay != nullptr) ? NormalOverlay->GetTriangle(TriID) : FIndex3i(-1, -1, -1);
+		FVertexInstanceID InstanceTri[3];
+		for (int j = 0; j < 3; ++j)
+		{
+			FIndex3i InstanceElem(Triangle[j], UVTriangle[j], NormalTriangle[j]);
+			if (InstanceList.Contains(InstanceElem) == false)
+			{
+				FVertexInstanceID NewInstanceID = AppendInstance(MapV[Triangle[j]]);
+				InstanceList.Add(InstanceElem, NewInstanceID);
+				SetInstance(NewInstanceID,
+					(UVTriangle[j] == -1 || UVOverlay == nullptr) ? FVector2f::Zero() : UVOverlay->GetElement(UVTriangle[j]),
+					(NormalTriangle[j] == -1 || NormalOverlay == nullptr) ? FVector3f::UnitY() : NormalOverlay->GetElement(NormalTriangle[j]));
+			}
+			InstanceTri[j] = InstanceList[InstanceElem];
+		}
+
+		FPolygonID NewPolyID = AppendTriangle(InstanceTri[0], InstanceTri[1], InstanceTri[2], AllGroupID);
+		if (bDoTransferPolyGroups)
+		{
+			SetPolyGroupID(NewPolyID, Mesh->GetTriangleGroup(TriID));
+		}
+	}
+}
+
+
+
+
+
+void FMeshDescriptionBuilder::Translate(const FVector& Translation)
+{
+	for (FVertexID VertexID : MeshDescription->Vertices().GetElementIDs())
+	{
+		FVector Position = VertexPositions.Get(VertexID);
+		Position += Translation;
+		VertexPositions.Set(VertexID, Position);
+	}
+}
+
+
+
+
+void FMeshDescriptionBuilder::SetAllEdgesHardness(bool bHard)
+{
+	TEdgeAttributesRef<bool> EdgeHardness =
+		MeshDescription->EdgeAttributes().GetAttributesRef<bool>(MeshAttribute::Edge::IsHard);
+	for (FEdgeID EdgeID : MeshDescription->Edges().GetElementIDs())
+	{
+		EdgeHardness.Set(EdgeID, 0, bHard);
+	}
+}
+
+
+FBox FMeshDescriptionBuilder::ComputeBoundingBox() const
+{
+	FAxisAlignedBox3f bounds = FAxisAlignedBox3f::Empty();
+	for ( FVertexID VertexID : MeshDescription->Vertices().GetElementIDs() )
+	{
+		bounds.Contain(VertexPositions.Get(VertexID));
+	}
+	return bounds;
+}
\ No newline at end of file
diff --git a/Engine/Plugins/Experimental/GeometryProcessing/Source/MeshConversion/Private/MeshDescriptionToDynamicMesh.cpp b/Engine/Plugins/Experimental/GeometryProcessing/Source/MeshConversion/Private/MeshDescriptionToDynamicMesh.cpp
new file mode 100644
--- /dev/null
+++ b/Engine/Plugins/Experimental/GeometryProcessing/Source/MeshConversion/Private/MeshDescriptionToDynamicMesh.cpp
@@ -0,0 +1,240 @@
+// Copyright 1998-2019 Epic Games, Inc. All Rights Reserved. 
+
+#include "MeshDescriptionToDynamicMesh.h"
+#include "MeshAttributes.h"
+#include "DynamicMeshAttributeSet.h"
+#include "DynamicMeshOverlay.h"
+#include "MeshDescriptionBuilder.h"
+
+
+struct FVertexUV
+{
+	int vid;
+	float x;
+	float y;
+	bool operator==(const FVertexUV & o) const
+	{
+		return vid == o.vid && x == o.x && y == o.y;
+	}
+};
+FORCEINLINE uint32 GetTypeHash(const FVertexUV& Vector)
+{
+	// ugh copied from FVector clearly should not be using CRC for hash!!
+	return FCrc::MemCrc32(&Vector, sizeof(Vector));
+}
+
+
+class FUVWelder
+{
+public:
+	TMap<FVertexUV, int> UniqueVertexUVs;
+	FDynamicMeshUVOverlay* UVOverlay;
+
+	FUVWelder(FDynamicMeshUVOverlay* UVOverlayIn)
+	{
+		check(UVOverlayIn);
+		UVOverlay = UVOverlayIn;
+	}
+
+	int FindOrAddUnique(const FVector2D & UV, int VertexID)
+	{
+		FVertexUV VertUV = { VertexID, UV.X, UV.Y };
+
+		int NewIndex = -1;
+		if (UniqueVertexUVs.Contains(VertUV))
+		{
+			NewIndex = UniqueVertexUVs[VertUV];
+		}
+		else
+		{
+			NewIndex = UVOverlay->AppendElement(UV, VertexID);
+			UniqueVertexUVs.Add(VertUV, NewIndex);
+		}
+		return NewIndex;
+	}
+};
+
+
+
+
+struct FVertexNormal
+{
+	int vid;
+	float x;
+	float y;
+	float z;
+	bool operator==(const FVertexNormal & o) const
+	{
+		return vid == o.vid && x == o.x && y == o.y && z == o.z;
+	}
+};
+FORCEINLINE uint32 GetTypeHash(const FVertexNormal& Vector)
+{
+	// ugh copied from FVector clearly should not be using CRC for hash!!
+	return FCrc::MemCrc32(&Vector, sizeof(Vector));
+}
+
+
+class FNormalWelder
+{
+public:
+	TMap<FVertexNormal, int> UniqueVertexNormals;
+	FDynamicMeshNormalOverlay* NormalOverlay;
+
+	FNormalWelder(FDynamicMeshNormalOverlay* NormalOverlayIn)
+	{
+		check(NormalOverlayIn);
+		NormalOverlay = NormalOverlayIn;
+	}
+
+	int FindOrAddUnique(const FVector & Normal, int VertexID)
+	{
+		FVertexNormal VertNormal = { VertexID, Normal.X, Normal.Y, Normal.Z };
+
+		int NewIndex = -1;
+		if (UniqueVertexNormals.Contains(VertNormal))
+		{
+			NewIndex = UniqueVertexNormals[VertNormal];
+		}
+		else
+		{
+			NewIndex = NormalOverlay->AppendElement(Normal, VertexID);
+			UniqueVertexNormals.Add(VertNormal, NewIndex);
+		}
+		return NewIndex;
+	}
+};
+
+
+
+
+
+
+
+
+
+
+
+
+
+void FMeshDescriptionToDynamicMesh::Convert(const FMeshDescription* MeshIn, FDynamicMesh3& MeshOut)
+{
+	// look up vertex positions
+	const FVertexArray& VertexIDs = MeshIn->Vertices();
+	TVertexAttributesConstRef<FVector> VertexPositions =
+		MeshIn->VertexAttributes().GetAttributesRef<FVector>(MeshAttribute::Vertex::Position);
+
+	// copy vertex positions
+	for (const FVertexID VertexID : VertexIDs.GetElementIDs())
+	{
+		const FVector Position = VertexPositions.Get(VertexID);
+		int NewVertIdx = MeshOut.AppendVertex(Position);
+		check(NewVertIdx == VertexID.GetValue());
+	}
+
+	// look up vertex-instance UVs and normals
+	// @todo: does the MeshDescription always have UVs and Normals?
+	TVertexInstanceAttributesConstRef<FVector2D> InstanceUVs =
+		MeshIn->VertexInstanceAttributes().GetAttributesRef<FVector2D>(MeshAttribute::VertexInstance::TextureCoordinate);
+	TVertexInstanceAttributesConstRef<FVector> InstanceNormals =
+		MeshIn->VertexInstanceAttributes().GetAttributesRef<FVector>(MeshAttribute::VertexInstance::Normal);
+
+	// enable attributes on output mesh
+	MeshOut.EnableAttributes();
+	FDynamicMeshUVOverlay* UVOverlay = MeshOut.Attributes()->PrimaryUV();
+	FDynamicMeshNormalOverlay* NormalOverlay = MeshOut.Attributes()->PrimaryNormals();
+
+	TPolygonAttributesConstRef<int> PolyGroups =
+		MeshIn->PolygonAttributes().GetAttributesRef<int>(ExtendedMeshAttribute::PolyTriGroups);
+
+	// base triangle groups will track polygons
+	if (bEnableOutputGroups)
+	{
+		MeshOut.EnableTriangleGroups(0);
+	}
+
+	int NumVertices = MeshIn->Vertices().Num();
+	int NumPolygons = MeshIn->Polygons().Num();
+	int NumVtxInstances = MeshIn->VertexInstances().Num();
+	if (bPrintDebugMessages)
+	{
+		UE_LOG(LogTemp, Warning, TEXT("FMeshDescriptionToDynamicMesh: MeshDescription verts %d polys %d instances %d"), NumVertices, NumPolygons, NumVtxInstances);
+	}
+
+	// reserve space in MeshOut?
+
+	// used to merge coincident elements so that we get actual topology
+	FUVWelder UVWelder(UVOverlay);
+	FNormalWelder NormalWelder(NormalOverlay);
+
+	const FPolygonArray& Polygons = MeshIn->Polygons();
+	for (const FPolygonID PolygonID : Polygons.GetElementIDs())
+	{
+		FPolygonGroupID PolygonGroupID = MeshIn->GetPolygonPolygonGroup(PolygonID);
+
+		const TArray<FMeshTriangle>& Triangles = MeshIn->GetPolygonTriangles(PolygonID);
+		for (FMeshTriangle Triangle : Triangles)
+		{
+			FVertexInstanceID InstanceTri[3];
+			InstanceTri[0] = Triangle.VertexInstanceID0; 
+			InstanceTri[1] = Triangle.VertexInstanceID1; 
+			InstanceTri[2] = Triangle.VertexInstanceID2;
+
+			int GroupID = 0;
+			if (GroupMode == EPrimaryGroupMode::SetToPolyGroup)
+			{
+				if (PolyGroups.IsValid())
+				{
+					GroupID = PolyGroups.Get(PolygonID, 0);
+				}
+			}
+			else if (GroupMode == EPrimaryGroupMode::SetToPolygonID)
+			{
+				GroupID = PolygonID.GetValue();
+			}
+			else if (GroupMode == EPrimaryGroupMode::SetToPolygonGroupID)
+			{
+				GroupID = PolygonGroupID.GetValue();
+			}
+			
+
+			// append triangle
+			int VertexID0 = MeshIn->GetVertexInstance(Triangle.VertexInstanceID0).VertexID.GetValue();
+			int VertexID1 = MeshIn->GetVertexInstance(Triangle.VertexInstanceID1).VertexID.GetValue();
+			int VertexID2 = MeshIn->GetVertexInstance(Triangle.VertexInstanceID2).VertexID.GetValue();
+			int NewTriangleID = MeshOut.AppendTriangle(VertexID0, VertexID1, VertexID2, GroupID);
+			FIndex3i Tri(VertexID0, VertexID1, VertexID2);
+
+			if (UVOverlay != nullptr)
+			{
+				FIndex3i TriUV;
+				for (int j = 0; j < 3; ++j)
+				{
+					FVector2D UV = InstanceUVs.Get(InstanceTri[j]);
+					TriUV[j] = UVWelder.FindOrAddUnique(UV, Tri[j]);
+				}
+				UVOverlay->SetTriangle(NewTriangleID, TriUV);
+			}
+
+			if (NormalOverlay != nullptr)
+			{
+				FIndex3i TriNormals;
+				for (int j = 0; j < 3; ++j)
+				{
+					FVector Normal = InstanceNormals.Get(InstanceTri[j]);
+					TriNormals[j] = NormalWelder.FindOrAddUnique(Normal, Tri[j]);
+				}
+				NormalOverlay->SetTriangle(NewTriangleID, TriNormals);
+			}
+
+		}
+	}
+
+	if (bPrintDebugMessages)
+	{
+		int NumUVs = (UVOverlay != nullptr) ? UVOverlay->MaxElementID() : 0;
+		int NumNormals = (NormalOverlay != nullptr) ? NormalOverlay->MaxElementID() : 0;
+		UE_LOG(LogTemp, Warning, TEXT("FMeshDescriptionToDynamicMesh:  FDynamicMesh verts %d triangles %d uvs %d normals %d"), MeshOut.MaxVertexID(), MeshOut.MaxTriangleID(), NumUVs, NumNormals);
+	}
+
+}
diff --git a/Engine/Plugins/Experimental/GeometryProcessing/Source/MeshConversion/Public/DynamicMeshToMeshDescription.h b/Engine/Plugins/Experimental/GeometryProcessing/Source/MeshConversion/Public/DynamicMeshToMeshDescription.h
new file mode 100644
--- /dev/null
+++ b/Engine/Plugins/Experimental/GeometryProcessing/Source/MeshConversion/Public/DynamicMeshToMeshDescription.h
@@ -0,0 +1,49 @@
+// Copyright 1998-2019 Epic Games, Inc. All Rights Reserved.
+
+#pragma once
+
+#include "CoreMinimal.h"
+#include "DynamicMesh3.h"
+#include "MeshDescription.h"
+
+/**
+ * Convert FDynamicMesh3 to FMeshDescription
+ *
+ */
+class MESHCONVERSION_API FDynamicMeshToMeshDescription
+{
+public:
+	/** If true, will print some possibly-helpful debugging spew to output log */
+	bool bPrintDebugMessages = false;
+
+	/** Should DynamicMesh triangle groups be transfered to MeshDescription via custom PolyTriGroups attribute */
+	bool bSetPolyGroups = true;
+
+	/**
+	 * Default conversion of DynamicMesh to MeshDescription. Calls functions below depending on mesh state
+	 */
+	void Convert(const FDynamicMesh3* MeshIn, FMeshDescription& MeshOut);
+
+
+	//
+	// Internal functions that you can also call directly
+	//
+
+	/**
+	 * Ignore any Attributes on input Mesh, calculate per-vertex normals and have MeshDescription compute tangents.
+	 * One VertexInstance per input vertex is generated
+	 */
+	void Convert_NoAttributes(const FDynamicMesh3* MeshIn, FMeshDescription& MeshOut);
+
+	/**
+	 * Convert while minimizing VertexInstance count, IE new VertexInstances are only created 
+	 * if a unique UV or Normal is required.
+	 */
+	void Convert_SharedInstances(const FDynamicMesh3* MeshIn, FMeshDescription& MeshOut);
+
+	/**
+	 * Convert with no shared VertexInstances. A new VertexInstance is created for
+	 * each triangle vertex (ie corner). However vertex positions are shared.
+	 */
+	void Convert_NoSharedInstances(const FDynamicMesh3* MeshIn, FMeshDescription& MeshOut);
+};
\ No newline at end of file
diff --git a/Engine/Plugins/Experimental/GeometryProcessing/Source/MeshConversion/Public/MeshDescriptionBuilder.h b/Engine/Plugins/Experimental/GeometryProcessing/Source/MeshConversion/Public/MeshDescriptionBuilder.h
new file mode 100644
--- /dev/null
+++ b/Engine/Plugins/Experimental/GeometryProcessing/Source/MeshConversion/Public/MeshDescriptionBuilder.h
@@ -0,0 +1,95 @@
+// Copyright 1998-2019 Epic Games, Inc. All Rights Reserved.
+
+#pragma once
+
+#include "CoreMinimal.h"
+#include "MeshDescription.h"
+
+class FDynamicMesh3;
+
+
+/**
+ * These attributes are used to store custom modeling tools data on a MeshDescription
+ */
+namespace ExtendedMeshAttribute
+{
+	extern MESHCONVERSION_API const FName PolyTriGroups;
+}
+
+
+/**
+ * Utility class to construct MeshDescription instances
+ */
+class MESHCONVERSION_API FMeshDescriptionBuilder
+{
+public:
+	void SetMeshDescription(FMeshDescription* Description);
+
+	/** Enable per-triangle integer attribute named PolyTriGroups */
+	void EnablePolyGroups();
+
+	/** Append vertex and return new vertex ID */
+	FVertexID AppendVertex(const FVector& Position);
+	
+	/** Append new vertex instance and return ID */
+	FVertexInstanceID AppendInstance(const FVertexID& VertexID);
+
+	/** Set the UV and Normal of a vertex instance*/
+	void SetInstance(const FVertexInstanceID& InstanceID, const FVector2D& InstanceUV, const FVector& InstanceNormal);
+
+	/** Create a new polygon group and return it's ID */
+	FPolygonGroupID AppendPolygonGroup();
+
+
+	/** Append a triangle to the mesh with the given PolygonGroup ID */
+	FPolygonID AppendTriangle(const FVertexID& Vertex0, const FVertexID& Vertex1, const FVertexID& Vertex2, const FPolygonGroupID& PolygonGroup);
+
+	/** Append a triangle to the mesh with the given PolygonGroup ID, and optionally with triangle-vertex UVs and Normals */
+	FPolygonID AppendTriangle(const FVertexID* Triangle, const FPolygonGroupID& PolygonGroup, 
+		const FVector2D* VertexUVs = nullptr, const FVector* VertexNormals = nullptr);
+
+	/** 
+	 * Append an arbitrary polygon to the mesh with the given PolygonGroup ID, and optionally with polygon-vertex UVs and Normals
+	 * Unique Vertex instances will be created for each polygon-vertex.
+	 */
+	FPolygonID AppendPolygon(const TArray<FVertexID>& Vertices, const FPolygonGroupID& PolygonGroup, 
+		const TArray<FVector2D>* VertexUVs = nullptr, const TArray<FVector>* VertexNormals = nullptr);
+
+	/**
+	 * Append a triangle to the mesh using the given vertex instances and PolygonGroup ID
+	 */
+	FPolygonID AppendTriangle(const FVertexInstanceID& Instance0, const FVertexInstanceID& Instance1, const FVertexInstanceID& Instance2, const FPolygonGroupID& PolygonGroup);
+
+
+	/**
+	 * Append an entire mesh to this mesh. This will create a minimal set of shared vertex instances.
+	 * @param bSetPolyGroups if true, we transfer the input mesh triangle groups to the PolyTriGroups attribute
+	 */
+	void AppendMesh(const FDynamicMesh3* Mesh, bool bSetPolyGroups);
+
+
+
+	/** Set the PolyTriGroups attribute value to a specific GroupID for a Polygon */
+	void SetPolyGroupID(const FPolygonID& PolygonID, int GroupID);
+
+	/** Set MeshAttribute::Edge::IsHard to true for all edges */
+	void SetAllEdgesHardness(bool bHard);
+
+	/** Translate the MeshDescription vertex positions */
+	void Translate(const FVector& Translation);
+
+	/** Return the current bounding box of the mesh */
+	FBox ComputeBoundingBox() const;
+
+protected:
+	FMeshDescription* MeshDescription;
+
+	TVertexAttributesRef<FVector> VertexPositions;
+	TVertexInstanceAttributesRef<FVector2D> InstanceUVs;
+	TVertexInstanceAttributesRef<FVector> InstanceNormals;
+	TArray<FVertexID> TempBuffer;
+	TArray<FVector2D> UVBuffer;
+	TArray<FVector> NormalBuffer;
+
+	TPolygonAttributesRef<int> PolyGroups;
+};
\ No newline at end of file
diff --git a/Engine/Plugins/Experimental/GeometryProcessing/Source/MeshSolverUtilities/MeshSolverUtilities.Build.cs b/Engine/Plugins/Experimental/GeometryProcessing/Source/MeshSolverUtilities/MeshSolverUtilities.Build.cs
new file mode 100644
--- /dev/null
+++ b/Engine/Plugins/Experimental/GeometryProcessing/Source/MeshSolverUtilities/MeshSolverUtilities.Build.cs
@@ -0,0 +1,34 @@
+// Copyright 1998-2019 Epic Games, Inc. All Rights Reserved.
+
+using UnrealBuildTool;
+
+public class MeshSolverUtilities : ModuleRules
+{
+    public MeshSolverUtilities(ReadOnlyTargetRules Target) : base(Target)
+    {
+        PCHUsage = ModuleRules.PCHUsageMode.UseExplicitOrSharedPCHs;
+        //PCHUsage = ModuleRules.PCHUsageMode.NoSharedPCHs;
+
+        PublicDependencyModuleNames.AddRange(
+            new string[] {
+                "Core",
+                "Eigen"
+            }
+		);
+
+        PrivateDependencyModuleNames.AddRange(
+            new string[]
+            {
+                "CoreUObject",
+                "Engine",
+				"GeometricObjects",
+				"DynamicMesh"
+                //"GeometricObjects",
+			    //"DynamicMesh"
+				//"Slate",
+				//"SlateCore",
+				// ... add private dependencies that you statically link with here ...	
+			}
+		);
+    }
+}
\ No newline at end of file
diff --git a/Engine/Plugins/Experimental/GeometryProcessing/Source/MeshSolverUtilities/Private/FSOAPositions.h b/Engine/Plugins/Experimental/GeometryProcessing/Source/MeshSolverUtilities/Private/FSOAPositions.h
new file mode 100644
--- /dev/null
+++ b/Engine/Plugins/Experimental/GeometryProcessing/Source/MeshSolverUtilities/Private/FSOAPositions.h
@@ -0,0 +1,79 @@
+// Copyright 1998-2019 Epic Games, Inc. All Rights Reserved.
+
+#pragma once
+
+#include "CoreMinimal.h"
+
+#if defined(_MSC_VER) && USING_CODE_ANALYSIS
+#pragma warning(push)
+#pragma warning(disable : 6011)
+#pragma warning(disable : 6387)
+#pragma warning(disable : 6313)
+#pragma warning(disable : 6294)
+#endif
+THIRD_PARTY_INCLUDES_START
+#include <Eigen/Core>
+THIRD_PARTY_INCLUDES_END
+#if defined(_MSC_VER) && USING_CODE_ANALYSIS
+#pragma warning(pop)
+#endif
+
+#include "FSparseMatrixD.h"
+
+/**
+* A struct of arrays representation used to hold vertex positions 
+* in three vectors that can interface with the eigen library
+*/
+class FSOAPositions
+{
+public:
+	typedef typename FSparseMatrixD::Scalar  ScalarType;
+	typedef Eigen::Matrix<ScalarType, Eigen::Dynamic, 1>  VectorType;
+
+
+	FSOAPositions(int32 Size)
+		: XVector(Size)
+		, YVector(Size)
+		, ZVector(Size)
+	{}
+
+	FSOAPositions()
+	{}
+
+	VectorType& Array(int32 i)
+	{
+		return (i == 0) ? XVector : (i == 1) ? YVector : ZVector;
+	}
+	const VectorType& Array(int32 i) const
+	{
+		return (i == 0) ? XVector : (i == 1) ? YVector : ZVector;
+	}
+
+	VectorType XVector;
+	VectorType YVector;
+	VectorType ZVector;
+
+	void SetZero(int32 NumElements)
+	{
+		XVector.setZero(NumElements);
+		YVector.setZero(NumElements);
+		ZVector.setZero(NumElements);
+	}
+
+	// Test that all the arrays have the same given size.
+	bool bHasSize(int32 Size) const
+	{
+		return (XVector.rows() == Size && YVector.rows() == Size && ZVector.rows() == Size);
+	}
+
+	int32 Num() const
+	{
+		int32 Size = XVector.rows();
+		if (!bHasSize(Size))
+		{
+			Size = -1;
+		}
+		return Size;
+	}
+
+};
\ No newline at end of file
diff --git a/Engine/Plugins/Experimental/GeometryProcessing/Source/MeshSolverUtilities/Private/LaplacianMeshSmoother.cpp b/Engine/Plugins/Experimental/GeometryProcessing/Source/MeshSolverUtilities/Private/LaplacianMeshSmoother.cpp
new file mode 100644
--- /dev/null
+++ b/Engine/Plugins/Experimental/GeometryProcessing/Source/MeshSolverUtilities/Private/LaplacianMeshSmoother.cpp
@@ -0,0 +1,695 @@
+// Copyright 1998-2019 Epic Games, Inc. All Rights Reserved.
+
+#include "LaplacianMeshSmoother.h"
+
+#include "FSOAPositions.h"
+
+#include "LaplacianOperators.h"
+#include "MatrixSolver.h"
+#include "MeshSmoothingUtilities.h"
+
+#include "ProfilingDebugging/ScopedTimers.h"
+#include "Async/ParallelFor.h"
+
+
+DEFINE_LOG_CATEGORY_STATIC(LogMeshSmoother, Log, All);
+
+
+double ComputeDistSqrd(const FSOAPositions& VecA, const FSOAPositions& VecB)
+{
+	const int32 NumA = VecA.Num();
+	const int32 NumB = VecB.Num();
+	checkSlow(NumA == NumB);
+
+#if 0
+	// the eigne way?
+	double Tmp  = (VecA.Array(0) - VecB.Array(0)).dot(VecA.Array(0) - VecB.Array(0))
+				+ (VecA.Array(1) - VecB.Array(1)).dot(VecA.Array(1) - VecB.Array(1))
+				+ (VecA.Array(2) - VecB.Array(2)).dot(VecA.Array(2) - VecB.Array(2));
+		     
+#endif
+
+	// doing it by hand so we can break when the error is large
+	double DistSqrd = 0.;
+	{
+		const auto& AX = VecA.Array(0);
+		const auto& AY = VecA.Array(1);
+		const auto& AZ = VecA.Array(2);
+
+		const auto& BX = VecB.Array(0);
+		const auto& BY = VecB.Array(1);
+		const auto& BZ = VecB.Array(2);
+
+
+		for (int32 i = 0; i < NumA; ++i)
+		{
+			double TmpX = AX(i)-BX(i);
+			double TmpY = AY(i)-BY(i);
+			double TmpZ = AZ(i)-BZ(i);
+
+			TmpX *= TmpX;
+			TmpY *= TmpY;
+			TmpZ *= TmpZ;
+			double TmpT = TmpX + TmpY + TmpZ;
+			DistSqrd += TmpT;
+		}
+	}
+	return DistSqrd;
+}
+
+
+
+FConstrainedMeshOperator::FConstrainedMeshOperator(const FDynamicMesh3& DynamicMesh, const ELaplacianWeightScheme Scheme, const EMatrixSolverType MatrixSolverType)
+	: VertexCount(DynamicMesh.VertexCount())
+{
+
+	Laplacian = ConstructLaplacian(Scheme, DynamicMesh, VtxLinearization, &EdgeVerts);
+	FSparseMatrixD& LMatrix = *(Laplacian);
+
+	TUniquePtr<FSparseMatrixD> LTL(new FSparseMatrixD(Laplacian->rows(), Laplacian->cols()));
+	FSparseMatrixD& LTLMatrix = *(LTL);
+
+	bool bIsLaplacianSymmetric = (Scheme == ELaplacianWeightScheme::Valence || Scheme == ELaplacianWeightScheme::Uniform);
+
+	if (bIsLaplacianSymmetric)
+	{
+		// Laplacian is symmetric, i.e. equal to its transpose
+		LTLMatrix = LMatrix * LMatrix;
+	
+		ConstrainedSolver.Reset(new FConstrainedSolver(LTL, MatrixSolverType));
+	}
+	else
+	{
+		// the laplacian 
+		LTLMatrix = LMatrix.transpose() * LMatrix;
+		ConstrainedSolver.Reset(new FConstrainedSolver(LTL, MatrixSolverType));
+	}
+
+}
+
+void FConstrainedMeshOperator::AddConstraint(const int32 VtxId, const double Weight, const FVector3d& Pos, const bool bPostFix)
+{
+
+	bConstraintPositionsDirty = true;
+	bConstraintWeightsDirty   = true;
+	int32 Index = VtxLinearization.ToIndex()[VtxId];
+
+	ConstraintPositionMap.Add(TTuple<int32, FConstraintPosition>(Index, FConstraintPosition(Pos, bPostFix)));
+	ConstraintWeightMap.Add(TTuple<int32, double>(Index, Weight));
+
+}
+
+
+bool FConstrainedMeshOperator::UpdateConstraintPosition(const int32 VtxId, const FVector3d& Pos, const bool bPostFix)
+{
+	bConstraintPositionsDirty = true;
+	int32 Index = VtxLinearization.ToIndex()[VtxId];
+	// Add should over-write any existing value for this key
+	ConstraintPositionMap.Add(TTuple<int32, FConstraintPosition>(Index, FConstraintPosition(Pos, bPostFix)));
+
+	return ConstraintWeightMap.Contains(VtxId);
+}
+
+bool FConstrainedMeshOperator::UpdateConstraintWeight(const int32 VtxId, const double Weight)
+{
+
+	bConstraintWeightsDirty = true;
+	int32 Index = VtxLinearization.ToIndex()[VtxId];
+	// Add should over-write any existing value for this key
+	ConstraintWeightMap.Add(TTuple<int32, double>(Index, Weight));
+
+	return ConstraintPositionMap.Contains(VtxId);
+}
+
+
+bool FConstrainedMeshOperator::IsConstrained(const int32 VtxId) const
+{
+	if (VtxId > VtxLinearization.ToIndex().Num()) return false;
+
+	int32 Index = VtxLinearization.ToIndex()[VtxId];
+
+	return ConstraintWeightMap.Contains(Index);
+}
+
+void FConstrainedMeshOperator::UpdateSolverConstraints()
+{
+	if (bConstraintWeightsDirty)
+	{
+		ConstrainedSolver->SetConstraintWeights(ConstraintWeightMap);
+		bConstraintWeightsDirty = false;
+	}
+	
+	if (bConstraintPositionsDirty)
+	{
+		ConstrainedSolver->SetContraintPositions(ConstraintPositionMap);
+		bConstraintPositionsDirty = false;
+	}	
+}
+
+bool FConstrainedMeshOperator::Linearize(const FSOAPositions& PositionalVector, TArray<FVector3d>& LinearArray) const
+{
+	// Number of positions
+
+	const int32 Num = PositionalVector.XVector.rows();
+
+	// early out if the x,y,z arrays in the PositionalVector have different lengths
+	if (!PositionalVector.bHasSize(Num))
+	{
+		return false;
+	}
+
+	// 
+	const auto& IndexToVtxId = VtxLinearization.ToId();
+	const int32 MaxVtxId = IndexToVtxId.Num(); // NB: this is really max_used + 1 in the mesh.  See  FDynamicMesh3::MaxVertexID()
+
+
+
+	LinearArray.Empty(MaxVtxId);
+	LinearArray.AddUninitialized(MaxVtxId);
+
+	for (int32 i = 0; i < Num; ++i)
+	{
+		const int32 VtxId = IndexToVtxId[i];
+
+		LinearArray[VtxId] = FVector3d(PositionalVector.XVector.coeff(i), PositionalVector.YVector.coeff(i), PositionalVector.ZVector.coeff(i));
+	}
+
+	return true;
+}
+
+void FConstrainedMeshOperator::ExtractVertexPositions(const FDynamicMesh3& DynamicMesh, FSOAPositions& VertexPositions) const
+{
+	VertexPositions.SetZero(VertexCount);
+
+
+	const TArray<int32>& ToIndex = VtxLinearization.ToIndex();
+
+
+	for (int32 VtxId : DynamicMesh.VertexIndicesItr())
+	{
+		const int32 i = ToIndex[VtxId];
+
+		checkSlow(i != FDynamicMesh3::InvalidID);
+
+		const FVector3d& Vertex = DynamicMesh.GetVertex(VtxId);
+		// coeffRef - coeff access without range checking
+		VertexPositions.XVector.coeffRef(i) = Vertex.X;
+		VertexPositions.YVector.coeffRef(i) = Vertex.Y;
+		VertexPositions.ZVector.coeffRef(i) = Vertex.Z;
+	}
+}
+
+void FConstrainedMeshOperator::UpdateWithPostFixConstraints(FSOAPositions& PositionVector) const
+{
+	for (const auto& ConstraintPosition : ConstraintPositionMap)
+	{
+		const int32 Index = ConstraintPosition.Key;
+		const FConstraintPosition& Constraint = ConstraintPosition.Value;
+
+		// we only care about post-fix constraints
+
+		if (Constraint.bPostFix)
+		{
+			const FVector3d& Pos = Constraint.Position;
+			PositionVector.XVector[Index] = Pos.X;
+			PositionVector.YVector[Index] = Pos.Y;
+			PositionVector.ZVector[Index] = Pos.Z;
+		}
+	}
+}
+
+
+
+FConstrainedMeshDeformer::FConstrainedMeshDeformer(const FDynamicMesh3& DynamicMesh, const ELaplacianWeightScheme LaplacianType)
+	: FConstrainedMeshOperator(DynamicMesh, LaplacianType, EMatrixSolverType::LU)
+	, LaplacianVectors(FConstrainedMeshOperator::VertexCount)
+{
+
+	// The current vertex positions 
+	
+	// Note: the OriginalVertexPositions are being stored as member data 
+	// for use if the solver is iterative.
+	// FSOAPositions OriginalVertexPositions; 
+	ExtractVertexPositions(DynamicMesh, OriginalVertexPositions);
+	
+	
+	// The biharmonic part of the constrained solver
+	//   Biharmonic := Laplacian^{T} * Laplacian
+	
+	const auto& Biharmonic = ConstrainedSolver->Biharmonic();
+
+	// Compute the Laplacian Vectors
+	//    := Biharmonic * VertexPostion
+	// In the case of the cotangent laplacian this can be identified as the mean curvature * normal.
+	for (int32 i = 0; i < 3; ++i)
+	{
+		LaplacianVectors.Array(i) = Biharmonic * OriginalVertexPositions.Array(i);
+	}
+}
+
+bool FConstrainedMeshDeformer::Deform(TArray<FVector3d>& PositionBuffer) 
+{
+	
+	// Update constraints.  This only trigger solver rebuild if the weights were updated.
+	UpdateSolverConstraints();
+
+	// Allocate space for the result as a struct of arrays
+	FSOAPositions SolutionVector(VertexCount);
+
+	// Solve the linear system
+	// NB: the original positions will only be used if the underlying solver type is iterative	
+	bool bSuccess = ConstrainedSolver->SolveWithGuess(OriginalVertexPositions, LaplacianVectors, SolutionVector);
+	
+	// Move any vertices to match bPostFix constraints
+
+	UpdateWithPostFixConstraints(SolutionVector);
+
+	// Copy the result into the array of stucts form.  
+	// NB: this re-indexes so the results can be looked up using VtxId
+
+	Linearize(SolutionVector, PositionBuffer);
+
+	// the matrix solve state
+	return bSuccess;
+
+}
+
+
+bool FBiHarmonicMeshSmoother::ComputeSmoothedMeshPositions(TArray<FVector3d>& UpdatedPositions)
+{
+
+	UpdateSolverConstraints();
+
+	// Solves the constrained system and updates the mesh 
+
+	FSOAPositions SolutionVector(VertexCount);
+
+	bool bSuccess = ConstrainedSolver->Solve(SolutionVector);
+
+
+	// Move any vertices to match bPostFix constraints
+
+	UpdateWithPostFixConstraints(SolutionVector);
+
+	// Move vertices to solution positions
+
+	Linearize(SolutionVector, UpdatedPositions);
+
+	return bSuccess;
+}
+
+bool FCGBiHarmonicMeshSmoother::ComputeSmoothedMeshPositions(TArray<FVector3d>& UpdatedPositions)
+{
+
+	UpdateSolverConstraints();
+
+	// Solves the constrained system and updates the mesh 
+
+	// Solves the constrained system
+
+	FSOAPositions SolutionVector(VertexCount);
+
+	bool bSuccess = ConstrainedSolver->Solve(SolutionVector);
+
+
+	// Move any vertices to match bPostFix constraints
+
+	UpdateWithPostFixConstraints(SolutionVector);
+
+	// Move vertices to solution positions
+
+	Linearize(SolutionVector, UpdatedPositions);
+
+	return bSuccess;
+}
+
+
+
+FDiffusionIntegrator::FDiffusionIntegrator(const FDynamicMesh3& DynamicMesh, const ELaplacianWeightScheme Scheme)
+{
+	Id = 0;
+	bIsSymmetric = false;
+	MinDiagonalValue = 0;
+
+	VertexCount = DynamicMesh.VertexCount();
+	
+	// Allocate the buffers.
+	Tmp[0].SetZero(VertexCount);
+	Tmp[1].SetZero(VertexCount);
+
+
+}
+
+void FDiffusionIntegrator::Initialize(const FDynamicMesh3& DynamicMesh, const ELaplacianWeightScheme Scheme)
+{
+	// Construct the laplacian, and extract the mapping for vertices (VtxLinearization)
+	DiffusionOperator = ConstructDiffusionOperator(Scheme, DynamicMesh, bIsSymmetric, VtxLinearization, &EdgeVerts);
+
+	const auto& ToVertId = VtxLinearization.ToId();
+
+	// Extract current positions.
+	for (int32 i = 0; i < VertexCount; ++i)
+	{
+		int32 VtxId = ToVertId[i];
+
+		const FVector3d Pos = DynamicMesh.GetVertex(VtxId);
+		Tmp[0].XVector[i] = Pos.X;
+		Tmp[0].YVector[i] = Pos.Y;
+		Tmp[0].ZVector[i] = Pos.Z;
+	}
+
+	const FSparseMatrixD&  M = *DiffusionOperator;
+
+	// Find the min diagonal entry (all should be negative).
+	int32 Rank = M.rows();
+	MinDiagonalValue = FSparseMatrixD::Scalar(0);
+	for (int32 i = 0; i < Rank; ++i)
+	{
+		auto Diag = M.coeff(i, i);
+		MinDiagonalValue = FMath::Min(Diag, MinDiagonalValue);
+	}
+
+
+#if 0
+	// testing - how to print the matrix to debug output 
+
+	std::stringstream  ss;
+	ss << Eigen::MatrixXd(M) << std::endl;
+
+	FString Foo = ss.str().c_str();
+	FPlatformMisc::LowLevelOutputDebugStringf(*Foo);
+#endif
+
+}
+
+
+
+void FDiffusionIntegrator::Integrate_ForwardEuler(int32 NumSteps, double Alpha, double)
+{
+	Alpha = FMath::Clamp(Alpha, 0., 1.);
+
+	const FSparseMatrixD& M = *DiffusionOperator;
+	FSparseMatrixD::Scalar TimeStep = -Alpha / MinDiagonalValue;
+	Id = 0;
+	for (int32 s = 0; s < NumSteps; ++s)
+	{
+
+		int32 SrcBuffer = Id;
+		Id = 1 - Id;
+		Tmp[Id].XVector = Tmp[SrcBuffer].XVector + TimeStep * M * Tmp[SrcBuffer].XVector;
+		Tmp[Id].YVector = Tmp[SrcBuffer].YVector + TimeStep * M * Tmp[SrcBuffer].YVector;
+		Tmp[Id].ZVector = Tmp[SrcBuffer].ZVector + TimeStep * M * Tmp[SrcBuffer].ZVector;
+
+	}
+
+}
+
+
+
+void FDiffusionIntegrator::Integrate_BackwardEuler(const EMatrixSolverType MatrixSolverType, int32 NumSteps, double Alpha, double Intensity)
+{
+
+	//typedef typename TMatrixSolverTrait<EMatrixSolverType::LU>::MatrixSolverType   MatrixSolverType;
+
+	// We solve 
+	// p^{n+1} - dt * L[p^{n+1}] = p^{n}
+	// 
+	// i.e.
+	// [I - dt * L ] p^{n+1} = p^{n}
+	//
+	// NB: in the case of the cotangent laplacian this would be better if we broke the L int
+	// L = (A^{-1}) H  where A is the "area matrix" (think "mass matrix"), then this would
+	// become
+	// [A - dt * H] p^{n+1} = Ap^{n}  
+	//  
+	// A - dt * H would be symmetric
+	//
+
+
+	const FSparseMatrixD& L = *DiffusionOperator;
+	// Identity matrix
+	FSparseMatrixD Ident(L.rows(), L.cols());
+	Ident.setIdentity();
+
+	FSparseMatrixD::Scalar TimeStep = Alpha * FMath::Min(Intensity, 1.e6);
+	
+	FSparseMatrixD M = Ident -TimeStep * L;
+
+	
+	M.makeCompressed();
+
+	TUniquePtr<IMatrixSolverBase> MatrixSolver = ContructMatrixSolver(MatrixSolverType);
+
+	MatrixSolver->SetUp(M, bIsSymmetric);
+
+	if (MatrixSolver->bIsIterative())
+	{
+		IIterativeMatrixSolverBase* IterativeSolver = (IIterativeMatrixSolverBase*)MatrixSolver.Get();
+
+		bool bForceSingleThreaded = false;
+		Id = 0;
+		for (int32 s = 0; s < NumSteps; ++s)
+		{
+
+			int32 SrcBuffer = Id;
+			Id = 1 - Id;
+
+			// Old solution is the guess.
+			IterativeSolver->SolveWithGuess(Tmp[SrcBuffer], Tmp[SrcBuffer], Tmp[Id]);
+
+		}
+	}
+	else
+	{
+		bool bForceSingleThreaded = false;
+		Id = 0;
+		for (int32 s = 0; s < NumSteps; ++s)
+		{
+
+			int32 SrcBuffer = Id;
+			Id = 1 - Id;
+
+			MatrixSolver->Solve(Tmp[SrcBuffer], Tmp[Id]);
+		}
+	}
+	
+
+}
+
+void FDiffusionIntegrator::GetPositions(TArray<FVector3d>& PositionArray) const
+{
+	Linearize(Tmp[Id], PositionArray);
+}
+
+bool FDiffusionIntegrator::Linearize(const FSOAPositions& PositionalVector, TArray<FVector3d>& LinearArray) const
+{
+	// Number of positions
+
+	const int32 Num = PositionalVector.XVector.rows();
+
+	// early out if the x,y,z arrays in the PositionalVector have different lengths
+	if (!PositionalVector.bHasSize(Num))
+	{
+		return false;
+	}
+
+	// 
+	const auto& IndexToVtxId = VtxLinearization.ToId();
+	const int32 MaxVtxId = IndexToVtxId.Num(); // NB: this is really max_used + 1 in the mesh.  See  FDynamicMesh3::MaxVertexID()
+
+
+
+	LinearArray.Empty(MaxVtxId);
+	LinearArray.AddUninitialized(MaxVtxId);
+
+	for (int32 i = 0; i < Num; ++i)
+	{
+		const int32 VtxId = IndexToVtxId[i];
+
+		LinearArray[VtxId] = FVector3d(PositionalVector.XVector.coeff(i), PositionalVector.YVector.coeff(i), PositionalVector.ZVector.coeff(i));
+	}
+
+	return true;
+}
+
+
+TUniquePtr<FSparseMatrixD> FLaplacianDiffusionMeshSmoother::ConstructDiffusionOperator( const ELaplacianWeightScheme Scheme,
+																						const FDynamicMesh3& DynamicMesh,
+																						bool& bIsOperatorSymmetric,
+																						FVertexLinearization& Linearization,
+																						TArray<int32>* EdgeVtx ) 
+{
+	bIsOperatorSymmetric = bIsSymmetricLaplacian(Scheme);
+
+
+	// Construct the laplacian, and extract the mapping for vertices (VtxLinearization)
+	TUniquePtr<FSparseMatrixD> Laplacian = ConstructLaplacian(Scheme, DynamicMesh, VtxLinearization, &EdgeVerts);
+	Laplacian->makeCompressed();
+	return Laplacian;
+};
+
+TUniquePtr<FSparseMatrixD> FBiHarmonicDiffusionMeshSmoother::ConstructDiffusionOperator( const ELaplacianWeightScheme Scheme,
+																						 const FDynamicMesh3& DynamicMesh,
+																						 bool& bIsOperatorSymmetric,
+																						 FVertexLinearization& Linearization,
+																						 TArray<int32>* EdgeVtx ) 
+{
+	bIsOperatorSymmetric = true;
+	
+	// Construct the laplacian, and extract the mapping for vertices (VtxLinearization)
+	TUniquePtr<FSparseMatrixD> Laplacian = ConstructLaplacian(Scheme, DynamicMesh, VtxLinearization, &EdgeVerts);
+	const FSparseMatrixD& L = *Laplacian;
+	TUniquePtr<FSparseMatrixD> MatrixOperator(new FSparseMatrixD());
+	
+	bool bIsLaplacianSymmetric = bIsSymmetricLaplacian(Scheme);
+
+	if (bIsLaplacianSymmetric)
+	{
+		
+		*MatrixOperator = -1. * L * L;
+	}
+	else
+	{
+		*MatrixOperator = -1. * L.transpose() * L;
+	}
+
+	MatrixOperator->makeCompressed();
+	return MatrixOperator;
+};
+
+
+void MeshSmoothingOperators::ComputeSmoothing_BiHarmonic(const ELaplacianWeightScheme WeightScheme, const FDynamicMesh3& OriginalMesh, 
+	                                                     const double Speed, const double Intensity, const int32 NumIterations, TArray<FVector3d>& PositionArray)
+{
+
+
+
+	// This is equivalent to taking a single backward Euler time step of bi-harmonic diffusion
+	// where L is the Laplacian (Del^2) , and L^T L is an approximation of the Del^4.
+	// 
+	// dp/dt = - k*k L^T L[p]
+	// with 
+	// weight = 1 / (k * Sqrt[dt] )
+	//
+	// p^{n+1} + dt * k * k L^TL [p^{n+1}] = p^{n}
+	//
+	// re-write as
+	// L^TL[p^{n+1}] + weight * weight p^{n+1} = weight * weight p^{n}
+
+#ifndef EIGEN_MPL2_ONLY
+	const EMatrixSolverType MatrixSolverType = EMatrixSolverType::LTL;
+#else
+	// const EMatrixSolverType MatrixSolverType = EMatrixSolverType::LU;
+	// const EMatrixSolverType MatrixSolverType = EMatrixSolverType::PCG;
+
+	// The Symmetric Laplacians are SPD, and so are the LtL Operators
+	const EMatrixSolverType MatrixSolverType = (bIsSymmetricLaplacian(WeightScheme))? EMatrixSolverType::PCG : EMatrixSolverType::LU;
+	
+#endif
+
+
+	FString DebugLogString = FString::Printf(TEXT("Biharmonic Smoothing of mesh with %d verts "), OriginalMesh.VertexCount()) + LaplacianSchemeName(WeightScheme) + MatrixSolverName(MatrixSolverType);
+
+	FScopedDurationTimeLogger Timmer(DebugLogString);
+
+	FBiHarmonicDiffusionMeshSmoother BiHarmonicDiffusionSmoother(OriginalMesh, WeightScheme);
+
+	BiHarmonicDiffusionSmoother.Integrate_BackwardEuler(MatrixSolverType, NumIterations, Speed, Intensity);
+
+	BiHarmonicDiffusionSmoother.GetPositions(PositionArray);
+
+
+}
+
+void MeshSmoothingOperators::ComputeSmoothing_ImplicitBiHarmonicPCG( const ELaplacianWeightScheme WeightScheme, const FDynamicMesh3& OriginalMesh, 
+	                                                                 const double Speed, const double Weight, const int32 MaxIterations, TArray<FVector3d>& PositionArray)
+{
+
+	// This is equivalent to taking a single backward Euler time step of bi-harmonic diffusion
+	// where L is the Laplacian (Del^2) , and L^T L is an approximation of the Del^4.
+	// 
+	// dp/dt = - k*k L^T L[p]
+	// with 
+	// weight = 1 / (k * Sqrt[dt] )
+	//
+	// p^{n+1} + dt * k * k L^TL [p^{n+1}] = p^{n}
+	//
+	// re-write as
+	// L^TL[p^{n+1}] + weight * weight p^{n+1} = weight * weight p^{n}
+
+	FString DebugLogString = FString::Printf(TEXT("PCG Biharmonic Smoothing of mesh with %d verts "), OriginalMesh.VertexCount()) + LaplacianSchemeName(WeightScheme);
+
+	FScopedDurationTimeLogger Timmer(DebugLogString);
+
+	if (MaxIterations < 1) return;
+
+	FCGBiHarmonicMeshSmoother Smoother(OriginalMesh, WeightScheme);
+
+	// Treat all vertices as constraints with the same weight
+	const bool bPostFix = false;
+
+	for (int32 VertId : OriginalMesh.VertexIndicesItr())
+	{
+		FVector3d Pos = OriginalMesh.GetVertex(VertId);
+
+		Smoother.AddConstraint(VertId, Weight, Pos, bPostFix);
+	}
+
+	Smoother.SetMaxIterations(MaxIterations);
+	Smoother.SetTolerance(1.e-4);
+
+	bool bSuccess = Smoother.ComputeSmoothedMeshPositions(PositionArray);
+
+}
+
+void  MeshSmoothingOperators::ComputeSmoothing_Diffusion( const ELaplacianWeightScheme WeightScheme, const FDynamicMesh3& OriginalMesh, bool bForwardEuler,
+	                                                      const double Speed, const double Intensity, const int32 IterationCount, TArray<FVector3d>& PositionArray)
+{
+
+#ifndef EIGEN_MPL2_ONLY
+	const EMatrixSolverType MatrixSolverType = EMatrixSolverType::LTL;
+#else
+	const EMatrixSolverType MatrixSolverType = EMatrixSolverType::LU;
+	//const EMatrixSolverType MatrixSolverType = EMatrixSolverType::PCG;
+	//const EMatrixSolverType MatrixSolverType = EMatrixSolverType::BICGSTAB;
+#endif
+
+	FString DebugLogString = FString::Printf(TEXT("Diffusion Smoothing of mesh with %d verts"), OriginalMesh.VertexCount());
+	if (!bForwardEuler)
+	{
+		DebugLogString += MatrixSolverName(MatrixSolverType);
+	}
+
+	FScopedDurationTimeLogger Timmer(DebugLogString);
+
+	if (IterationCount < 1) return;
+
+	FLaplacianDiffusionMeshSmoother Smoother(OriginalMesh, WeightScheme);
+
+	if (bForwardEuler)
+	{
+		Smoother.Integrate_ForwardEuler(IterationCount, Speed, Intensity);
+	}
+	else
+	{
+		Smoother.Integrate_BackwardEuler(MatrixSolverType, IterationCount, Speed, Intensity);
+	}
+
+	Smoother.GetPositions(PositionArray);
+};
+
+
+TUniquePtr<MeshDeformingOperators::IConstrainedMeshOperator> MeshDeformingOperators::ConstructConstrainedMeshDeformer(const ELaplacianWeightScheme WeightScheme, const FDynamicMesh3& DynamicMesh)
+{
+	TUniquePtr<MeshDeformingOperators::IConstrainedMeshOperator> Deformer(new FConstrainedMeshDeformer(DynamicMesh, WeightScheme));
+
+	return Deformer;
+}
+
+
+TUniquePtr<MeshDeformingOperators::IConstrainedMeshOperator> MeshDeformingOperators::ConstructConstrainedMeshSmoother(const ELaplacianWeightScheme WeightScheme, const FDynamicMesh3& DynamicMesh)
+{
+	TUniquePtr<MeshDeformingOperators::IConstrainedMeshOperator> Deformer(new FBiHarmonicMeshSmoother(DynamicMesh, WeightScheme));
+
+	return Deformer;
+}
\ No newline at end of file
