move GeometryProcessing plugin from NotForLicensees to Experimental

diff --git a/Engine/Plugins/Experimental/GeometryProcessing/Source/MeshSolverUtilities/Private/LaplacianMeshSmoother.h b/Engine/Plugins/Experimental/GeometryProcessing/Source/MeshSolverUtilities/Private/LaplacianMeshSmoother.h
new file mode 100644
--- /dev/null
+++ b/Engine/Plugins/Experimental/GeometryProcessing/Source/MeshSolverUtilities/Private/LaplacianMeshSmoother.h
@@ -0,0 +1,249 @@
+// Copyright 1998-2019 Epic Games, Inc. All Rights Reserved.
+
+#pragma once
+
+#include "CoreMinimal.h"
+#include "DynamicMesh3.h"
+#include "ConstrainedPoissonSolver.h"
+#include "FSOAPositions.h"
+#include "MatrixSolver.h"
+#include "MeshSmoothingUtilities.h"
+#include "MeshElementLinearizations.h"
+
+
+
+class FConstrainedMeshOperator : public MeshDeformingOperators::IConstrainedMeshOperator
+{
+public:
+	typedef FConstrainedSolver::FConstraintPosition   FConstraintPosition;
+
+
+	FConstrainedMeshOperator(const FDynamicMesh3& DynamicMesh, const ELaplacianWeightScheme Scheme, const EMatrixSolverType MatrixSolverType);
+	virtual ~FConstrainedMeshOperator() override {}
+
+	// Add constraint associated with given vertex id.
+	void AddConstraint(const int32 VtxId, const double Weight, const FVector3d& Pos, const bool bPostFix) override;
+
+
+	bool UpdateConstraintPosition(const int32 VtxId, const FVector3d& Position, const bool bPostFix) override;
+
+	// The underlying solver will have to refactor the matrix if this is done
+	bool UpdateConstraintWeight(const int32 VtxId, const double Weight) override;
+
+	// Clear all constraints associated with this smoother
+	void ClearConstraints() override { FConstrainedMeshOperator::ClearConstraintPositions(); FConstrainedMeshOperator::ClearConstraintWeights(); }
+
+	void ClearConstraintWeights() override { ConstraintWeightMap.Empty();    bConstraintWeightsDirty = true; }
+
+	void ClearConstraintPositions() override { ConstraintPositionMap.Empty();  bConstraintPositionsDirty = true; }
+
+
+	// Test if for constraint associated with given vertex id.
+	bool IsConstrained(const int32 VtxId) const override;
+
+	
+
+	virtual bool Deform(TArray<FVector3d>& PositionBuffer) override { return false;  }
+
+protected:
+
+	// Sync constraints with internal solver.
+	void UpdateSolverConstraints();
+
+
+	void ExtractVertexPositions(const FDynamicMesh3& DynamicMesh, FSOAPositions& Positions) const ;
+
+	//void UpdateMeshWithConstraints();
+	// Respect any bPostFix constraints by moving those vertices to position defined by said constraint.
+	void UpdateWithPostFixConstraints(FSOAPositions& PositionalVector) const;
+
+
+	// converts the positionalvector to a TArray<FVector3d> where the offset in the array is implicitly the
+	// VtxId in the mesh, and not ness the matrix row id
+	// NB: the resulting array is treated as sparse and may have un-initialized elements.
+	bool Linearize(const FSOAPositions& PositionalVector, TArray<FVector3d>& LinearArray) const;
+
+
+protected:
+
+	// The Key (int32) here is the vertex index not vertex ID
+	// making it the same as the matrix row
+	TMap<int32, FConstraintPosition> ConstraintPositionMap;
+	TMap<int32, double>              ConstraintWeightMap;
+
+	bool bConstraintPositionsDirty = true;
+	bool bConstraintWeightsDirty   = true;
+
+	// Cache the vertex count.
+	int32                 VertexCount;
+
+	// Used to map between VtxId and vertex Index in linear vector..
+	FVertexLinearization  VtxLinearization;
+
+	// I don't know if we want to keep this after the constructor
+	TUniquePtr<FSparseMatrixD>             Laplacian;
+	TArray<int32>                          EdgeVerts;
+
+	// Actual solver that manages the various linear algebra bits.
+	TUniquePtr<FConstrainedSolver>         ConstrainedSolver;
+};
+
+
+class FConstrainedMeshDeformer : public FConstrainedMeshOperator
+{
+public:
+	FConstrainedMeshDeformer(const FDynamicMesh3& DynamicMesh, const ELaplacianWeightScheme LaplacianType);
+	~FConstrainedMeshDeformer() override {}
+
+	bool Deform(TArray<FVector3d>& PositionBuffer) override;
+
+private:
+
+	FSOAPositions LaplacianVectors;
+	FSOAPositions OriginalVertexPositions;
+};
+
+
+class FBiHarmonicMeshSmoother : public FConstrainedMeshOperator
+{
+public:
+	typedef FConstrainedMeshOperator         MyBaseType;
+
+	FBiHarmonicMeshSmoother(const FDynamicMesh3& DynamicMesh, const ELaplacianWeightScheme Scheme) :
+		MyBaseType(DynamicMesh, Scheme, EMatrixSolverType::LU)
+	{}
+
+	bool Deform(TArray<FVector3d>& UpdatedPositions) override 
+	{
+		return	ComputeSmoothedMeshPositions(UpdatedPositions);
+	}
+
+	// (Direct) Solve the constrained system and populate the UpdatedPositions with the result 
+	bool ComputeSmoothedMeshPositions(TArray<FVector3d>& UpdatedPositions);
+
+};
+
+
+// NB: This conjugate gradient solver could be updated to use  solveWithGuess() method on the iterative solver
+class FCGBiHarmonicMeshSmoother : public FConstrainedMeshOperator
+{
+public:
+	typedef FConstrainedMeshOperator      MyBaseType;
+
+	FCGBiHarmonicMeshSmoother(const FDynamicMesh3& DynamicMesh, const ELaplacianWeightScheme Scheme) :
+		MyBaseType(DynamicMesh, Scheme, EMatrixSolverType::BICGSTAB)
+	{}
+
+	bool Deform(TArray<FVector3d>& UpdatedPositions) override
+	{
+		return	ComputeSmoothedMeshPositions(UpdatedPositions);
+	}
+
+	void SetMaxIterations(int32 MaxIterations) 
+	{ 
+		IIterativeMatrixSolverBase* Solver = ConstrainedSolver->GetMatrixSolverIterativeBase(); 
+		if (Solver)
+		{
+			Solver->SetIterations(MaxIterations);
+		}
+	}
+
+	void SetTolerance(double Tol)
+	{
+		IIterativeMatrixSolverBase* Solver = ConstrainedSolver->GetMatrixSolverIterativeBase();
+		if (Solver)
+		{
+			Solver->SetTolerance(Tol);
+		}
+	}
+
+	// (Iterative) Solve the constrained system and populate the UpdatedPositions with the result 
+	bool ComputeSmoothedMeshPositions(TArray<FVector3d>& UpdatedPositions);
+
+
+};
+
+
+class FDiffusionIntegrator
+{
+public:
+
+	FDiffusionIntegrator(const FDynamicMesh3& DynamicMesh, const ELaplacianWeightScheme Scheme);
+	virtual ~FDiffusionIntegrator() {}
+
+	void Integrate_ForwardEuler(int32 NumSteps, double Speed, double);
+
+	// Note: 
+	void Integrate_BackwardEuler(const EMatrixSolverType MatrixSolverType, int32 NumSteps, double Speed, double Intensity);
+
+	void GetPositions(TArray<FVector3d>& PositionArray) const;
+
+protected:
+
+	// The derived class has to implement this.
+	virtual TUniquePtr<FSparseMatrixD> ConstructDiffusionOperator(const ELaplacianWeightScheme Scheme, 
+		                                                          const FDynamicMesh3& Mesh,
+		                                                          bool& bIsOperatorSymmetric,
+		                                                          FVertexLinearization& Linearization,
+		                                                          TArray<int32>* EdgeVtx) = 0 ;
+
+	void Initialize(const FDynamicMesh3& DynamicMesh, const ELaplacianWeightScheme Scheme);
+
+protected:
+	bool Linearize(const FSOAPositions& PositionalVector, TArray<FVector3d>& LinearArray) const;
+
+	// Cache the vertex count.
+	int32 VertexCount;
+
+	// Used to map vertex ID to entry in linear vector..
+	FVertexLinearization  VtxLinearization;
+
+	// I don't know if we want to keep this after the constructor
+	bool                                   bIsSymmetric;
+	TUniquePtr<FSparseMatrixD>             DiffusionOperator;
+	TArray<int32>                          EdgeVerts;
+	FSparseMatrixD::Scalar                 MinDiagonalValue;
+
+	FSOAPositions                       Tmp[2];
+	int32                               Id; // double buffer id
+};
+
+
+class FLaplacianDiffusionMeshSmoother : public FDiffusionIntegrator
+{
+public:
+
+	FLaplacianDiffusionMeshSmoother(const FDynamicMesh3& DynamicMesh, const ELaplacianWeightScheme Scheme)
+	: FDiffusionIntegrator(DynamicMesh, Scheme)
+	{ 
+		Initialize(DynamicMesh, Scheme);
+	}
+	
+protected:
+	TUniquePtr<FSparseMatrixD> ConstructDiffusionOperator( const ELaplacianWeightScheme Scheme,
+														   const FDynamicMesh3& Mesh,
+		                                                   bool& bIsOperatorSymmetric,
+		                                                   FVertexLinearization& Linearization,
+		                                                   TArray<int32>* EdgeVtx) override ;
+	
+
+};
+
+class  FBiHarmonicDiffusionMeshSmoother : public FDiffusionIntegrator
+{
+public:
+
+	FBiHarmonicDiffusionMeshSmoother(const FDynamicMesh3& DynamicMesh, const ELaplacianWeightScheme Scheme)
+	: FDiffusionIntegrator(DynamicMesh, Scheme)
+	{
+		Initialize(DynamicMesh, Scheme);
+	}
+
+protected:
+	TUniquePtr<FSparseMatrixD> ConstructDiffusionOperator( const ELaplacianWeightScheme Scheme,
+														   const FDynamicMesh3& Mesh,
+														   bool& bIsOperatorSymmetric,
+														   FVertexLinearization& Linearization,
+														   TArray<int32>* EdgeVtx) override;
+};
+
diff --git a/Engine/Plugins/Experimental/GeometryProcessing/Source/MeshSolverUtilities/Private/LaplacianOperators.cpp b/Engine/Plugins/Experimental/GeometryProcessing/Source/MeshSolverUtilities/Private/LaplacianOperators.cpp
new file mode 100644
--- /dev/null
+++ b/Engine/Plugins/Experimental/GeometryProcessing/Source/MeshSolverUtilities/Private/LaplacianOperators.cpp
@@ -0,0 +1,1164 @@
+// Copyright 1998-2019 Epic Games, Inc. All Rights Reserved.
+
+#include "LaplacianOperators.h"
+
+#include <cmath> // double version of sqrt
+#include <vector> // used by eigen to initialize sparse matrix
+
+#if defined(_MSC_VER) && USING_CODE_ANALYSIS
+#pragma warning(push)
+#pragma warning(disable : 6011)
+#pragma warning(disable : 6387)
+#pragma warning(disable : 6313)
+#pragma warning(disable : 6294)
+#endif
+THIRD_PARTY_INCLUDES_START
+#include <Eigen/Dense> // for Matrix4d in testing
+THIRD_PARTY_INCLUDES_END
+#if defined(_MSC_VER) && USING_CODE_ANALYSIS
+#pragma warning(pop)
+#endif
+
+
+#define LAPLACIAN_SKIP_BOUNDARY 0
+
+
+
+FString LaplacianSchemeName(const ELaplacianWeightScheme Scheme)
+{
+	FString LaplacianName;
+	switch (Scheme)
+	{
+	case ELaplacianWeightScheme::ClampedCotangent:
+		LaplacianName = FString(TEXT("Clamped Cotangent Laplacian"));
+		break;
+	case ELaplacianWeightScheme::Cotangent:
+		LaplacianName = FString(TEXT("Cotangent Laplacian"));
+		break;
+	case ELaplacianWeightScheme::Umbrella:
+		LaplacianName = FString(TEXT("Umbrella Laplacian"));
+		break;
+	case ELaplacianWeightScheme::MeanValue:
+		LaplacianName = FString(TEXT("MeanValue Laplacian"));
+		break;
+	case ELaplacianWeightScheme::Uniform:
+		LaplacianName = FString(TEXT("Uniform Laplacian"));
+		break;
+	case ELaplacianWeightScheme::Valence:
+		LaplacianName = FString(TEXT("Valence Laplacian"));
+		break;
+	default:
+		check(0 && "Unknown Laplacian Weight Scheme Enum");
+	}
+
+	return LaplacianName;
+}
+
+
+// Utility to compute the number of elements in the sparse laplacian matrix
+int32 ComputeNumMatrixElements(const FDynamicMesh3& DynamicMesh, const TArray<int32>& ToVtxId)
+{
+	const int32 NumVerts = ToVtxId.Num();
+	TArray<int32> OneRingSize;
+	{
+		OneRingSize.SetNumUninitialized(NumVerts);
+
+		for (int32 i = 0; i < NumVerts; ++i)
+		{
+			const int32 VertId = ToVtxId[i];
+			OneRingSize[i] = DynamicMesh.GetVtxEdgeCount(VertId);
+		}
+	}
+
+	// Compute the total number of entries in the sparse matrix
+	int32 NumMatrixEntries = 0;
+	{
+		for (int32 i = 0; i < NumVerts; ++i)
+		{
+			NumMatrixEntries += 1 + OneRingSize[i]; // myself plus my neighbors
+		}
+
+	}
+
+	return NumMatrixEntries;
+}
+
+TUniquePtr<FSparseMatrixD> ConstructUniformLaplacian(const FDynamicMesh3& DynamicMesh, FVertexLinearization& VertexMap, TArray<int32>* BoundaryVerts)
+{
+
+
+	typedef FSparseMatrixD::Scalar    ScalarT;
+	typedef Eigen::Triplet<ScalarT>  MatrixTripletT;
+
+	if (BoundaryVerts)
+	{
+		// empty
+		BoundaryVerts->Empty();
+	}
+
+	// Sync the mapping between the mesh vertex ids and their offsets in a nomimal linear array.
+	VertexMap.Reset(DynamicMesh);
+
+	const TArray<int32>& ToMeshV = VertexMap.ToId();
+	const TArray<int32>& ToIndex = VertexMap.ToIndex();
+	const int32 NumVerts = VertexMap.NumVerts();
+
+	// Eigen constructs a sparse matrix from a linearized array of matrix entries.
+
+	std::vector<MatrixTripletT> MatrixTripelList;
+	{
+		int32 NumMatrixEntries = ComputeNumMatrixElements(DynamicMesh, ToMeshV);
+		MatrixTripelList.reserve(NumMatrixEntries);
+	}
+
+	// Construct Laplacian Matrix: loop over verts constructing the corresponding matrix row.
+
+	for (int32 i = 0; i < NumVerts; ++i)
+	{
+		const int32 VertId = ToMeshV[i];
+
+		if (DynamicMesh.IsBoundaryVertex(VertId))
+		{
+			if (BoundaryVerts)
+			{
+				BoundaryVerts->Add(VertId);
+			}
+#if LAPLACIAN_SKIP_BOUNDARY == 1  
+			MatrixTripelList.push_back(MatrixTripletT(i, i, 0.));
+			continue;
+#endif 
+		}
+
+
+		ScalarT CenterWeight = ScalarT(0); // equal and opposite the sum of the neighbor weights
+		for (int NeighborVertId : DynamicMesh.VtxVerticesItr(VertId))
+		{
+			const int32 j = ToIndex[NeighborVertId];
+
+
+			ScalarT NeighborWeight = ScalarT(1);
+			CenterWeight += NeighborWeight;
+
+			// add the neighbor 
+			MatrixTripelList.push_back(MatrixTripletT(i, j, NeighborWeight));
+
+		}
+		// add the center
+		MatrixTripelList.push_back(MatrixTripletT(i, i, -CenterWeight));
+
+	}
+
+	// SparseMatrix doesn't have (SparseMatrix&& ) constructor 
+
+	TUniquePtr<FSparseMatrixD> LaplacianMatrix(new FSparseMatrixD(NumVerts, NumVerts));
+	LaplacianMatrix->setFromTriplets(MatrixTripelList.begin(), MatrixTripelList.end());
+
+	LaplacianMatrix->makeCompressed();
+
+	return LaplacianMatrix;
+
+}
+
+TUniquePtr<FSparseMatrixD> ConstructUmbrellaLaplacian(const FDynamicMesh3& DynamicMesh, FVertexLinearization& VertexMap, TArray<int32>* BoundaryVerts)
+{
+
+
+	typedef FSparseMatrixD::Scalar    ScalarT;
+	typedef Eigen::Triplet<ScalarT>  MatrixTripletT;
+
+	if (BoundaryVerts)
+	{
+		// empty
+		BoundaryVerts->Empty();
+	}
+
+	// Sync the mapping between the mesh vertex ids and their offsets in a nomimal linear array.
+	VertexMap.Reset(DynamicMesh);
+
+	const TArray<int32>& ToMeshV = VertexMap.ToId();
+	const TArray<int32>& ToIndex = VertexMap.ToIndex();
+	const int32 NumVerts = VertexMap.NumVerts();
+
+	// Eigen constructs a sparse matrix from a linearized array of matrix entries.
+
+	std::vector<MatrixTripletT> MatrixTripelList;
+	{
+		int32 NumMatrixEntries = ComputeNumMatrixElements(DynamicMesh, ToMeshV);
+		MatrixTripelList.reserve(NumMatrixEntries);
+	}
+
+	// Cache valency of each vertex.
+	// Number of non-zero elements in the i'th row = 1 + OneRingSize(i)
+	TArray<int32> OneRingSize;
+	{
+		OneRingSize.SetNumUninitialized(NumVerts);
+
+		for (int32 i = 0; i < NumVerts; ++i)
+		{
+			const int32 VertId = ToMeshV[i];
+			OneRingSize[i] = DynamicMesh.GetVtxEdgeCount(VertId);
+		}
+	}
+
+	// Construct Laplacian Matrix: loop over verts constructing the corresponding matrix row.
+
+	for (int32 i = 0; i < NumVerts; ++i)
+	{
+		const int32 VertId = ToMeshV[i];
+		const int32 Valence = OneRingSize[i];
+		double InvValence = (Valence != 0) ? 1. / double(Valence) : 0.;
+
+		if (DynamicMesh.IsBoundaryVertex(VertId))
+		{
+			if (BoundaryVerts)
+			{
+				BoundaryVerts->Add(VertId);
+			}
+#if LAPLACIAN_SKIP_BOUNDARY == 1  
+			MatrixTripelList.push_back(MatrixTripletT(i, i, 0.));
+			continue;
+#endif 
+		}
+
+		for (int NeighborVertId : DynamicMesh.VtxVerticesItr(VertId))
+		{
+			const int32 j = ToIndex[NeighborVertId];
+
+			// add the neighbor 
+			MatrixTripelList.push_back(MatrixTripletT(i, j, InvValence));
+
+		}
+		// add the center
+		MatrixTripelList.push_back(MatrixTripletT(i, i, -ScalarT(1)));
+
+	}
+
+	// SparseMatrix doesn't have (SparseMatrix&& ) constructor 
+
+	TUniquePtr<FSparseMatrixD> LaplacianMatrix(new FSparseMatrixD(NumVerts, NumVerts));
+	LaplacianMatrix->setFromTriplets(MatrixTripelList.begin(), MatrixTripelList.end());
+
+	LaplacianMatrix->makeCompressed();
+
+	return LaplacianMatrix;
+
+}
+
+TUniquePtr<FSparseMatrixD> ConstructValenceWeightedLaplacian(const FDynamicMesh3& DynamicMesh, FVertexLinearization& VertexMap, TArray<int32>* BoundaryVerts)
+{
+
+
+	typedef FSparseMatrixD::Scalar    ScalarT;
+	typedef Eigen::Triplet<ScalarT>  MatrixTripletT;
+
+	if (BoundaryVerts)
+	{
+		// empty
+		BoundaryVerts->Empty();
+	}
+
+	// Sync the mapping between the mesh vertex ids and their offsets in a nomimal linear array.
+	VertexMap.Reset(DynamicMesh);
+
+	const TArray<int32>& ToMeshV = VertexMap.ToId();
+	const TArray<int32>& ToIndex = VertexMap.ToIndex();
+	const int32 NumVerts = VertexMap.NumVerts();
+
+	// Cache valency of each vertex.
+	// Number of non-zero elements in the i'th row = 1 + OneRingSize(i)
+	TArray<int32> OneRingSize;
+	{
+		OneRingSize.SetNumUninitialized(NumVerts);
+
+		for (int32 i = 0; i < NumVerts; ++i)
+		{
+			const int32 VertId = ToMeshV[i];
+			OneRingSize[i] = DynamicMesh.GetVtxEdgeCount(VertId);
+		}
+	}
+
+
+	// Compute the total number of entries in the sparse matrix
+	int32 NumMatrixEntries = 0;
+	{
+		for (int32 i = 0; i < NumVerts; ++i)
+		{
+			NumMatrixEntries += 1 + OneRingSize[i]; // myself plus my neighbors
+		}
+
+	}
+	std::vector<MatrixTripletT> MatrixTripelList;
+	MatrixTripelList.reserve(NumMatrixEntries);
+
+
+
+	// Construct Laplacian Matrix: loop over verts constructing the corresponding matrix row.
+
+	for (int32 i = 0; i < NumVerts; ++i)
+	{
+		const int32 VertId = ToMeshV[i];
+		const int32 IOneRingSize = OneRingSize[i];
+
+		if (DynamicMesh.IsBoundaryVertex(VertId))
+		{
+			if (BoundaryVerts)
+			{
+				BoundaryVerts->Add(VertId);
+			}
+#if LAPLACIAN_SKIP_BOUNDARY == 1  
+			MatrixTripelList.push_back(MatrixTripletT(i, i, 0.));
+			continue;
+#endif 
+		}
+
+
+		ScalarT CenterWeight = ScalarT(0); // equal and opposite the sum of the neighbor weights
+		for (int NeighborVertId : DynamicMesh.VtxVerticesItr(VertId))
+		{
+			const int32 j = ToIndex[NeighborVertId];
+			const int32 JOneRingSize = OneRingSize[j];
+
+
+			ScalarT NeighborWeight = ScalarT(1) / std::sqrt(IOneRingSize + JOneRingSize);
+			CenterWeight += NeighborWeight;
+
+			// add the neighbor 
+			MatrixTripelList.push_back(MatrixTripletT(i, j, NeighborWeight));
+
+		}
+		// add the center
+		MatrixTripelList.push_back(MatrixTripletT(i, i, -CenterWeight));
+
+	}
+
+	// SparseMatrix doesn't have (SparseMatrix&& ) constructor 
+
+	TUniquePtr<FSparseMatrixD> LaplacianMatrix(new FSparseMatrixD(NumVerts, NumVerts));
+	LaplacianMatrix->setFromTriplets(MatrixTripelList.begin(), MatrixTripelList.end());
+
+	LaplacianMatrix->makeCompressed();
+
+	return LaplacianMatrix;
+
+}
+
+
+/**
+* The per-triangle data used in constructing the cotangent weighted laplacian.
+*
+*/
+class CotanTriangleData
+{
+public:
+
+	typedef FVector3d TriangleVertices[3];
+
+	CotanTriangleData() = default;
+
+	CotanTriangleData(const CotanTriangleData& Other)
+	{
+
+		Cotangent[0] = Other.Cotangent[0];
+		Cotangent[1] = Other.Cotangent[1];
+		Cotangent[2] = Other.Cotangent[2];
+
+		VoronoiArea[0] = Other.VoronoiArea[0];
+		VoronoiArea[1] = Other.VoronoiArea[1];
+		VoronoiArea[2] = Other.VoronoiArea[2];
+
+		OppositeEdge[0] = Other.OppositeEdge[0];
+		OppositeEdge[1] = Other.OppositeEdge[1];
+		OppositeEdge[2] = Other.OppositeEdge[2];
+
+	}
+
+
+	CotanTriangleData(const FDynamicMesh3& DynamicMesh, int32 TriId)
+	{
+		Initialize(DynamicMesh, TriId);
+	}
+
+	void Initialize(const FDynamicMesh3& DynamicMesh, int32 SrcTriId)
+	{
+		TriId = SrcTriId;
+
+		// edges: ab, bc, ca
+		FIndex3i EdgeIds = DynamicMesh.GetTriEdges(TriId);
+
+		FVector3d VertA, VertB, VertC;
+		DynamicMesh.GetTriVertices(TriId, VertA, VertB, VertC);
+
+		const FVector3d EdgeAB(VertB - VertA);
+		const FVector3d EdgeAC(VertC - VertA);
+		const FVector3d EdgeBC(VertC - VertB);
+
+
+		OppositeEdge[0] = EdgeIds[1]; // EdgeBC is opposite vert A
+		OppositeEdge[1] = EdgeIds[2]; // EdgeAC is opposite vert B
+		OppositeEdge[2] = EdgeIds[0]; // EdgeAB is opposite vert C
+
+
+		// NB: didn't use VectorUtil::Area() so we can re-use the Edges.
+		//     also this formulation of area is always positive.
+
+		const double TwiceArea = EdgeAB.Cross(EdgeAC).Length();
+
+		// NB: Area = 1/2 || EdgeA X EdgeB ||  where EdgeA and EdgeB are any two edges in the triangle.
+
+		// Compute the Voronoi areas
+		// 
+		// From Discrete Differential-Geometry Operators for Triangulated 2-Manifolds (Meyer, Desbrun, Schroder, Barr)
+		// http://www.geometry.caltech.edu/pubs/DMSB_III.pdf
+		//    Given triangle P,Q, R the voronoi area at P is given by
+		//    Area = (1/8) * ( |PR|**2 Cot<Q  + |PQ|**2 Cot<R ) 
+
+		if (TwiceArea > 2. * SmallTriangleArea)
+		{
+
+
+
+			// Compute the cotangent of the angle between V1 and V2 
+			// as the ratio  V1.Dot.V2 / || V1 Cross V2 ||
+
+			// Cotangent[i] is cos(theta)/sin(theta) at the i'th vertex.
+
+			Cotangent[0] =  EdgeAB.Dot(EdgeAC) / TwiceArea;
+			Cotangent[1] = -EdgeAB.Dot(EdgeBC) / TwiceArea;
+			Cotangent[2] =  EdgeAC.Dot(EdgeBC) / TwiceArea;
+
+			if (bIsObtuse())
+			{
+				const double Area = 0.5 * TwiceArea;
+
+				// Voronoi inappropriate case.  Instead use Area(T)/2 at obtuse corner
+				// and Area(T) / 4 at the other corners.
+
+				VoronoiArea[0] = 0.25 * Area;
+				VoronoiArea[1] = 0.25 * Area;
+				VoronoiArea[2] = 0.25 * Area;
+
+				for (int i = 0; i < 3; ++i)
+				{
+					if (Cotangent[i] < 0.)
+					{
+						VoronoiArea[i] = 0.5 * Area;
+					}
+				}
+			}
+			else
+			{
+				// If T is non-obtuse.
+
+				const double EdgeABSqLength = EdgeAB.SquaredLength();
+				const double EdgeACSqLength = EdgeAC.SquaredLength();
+				const double EdgeBCSqLength = EdgeBC.SquaredLength();
+
+				VoronoiArea[0] = EdgeABSqLength * Cotangent[1] + EdgeACSqLength * Cotangent[2];
+				VoronoiArea[1] = EdgeABSqLength * Cotangent[0] + EdgeBCSqLength * Cotangent[2];
+				VoronoiArea[2] = EdgeACSqLength * Cotangent[0] + EdgeBCSqLength * Cotangent[1];
+
+				const double Inv8 = .125; // 1/8
+
+				VoronoiArea[0] *= Inv8;
+				VoronoiArea[1] *= Inv8;
+				VoronoiArea[2] *= Inv8;
+			}
+		}
+		else
+		{
+			// default small triangle - equilateral 
+			double CotOf60 = 1. / FMath::Sqrt(3.f);
+			Cotangent[0] = CotOf60;
+			Cotangent[1] = CotOf60;
+			Cotangent[2] = CotOf60;
+
+			VoronoiArea[0] = SmallTriangleArea / 3.;
+			VoronoiArea[1] = SmallTriangleArea / 3.;
+			VoronoiArea[2] = SmallTriangleArea / 3.;
+		}
+	}
+
+	int32 GetLocalEdgeIdx(const int32 DynamicsMeshEdgeId) const
+	{
+		int32 Result = -1;
+		if (DynamicsMeshEdgeId == OppositeEdge[0])
+		{
+			Result = 0;
+		}
+		else if (DynamicsMeshEdgeId == OppositeEdge[1])
+		{
+			Result = 1;
+		}
+		else if (DynamicsMeshEdgeId == OppositeEdge[2])
+		{
+			Result = 2;
+		}
+		return Result;
+	}
+
+	/** helper to return the cotangent of the angle opposite the given edge
+	*
+	*   @param DynamicsMeshEdgeId is the id used by FDynamicMesh3 for this edge.
+	*   @param bValid will false on return if the requested edge is not part of this triangle
+	*   @return Cotangent of the opposite angle.
+	*/
+	double GetOpposingCotangent(const int32 DynamicsMeshEdgeId, bool& bValid) const
+	{
+
+		double OpposingCotangent = -1.;
+		bValid = false;
+		int32 LocalEdgeIdx = GetLocalEdgeIdx(DynamicsMeshEdgeId);
+		if (LocalEdgeIdx > -1)
+		{
+			OpposingCotangent = Cotangent[LocalEdgeIdx];
+			bValid = true;
+		}
+
+		return OpposingCotangent;
+	}
+
+	double GetOpposingCotangent(const int32 DynamicsMeshEdgeId) const
+	{
+		bool bValid;
+		double OpposingCotangent = GetOpposingCotangent(DynamicsMeshEdgeId, bValid);
+		checkSlow(bValid);
+		return OpposingCotangent;
+	}
+
+	bool bIsObtuse() const
+	{
+		return (Cotangent[0] < 0.f || Cotangent[1] < 0.f || Cotangent[2] < 0.f);
+	}
+
+
+	// The "floor" for triangle area.
+	// NB: the cotan laplacian has terms ~ 1/TriArea
+	//     and the deformation matrix has terms ~ 1/TriArea**2
+
+	static constexpr double SmallTriangleArea = 1.e-4;
+
+	// testing
+	int32 TriId = -1;
+
+	/** Total byte count: 6 double + 3 int32 = 60 bytes.   */
+
+	// Cotangent[i] is cos(theta)/sin(theta) at the i'th vertex.
+
+	double Cotangent[3] = { 0. };
+
+	// VoronoiArea[i] is the voronoi area about the i'th vertex in this triangle.
+
+	double VoronoiArea[3] = { 0. };
+
+	// OppositeEdge[i] = Corresponding DynamicsMesh3::EdgeId for the edge that is opposite
+	// the i'th vertex in this triangle
+
+	int32 OppositeEdge[3] = { -1 };
+
+};
+
+
+/**
+* The per-triangle data used in constructing the mean-value weighted laplacian.
+*
+*/
+class MeanValueTriangleData
+{
+
+public:
+
+	MeanValueTriangleData(const MeanValueTriangleData& Other)
+		: TriId(Other.TriId)
+		, TriVtxIds(Other.TriVtxIds)
+		, TriEdgeIds(Other.TriEdgeIds)
+		, bDegenerate(Other.bDegenerate)
+	{
+		EdgeLength[0] = Other.EdgeLength[0];
+		EdgeLength[1] = Other.EdgeLength[1];
+		EdgeLength[2] = Other.EdgeLength[2];
+
+		TanHalfAngle[0] = Other.TanHalfAngle[0];
+		TanHalfAngle[1] = Other.TanHalfAngle[1];
+		TanHalfAngle[2] = Other.TanHalfAngle[2];
+	}
+
+	void Initialize(const FDynamicMesh3& DynamicMesh, int32 SrcTriId)
+	{
+		TriId = SrcTriId;
+
+		// VertAId, VertBId, VertCId
+		TriVtxIds = DynamicMesh.GetTriangle(TriId);
+		TriEdgeIds = DynamicMesh.GetTriEdges(TriId);
+
+		FVector3d VertA, VertB, VertC;
+		DynamicMesh.GetTriVertices(TriId, VertA, VertB, VertC);
+
+		const FVector3d EdgeAB(VertB - VertA);
+		const FVector3d EdgeAC(VertC - VertA);
+		const FVector3d EdgeBC(VertC - VertB);
+
+		EdgeLength[0] = EdgeAB.Length();
+		EdgeLength[1] = EdgeAC.Length();
+		EdgeLength[2] = EdgeBC.Length();
+
+		constexpr double SmallEdge = 1e-4;
+
+		bDegenerate = (EdgeLength[0] < SmallEdge || EdgeLength[1] < SmallEdge || EdgeLength[2] < SmallEdge);
+
+		// Compute tan(angle/2) = Sqrt[ (1-cos) / (1 + cos)]
+
+		const double ABdotAC =  EdgeAB.Dot(EdgeAC);
+		const double BCdotBA = -EdgeBC.Dot(EdgeAB);
+		const double CAdotCB =  EdgeAC.Dot(EdgeBC);
+
+
+		const double RegularizingConst = 1.e-6; // keeps us from dividing by zero when making tan[180/2] = sin[90]/cos[90] = inf
+		TanHalfAngle[0] = (EdgeLength[0] * EdgeLength[1] - ABdotAC) / (EdgeLength[0] * EdgeLength[1] + ABdotAC + RegularizingConst);
+		TanHalfAngle[1] = (EdgeLength[0] * EdgeLength[2] - BCdotBA) / (EdgeLength[0] * EdgeLength[2] + BCdotBA + RegularizingConst);
+		TanHalfAngle[2] = (EdgeLength[1] * EdgeLength[2] - CAdotCB) / (EdgeLength[1] * EdgeLength[2] + CAdotCB + RegularizingConst);
+
+		// The ABS is just a precaution.. mathematically these should all be positive, but very small angles may result in negative values.
+
+		TanHalfAngle[0] = std::sqrt(FMath::Abs(TanHalfAngle[0])); // at vertA
+		TanHalfAngle[1] = std::sqrt(FMath::Abs(TanHalfAngle[1])); // at vertB
+		TanHalfAngle[2] = std::sqrt(FMath::Abs(TanHalfAngle[2])); // at vertC
+#if 0
+		// testing
+		double totalAngle = 2. * (std::atan(TanHalfAngle[0]) + std::atan(TanHalfAngle[1]) + std::atan(TanHalfAngle[2]));
+		
+		double angleError = M_PI - totalAngle;
+#endif
+	}
+
+	// return Tan(angle / 2) for the corner indicated by this vert id.
+	double GetTanHalfAngle(int32 VtxId) const
+	{
+		int32 Offset = 0;
+
+		while (VtxId != TriVtxIds[Offset])
+		{
+			Offset++;
+			checkSlow(Offset < 3);
+		}
+
+		return TanHalfAngle[Offset];
+	}
+
+	// return the length of the indicated edge
+	double GetEdgeLenght(int32 EdgeId) const
+	{
+		int32 Offset = 0;
+
+		while (EdgeId != TriEdgeIds[Offset])
+		{
+			Offset++;
+			checkSlow(Offset < 3);
+		}
+
+		return EdgeLength[Offset];
+	}
+
+
+	int32 TriId = -1;
+	FIndex3i TriVtxIds;
+	FIndex3i TriEdgeIds;
+
+	bool bDegenerate = true;
+
+	double EdgeLength[3] = { 0. };
+	double TanHalfAngle[3] = { 0. };
+
+
+};
+
+/**
+* Return and array in triangle order that holds the per-triangle derived data needed
+*/
+template <typename TriangleDataType>
+TArray<TriangleDataType>  ConstructTriangleDataArray(const FDynamicMesh3& DynamicMesh, const FTriangleLinearization& TriangleLinearization)
+{
+	TArray< TriangleDataType > TriangleDataArray;
+
+	const int32 NumTris = TriangleLinearization.NumTris();
+	TriangleDataArray.SetNumUninitialized(NumTris);
+
+	const auto& ToTriIdx = TriangleLinearization.ToIndex();
+
+	for (int32 i = 0; i < NumTris; ++i)
+	{
+		// Current triangle
+
+		const int32 TriIdx = ToTriIdx[i];
+
+		// Compute all the geometric data needed for this triangle.
+
+		TriangleDataArray[i].Initialize(DynamicMesh, TriIdx);
+	}
+
+	return TriangleDataArray;
+}
+
+
+TUniquePtr<FSparseMatrixD> ConstructCotangentLaplacian(const FDynamicMesh3& DynamicMesh, FVertexLinearization& VertexMap, FSparseMatrixD& AreaMatrix, TArray<int32>* BoundaryVerts)
+{
+
+	if (BoundaryVerts)
+	{
+		// empty
+		BoundaryVerts->Empty();
+	}
+
+	typedef FSparseMatrixD::Scalar  ScalarT;
+	typedef Eigen::Triplet<ScalarT> MatrixTripletT;
+
+
+	// Create a mapping between the ordering of the vertex indices in the Dynamic Mesh
+	// and the actual storage. 
+	VertexMap.Reset(DynamicMesh);
+
+	const TArray<int32>& ToMeshV = VertexMap.ToId();
+	const TArray<int32>& ToIndex = VertexMap.ToIndex();
+	const int32 NumVerts = VertexMap.NumVerts();
+
+
+	// Create the mapping of triangles
+
+	FTriangleLinearization TriangleMap(DynamicMesh);
+
+	const TArray<int32>& ToMeshTri = TriangleMap.ToId();
+	const TArray<int32>& ToTriIdx = TriangleMap.ToIndex();
+	const int32 NumTris = TriangleMap.NumTris();
+
+
+	// Clear space for the areas
+	std::vector<MatrixTripletT> DiagonalTriplets;
+	DiagonalTriplets.reserve(NumVerts);
+
+	FSparseMatrixD Diagonals;
+	Diagonals.reserve(NumVerts);
+
+
+	// Create an array that holds all the geometric information we need for each triangle.
+
+	TArray<CotanTriangleData>  CotangentTriangleDataArray = ConstructTriangleDataArray<CotanTriangleData>(DynamicMesh, TriangleMap);
+
+
+	// Eigen constructs a sparse matrix from a linearized array of matrix entries.
+
+	std::vector<MatrixTripletT> MatrixTripelList;
+	{
+		int32 NumMatrixEntries = ComputeNumMatrixElements(DynamicMesh, ToMeshV);
+		MatrixTripelList.reserve(NumMatrixEntries);
+	}
+
+
+	// Construct Laplacian Matrix: loop over verts constructing the corresponding matrix row.
+	//                             store the id of the boundary verts for later use.
+
+	for (int32 i = 0; i < NumVerts; ++i)
+	{
+		const int32 IVertId = ToMeshV[i]; // I - the row
+
+
+		if (DynamicMesh.IsBoundaryVertex(IVertId))
+		{
+			if (BoundaryVerts)
+			{
+				BoundaryVerts->Add(IVertId);
+			}
+#if LAPLACIAN_SKIP_BOUNDARY == 1  
+			MatrixTripelList.push_back(MatrixTripletT(i, i, 0.));
+			DiagonalTriplets.push_back(MatrixTripletT(i, i, 1.));
+			continue;
+#endif 
+		}
+
+
+		// Compute the Voronoi area for this vertex.
+		double WeightArea = 0.;
+		for (int32 TriId : DynamicMesh.VtxTrianglesItr(IVertId))
+		{
+			const int32 TriIdx = ToTriIdx[TriId];
+			const CotanTriangleData& TriData = CotangentTriangleDataArray[TriIdx];
+
+
+			// The three VertIds for this triangle.
+			const FIndex3i TriVertIds = DynamicMesh.GetTriangle(TriId);
+
+			// Which of the corners is IVertId?
+			int32 Offset = 0;
+			while (TriVertIds[Offset] != IVertId)
+			{
+				Offset++;
+				checkSlow(Offset < 3);
+			}
+
+			WeightArea += TriData.VoronoiArea[Offset];
+		}
+
+
+		double WeightII = 0.; // accumulate to equal and opposite the sum of the neighbor weights
+
+		// for each connecting edge
+
+		for (int32 EdgeId : DynamicMesh.VtxEdgesItr(IVertId))
+		{
+			// [v0, v1, t0, t1]:  NB: both t0 & t1 exist since IVert isn't a boundary vert.
+			FIndex4i Edge = DynamicMesh.GetEdge(EdgeId);
+
+
+			// the other vert in the edge - identifies the matrix column
+			const int32 JVertId = (Edge[0] == IVertId) ? Edge[1] : Edge[0];  // J - the column
+
+			checkSlow(JVertId != IVertId);
+
+			// Get the cotangents for this edge.
+
+			const int32 Tri0Idx = ToTriIdx[Edge[2]];
+			const CotanTriangleData& Tri0Data = CotangentTriangleDataArray[Tri0Idx];
+			const double CotanAlpha = Tri0Data.GetOpposingCotangent(EdgeId);
+
+
+			// The second triangle will be invalid if this is an edge!
+
+			const double CotanBeta = (Edge[3] != FDynamicMesh3::InvalidID) ? CotangentTriangleDataArray[ToTriIdx[Edge[3]]].GetOpposingCotangent(EdgeId) : 0.;
+
+			double WeightIJ = 0.5 * (CotanAlpha + CotanBeta);
+			WeightII += WeightIJ;
+
+			const int32 j = ToIndex[JVertId];
+
+			MatrixTripelList.push_back(MatrixTripletT(i, j, WeightIJ));
+
+
+		}
+
+		MatrixTripelList.push_back(MatrixTripletT(i, i, -WeightII));
+
+		DiagonalTriplets.push_back(MatrixTripletT(i, i, WeightArea));
+
+
+	}
+
+	Diagonals.setFromTriplets(DiagonalTriplets.begin(), DiagonalTriplets.end());
+	AreaMatrix.swap(Diagonals);
+	AreaMatrix.makeCompressed();
+
+	TUniquePtr<FSparseMatrixD> LaplacianMatrix(new FSparseMatrixD(NumVerts, NumVerts));
+	LaplacianMatrix->setFromTriplets(MatrixTripelList.begin(), MatrixTripelList.end());
+	LaplacianMatrix->makeCompressed();
+
+
+	return LaplacianMatrix;
+
+}
+
+TUniquePtr<FSparseMatrixD> ConstructCotangentLaplacian(const FDynamicMesh3& DynamicMesh, FVertexLinearization& VertexMap, const bool bClampWeights, TArray<int32>* BoundaryVerts)
+{
+
+	if (BoundaryVerts)
+	{
+		// empty
+		BoundaryVerts->Empty();
+	}
+
+	typedef FSparseMatrixD::Scalar  ScalarT;
+	typedef Eigen::Triplet<ScalarT> MatrixTripletT;
+
+
+	// Create a mapping between the ordering of the vertex indices in the Dynamic Mesh
+	// and the actual storage. 
+	VertexMap.Reset(DynamicMesh);
+
+	const TArray<int32>& ToMeshV = VertexMap.ToId();
+	const TArray<int32>& ToIndex = VertexMap.ToIndex();
+	const int32 NumVerts = VertexMap.NumVerts();
+
+
+	// Map the triangles.
+
+	FTriangleLinearization TriangleMap(DynamicMesh);
+
+	const TArray<int32>& ToMeshTri = TriangleMap.ToId();
+	const TArray<int32>& ToTriIdx = TriangleMap.ToIndex();
+	const int32 NumTris = TriangleMap.NumTris();
+
+	// Create an array that holds all the geometric information we need for each triangle.
+
+	TArray<CotanTriangleData>  CotangentTriangleDataArray = ConstructTriangleDataArray<CotanTriangleData>(DynamicMesh, TriangleMap);
+
+	// Eigen constructs a sparse matrix from a linearized array of matrix entries.
+
+	std::vector<MatrixTripletT> MatrixTripelList;
+	{
+		int32 NumMatrixEntries = ComputeNumMatrixElements(DynamicMesh, ToMeshV);
+		MatrixTripelList.reserve(NumMatrixEntries);
+	}
+
+
+	// Construct Laplacian Matrix: loop over verts constructing the corresponding matrix row.
+	//                             skipping the boundary verts for later use.
+
+	for (int32 i = 0; i < NumVerts; ++i)
+	{
+		const int32 IVertId = ToMeshV[i]; // I - the row
+
+
+		if (DynamicMesh.IsBoundaryVertex(IVertId))
+		{
+			if (BoundaryVerts)
+			{
+				BoundaryVerts->Add(IVertId);
+			}
+#if LAPLACIAN_SKIP_BOUNDARY == 1  
+			MatrixTripelList.push_back(MatrixTripletT(i, i, 0.));
+			continue;
+#endif 
+		}
+
+
+		// Compute the Voronoi area for this vertex.
+		double WeightArea = 0.;
+		for (int32 TriId : DynamicMesh.VtxTrianglesItr(IVertId))
+		{
+			const int32 TriIdx = ToTriIdx[TriId];
+			const CotanTriangleData& TriData = CotangentTriangleDataArray[TriIdx];
+
+
+			// The three VertIds for this triangle.
+			const FIndex3i TriVertIds = DynamicMesh.GetTriangle(TriId);
+
+			// Which of the corners is IVertId?
+			int32 Offset = 0;
+			while (TriVertIds[Offset] != IVertId)
+			{
+				Offset++;
+				checkSlow(Offset < 3);
+			}
+
+			WeightArea += TriData.VoronoiArea[Offset];
+		}
+
+
+		double WeightII = 0.; // accumulate to equal and opposite the sum of the neighbor weights
+
+		// for each connecting edge
+
+		for (int32 EdgeId : DynamicMesh.VtxEdgesItr(IVertId))
+		{
+			// [v0, v1, t0, t1]:  NB: both t0 & t1 exist since IVert isn't a boundary vert.
+			FIndex4i Edge = DynamicMesh.GetEdge(EdgeId);
+
+
+			// the other vert in the edge - identifies the matrix column
+			const int32 JVertId = (Edge[0] == IVertId) ? Edge[1] : Edge[0];  // J - the column
+
+			checkSlow(JVertId != IVertId);
+
+			// Get the cotangents for this edge.
+
+			const int32 Tri0Idx = ToTriIdx[Edge[2]];
+			const CotanTriangleData& Tri0Data = CotangentTriangleDataArray[Tri0Idx];
+			const double CotanAlpha = Tri0Data.GetOpposingCotangent(EdgeId);
+
+			// The second triangle will be invalid if this is an edge!
+
+			const double CotanBeta = (Edge[3] != FDynamicMesh3::InvalidID) ? CotangentTriangleDataArray[ToTriIdx[Edge[3]]].GetOpposingCotangent(EdgeId) : 0.;
+
+			double WeightIJ = 0.5 * (CotanAlpha + CotanBeta);
+
+			// clamp the weight
+			if (bClampWeights)
+			{
+				WeightIJ = FMath::Clamp(WeightIJ, -1.e5 * WeightArea, 1.e5 * WeightArea);
+			}
+
+			WeightII += WeightIJ;
+
+			const int32 j = ToIndex[JVertId];
+
+			MatrixTripelList.push_back(MatrixTripletT(i, j, WeightIJ / WeightArea));
+
+		}
+
+		MatrixTripelList.push_back(MatrixTripletT(i, i, -WeightII / WeightArea));
+
+	}
+
+#if 0
+	// used in degubing to insepct the matrix for testing
+	std::vector<MatrixTripletT> Tmp = MatrixTripelList;
+
+	auto SortFunctor = [](const MatrixTripletT& A, const MatrixTripletT& B)->bool
+	{
+		return (A.value() < B.value());
+	};
+
+	std::sort(Tmp.begin(), Tmp.end(), SortFunctor);
+
+	auto ReverseSortFunctor = [](const MatrixTripletT& A, const MatrixTripletT& B)->bool
+	{
+		return (A.value() > B.value());
+	};
+
+	std::sort(Tmp.begin(), Tmp.end(), ReverseSortFunctor);
+#endif 
+
+	TUniquePtr<FSparseMatrixD> LaplacianMatrix(new FSparseMatrixD(NumVerts, NumVerts));
+	LaplacianMatrix->setFromTriplets(MatrixTripelList.begin(), MatrixTripelList.end());
+	LaplacianMatrix->makeCompressed();
+
+	return LaplacianMatrix;
+
+}
+
+
+
+
+
+
+
+TUniquePtr<FSparseMatrixD> ConstructMeanValueWeightLaplacian(const FDynamicMesh3& DynamicMesh, FVertexLinearization& VertexMap, TArray<int32>* BoundaryVerts)
+{
+
+	if (BoundaryVerts)
+	{
+		// empty
+		BoundaryVerts->Empty();
+	}
+
+	typedef FSparseMatrixD::Scalar  ScalarT;
+	typedef Eigen::Triplet<ScalarT> MatrixTripletT;
+
+
+	// Create a mapping between the ordering of the vertex indices in the Dynamic Mesh
+	// and the actual storage. 
+	VertexMap.Reset(DynamicMesh);
+
+	const TArray<int32>& ToMeshV = VertexMap.ToId();
+	const TArray<int32>& ToIndex = VertexMap.ToIndex();
+	const int32 NumVerts = VertexMap.NumVerts();
+
+
+	// Map the triangles.
+
+	FTriangleLinearization TriangleMap(DynamicMesh);
+
+	const TArray<int32>& ToMeshTri = TriangleMap.ToId();
+	const TArray<int32>& ToTriIdx = TriangleMap.ToIndex();
+	const int32 NumTris = TriangleMap.NumTris();
+
+
+	// Create an array that holds all the geometric information we need for each triangle.
+
+	TArray<MeanValueTriangleData>  TriangleDataArray = ConstructTriangleDataArray<MeanValueTriangleData>(DynamicMesh, TriangleMap);
+
+
+	// Eigen constructs a sparse matrix from a linearized array of matrix entries.
+
+	std::vector<MatrixTripletT> MatrixTripelList;
+	{
+		int32 NumMatrixEntries = ComputeNumMatrixElements(DynamicMesh, ToMeshV);
+		MatrixTripelList.reserve(NumMatrixEntries);
+	}
+
+
+	// Construct Laplacian Matrix: loop over verts constructing the corresponding matrix row.
+	//                             skipping the boundary verts for later use.
+
+	for (int32 i = 0; i < NumVerts; ++i)
+	{
+		const int32 IVertId = ToMeshV[i]; // I - the row
+
+
+		if (DynamicMesh.IsBoundaryVertex(IVertId))
+		{
+			if (BoundaryVerts)
+			{
+				BoundaryVerts->Add(IVertId);
+			}
+#if LAPLACIAN_SKIP_BOUNDARY == 1  
+			MatrixTripelList.push_back(MatrixTripletT(i, i, 0.));
+			continue;
+#endif 
+		}
+
+
+
+		double WeightII = 0.; // accumulate to equal and opposite the sum of the neighbor weights
+
+		// for each connecting edge
+
+		for (int32 EdgeId : DynamicMesh.VtxEdgesItr(IVertId))
+		{
+			// [v0, v1, t0, t1]:  NB: both t0 & t1 exist since IVert isn't a boundary vert.
+			FIndex4i Edge = DynamicMesh.GetEdge(EdgeId);
+
+			// the other vert in the edge - identifies the matrix column
+			const int32 JVertId = (Edge[0] == IVertId) ? Edge[1] : Edge[0];  // J - the column
+
+			// Get the cotangents for this edge.
+
+			const int32 Tri0Idx = ToTriIdx[Edge[2]];
+			const auto& Tri0Data = TriangleDataArray[Tri0Idx];
+			double TanHalfAngleSum = Tri0Data.GetTanHalfAngle(IVertId);
+			double EdgeLength = FMath::Max(1.e-5, Tri0Data.GetEdgeLenght(EdgeId)); // Clamp the length
+
+			// The second triangle will be invalid if this is an edge!
+
+			TanHalfAngleSum += (Edge[3] != FDynamicMesh3::InvalidID) ? TriangleDataArray[ToTriIdx[Edge[3]]].GetTanHalfAngle(IVertId) : 0.;
+
+			double WeightIJ = TanHalfAngleSum / EdgeLength;
+			WeightII += WeightIJ;
+
+			const int32 j = ToIndex[JVertId];
+
+			MatrixTripelList.push_back(MatrixTripletT(i, j, WeightIJ));
+		}
+
+		MatrixTripelList.push_back(MatrixTripletT(i, i, -WeightII));
+
+
+	}
+
+	TUniquePtr<FSparseMatrixD> LaplacianMatrix(new FSparseMatrixD(NumVerts, NumVerts));
+	LaplacianMatrix->setFromTriplets(MatrixTripelList.begin(), MatrixTripelList.end());
+	LaplacianMatrix->makeCompressed();
+
+	return LaplacianMatrix;
+
+}
+
+TUniquePtr<FSparseMatrixD> ConstructLaplacian(const ELaplacianWeightScheme Scheme, const FDynamicMesh3& DynamicMesh, FVertexLinearization& VertexMap, TArray<int32>* BoundaryVerts)
+{
+
+	switch (Scheme)
+	{
+	default:
+	case ELaplacianWeightScheme::Uniform:
+		return ConstructUniformLaplacian(DynamicMesh, VertexMap, BoundaryVerts);
+		break;
+	case ELaplacianWeightScheme::Umbrella:
+		return ConstructUmbrellaLaplacian(DynamicMesh, VertexMap, BoundaryVerts);
+		break;
+	case ELaplacianWeightScheme::Valence:
+		return ConstructValenceWeightedLaplacian(DynamicMesh, VertexMap, BoundaryVerts);
+		break;
+	case ELaplacianWeightScheme::Cotangent:
+	{
+		bool bClampWeights = false;
+		return ConstructCotangentLaplacian(DynamicMesh, VertexMap, bClampWeights, BoundaryVerts);
+		break;
+	}
+	case ELaplacianWeightScheme::ClampedCotangent:
+	{
+		bool bClampWeights = true;
+		return ConstructCotangentLaplacian(DynamicMesh, VertexMap, bClampWeights, BoundaryVerts);
+		break;
+	}
+	case ELaplacianWeightScheme::MeanValue:
+		return ConstructMeanValueWeightLaplacian(DynamicMesh, VertexMap, BoundaryVerts);
+		break;
+	}
+}
+
diff --git a/Engine/Plugins/Experimental/GeometryProcessing/Source/MeshSolverUtilities/Private/LaplacianOperators.h b/Engine/Plugins/Experimental/GeometryProcessing/Source/MeshSolverUtilities/Private/LaplacianOperators.h
new file mode 100644
--- /dev/null
+++ b/Engine/Plugins/Experimental/GeometryProcessing/Source/MeshSolverUtilities/Private/LaplacianOperators.h
@@ -0,0 +1,156 @@
+// Copyright 1998-2019 Epic Games, Inc. All Rights Reserved.
+
+#pragma once
+
+#include "CoreMinimal.h"
+#include "DynamicMesh3.h"
+
+#include "MeshElementLinearizations.h"
+
+#include "FSparseMatrixD.h"
+
+#include "MeshSmoothingUtilities.h" // for ELaplacianWeightScheme
+
+
+/**
+* Construct a sparse matrix representation of a uniform weighted Laplacian.
+* The uniform weighted Laplacian is defined solely in terms of the connectivity
+* of the mesh.  Note, by construction this should be a symmetric matrix.
+*
+* Row i represents the Laplacian at vert_i, the non-zero entries correspond to
+* the incident one-ring vertices vert_j.
+*
+* L_{ij} = 1                      if vert_j is in the one-ring of vert_i
+* L_{ii} = -Sum{ L_{ij}, j != i}
+*
+*
+* @param DynamicMesh3   The triangle mesh
+* @param VertexMap      Additional arrays used to map between vertexID and offset in a linear array (i.e. the row)
+* @param BoundaryVerts  Optionally capture the verts that are on mesh edges.
+*/
+TUniquePtr<FSparseMatrixD> ConstructUniformLaplacian(const FDynamicMesh3& DynamicMesh, FVertexLinearization& VertexMap, TArray<int32>* BoundaryVerts = NULL);
+
+/**
+* Construct a sparse matrix representation of an umbrella weighted Laplacian.
+* This Laplacian is defined solely in terms of the connectivity
+* of the mesh.  Note, there is no expectation that the resulting matrix will be symmetric.
+*
+* Row i represents the Laplacian at vert_i, the non-zero entries correspond to
+* the incident one-ring vertices vert_j.
+*
+* L_{ij} = 1 / valence(of i)                      if vert_j is in the one-ring of vert_i
+* L_{ii} = -Sum{ L_{ij}, j != i} = -1
+*
+*
+* @param DynamicMesh3   The triangle mesh
+* @param VertexMap      Additional arrays used to map between vertexID and offset in a linear array (i.e. the row)
+* @param BoundaryVerts  Optionally capture the verts that are on mesh edges.
+*/
+TUniquePtr<FSparseMatrixD> ConstructUmbrellaLaplacian(const FDynamicMesh3& DynamicMesh, FVertexLinearization& VertexMap, TArray<int32>* BoundaryVerts = NULL);
+
+/**
+* Construct a sparse matrix representation of a valence-weighted Laplacian.
+* The valence weighted Laplacian is defined solely in terms of the connectivity
+* of the mesh.  Note, by construction this should be a symmetric matrix.
+*
+* Row i represents the Laplacian at vert_i, the non-zero entries correspond to
+* the incident one-ring vertices vert_j.
+*
+* L_{ij} = 1/\sqrt(valence(i) + valence(j))   if vert_j is in the one-ring of vert_i
+* L_{ii} = -Sum{ L_{ij}, j != i}
+*
+*
+* @param DynamicMesh3   The triangle mesh
+* @param VertexMap      Additional arrays used to map between vertexID and offset in a linear array (i.e. the row)
+* @param BoundaryVerts  Optionally capture the verts that are on mesh edges.
+*/
+TUniquePtr<FSparseMatrixD>  ConstructValenceWeightedLaplacian(const FDynamicMesh3& DynamicMesh3, FVertexLinearization& VertexMap, TArray<int32>* BoundaryVerts = NULL);
+
+/**
+* Construct a sparse matrix representation using a cotangent-weighted  Laplacian.
+* NB: there is no reason to expect this to be a symmetric matrix.
+*
+* @param DynamicMesh3   The triangle mesh
+* @param VertexMap      Additional arrays used to map between vertexID and offset in a linear array (i.e. the row)
+* @param bClampWeights  Indicates if the off-diagonal weights should be clamped on construction.
+*                       in practice this is desirable when creating the biharmonic operator, but not the mean curvature flow operator
+* @param BoundaryVerts  Optionally capture the verts that are on mesh edges.
+*/
+
+TUniquePtr<FSparseMatrixD> ConstructCotangentLaplacian(const FDynamicMesh3& DynamicMesh, FVertexLinearization& VertexMap, const bool bClampWeights = false, TArray<int32>* BoundaryVerts = NULL);
+
+/**
+* Construct a sparse matrix representation using a cotangent-weighted  Laplacian.
+* but returns the result in two symmetric parts.
+*
+*  (AreaMatrix^{-1}) * L_hat  = Cotangent weighted Laplacian.
+*
+* @return L_hat = Laplacian Without Area weighting. - this will be symmetric
+*
+* @param On return - diagonal matrix holding the voronoi areas for each vertex.  Area 1 is assigned to any edge vertex
+*
+*
+* @param DynamicMesh3   The triangle mesh
+* @param VertexMap      Additional arrays used to map between vertexID and offset in a linear array (i.e. the row)
+* @param BoundaryVerts  Optionally capture the verts that are on mesh edges.
+*/
+
+TUniquePtr<FSparseMatrixD> ConstructCotangentLaplacian(const FDynamicMesh3& DynamicMesh, FVertexLinearization& VertexMap, FSparseMatrixD& AreaMatrix, TArray<int32>* BoundaryVerts = NULL);
+
+
+
+/**
+* Construct a sparse matrix representation using a meanvalue-weighted  Laplacian.
+* NB: there is no reason to expect this to be a symmetric matrix.
+*
+* @param DynamicMesh3   The triangle mesh
+* @param VertexMap      Additional arrays used to map between vertexID and offset in a linear array (i.e. the row)
+* @param bClampWeights  Indicates if the off-diagonal weights should be clamped on construction.
+*                       in practice this is desirable when creating the biharmonic operator, but not the mean curvature flow operator
+* @param BoundaryVerts  Optionally capture the verts that are on mesh edges.
+*/
+
+TUniquePtr<FSparseMatrixD> ConstructMeanValueWeightLaplacian(const FDynamicMesh3& DynamicMesh, FVertexLinearization& VertexMap, TArray<int32>* BoundaryVerts = NULL);
+
+TUniquePtr<FSparseMatrixD> ConstructLaplacian(const ELaplacianWeightScheme Scheme, const FDynamicMesh3& DynamicMesh, FVertexLinearization& VertexMap, TArray<int32>* BoundaryVerts = NULL);
+
+
+/**
+* Utility to map the enum names for debuging etc.
+*/
+FString LaplacianSchemeName(const ELaplacianWeightScheme Scheme);
+
+
+/**
+* Only certian laplacian operators are symmetric...
+*/
+static bool bIsSymmetricLaplacian(const ELaplacianWeightScheme Scheme)
+{
+	bool bSymmetric = false;
+	switch (Scheme)
+	{
+	case ELaplacianWeightScheme::ClampedCotangent:
+		bSymmetric = false;
+		break;
+	case ELaplacianWeightScheme::Cotangent:
+		bSymmetric = false;
+		break;
+	case ELaplacianWeightScheme::Umbrella:
+		bSymmetric = false;
+		break;
+	case ELaplacianWeightScheme::MeanValue:
+		bSymmetric = false;
+		break;
+	case ELaplacianWeightScheme::Uniform:
+		bSymmetric = true;
+		break;
+	case ELaplacianWeightScheme::Valence:
+		bSymmetric = true;
+		break;
+	default:
+		check(0);
+	}
+	return bSymmetric;
+}
+
+
diff --git a/Engine/Plugins/Experimental/GeometryProcessing/Source/MeshSolverUtilities/Private/MeshElementLinearizations.h b/Engine/Plugins/Experimental/GeometryProcessing/Source/MeshSolverUtilities/Private/MeshElementLinearizations.h
new file mode 100644
--- /dev/null
+++ b/Engine/Plugins/Experimental/GeometryProcessing/Source/MeshSolverUtilities/Private/MeshElementLinearizations.h
@@ -0,0 +1,113 @@
+// Copyright 1998-2019 Epic Games, Inc. All Rights Reserved.
+
+#pragma once
+
+#include "CoreMinimal.h"
+#include "DynamicMesh3.h"
+
+
+/**
+* Used linearize the VtxIds in a mesh as a single array and allow mapping from array offset to mesh VtxId.
+* Generally, the array offset will correspond to a matrix row when forming a Laplacian
+*/
+class FMeshElementLinearization
+{
+public:
+
+	FMeshElementLinearization() = default;
+
+	// Lookup   ToVtxId(Index) = VtxId;
+	const TArray<int32>& ToId() const { return ToIdMap; }
+
+	// Lookup   ToIndex(VtxId) = Index;  may return FDynamicMesh3::InvalidID 
+	const TArray<int32>& ToIndex() const { return ToIndexMap; }
+
+	int32 NumIds() const { return ToIdMap.Num(); }
+
+	// Following the FDynamicMesh3 convention this is really MaxId + 1
+	int32 MaxId() const { return ToIndexMap.Num(); }
+
+	void Empty() { ToIdMap.Empty();  ToIndexMap.Empty(); }
+
+	void Populate(const int32 MaxId, const int32 Count, FRefCountVector::IndexEnumerable ElementItr)
+	{
+		ToIndexMap.SetNumUninitialized(MaxId);
+		ToIdMap.SetNumUninitialized(Count);
+
+		for (int32 i = 0; i < MaxId; ++i)
+		{
+			ToIndexMap[i] = FDynamicMesh3::InvalidID;
+		}
+
+		// create the mapping
+		{
+			int32 N = 0;
+			for (int32 Id : ElementItr)
+			{
+				ToIdMap[N] = Id;
+				ToIndexMap[Id] = N;
+				N++;
+			}
+		}
+	}
+
+protected:
+	TArray<int32>  ToIdMap;
+	TArray<int32>  ToIndexMap;
+
+private:
+	FMeshElementLinearization(const FMeshElementLinearization&);
+};
+
+
+/**
+* Used linearize the VtxIds in a mesh as a single array and allow mapping from array offset to mesh VtxId.
+* Generally, the array offset will correspond to a matrix row when forming a Laplacian
+*/
+class FVertexLinearization : public FMeshElementLinearization
+{
+public:
+	FVertexLinearization() = default;
+	FVertexLinearization(const FDynamicMesh3& DynamicMesh)
+	{
+		Reset(DynamicMesh);
+	}
+
+	void Reset(const FDynamicMesh3& DynamicMesh)
+	{
+		Empty();
+		FMeshElementLinearization::Populate(DynamicMesh.MaxVertexID(), DynamicMesh.VertexCount(), DynamicMesh.VertexIndicesItr());
+	}
+
+	int32 NumVerts() const { return FMeshElementLinearization::NumIds(); }
+
+private:
+	FVertexLinearization(const FVertexLinearization&);
+};
+
+/**
+* Used linearize the TriIds in a mesh as a single array and allow mapping from array offset to mesh TriId.
+*
+*/
+class FTriangleLinearization : public FMeshElementLinearization
+{
+public:
+	FTriangleLinearization() = default;
+
+	FTriangleLinearization(const FDynamicMesh3& DynamicMesh)
+	{
+		Reset(DynamicMesh);
+	}
+
+	void Reset(const FDynamicMesh3& DynamicMesh)
+	{
+		Empty();
+		FMeshElementLinearization::Populate(DynamicMesh.MaxTriangleID(), DynamicMesh.TriangleCount(), DynamicMesh.TriangleIndicesItr());
+	}
+
+	int32 NumTris() const { return FMeshElementLinearization::NumIds(); }
+
+
+private:
+	FTriangleLinearization(const FTriangleLinearization&);
+};
diff --git a/Engine/Plugins/Experimental/GeometryProcessing/Source/MeshSolverUtilities/Private/MeshSolverUtilitiesModule.cpp b/Engine/Plugins/Experimental/GeometryProcessing/Source/MeshSolverUtilities/Private/MeshSolverUtilitiesModule.cpp
new file mode 100644
--- /dev/null
+++ b/Engine/Plugins/Experimental/GeometryProcessing/Source/MeshSolverUtilities/Private/MeshSolverUtilitiesModule.cpp
@@ -0,0 +1,19 @@
+// Copyright 1998-2019 Epic Games, Inc. All Rights Reserved.
+
+#include "MeshSolverUtilitiesModule.h"
+
+#define LOCTEXT_NAMESPACE "FMeshSolverUtilitiesModule"
+
+void FMeshSolverUtilitiesModule::StartupModule()
+{
+}
+
+void FMeshSolverUtilitiesModule::ShutdownModule()
+{
+	// This function may be called during shutdown to clean up your module.  For modules that support dynamic reloading,
+	// we call this function before unloading the module.
+}
+
+#undef LOCTEXT_NAMESPACE
+
+IMPLEMENT_MODULE(FMeshSolverUtilitiesModule, MeshSolverUtilities)
\ No newline at end of file
diff --git a/Engine/Plugins/Experimental/GeometryProcessing/Source/MeshSolverUtilities/Public/MeshSmoothingUtilities.h b/Engine/Plugins/Experimental/GeometryProcessing/Source/MeshSolverUtilities/Public/MeshSmoothingUtilities.h
new file mode 100644
--- /dev/null
+++ b/Engine/Plugins/Experimental/GeometryProcessing/Source/MeshSolverUtilities/Public/MeshSmoothingUtilities.h
@@ -0,0 +1,162 @@
+// Copyright 1998-2019 Epic Games, Inc. All Rights Reserved.
+
+#pragma once
+
+#include "CoreMinimal.h"
+#include "DynamicMesh3.h"
+
+
+
+enum class ELaplacianWeightScheme
+{
+	Uniform,
+	Umbrella,
+	Valence,
+	MeanValue,
+	Cotangent,
+	ClampedCotangent
+};
+
+namespace MeshSmoothingOperators
+{
+	/**
+	*
+	* Note: for discussion of implicit / explicit integration of diffusion and biharmonic equations 
+	*       see "Implicit Fairing of Irregular Meshes using Diffusion and Curvature Flow" - M Desbrun 99.
+	*       although the following suggests an additional source term could be included in the implicit solve for better accuracy.
+	*       or "Generalized Surface Flows for Mesh Processing" Eckstein et al. 2007
+	*/
+
+	/**
+	* This is equivalent to taking a single backward Euler time step of bi-harmonic diffusion
+	* where L is the Laplacian (Del^2) , and L^T L is an approximation of the Del^4.
+	*
+	* dp/dt = - k*k L^T L[p]
+	* with
+	* weight = 1 / (k * Sqrt[dt] )
+	*
+	* p^{n+1} + dt * k * k L^TL [p^{n+1}] = p^{n}
+	*
+	* re-write as
+	* L^TL[p^{n+1}] + weight * weight p^{n+1} = weight * weight p^{n}
+	
+	*
+	* The result is returned in the PositionArray
+	*/
+	void MESHSOLVERUTILITIES_API ComputeSmoothing_BiHarmonic(const ELaplacianWeightScheme WeightingScheme, const FDynamicMesh3& OriginalMesh,
+		const double Speed, const double Weight, const int32 NumIterations, TArray<FVector3d>& PositionArray);
+
+	void  MESHSOLVERUTILITIES_API ComputeSmoothing_ImplicitBiHarmonicPCG(const ELaplacianWeightScheme WeightScheme, const FDynamicMesh3& OriginalMesh,
+		const double Speed, const double Weight, const int32 MaxIterations, TArray<FVector3d>& PositionArray);
+
+	/**
+	* This is equivalent to forward or backward Euler time steps of the diffusion equation
+	*
+	* dp/dt = L[p]
+	*
+	* p^{n+1} = p^{n} + dt L[p^{n}]
+	*
+	* with dt = Speed / Max(|w_ii|)
+	*
+	* here w_ii are the diagonal values of L
+	*/
+	void  MESHSOLVERUTILITIES_API ComputeSmoothing_Diffusion(const ELaplacianWeightScheme WeightScheme, const FDynamicMesh3& OriginalMesh, bool bForwardEuler,
+		const double Speed, double Weight, const int32 NumIterations, TArray<FVector3d>& PositionArray);
+
+}
+
+namespace MeshDeformingOperators
+{
+	class MESHSOLVERUTILITIES_API IConstrainedMeshOperator
+	{
+	public:
+
+		virtual ~IConstrainedMeshOperator() {};
+
+		// Add or update a constraint associated with VtxId
+		virtual void AddConstraint(const int32 VtxId, const double Weight, const FVector3d& Position, const bool bPostFix) = 0;
+		
+		// Update or create a constraint position associated with @param VtxId
+		// @return true if a constraint weight is associated with @param VtxId
+		virtual bool UpdateConstraintPosition(const int32 VtxId, const FVector3d& Position, const bool bPostFix) = 0;
+		
+		// Update or create a constraint weight associated with @param VtxId
+		// @return true if a constraint position is associated with @param VtxId
+		virtual bool UpdateConstraintWeight(const int32 VtxId, const double Weight) = 0;
+		
+		// Clear all constraints (Positions and Weights)
+		virtual void ClearConstraints() = 0;
+		
+		// Clear all Constraint Weights
+		virtual void ClearConstraintWeights() = 0;
+		
+		// Clear all Constraint Positions
+		virtual void ClearConstraintPositions() = 0;
+
+
+		// Test if a non-zero weighted constraint is associated with VtxId 
+		virtual bool IsConstrained(const int32 VtxId) const = 0;
+		
+		// Returns the vertex locations of the deformed mesh. 
+		// Note the array may have empty elements as the index matches the Mesh based VtxId.  PositionBuffer[VtxId] = Pos.
+		virtual bool Deform(TArray<FVector3d>& PositionBuffer) = 0;
+	};
+
+	/**
+	*  Solves the linear system for p_vec 
+	* 
+	*         ( Transpose(L) * L   + (0  0      )  ) p_vec = source_vec + ( 0              )
+	*		  (                      (0 lambda^2)  )                      ( lambda^2 c_vec )
+	*
+	*   where:  L := laplacian for the mesh,
+	*           source_vec := Transpose(L)*L mesh_vertex_positions
+	*           lambda := weights
+	*           c_vec := constrained positions
+	*
+	* Expected Use:
+	*
+	*   // Create Deformation Solver from Mesh
+	*   TUniquePtr<IConstrainedMeshOperator>  MeshDeformer = ConstructConstrainedMeshDeformer(ELaplacianWeightScheme::ClampedCotangent, DynamicMesh);
+	*
+	*   // Add constraints.
+	*   for..
+	*   {
+	*   	int32 VtxId = ..; double Weight = ..; FVector3d TargetPos = ..;  bool bPostFix = ...;
+	*   	MeshDeformer->AddConstraint(VtxId, Weight, TargetPos, bPostFix);
+	*   }
+	*
+	*   // Solve for new mesh vertex locations
+	*   TArray<FVector3d> PositionBuffer;
+	*   MeshDeformer->Deform(PositionBuffer);
+	*
+	*   // Update Mesh? for (int32 VtxId : DynamicMesh.VertexIndices()) DynamicMesh.SetVertex(VtxId, PositionBuffer[VtxId]);
+	*   ...
+	* 
+	*   // Update constraint positions.
+	*   for ..
+	*   {
+	*   	int32 VtxId = ..;  FVector3d TargetPos = ..; bool bPostFix = ...;
+	*	    MeshDeformer->UpdateConstraintPosition(VtxId, TargetPos, bPostFix);
+	*   }
+	*
+	*   // Solve for new vertex locations.
+	*   MeshDeformer->Deform(PositionBuffer);
+	*   // Update Mesh?
+	*/
+	TUniquePtr<IConstrainedMeshOperator> MESHSOLVERUTILITIES_API ConstructConstrainedMeshDeformer(const ELaplacianWeightScheme WeightScheme, const FDynamicMesh3& DynamicMesh);
+
+	/**
+	*  Solves the linear system for p_vec
+	*
+	*         ( Transpose(L) * L   + (0  0      )  ) p_vec = ( 0              )
+	*		  (                      (0 lambda^2)  )         ( lambda^2 c_vec )
+	*
+	*   where:  L := laplacian for the mesh,
+	*           lambda := weights
+	*           c_vec := constrained positions
+	*
+	* Expected Use: same as the ConstrainedMeshDeformer above.
+	*
+	*/
+	TUniquePtr<IConstrainedMeshOperator> MESHSOLVERUTILITIES_API ConstructConstrainedMeshSmoother(const ELaplacianWeightScheme WeightScheme, const FDynamicMesh3& DynamicMesh);
+}
\ No newline at end of file
diff --git a/Engine/Plugins/Experimental/GeometryProcessing/Source/MeshSolverUtilities/Public/MeshSolverUtilitiesModule.h b/Engine/Plugins/Experimental/GeometryProcessing/Source/MeshSolverUtilities/Public/MeshSolverUtilitiesModule.h
new file mode 100644
--- /dev/null
+++ b/Engine/Plugins/Experimental/GeometryProcessing/Source/MeshSolverUtilities/Public/MeshSolverUtilitiesModule.h
@@ -0,0 +1,15 @@
+// Copyright 1998-2019 Epic Games, Inc. All Rights Reserved.
+
+#pragma once
+
+#include "CoreMinimal.h"
+#include "Modules/ModuleManager.h"
+
+class FMeshSolverUtilitiesModule : public IModuleInterface
+{
+public:
+
+	/** IModuleInterface implementation */
+	virtual void StartupModule() override;
+	virtual void ShutdownModule() override;
+};
\ No newline at end of file
