Add conversion from FVector3 to FColor
#rnx

diff --git a/Engine/Plugins/Experimental/GeometryProcessing/Source/GeometricObjects/Public/VectorTypes.h b/Engine/Plugins/Experimental/GeometryProcessing/Source/GeometricObjects/Public/VectorTypes.h
--- a/Engine/Plugins/Experimental/GeometryProcessing/Source/GeometricObjects/Public/VectorTypes.h
+++ b/Engine/Plugins/Experimental/GeometryProcessing/Source/GeometricObjects/Public/VectorTypes.h
@@ -57,6 +57,13 @@ struct FVector3
 		Z = (T)Vec.Z;
 	}
 
+
+	inline operator FColor() const
+	{
+		return FColor((float)X, (float)Y, (float)Z);
+	}
+
+
 	static FVector3<T> Zero()
 	{
 		return FVector3<T>((T)0, (T)0, (T)0);
