fix signed/unsigned mismatch

diff --git a/Engine/Plugins/Experimental/GeometryProcessing/Source/GeometryAlgorithms/Private/Arrangement2d.cpp b/Engine/Plugins/Experimental/GeometryProcessing/Source/GeometryAlgorithms/Private/Arrangement2d.cpp
--- a/Engine/Plugins/Experimental/GeometryProcessing/Source/GeometryAlgorithms/Private/Arrangement2d.cpp
+++ b/Engine/Plugins/Experimental/GeometryProcessing/Source/GeometryAlgorithms/Private/Arrangement2d.cpp
@@ -222,7 +222,7 @@ bool FArrangement2d::AttemptTriangulate(TArray<FIntVector> &Triangles, TArray<in
 		}
 
 		// copy uneaten triangles
-		for (size_t i = 0; i < TriNum; i++)
+		for (int i = 0; i < TriNum; i++)
 		{
 			if (!Eaten[i])
 			{
