Fix GeometricObjects' axis aligned box extents() to return half-diagonals, as expected by the code that uses it

diff --git a/Engine/Plugins/Experimental/GeometryProcessing/Source/GeometricObjects/Public/BoxTypes.h b/Engine/Plugins/Experimental/GeometryProcessing/Source/GeometricObjects/Public/BoxTypes.h
--- a/Engine/Plugins/Experimental/GeometryProcessing/Source/GeometricObjects/Public/BoxTypes.h
+++ b/Engine/Plugins/Experimental/GeometryProcessing/Source/GeometricObjects/Public/BoxTypes.h
@@ -222,7 +222,7 @@ struct FAxisAlignedBox3
 
 	FVector3<RealType> Extents() const
 	{
-		return FVector3<RealType>(Max.X - Min.X, Max.Y - Min.Y, Max.Z - Min.Z);
+		return (Max - Min) * (RealType).5;
 	}
 
 	void Contain(const FVector3<RealType>& V)
@@ -379,7 +379,7 @@ struct FAxisAlignedBox2
 
 	FVector2<RealType> Extents() const
 	{
-		return FVector2<RealType>(Max.X - Min.X, Max.Y - Min.Y);
+		return (Max - Min) * (RealType).5;
 	}
 
 	inline void Contain(const FVector2<RealType>& V)
