Change interval's extent() to match the change to the aabb extents() (make it half of max-min instead of just max-min)

diff --git a/Engine/Plugins/Experimental/GeometryProcessing/Source/GeometricObjects/Public/BoxTypes.h b/Engine/Plugins/Experimental/GeometryProcessing/Source/GeometricObjects/Public/BoxTypes.h
--- a/Engine/Plugins/Experimental/GeometryProcessing/Source/GeometricObjects/Public/BoxTypes.h
+++ b/Engine/Plugins/Experimental/GeometryProcessing/Source/GeometricObjects/Public/BoxTypes.h
@@ -33,7 +33,7 @@ struct FInterval1
 
 	RealType Extent() const
 	{
-		return Max - Min;
+		return (Max - Min)*(RealType).5;
 	}
 	RealType Length() const
 	{
