Copy header include fix from //UE4/Dev-Physics to Dev-Editor (//UE4/Dev-Editor)

diff --git a/Engine/Plugins/Experimental/GeometryProcessing/Source/GeometricObjects/Public/MeshQueries.h b/Engine/Plugins/Experimental/GeometryProcessing/Source/GeometricObjects/Public/MeshQueries.h
--- a/Engine/Plugins/Experimental/GeometryProcessing/Source/GeometricObjects/Public/MeshQueries.h
+++ b/Engine/Plugins/Experimental/GeometryProcessing/Source/GeometricObjects/Public/MeshQueries.h
@@ -5,6 +5,7 @@
 #include "Distance/DistPoint3Triangle3.h"
 #include "Intersection/IntrRay3Triangle3.h"
 #include "BoxTypes.h"
+#include "IndexTypes.h"
 
 template <class TriangleMeshType>
 class TMeshQueries
