Add a few new functions to FFrame3 and VectorUtil, cleanup and document VectorUtil functions
#rnx

diff --git a/Engine/Plugins/Experimental/GeometryProcessing/Source/GeometricObjects/Public/FrameTypes.h b/Engine/Plugins/Experimental/GeometryProcessing/Source/GeometricObjects/Public/FrameTypes.h
--- a/Engine/Plugins/Experimental/GeometryProcessing/Source/GeometricObjects/Public/FrameTypes.h
+++ b/Engine/Plugins/Experimental/GeometryProcessing/Source/GeometricObjects/Public/FrameTypes.h
@@ -204,23 +204,23 @@ struct FFrame3
 	/**
 	 * Project 3D point into plane and convert to UV coordinates in that plane
 	 * @param Pos 3D position
-	 * @param nPlaneNormalAxis which plane to project onto, identified by perpendicular normal. Default is 2, ie normal is Z, plane is (X,Y)
+	 * @param PlaneNormalAxis which plane to project onto, identified by perpendicular normal. Default is 2, ie normal is Z, plane is (X,Y)
 	 * @return 2D coordinates in UV plane, relative to origin
 	 */
-	FVector2<RealType> ToPlaneUV(const FVector3<RealType>& Pos, int nPlaneNormalAxis) const
+	FVector2<RealType> ToPlaneUV(const FVector3<RealType>& Pos, int PlaneNormalAxis) const
 	{
-		int nAxis0 = 0, nAxis1 = 1;
-		if (nPlaneNormalAxis == 0)
+		int Axis0 = 0, Axis1 = 1;
+		if (PlaneNormalAxis == 0)
 		{
-			nAxis0 = 2;
+			Axis0 = 2;
 		}
-		else if (nPlaneNormalAxis == 1)
+		else if (PlaneNormalAxis == 1)
 		{
-			nAxis1 = 2;
+			Axis1 = 2;
 		}
 		FVector3<RealType> LocalPos = Pos - Origin;
-		RealType U = LocalPos.Dot(GetAxis(nAxis0));
-		RealType V = LocalPos.Dot(GetAxis(nAxis1));
+		RealType U = LocalPos.Dot(GetAxis(Axis0));
+		RealType V = LocalPos.Dot(GetAxis(Axis1));
 		return FVector2<RealType>(U, V);
 	}
 
@@ -229,7 +229,7 @@ struct FFrame3
 	/**
 	 * Map a point from local UV plane coordinates to the corresponding 3D point in one of the planes of the frame
 	 * @param PosUV local UV plane coordinates
-	 * @param nPlaneNormalAxis which plane to map to, identified by perpendicular normal. Default is 2, ie normal is Z, plane is (X,Y)
+	 * @param PlaneNormalAxis which plane to map to, identified by perpendicular normal. Default is 2, ie normal is Z, plane is (X,Y)
 	 * @return 3D coordinates in frame's plane (including Origin translation)
 	 */
 	FVector3<RealType> FromPlaneUV(const FVector2<RealType>& PosUV, int PlaneNormalAxis) const
@@ -248,6 +248,21 @@ struct FFrame3
 
 
 
+	/**
+	 * Project a point onto one of the planes of the frame
+	 * @param Pos 3D position
+	 * @param PlaneNormalAxis which plane to project onto, identified by perpendicular normal. Default is 2, ie normal is Z, plane is (X,Y)
+	 * @return 3D coordinate in the plane
+	 */
+	FVector3<RealType> ToPlane(const FVector3<RealType>& Pos, int PlaneNormalAxis) const
+	{
+		FVector3<RealType> Normal = GetAxis(PlaneNormalAxis);
+		FVector3<RealType> LocalVec = Pos - Origin;
+		RealType SignedDist = LocalVec.Dot(Normal);
+		return Pos - SignedDist * Normal;
+	}
+
+
 
 
 	/**
diff --git a/Engine/Plugins/Experimental/GeometryProcessing/Source/GeometricObjects/Public/VectorUtil.h b/Engine/Plugins/Experimental/GeometryProcessing/Source/GeometricObjects/Public/VectorUtil.h
--- a/Engine/Plugins/Experimental/GeometryProcessing/Source/GeometricObjects/Public/VectorUtil.h
+++ b/Engine/Plugins/Experimental/GeometryProcessing/Source/GeometricObjects/Public/VectorUtil.h
@@ -5,7 +5,6 @@
 #include "MathUtil.h"
 #include "VectorTypes.h"
 
-// these should go somewhere else?
 
 enum class EIntersectionResult
 {
@@ -28,11 +27,18 @@ enum class EIntersectionType
 
 namespace VectorUtil
 {
+	/**
+	 * @return true if all components of V are finite
+	 */
 	template <typename RealType>
 	inline bool IsFinite(const FVector2<RealType>& V)
 	{
 		return TMathUtil<RealType>::IsFinite(V.X) && TMathUtil<RealType>::IsFinite(V.Y);
 	}
+
+	/**
+	 * @return true if all components of V are finite
+	 */
 	template <typename RealType>
 	inline bool IsFinite(const FVector3<RealType>& V)
 	{
@@ -40,12 +46,18 @@ namespace VectorUtil
 	}
 
 
+	/**
+	 * @return input Value clamped to range [MinValue, MaxValue]
+	 */
 	template <typename RealType>
 	inline RealType Clamp(RealType Value, RealType MinValue, RealType MaxValue)
 	{
 		return (Value < MinValue) ? MinValue : ((Value > MaxValue) ? MaxValue : Value);
 	}
 
+	/**
+	 * @return normalized vector that is perpendicular to triangle V0,V1,V2  (triangle normal)
+	 */
 	template <typename RealType>
 	inline FVector3<RealType> Normal(const FVector3<RealType>& V0, const FVector3<RealType>& V1, const FVector3<RealType>& V2)
 	{
@@ -59,6 +71,9 @@ namespace VectorUtil
 		return vCross.Normalized();
 	}
 
+	/**
+	 * @return un-normalized direction that is parallel to normal of triangle V0,V1,V2
+	 */
 	template <typename RealType>
 	inline FVector3<RealType> FastNormalDirection(const FVector3<RealType>& V0, const FVector3<RealType>& V1, const FVector3<RealType>& V2)
 	{
@@ -67,23 +82,33 @@ namespace VectorUtil
 		//return (V1 - V0).Cross(V2 - V0);
 	}
 
+	/**
+	 * @return area of 3D triangle V0,V1,V2
+	 */
 	template <typename RealType>
 	inline RealType Area(const FVector3<RealType>& V0, const FVector3<RealType>& V1, const FVector3<RealType>& V2)
 	{
 		FVector3<RealType> edge1(V1 - V0);
 		FVector3<RealType> edge2(V2 - V0);
 		RealType Dot = edge1.Dot(edge2);
-		return (RealType)(0.5 * sqrt(edge1.SquaredLength() * edge2.SquaredLength() - Dot * Dot));
+		return (RealType)0.5 * TMathUtil<RealType>::Sqrt(edge1.SquaredLength() * edge2.SquaredLength() - Dot * Dot);
 	}
+
+	/**
+	 * @return area of 2D triangle V0,V1,V2
+	 */
 	template <typename RealType>
 	inline RealType Area(const FVector2<RealType>& V0, const FVector2<RealType>& V1, const FVector2<RealType>& V2)
 	{
 		FVector2<RealType> edge1(V1 - V0);
 		FVector2<RealType> edge2(V2 - V0);
 		RealType Dot = edge1.Dot(edge2);
-		return (RealType)(0.5 * sqrt(edge1.SquaredLength() * edge2.SquaredLength() - Dot * Dot));
+		return (RealType)0.5 * TMathUtil<RealType>::Sqrt(edge1.SquaredLength() * edge2.SquaredLength() - Dot * Dot);
 	}
 
+	/**
+	 * @return true if triangle V1,V2,V3 is obtuse
+	 */
 	template <typename RealType>
 	inline bool IsObtuse(const FVector3<RealType>& V1, const FVector3<RealType>& V2, const FVector3<RealType>& V3)
 	{
@@ -93,29 +118,37 @@ namespace VectorUtil
 		return (a2 + b2 < c2) || (b2 + c2 < a2) || (c2 + a2 < b2);
 	}
 
+	/**
+	 * Calculate Normal and Area of triangle V0,V1,V2
+	 * @return triangle normal
+	 */
 	template <typename RealType>
-	inline FVector3<RealType> FastNormalArea(const FVector3<RealType>& V0, const FVector3<RealType>& V1, const FVector3<RealType>& V2, double& Area)
+	inline FVector3<RealType> FastNormalArea(const FVector3<RealType>& V0, const FVector3<RealType>& V1, const FVector3<RealType>& V2, RealType& AreaOut)
 	{
 		FVector3<RealType> edge1(V1 - V0);
 		FVector3<RealType> edge2(V2 - V0);
 		// Unreal has Left-Hand Coordinate System so we need to reverse this cross-product to get proper triangle normal
 		FVector3d vCross = edge2.Cross(edge1);
 		//FVector3d vCross = edge1.Cross(edge2);
-		Area = RealType(0.5) * vCross.Normalize();
+		AreaOut = RealType(0.5) * vCross.Normalize();
 		return vCross;
 	}
 
+	/** @return true if Abs(A-B) is less than Epsilon */
 	template <typename RealType>
 	inline bool EpsilonEqual(RealType A, RealType B, RealType Epsilon)
 	{
 		return TMathUtil<RealType>::Abs(A - B) < Epsilon;
 	}
 
+	/** @return true if all coordinates of V0 and V1 are within Epsilon of eachother */
 	template <typename RealType>
 	inline bool EpsilonEqual(const FVector2<RealType>& V0, const FVector2<RealType>& V1, RealType Epsilon)
 	{
 		return EpsilonEqual(V0.X, V1.X, Epsilon) && EpsilonEqual(V0.Y, V1.Y, Epsilon);
 	}
+
+	/** @return true if all coordinates of V0 and V1 are within Epsilon of eachother */
 	template <typename RealType>
 	inline bool EpsilonEqual(const FVector3<RealType>& V0, const FVector3<RealType>& V1, RealType Epsilon)
 	{
@@ -123,12 +156,12 @@ namespace VectorUtil
 	}
 
 	/**
-	 * Returns two vectors perpendicular to n, as efficiently as possible.
-	 * Duff et all method, from https://graphics.pixar.com/library/OrthonormalB/paper.pdf
+	 * Calculates two vectors perpendicular to input Normal, as efficiently as possible.
 	 */
 	template <typename RealType>
 	inline void MakePerpVectors(const FVector3<RealType>& Normal, FVector3<RealType>& OutPerp1, FVector3<RealType>& OutPerp2)
 	{
+		// Duff et al method, from https://graphics.pixar.com/library/OrthonormalB/paper.pdf
 		if (Normal.Z < (RealType)0)
 		{
 			RealType A = (RealType)1 / ((RealType)1 - Normal.Z);
@@ -153,17 +186,21 @@ namespace VectorUtil
 		}
 	}
 
+	/**
+	 * Calculates angle between VFrom and VTo after projection onto plane with normal defined by PlaneN
+	 * @return angle in degrees
+	 */
 	template <typename RealType>
-	inline double PlaneAngleSignedD(const FVector3<RealType>& VFromIn, const FVector3<RealType>& VToIn, const FVector3<RealType>& PlaneN)
+	inline RealType PlaneAngleSignedD(const FVector3<RealType>& VFrom, const FVector3<RealType>& VTo, const FVector3<RealType>& PlaneN)
 	{
-		FVector3<RealType> vFrom = VFromIn - VFromIn.Dot(PlaneN) * PlaneN;
-		FVector3<RealType> vTo = VToIn - VToIn.Dot(PlaneN) * PlaneN;
+		FVector3<RealType> vFrom = VFrom - VFrom.Dot(PlaneN) * PlaneN;
+		FVector3<RealType> vTo = VTo - VTo.Dot(PlaneN) * PlaneN;
 		vFrom.Normalize();
 		vTo.Normalize();
 		FVector3<RealType> C = vFrom.Cross(vTo);
 		if (C.SquaredLength() < TMathUtil<RealType>::ZeroTolerance)
 		{ // vectors are parallel
-			return vFrom.Dot(vTo) < 0 ? 180.0 : 0;
+			return vFrom.Dot(vTo) < 0 ? (RealType)180 : (RealType)0;
 		}
 		RealType fSign = C.Dot(PlaneN) < 0 ? (RealType)-1 : (RealType)1;
 		return (RealType)(fSign * vFrom.AngleD(vTo));
@@ -171,7 +208,7 @@ namespace VectorUtil
 
 	/**
 	 * tan(theta/2) = +/- sqrt( (1-cos(theta)) / (1+cos(theta)) )
-	 * This function returns positive Value
+	 * @return positive value of tan(theta/2) where theta is angle between normalized vectors A and B
 	 */
 	template <typename RealType>
 	RealType VectorTanHalfAngle(const FVector3<RealType>& A, const FVector3<RealType>& B)
@@ -179,12 +216,14 @@ namespace VectorUtil
 		RealType cosAngle = A.Dot(B);
 		RealType sqr = ((RealType)1 - cosAngle) / ((RealType)1 + cosAngle);
 		sqr = Clamp(sqr, (RealType)0, TMathUtil<RealType>::MaxReal);
-		return (RealType)sqrt(sqr);
+		return TMathUtil<RealType>::Sqrt(sqr);
 	}
 
-	//! fast cotangent between two normalized vectors
-	//! cot = cos/sin, both of which can be computed from vector identities
-	//! returns zero if result would be unstable (eg infinity)
+	/**
+	 * Fast cotangent of angle between two vectors.
+	 * cot = cos/sin, both of which can be computed from vector identities
+	 * @return cotangent of angle between V1 and V2, or zero if result would be unstable (eg infinity)
+	 */ 
 	template <typename RealType>
 	RealType VectorCot(const FVector3<RealType>& V1, const FVector3<RealType>& V2)
 	{
@@ -194,14 +233,22 @@ namespace VectorUtil
 		RealType lensqr2 = V2.SquaredLength();
 		RealType d = Clamp(lensqr1 * lensqr2 - fDot * fDot, (RealType)0.0, TMathUtil<RealType>::MaxReal);
 		if (d < TMathUtil<RealType>::ZeroTolerance)
+		{
 			return (RealType)0;
+		}
 		else
-			return fDot / (RealType)sqrt(d);
+		{
+			return fDot / TMathUtil<RealType>::Sqrt(d);
+		}
 	}
 
+	/**
+	 * @return solid angle at point P for triangle A,B,C
+	 */
 	template <typename RealType>
-	inline double TriSolidAngle(FVector3<RealType> A, FVector3<RealType> B, FVector3<RealType> C, const FVector3<RealType>& P)
+	inline RealType TriSolidAngle(FVector3<RealType> A, FVector3<RealType> B, FVector3<RealType> C, const FVector3<RealType>& P)
 	{
+		// Formula from https://igl.ethz.ch/projects/winding-number/
 		A -= P;
 		B -= P;
 		C -= P;
@@ -212,4 +259,17 @@ namespace VectorUtil
 		return RealType(-2.0) * atan2(bottom, top);
 	}
 
+	/**
+	 * @return angle between vectors (A-CornerPt) and (B-CornerPt)
+	 */
+	template<typename RealType>
+	inline RealType OpeningAngleD(FVector3<RealType> A, FVector3<RealType> B, const FVector3<RealType>& P)
+	{
+		A -= P; 
+		A.Normalize();
+		B -= P;
+		B.Normalize();
+		return A.AngleD(B);
+	}
+
 }; // namespace VectorUtil
