Allow floating vertices in arrangement2d (as they could happen if you try to add a too-short edge anyway, and this will allow me to cleanly add steiner verts on my triangulations)

diff --git a/Engine/Plugins/Experimental/GeometryProcessing/Source/GeometryAlgorithms/Public/Arrangement2d.h b/Engine/Plugins/Experimental/GeometryProcessing/Source/GeometryAlgorithms/Public/Arrangement2d.h
--- a/Engine/Plugins/Experimental/GeometryProcessing/Source/GeometryAlgorithms/Public/Arrangement2d.h
+++ b/Engine/Plugins/Experimental/GeometryProcessing/Source/GeometryAlgorithms/Public/Arrangement2d.h
@@ -77,6 +77,14 @@ struct FArrangement2d
 		return false;
 	}
 
+	/**
+	 * Insert isolated point P into the arrangement
+	 */
+	void Insert(const FVector2d& Pt)
+	{
+		insert_point(Pt, VertexSnapTol);
+	}
+
 
 	/**
 	 * insert segment [A,B] into the arrangement
@@ -153,6 +161,52 @@ protected:
 		int VID;
 	};
 
+	/**
+	 * insert pt P into the arrangement, splitting existing edges as necessary
+	 */
+	bool insert_point(const FVector2d &P, double Tol = 0)
+	{
+		int PIdx = find_existing_vertex(P);
+		if (PIdx > -1)
+		{
+			return false;
+		}
+
+		// TODO: currently this tries to add the vertex on the closest edge below tolerance; we should instead insert at *every* edge below tolerance!  ... but that is more inconvenient to write
+		FVector2d x = FVector2d::Zero(), y = FVector2d::Zero();
+		double ClosestDistSq = Tol*Tol;
+		int FoundEdgeToSplit = -1;
+		for (int EID = 0, ExistingEdgeMax = Graph.MaxEdgeID(); EID < ExistingEdgeMax; EID++)
+		{
+			if (!Graph.IsEdge(EID))
+			{
+				continue;
+			}
+
+			Graph.GetEdgeV(EID, x, y);
+			FSegment2d Seg(x, y);
+			double DistSq = Seg.DistanceSquared(P);
+			if (DistSq < ClosestDistSq)
+			{
+				ClosestDistSq = DistSq;
+				FoundEdgeToSplit = EID;
+			}
+		}
+		if (FoundEdgeToSplit > -1)
+		{
+			FDynamicGraph2d::FEdgeSplitInfo splitInfo;
+			EMeshResult result = Graph.SplitEdge(FoundEdgeToSplit, splitInfo);
+			ensureMsgf(result == EMeshResult::Ok, TEXT("insert_into_segment: edge split failed?"));
+			Graph.SetVertex(splitInfo.VNew, P);
+			PointHash.InsertPointUnsafe(splitInfo.VNew, P);
+			return true;
+		}
+
+		int VID = Graph.AppendVertex(P);
+		PointHash.InsertPointUnsafe(VID, P);
+		return true;
+	}
+
 	/**
 	 * insert edge [A,B] into the arrangement, splitting existing edges as necessary
 	 */
@@ -177,15 +231,16 @@ protected:
 
 		// ok find all intersections
 		TArray<FIntersection> Hits;
-		find_intersecting_edges(A, B, Hits, Tol);
-		int N = Hits.Num();
+		find_intersecting_edges(A, B, Hits, Tol);		
 
 		// we are going to construct a list of <T,vertex_id> values along segment AB
 		TArray<FSegmentPoint> points;
 		FSegment2d segAB = FSegment2d(A, B);
 
+		find_intersecting_floating_vertices(segAB, a_idx, b_idx, points, Tol);
+
 		// insert intersections into existing segments
-		for (int i = 0; i < N; ++i)
+		for (int i = 0, N = Hits.Num(); i < N; ++i)
 		{
 			FIntersection Intr = Hits[i];
 			int EID = Intr.EID;
@@ -321,10 +376,11 @@ protected:
 	 */
 	int find_nearest_vertex(FVector2d Pt, double SearchRadius, int IgnoreVID = -1)
 	{
+		double SearchRadiusSquared = SearchRadius * SearchRadius;
 		auto FuncDist = [&](int B) { return Pt.DistanceSquared(Graph.GetVertex(B)); };
 		auto FuncIgnore = [&](int VID) { return VID == IgnoreVID; };
-		TPair<int, double> found = (IgnoreVID == -1) ? PointHash.FindNearestInRadius(Pt, SearchRadius, FuncDist)
-													 : PointHash.FindNearestInRadius(Pt, SearchRadius, FuncDist, FuncIgnore);
+		TPair<int, double> found = (IgnoreVID == -1) ? PointHash.FindNearestInRadius(Pt, SearchRadiusSquared, FuncDist)
+													 : PointHash.FindNearestInRadius(Pt, SearchRadiusSquared, FuncDist, FuncIgnore);
 		if (found.Key == PointHash.InvalidValue())
 		{
 			return -1;
@@ -337,10 +393,11 @@ protected:
 	 */
 	int find_nearest_boundary_vertex(FVector2d Pt, double SearchRadius, int IgnoreVID = -1)
 	{
+		double SearchRadiusSquared = SearchRadius * SearchRadius;
 		auto FuncDist = [&](int B) { return Pt.DistanceSquared(Graph.GetVertex(B)); };
 		auto FuncIgnore = [&](int VID) { return Graph.IsBoundaryVertex(VID) == false || VID == IgnoreVID; };
 		TPair<int, double> found =
-			PointHash.FindNearestInRadius(Pt, SearchRadius, FuncDist, FuncIgnore);
+			PointHash.FindNearestInRadius(Pt, SearchRadiusSquared, FuncDist, FuncIgnore);
 		if (found.Key == PointHash.InvalidValue())
 		{
 			return -1;
@@ -380,7 +437,31 @@ protected:
 				Hits.Add(FIntersection{EID, SideX, SideY, Intr});
 				num_hits++;
 			}
-		}
+		} 
+		
 		return (num_hits > 0);
 	}
+
+	bool find_intersecting_floating_vertices(const FSegment2d &SegAB, int32 AID, int32 BID, TArray<FSegmentPoint>& Hits, double Tol = 0)
+	{
+		int num_hits = 0;
+
+		for (int VID : Graph.VertexIndices())
+		{
+			if (Graph.GetVtxEdgeCount(VID) > 0 || VID == AID || VID == BID) // if it's an existing edge or on the currently added edge, it's not floating so skip it
+			{
+				continue;
+			}
+
+			FVector2d V = Graph.GetVertex(VID);
+			double T;
+			double DSQ = SegAB.DistanceSquared(V, T);
+			if (DSQ < Tol*Tol)
+			{
+				Hits.Add(FSegmentPoint{ T, VID });
+			}
+		}
+
+		return num_hits > 0;
+	}
 };
