Add Segment3, line nearest-point functions, clean up leftover doubles from porting
#rnx

diff --git a/Engine/Plugins/Experimental/GeometryProcessing/Source/GeometricObjects/Public/LineTypes.h b/Engine/Plugins/Experimental/GeometryProcessing/Source/GeometricObjects/Public/LineTypes.h
--- a/Engine/Plugins/Experimental/GeometryProcessing/Source/GeometricObjects/Public/LineTypes.h
+++ b/Engine/Plugins/Experimental/GeometryProcessing/Source/GeometricObjects/Public/LineTypes.h
@@ -50,7 +50,7 @@ struct FLine2
 	/**
 	 * @return point on line at given line parameter value (distance along line from origin)
 	 */
-	inline FVector2<T> PointAt(double LineParameter) const
+	inline FVector2<T> PointAt(T LineParameter) const
 	{
 		return Origin + LineParameter * Direction;
 	}
@@ -59,7 +59,7 @@ struct FLine2
 	/**
 	 * @return line parameter (ie distance from Origin) at nearest point on line to QueryPoint
 	 */
-	inline double Project(const FVector2<T>& QueryPoint) const
+	inline T Project(const FVector2<T>& QueryPoint) const
 	{
 		return (QueryPoint - Origin).Dot(Direction);
 	}
@@ -67,23 +67,33 @@ struct FLine2
 	/**
 	 * @return smallest squared distance from line to QueryPoint
 	 */
-	inline double DistanceSquared(const FVector2<T>& QueryPoint) const
+	inline T DistanceSquared(const FVector2<T>& QueryPoint) const
 	{
-		double t = (QueryPoint - Origin).Dot(Direction);
-		FVector2<T> proj = Origin + t * Direction;
+		T ParameterT = (QueryPoint - Origin).Dot(Direction);
+		FVector2<T> proj = Origin + ParameterT * Direction;
 		return (proj - QueryPoint).SquaredLength();
 	}
 
+	/**
+	 * @return nearest point on line to QueryPoint
+	 */
+	inline FVector2<T> NearestPoint(const FVector2<T>& QueryPoint) const
+	{
+		T ParameterT = (QueryPoint - Origin).Dot(Direction);
+		return Origin + ParameterT * Direction;
+	}
+
+
 	/**
 	 * @return +1 if QueryPoint is "right" of line, -1 if "left" or 0 if "on" line (up to given tolerance)
 	 */
-	inline int WhichSide(const FVector2<T>& QueryPoint, double OnLineTolerance = 0) const
+	inline int WhichSide(const FVector2<T>& QueryPoint, T OnLineTolerance = 0) const
 	{
-		double x0 = QueryPoint.X - Origin.X;
-		double y0 = QueryPoint.Y - Origin.Y;
-		double x1 = Direction.X;
-		double y1 = Direction.Y;
-		double det = x0 * y1 - x1 * y0;
+		T x0 = QueryPoint.X - Origin.X;
+		T y0 = QueryPoint.Y - Origin.Y;
+		T x1 = Direction.X;
+		T y1 = Direction.Y;
+		T det = x0 * y1 - x1 * y0;
 		return (det > OnLineTolerance ? +1 : (det < -OnLineTolerance ? -1 : 0));
 	}
 
@@ -95,15 +105,15 @@ struct FLine2
 	 * @param ParallelDotTolerance tolerance used to determine if lines are parallel (and hence cannot intersect)
 	 * @return true if lines intersect and IntersectionPointOut was computed
 	 */
-	bool IntersectionPoint(const FLine2& OtherLine, FVector2<T>& IntersectionPointOut, double ParallelDotTolerance = TMathUtil<T>::ZeroTolerance) const
+	bool IntersectionPoint(const FLine2<T>& OtherLine, FVector2<T>& IntersectionPointOut, T ParallelDotTolerance = TMathUtil<T>::ZeroTolerance) const
 	{
 		// see IntrFLine2FLine2 for more detailed explanation
 		FVector2<T> diff = OtherLine.Origin - Origin;
-		double D0DotPerpD1 = Direction.DotPerp(OtherLine.Direction);
+		T D0DotPerpD1 = Direction.DotPerp(OtherLine.Direction);
 		if (TMathUtil<T>::Abs(D0DotPerpD1) > ParallelDotTolerance)                     // FLines intersect in a single point.
 		{
-			double invD0DotPerpD1 = ((double)1) / D0DotPerpD1;
-			double diffDotPerpD1 = diff.DotPerp(OtherLine.Direction);
+			T invD0DotPerpD1 = ((T)1) / D0DotPerpD1;
+			T diffDotPerpD1 = diff.DotPerp(OtherLine.Direction);
 			T s = diffDotPerpD1 * invD0DotPerpD1;
 			IntersectionPointOut = Origin + s * Direction;
 			return true;
@@ -165,7 +175,7 @@ struct FLine3
 	/**
 	 * @return point on line at given line parameter value (distance along line from origin)
 	 */
-	inline FVector3<T> PointAt(double LineParameter) const
+	inline FVector3<T> PointAt(T LineParameter) const
 	{
 		return Origin + LineParameter * Direction;
 	}
@@ -174,7 +184,7 @@ struct FLine3
 	/**
 	 * @return line parameter (ie distance from Origin) at nearest point on line to QueryPoint
 	 */
-	inline double Project(const FVector3<T>& QueryPoint) const
+	inline T Project(const FVector3<T>& QueryPoint) const
 	{
 		return (QueryPoint - Origin).Dot(Direction);
 	}
@@ -182,13 +192,23 @@ struct FLine3
 	/**
 	 * @return smallest squared distance from line to QueryPoint
 	 */
-	inline double DistanceSquared(const FVector3<T>& QueryPoint) const
+	inline T DistanceSquared(const FVector3<T>& QueryPoint) const
 	{
-		double t = (QueryPoint - Origin).Dot(Direction);
+		T t = (QueryPoint - Origin).Dot(Direction);
 		FVector3<T> proj = Origin + t * Direction;
 		return (proj - QueryPoint).SquaredLength();
 	}
 
+	/**
+	 * @return nearest point on line to QueryPoint
+	 */
+	inline FVector3<T> NearestPoint(const FVector3<T>& QueryPoint) const
+	{
+		T ParameterT = (QueryPoint - Origin).Dot(Direction);
+		return Origin + ParameterT * Direction;
+	}
+
+
 };
 
 typedef FLine3<double> FLine3d;
diff --git a/Engine/Plugins/Experimental/GeometryProcessing/Source/GeometricObjects/Public/SegmentTypes.h b/Engine/Plugins/Experimental/GeometryProcessing/Source/GeometricObjects/Public/SegmentTypes.h
--- a/Engine/Plugins/Experimental/GeometryProcessing/Source/GeometricObjects/Public/SegmentTypes.h
+++ b/Engine/Plugins/Experimental/GeometryProcessing/Source/GeometricObjects/Public/SegmentTypes.h
@@ -36,7 +36,7 @@ public:
 	/**
 	 * Construct a segment from a Center Point, normalized Direction, and scalar Extent
 	 */
-	FSegment2(const FVector2<T>& CenterIn, const FVector2<T>& DirectionIn, double ExtentIn)
+	FSegment2(const FVector2<T>& CenterIn, const FVector2<T>& DirectionIn, T ExtentIn)
 	{
 		Center = CenterIn;
 		Direction = DirectionIn; 
@@ -103,7 +103,7 @@ public:
 	 * @param UnitParameter value in range [0,1]
 	 * @return point on segment at that linearly interpolates between start and end based on Unit Parameter
 	 */
-	inline FVector2<T> PointBetween(double UnitParameter) const
+	inline FVector2<T> PointBetween(T UnitParameter) const
 	{
 		return Center + ((T)2 * UnitParameter - (T)1) * Extent * Direction;
 	}
@@ -111,9 +111,9 @@ public:
 	/**
 	 * @return minimum squared distance from Point to segment
 	 */
-	inline double DistanceSquared(const FVector2<T>& Point) const
+	inline T DistanceSquared(const FVector2<T>& Point) const
 	{
-		double DistParameter;
+		T DistParameter;
 		return DistanceSquared(Point, DistParameter);
 	}
 
@@ -124,7 +124,7 @@ public:
 	 * @param DistParameterOut calculated distance parameter in range [-Extent,Extent]
 	 * @return minimum squared distance from Point to Segment
 	 */	
-	double DistanceSquared(const FVector2<T>& Point, double& DistParameterOut) const
+	T DistanceSquared(const FVector2<T>& Point, T& DistParameterOut) const
 	{
 		DistParameterOut = (Point - Center).Dot(Direction);
 		if (DistParameterOut >= Extent)
@@ -147,7 +147,7 @@ public:
 	 */
 	inline FVector2<T> NearestPoint(const FVector2<T>& QueryPoint) const
 	{
-		double t = (QueryPoint - Center).Dot(Direction);
+		T t = (QueryPoint - Center).Dot(Direction);
 		if (t >= Extent)
 		{
 			return EndPoint();
@@ -163,7 +163,7 @@ public:
 	/**
 	 * @return scalar projection of QueryPoint onto line of Segment (not clamped to Extents)
 	 */
-	inline double Project(const FVector2<T>& QueryPoint) const
+	inline T Project(const FVector2<T>& QueryPoint) const
 	{
 		return (QueryPoint - Center).Dot(Direction);
 	}
@@ -175,12 +175,12 @@ public:
 	 * @param Tolerance tolerance band in which we return 0
 	 * @return +1 if point is to right of line, -1 if left, and 0 if on line or within tolerance band
 	 */
-	int WhichSide(const FVector2<T>& QueryPoint, double Tolerance = 0)
+	int WhichSide(const FVector2<T>& QueryPoint, T Tolerance = 0)
 	{
 		// [TODO] subtract Center from test?
 		FVector2<T> EndPt = Center + Extent * Direction;
 		FVector2<T> StartPt = Center - Extent * Direction;
-		double det = -FVector2<T>::Orient(EndPt, StartPt, QueryPoint);
+		T det = -FVector2<T>::Orient(EndPt, StartPt, QueryPoint);
 		return (det > Tolerance ? +1 : (det < -Tolerance ? -1 : 0));
 	}
 
@@ -300,3 +300,185 @@ protected:
 typedef FSegment2<float> FSegment2f;
 typedef FSegment2<double> FSegment2d;
 
+
+
+
+
+
+
+
+/*
+ * 3D Line Segmented stored as Center point, normalized Direction vector, and scalar Extent
+ */
+template<typename T>
+struct FSegment3
+{
+public:
+	/** Center point of segment */
+	FVector3<T> Center;
+	/** normalized Direction vector of segment */
+	FVector3<T> Direction;
+	/** Extent of segment, which is half the total length */
+	T Extent;
+
+	/**
+	 * Construct a Segment from two Points
+	 */
+	FSegment3(const FVector3<T>& Point0, const FVector3<T>& Point1)
+	{
+		// set from endpoints 
+		Center = T(.5) * (Point0 + Point1);
+		Direction = Point1 - Point0;
+		Extent = T(.5) * Direction.Normalize();
+	}
+
+	/**
+	 * Construct a segment from a Center Point, normalized Direction, and scalar Extent
+	 */
+	FSegment3(const FVector3<T>& CenterIn, const FVector3<T>& DirectionIn, T ExtentIn)
+	{
+		Center = CenterIn;
+		Direction = DirectionIn;
+		Extent = ExtentIn;
+	}
+
+
+
+	/** Update the Segment with a new start point */
+	inline void SetStartPoint(const FVector3<T>& Point)
+	{
+		update_from_endpoints(Point, EndPoint());
+	}
+
+	/** Update the Segment with a new end point */
+	inline void SetEndPoint(const FVector3<T>& Point)
+	{
+		update_from_endpoints(StartPoint(), Point);
+	}
+
+	/** Reverse the segment */
+	void Reverse()
+	{
+		update_from_endpoints(EndPoint(), StartPoint());
+	}
+
+
+
+
+	/** @return start point of segment */
+	inline FVector3<T> StartPoint() const
+	{
+		return Center - Extent * Direction;
+	}
+
+	/** @return end point of segment */
+	inline FVector3<T> EndPoint() const
+	{
+		return Center + Extent * Direction;
+	}
+
+
+	/** @return the Length of the segment */
+	inline T Length() const
+	{
+		return (T)2 * Extent;
+	}
+
+	/** @return first (i == 0) or second (i == 1) endpoint of the Segment  */
+	inline FVector3<T> GetPointFromIndex(int i) const
+	{
+		return (i == 0) ? (Center - Extent * Direction) : (Center + Extent * Direction);
+	}
+
+	/**
+	 * @return point on segment at given (signed) Distance from the segment Origin
+	 */
+	inline FVector3<T> PointAt(T DistanceParameter) const
+	{
+		return Center + DistanceParameter * Direction;
+	}
+
+	/**
+	 * @param UnitParameter value in range [0,1]
+	 * @return point on segment at that linearly interpolates between start and end based on Unit Parameter
+	 */
+	inline FVector3<T> PointBetween(T UnitParameter) const
+	{
+		return Center + ((T)2 * UnitParameter - (T)1) * Extent * Direction;
+	}
+
+	/**
+	 * @return minimum squared distance from Point to segment
+	 */
+	inline T DistanceSquared(const FVector3<T>& Point) const
+	{
+		T DistParameter;
+		return DistanceSquared(Point, DistParameter);
+	}
+
+
+
+	/**
+	 * @param Point query point
+	 * @param DistParameterOut calculated distance parameter in range [-Extent,Extent]
+	 * @return minimum squared distance from Point to Segment
+	 */
+	T DistanceSquared(const FVector3<T>& Point, T& DistParameterOut) const
+	{
+		DistParameterOut = (Point - Center).Dot(Direction);
+		if (DistParameterOut >= Extent)
+		{
+			DistParameterOut = Extent;
+			return Point.DistanceSquared(EndPoint());
+		}
+		else if (DistParameterOut <= -Extent)
+		{
+			DistParameterOut = -Extent;
+			return Point.DistanceSquared(StartPoint());
+		}
+		FVector3<T> ProjectedPt = Center + DistParameterOut * Direction;
+		return ProjectedPt.DistanceSquared(Point);
+	}
+
+
+	/**
+	 * @return nearest point on segment to QueryPoint
+	 */
+	inline FVector3<T> NearestPoint(const FVector3<T>& QueryPoint) const
+	{
+		T t = (QueryPoint - Center).Dot(Direction);
+		if (t >= Extent)
+		{
+			return EndPoint();
+		}
+		if (t <= -Extent)
+		{
+			return StartPoint();
+		}
+		return Center + t * Direction;
+	}
+
+
+	/**
+	 * @return scalar projection of QueryPoint onto line of Segment (not clamped to Extents)
+	 */
+	inline T Project(const FVector3<T>& QueryPoint) const
+	{
+		return (QueryPoint - Center).Dot(Direction);
+	}
+
+
+
+protected:
+
+	// update segment based on new endpoints
+	inline void update_from_endpoints(const FVector3<T>& p0, const FVector3<T>& p1)
+	{
+		Center = 0.5 * (p0 + p1);
+		Direction = p1 - p0;
+		Extent = 0.5 * Direction.Normalize();
+	}
+
+};
+typedef FSegment3<float> FSegment3f;
+typedef FSegment3<double> FSegment3d;
