improve UV panel alignment in ExtrudeMesh
#rnx

diff --git a/Engine/Plugins/Experimental/GeometryProcessing/Source/DynamicMesh/Private/Operations/ExtrudeMesh.cpp b/Engine/Plugins/Experimental/GeometryProcessing/Source/DynamicMesh/Private/Operations/ExtrudeMesh.cpp
--- a/Engine/Plugins/Experimental/GeometryProcessing/Source/DynamicMesh/Private/Operations/ExtrudeMesh.cpp
+++ b/Engine/Plugins/Experimental/GeometryProcessing/Source/DynamicMesh/Private/Operations/ExtrudeMesh.cpp
@@ -107,7 +107,14 @@ bool FExtrudeMesh::Apply()
 			// @todo is there a simpler way to construct rotation from 3 known axes (third id Normal.Cross(UnitY))
 			//  (converting from matrix might end up being more efficient due to trig in ConstrainedAlignAxis?)
 			FFrame3f ProjectFrame(FVector3f::Zero(), Normal);
-			ProjectFrame.ConstrainedAlignAxis(1, FVector3f::UnitY(), ProjectFrame.Z());
+			if (FMathd::Abs(ProjectFrame.Y().Dot(FVector3f::UnitY())) < 0.01)
+			{
+				ProjectFrame.ConstrainedAlignAxis(0, FVector3f::UnitX(), ProjectFrame.Z());
+			}
+			else
+			{
+				ProjectFrame.ConstrainedAlignAxis(1, FVector3f::UnitY(), ProjectFrame.Z());
+			}
 			Editor.SetQuadUVsFromProjection(StitchResult.NewQuads[k], ProjectFrame, UVScaleFactor);
 		}
 
