bring FRay3 API into parity with FLine3 and FSegment3
#rnx

diff --git a/Engine/Plugins/Experimental/GeometryProcessing/Source/GeometricObjects/Public/RayTypes.h b/Engine/Plugins/Experimental/GeometryProcessing/Source/GeometricObjects/Public/RayTypes.h
--- a/Engine/Plugins/Experimental/GeometryProcessing/Source/GeometricObjects/Public/RayTypes.h
+++ b/Engine/Plugins/Experimental/GeometryProcessing/Source/GeometricObjects/Public/RayTypes.h
@@ -52,6 +52,50 @@ public:
 	}
 
 
+	/**
+	 * @return ray parameter (ie positive distance from Origin) at nearest point on ray to QueryPoint
+	 */
+	inline RealType Project(const FVector3<RealType>& QueryPoint) const
+	{
+		RealType LineParam = (QueryPoint - Origin).Dot(Direction);
+		return (LineParam < 0) ? 0 : LineParam;
+	}
+
+	/**
+	 * @return smallest squared distance from line to QueryPoint
+	 */
+	inline RealType DistanceSquared(const FVector3<RealType>& QueryPoint) const
+	{
+		RealType LineParam = (QueryPoint - Origin).Dot(Direction);
+		if (LineParam < 0)
+		{
+			return Origin.DistanceSquared(QueryPoint);
+		}
+		else
+		{
+			FVector3<RealType> NearestPt = Origin + LineParam * Direction;
+			return NearestPt.DistanceSquared(QueryPoint);
+		}
+	}
+
+	/**
+	 * @return nearest point on line to QueryPoint
+	 */
+	inline FVector3<RealType> NearestPoint(const FVector3<RealType>& QueryPoint) const
+	{
+		RealType LineParam = (QueryPoint - Origin).Dot(Direction);
+		if (LineParam < 0)
+		{
+			return Origin;
+		}
+		else
+		{
+			return Origin + LineParam * Direction;
+		}
+	}
+
+
+
 	// conversion operators
 
 	inline operator FRay() const
