MeshDescriptionBuilder can now recalculate normals w/o requiring polygon attributes, add DynamicMeshToMeshDescription::Update() that just copies vertex positions and updates normals (ie no topology updates)
#rnx

diff --git a/Engine/Plugins/Experimental/GeometryProcessing/Source/MeshConversion/Private/DynamicMeshToMeshDescription.cpp b/Engine/Plugins/Experimental/GeometryProcessing/Source/MeshConversion/Private/DynamicMeshToMeshDescription.cpp
--- a/Engine/Plugins/Experimental/GeometryProcessing/Source/MeshConversion/Private/DynamicMeshToMeshDescription.cpp
+++ b/Engine/Plugins/Experimental/GeometryProcessing/Source/MeshConversion/Private/DynamicMeshToMeshDescription.cpp
@@ -8,6 +8,35 @@
 #include "MeshNormals.h"
 
 
+
+void FDynamicMeshToMeshDescription::Update(const FDynamicMesh3* MeshIn, FMeshDescription& MeshOut)
+{
+	FMeshDescriptionBuilder Builder;
+	Builder.SetMeshDescription(&MeshOut);
+
+	check(MeshIn->IsCompactV());
+	check(MeshIn->VertexCount() == MeshOut.Vertices().Num());
+
+	// update positions
+	for (int VertID : MeshIn->VertexIndicesItr())
+	{
+		Builder.SetPosition(FVertexID(VertID), MeshIn->GetVertex(VertID));
+	}
+
+	// can't trust these yet...
+	//const FDynamicMeshUVOverlay* UVOverlay = MeshIn->HasAttributes() ? MeshIn->Attributes()->PrimaryUV() : nullptr;
+	//const FDynamicMeshNormalOverlay* NormalOverlay = MeshIn->HasAttributes() ? MeshIn->Attributes()->PrimaryNormals() : nullptr;
+
+	Builder.RecalculateInstanceNormals();
+}
+
+
+
+
+
+
+
+
 void FDynamicMeshToMeshDescription::Convert(const FDynamicMesh3* MeshIn, FMeshDescription& MeshOut)
 {
 	if (MeshIn->HasAttributes())
@@ -21,9 +50,6 @@ void FDynamicMeshToMeshDescription::Convert(const FDynamicMesh3* MeshIn, FMeshDe
 }
 
 
-
-
-
 void FDynamicMeshToMeshDescription::Convert_NoAttributes(const FDynamicMesh3* MeshIn, FMeshDescription& MeshOut)
 {
 	MeshOut.Empty();
diff --git a/Engine/Plugins/Experimental/GeometryProcessing/Source/MeshConversion/Private/MeshDescriptionBuilder.cpp b/Engine/Plugins/Experimental/GeometryProcessing/Source/MeshConversion/Private/MeshDescriptionBuilder.cpp
--- a/Engine/Plugins/Experimental/GeometryProcessing/Source/MeshConversion/Private/MeshDescriptionBuilder.cpp
+++ b/Engine/Plugins/Experimental/GeometryProcessing/Source/MeshConversion/Private/MeshDescriptionBuilder.cpp
@@ -74,6 +74,22 @@ FVertexInstanceID FMeshDescriptionBuilder::AppendInstance(const FVertexID& Verte
 
 
 
+void FMeshDescriptionBuilder::SetPosition(const FVertexID& VertexID, const FVector& NewPosition)
+{
+	VertexPositions.Set(VertexID, 0, NewPosition);
+}
+
+FVector FMeshDescriptionBuilder::GetPosition(const FVertexID& VertexID)
+{
+	return VertexPositions.Get(VertexID, 0);
+}
+
+FVector FMeshDescriptionBuilder::GetPosition(const FVertexInstanceID& InstanceID)
+{
+	return VertexPositions.Get(MeshDescription->GetVertexInstance(InstanceID).VertexID, 0);
+}
+
+
 void FMeshDescriptionBuilder::SetInstance(const FVertexInstanceID& InstanceID, const FVector2D& InstanceUV, const FVector& InstanceNormal)
 {
 	if (InstanceUVs.IsValid())
@@ -269,6 +285,53 @@ void FMeshDescriptionBuilder::SetAllEdgesHardness(bool bHard)
 }
 
 
+
+void FMeshDescriptionBuilder::RecalculateInstanceNormals()
+{
+	for (int k = 0; k < InstanceNormals.GetNumElements(); ++k)
+	{
+		InstanceNormals.Set(FVertexInstanceID(k), 0, FVector::ZeroVector);
+	}
+
+	const FPolygonArray& Polygons = MeshDescription->Polygons();
+	for (const FPolygonID PolygonID : Polygons.GetElementIDs())
+	{
+		const TArray<FMeshTriangle>& Triangles = MeshDescription->GetPolygonTriangles(PolygonID);
+		for (FMeshTriangle Triangle : Triangles)
+		{
+			FVector3d A = GetPosition(Triangle.VertexInstanceID0);
+			FVector3d B = GetPosition(Triangle.VertexInstanceID1);
+			FVector3d C = GetPosition(Triangle.VertexInstanceID2);
+			double Area = 1.0;
+			FVector FaceNormal = VectorUtil::FastNormalArea(A, B, C, Area);
+			if (Area > FMathf::ZeroTolerance)
+			{
+				FaceNormal *= Area;
+				InstanceNormals.Set(Triangle.VertexInstanceID0,
+					InstanceNormals.Get(Triangle.VertexInstanceID0, 0) + FaceNormal);
+				InstanceNormals.Set(Triangle.VertexInstanceID1,
+					InstanceNormals.Get(Triangle.VertexInstanceID1, 0) + FaceNormal);
+				InstanceNormals.Set(Triangle.VertexInstanceID2,
+					InstanceNormals.Get(Triangle.VertexInstanceID2, 0) + FaceNormal);
+			}
+		}
+	}
+
+
+	for (int k = 0; k < InstanceNormals.GetNumElements(); ++k)
+	{
+		FVector SumNormal = InstanceNormals.Get(FVertexInstanceID(k), 0);
+		SumNormal.Normalize();
+		if ( SumNormal.Size() < 0.99 )
+		{ 
+			SumNormal = FVector(1, 0, 0);
+		}
+		InstanceNormals.Set(FVertexInstanceID(k), 0, SumNormal);
+	}
+}
+
+
+
 FBox FMeshDescriptionBuilder::ComputeBoundingBox() const
 {
 	FAxisAlignedBox3f bounds = FAxisAlignedBox3f::Empty();
diff --git a/Engine/Plugins/Experimental/GeometryProcessing/Source/MeshConversion/Public/DynamicMeshToMeshDescription.h b/Engine/Plugins/Experimental/GeometryProcessing/Source/MeshConversion/Public/DynamicMeshToMeshDescription.h
--- a/Engine/Plugins/Experimental/GeometryProcessing/Source/MeshConversion/Public/DynamicMeshToMeshDescription.h
+++ b/Engine/Plugins/Experimental/GeometryProcessing/Source/MeshConversion/Public/DynamicMeshToMeshDescription.h
@@ -25,6 +25,13 @@ public:
 	void Convert(const FDynamicMesh3* MeshIn, FMeshDescription& MeshOut);
 
 
+	/**
+	 * Update existing MeshDescription based on DynamicMesh. Assumes mesh topology has not changed.
+	 * Copies positions, recalculates MeshDescription normals.
+	 */
+	void Update(const FDynamicMesh3* MeshIn, FMeshDescription& MeshOut);
+
+
 	//
 	// Internal functions that you can also call directly
 	//
diff --git a/Engine/Plugins/Experimental/GeometryProcessing/Source/MeshConversion/Public/MeshDescriptionBuilder.h b/Engine/Plugins/Experimental/GeometryProcessing/Source/MeshConversion/Public/MeshDescriptionBuilder.h
--- a/Engine/Plugins/Experimental/GeometryProcessing/Source/MeshConversion/Public/MeshDescriptionBuilder.h
+++ b/Engine/Plugins/Experimental/GeometryProcessing/Source/MeshConversion/Public/MeshDescriptionBuilder.h
@@ -25,21 +25,38 @@ class MESHCONVERSION_API FMeshDescriptionBuilder
 public:
 	void SetMeshDescription(FMeshDescription* Description);
 
-	/** Enable per-triangle integer attribute named PolyTriGroups */
-	void EnablePolyGroups();
-
 	/** Append vertex and return new vertex ID */
 	FVertexID AppendVertex(const FVector& Position);
-	
+
+	/** Return position of vertex */
+	FVector GetPosition(const FVertexID& VertexID);
+
+	/** Return position of vertex parent of instance */
+	FVector GetPosition(const FVertexInstanceID& InstanceID);
+
+	/** Set the position of a vertex */
+	void SetPosition(const FVertexID& VertexID, const FVector& NewPosition);
+
+
+
 	/** Append new vertex instance and return ID */
 	FVertexInstanceID AppendInstance(const FVertexID& VertexID);
 
 	/** Set the UV and Normal of a vertex instance*/
 	void SetInstance(const FVertexInstanceID& InstanceID, const FVector2D& InstanceUV, const FVector& InstanceNormal);
 
+
+
+	/** Enable per-triangle integer attribute named PolyTriGroups */
+	void EnablePolyGroups();
+
 	/** Create a new polygon group and return it's ID */
 	FPolygonGroupID AppendPolygonGroup();
 
+	/** Set the PolyTriGroups attribute value to a specific GroupID for a Polygon */
+	void SetPolyGroupID(const FPolygonID& PolygonID, int GroupID);
+
+
 
 	/** Append a triangle to the mesh with the given PolygonGroup ID */
 	FPolygonID AppendTriangle(const FVertexID& Vertex0, const FVertexID& Vertex1, const FVertexID& Vertex2, const FPolygonGroupID& PolygonGroup);
@@ -69,15 +86,15 @@ public:
 
 
 
-	/** Set the PolyTriGroups attribute value to a specific GroupID for a Polygon */
-	void SetPolyGroupID(const FPolygonID& PolygonID, int GroupID);
-
 	/** Set MeshAttribute::Edge::IsHard to true for all edges */
 	void SetAllEdgesHardness(bool bHard);
 
 	/** Translate the MeshDescription vertex positions */
 	void Translate(const FVector& Translation);
 
+	/** Calculate normals on the MeshDescription vertex instances, using area-weighted face average */
+	void RecalculateInstanceNormals();
+
 	/** Return the current bounding box of the mesh */
 	FBox ComputeBoundingBox() const;
 
