don't use MeshDescription normal/tangent estimation that requires a bunch of extra attributes
#rnx

diff --git a/Engine/Plugins/Experimental/GeometryProcessing/Source/MeshConversion/Private/DynamicMeshToMeshDescription.cpp b/Engine/Plugins/Experimental/GeometryProcessing/Source/MeshConversion/Private/DynamicMeshToMeshDescription.cpp
--- a/Engine/Plugins/Experimental/GeometryProcessing/Source/MeshConversion/Private/DynamicMeshToMeshDescription.cpp
+++ b/Engine/Plugins/Experimental/GeometryProcessing/Source/MeshConversion/Private/DynamicMeshToMeshDescription.cpp
@@ -104,7 +104,7 @@ void FDynamicMeshToMeshDescription::Convert_NoAttributes(const FDynamicMesh3* Me
 		}
 	}
 
-	MeshOut.ComputeTangentsAndNormals(EComputeNTBsOptions::Tangents);
+	Builder.RecalculateInstanceNormals();
 }
 
 
