add FDynamicMesh3 operator=, and FFrame util function
#rnx

diff --git a/Engine/Plugins/Experimental/GeometryProcessing/Source/DynamicMesh/Private/DynamicMesh3.cpp b/Engine/Plugins/Experimental/GeometryProcessing/Source/DynamicMesh/Private/DynamicMesh3.cpp
--- a/Engine/Plugins/Experimental/GeometryProcessing/Source/DynamicMesh/Private/DynamicMesh3.cpp
+++ b/Engine/Plugins/Experimental/GeometryProcessing/Source/DynamicMesh/Private/DynamicMesh3.cpp
@@ -91,6 +91,13 @@ FDynamicMesh3::FDynamicMesh3(const FDynamicMesh3& CopyMesh, bool bCompact, bool
 
 
 
+const FDynamicMesh3& FDynamicMesh3::operator=(const FDynamicMesh3& CopyMesh)
+{
+	Copy(CopyMesh);
+	return *this;
+}
+
+
 FDynamicMesh3::FDynamicMesh3(const FMeshShapeGenerator* Generator)
 {
 	Copy(Generator);
diff --git a/Engine/Plugins/Experimental/GeometryProcessing/Source/DynamicMesh/Public/DynamicMesh3.h b/Engine/Plugins/Experimental/GeometryProcessing/Source/DynamicMesh/Public/DynamicMesh3.h
--- a/Engine/Plugins/Experimental/GeometryProcessing/Source/DynamicMesh/Public/DynamicMesh3.h
+++ b/Engine/Plugins/Experimental/GeometryProcessing/Source/DynamicMesh/Public/DynamicMesh3.h
@@ -236,6 +236,9 @@ public:
 	/** Initialize mesh from the output of a MeshShapeGenerator (assumes Generate() was already called) */
 	FDynamicMesh3(const FMeshShapeGenerator* Generator);
 
+	/** copy assignment operator */
+	const FDynamicMesh3& operator=(const FDynamicMesh3& CopyMesh);
+
 	//
 	// Copy functions to construct a mesh from an input mesh
 	//
diff --git a/Engine/Plugins/Experimental/GeometryProcessing/Source/GeometricObjects/Public/FrameTypes.h b/Engine/Plugins/Experimental/GeometryProcessing/Source/GeometricObjects/Public/FrameTypes.h
--- a/Engine/Plugins/Experimental/GeometryProcessing/Source/GeometricObjects/Public/FrameTypes.h
+++ b/Engine/Plugins/Experimental/GeometryProcessing/Source/GeometricObjects/Public/FrameTypes.h
@@ -138,6 +138,11 @@ struct FFrame3
 		return Rotation * FVector3<RealType>(X,Y,Z) + Origin;
 	}
 
+	/** @return point at distances along frame axes */
+	FVector3<RealType> PointAt(const FVector3<RealType>& Point) const
+	{
+		return Rotation * FVector3<RealType>(Point.X, Point.Y, Point.Z) + Origin;
+	}
 	
 	/** @return input Point transformed into local coordinate system of Frame */
 	FVector3<RealType> ToFramePoint(const FVector3<RealType>& Point) const
