fix arrangment2d point location bugs; also return vertex indices of added points to allow caller to keep track of them

diff --git a/Engine/Plugins/Experimental/GeometryProcessing/Source/GeometricObjects/Public/Spatial/PointHashGrid2.h b/Engine/Plugins/Experimental/GeometryProcessing/Source/GeometricObjects/Public/Spatial/PointHashGrid2.h
--- a/Engine/Plugins/Experimental/GeometryProcessing/Source/GeometricObjects/Public/Spatial/PointHashGrid2.h
+++ b/Engine/Plugins/Experimental/GeometryProcessing/Source/GeometricObjects/Public/Spatial/PointHashGrid2.h
@@ -162,14 +162,15 @@ public:
 	 */
 	TPair<PointDataType, double> FindNearestInRadius(
 		const FVector2d& QueryPoint, double Radius, 
-		TFunction<double(const PointDataType&)> DistanceFunc,
+		TFunction<double(const PointDataType&)> DistanceSqFunc,
 		TFunction<bool(const PointDataType&)> IgnoreFunc = [](const PointDataType& data) { return false; }) const
 	{
 		FVector2i min_idx = Indexer.ToGrid(QueryPoint - Radius * FVector2d::One());
 		FVector2i max_idx = Indexer.ToGrid(QueryPoint + Radius * FVector2d::One());
 
-		double min_dist = TNumericLimits<double>::Max();
+		double min_distsq = TNumericLimits<double>::Max();
 		PointDataType nearest = InvalidValue();
+		double RadiusSquared = Radius * Radius;
 
 		TArray<PointDataType> Values;
 		for (int yi = min_idx.Y; yi <= max_idx.Y; yi++) 
@@ -185,17 +186,17 @@ public:
 					{
 						continue;
 					}
-					double dist = DistanceFunc(Value);
-					if (dist < Radius && dist < min_dist)
+					double distsq = DistanceSqFunc(Value);
+					if (distsq < RadiusSquared && distsq < min_distsq)
 					{
 						nearest = Value;
-						min_dist = dist;
+						min_distsq = distsq;
 					}
 				}
 			}
 		}
 
-		return TPair<PointDataType, double>(nearest, min_dist);
+		return TPair<PointDataType, double>(nearest, min_distsq);
 	}
 };
 
diff --git a/Engine/Plugins/Experimental/GeometryProcessing/Source/GeometricObjects/Public/Util/GridIndexing2.h b/Engine/Plugins/Experimental/GeometryProcessing/Source/GeometricObjects/Public/Util/GridIndexing2.h
--- a/Engine/Plugins/Experimental/GeometryProcessing/Source/GeometricObjects/Public/Util/GridIndexing2.h
+++ b/Engine/Plugins/Experimental/GeometryProcessing/Source/GeometricObjects/Public/Util/GridIndexing2.h
@@ -21,7 +21,7 @@ struct TScaleGridIndexer2
 	/** Convert real-valued point to integer grid coordinates */
 	inline FVector2i ToGrid(const FVector2<RealType>& P) const
 	{
-		return FVector2i((int)P.X / CellSize, (int)P.Y / CellSize);
+		return FVector2i(int(P.X / CellSize), int(P.Y / CellSize));
 	}
 
 	/** Convert integer grid coordinates to real-valued point */
diff --git a/Engine/Plugins/Experimental/GeometryProcessing/Source/GeometryAlgorithms/Public/Arrangement2d.h b/Engine/Plugins/Experimental/GeometryProcessing/Source/GeometryAlgorithms/Public/Arrangement2d.h
--- a/Engine/Plugins/Experimental/GeometryProcessing/Source/GeometryAlgorithms/Public/Arrangement2d.h
+++ b/Engine/Plugins/Experimental/GeometryProcessing/Source/GeometryAlgorithms/Public/Arrangement2d.h
@@ -77,12 +77,33 @@ struct FArrangement2d
 		return false;
 	}
 
+	/**
+	 * Subdivide edge at a given position
+	 */
+	FIndex2i SplitEdgeAtPoint(int EdgeID, FVector2d Point)
+	{
+		FDynamicGraph2d::FEdgeSplitInfo splitInfo;
+		EMeshResult result = Graph.SplitEdge(EdgeID, splitInfo);
+		ensureMsgf(result == EMeshResult::Ok, TEXT("SplitEdgeAtPoint: edge split failed?"));
+		Graph.SetVertex(splitInfo.VNew, Point);
+		PointHash.InsertPointUnsafe(splitInfo.VNew, Point);
+		return FIndex2i(splitInfo.VNew, splitInfo.ENewBN);
+	}
+
+	/**
+	 * Check if vertex exists in region
+	 */
+	bool HasVertexNear(FVector2d Point, double SearchRadius)
+	{
+		return find_nearest_vertex(Point, SearchRadius) > -1;
+	}
+
 	/**
 	 * Insert isolated point P into the arrangement
 	 */
-	void Insert(const FVector2d& Pt)
+	int Insert(const FVector2d& Pt)
 	{
-		insert_point(Pt, VertexSnapTol);
+		return insert_point(Pt, VertexSnapTol);
 	}
 
 
@@ -164,12 +185,12 @@ protected:
 	/**
 	 * insert pt P into the arrangement, splitting existing edges as necessary
 	 */
-	bool insert_point(const FVector2d &P, double Tol = 0)
+	int insert_point(const FVector2d &P, double Tol = 0)
 	{
 		int PIdx = find_existing_vertex(P);
 		if (PIdx > -1)
 		{
-			return false;
+			return -1;
 		}
 
 		// TODO: currently this tries to add the vertex on the closest edge below tolerance; we should instead insert at *every* edge below tolerance!  ... but that is more inconvenient to write
@@ -199,12 +220,12 @@ protected:
 			ensureMsgf(result == EMeshResult::Ok, TEXT("insert_into_segment: edge split failed?"));
 			Graph.SetVertex(splitInfo.VNew, P);
 			PointHash.InsertPointUnsafe(splitInfo.VNew, P);
-			return true;
+			return splitInfo.VNew;
 		}
 
 		int VID = Graph.AppendVertex(P);
 		PointHash.InsertPointUnsafe(VID, P);
-		return true;
+		return VID;
 	}
 
 	/**
@@ -229,6 +250,19 @@ protected:
 			B = Graph.GetVertex(b_idx);
 		}
 
+		// handle tiny-segment case
+		double SegLenSq = A.DistanceSquared(B);
+		if (SegLenSq <= VertexSnapTol*VertexSnapTol)
+		{
+			// seg is too short and was already on an existing vertex; just consider that vertex to be the inserted segment
+			if (a_idx >= 0 || b_idx >= 0)
+			{
+				return false;
+			}
+			// seg is too short and wasn't on an existing vertex; add it as an isolated vertex
+			return insert_point(A, Tol) != -1;
+		}
+
 		// ok find all intersections
 		TArray<FIntersection> Hits;
 		find_intersecting_edges(A, B, Hits, Tol);		
@@ -376,11 +410,10 @@ protected:
 	 */
 	int find_nearest_vertex(FVector2d Pt, double SearchRadius, int IgnoreVID = -1)
 	{
-		double SearchRadiusSquared = SearchRadius * SearchRadius;
-		auto FuncDist = [&](int B) { return Pt.DistanceSquared(Graph.GetVertex(B)); };
+		auto FuncDistSq = [&](int B) { return Pt.DistanceSquared(Graph.GetVertex(B)); };
 		auto FuncIgnore = [&](int VID) { return VID == IgnoreVID; };
-		TPair<int, double> found = (IgnoreVID == -1) ? PointHash.FindNearestInRadius(Pt, SearchRadiusSquared, FuncDist)
-													 : PointHash.FindNearestInRadius(Pt, SearchRadiusSquared, FuncDist, FuncIgnore);
+		TPair<int, double> found = (IgnoreVID == -1) ? PointHash.FindNearestInRadius(Pt, SearchRadius, FuncDistSq)
+													 : PointHash.FindNearestInRadius(Pt, SearchRadius, FuncDistSq, FuncIgnore);
 		if (found.Key == PointHash.InvalidValue())
 		{
 			return -1;
@@ -393,11 +426,10 @@ protected:
 	 */
 	int find_nearest_boundary_vertex(FVector2d Pt, double SearchRadius, int IgnoreVID = -1)
 	{
-		double SearchRadiusSquared = SearchRadius * SearchRadius;
-		auto FuncDist = [&](int B) { return Pt.DistanceSquared(Graph.GetVertex(B)); };
+		auto FuncDistSq = [&](int B) { return Pt.DistanceSquared(Graph.GetVertex(B)); };
 		auto FuncIgnore = [&](int VID) { return Graph.IsBoundaryVertex(VID) == false || VID == IgnoreVID; };
 		TPair<int, double> found =
-			PointHash.FindNearestInRadius(Pt, SearchRadiusSquared, FuncDist, FuncIgnore);
+			PointHash.FindNearestInRadius(Pt, SearchRadius, FuncDistSq, FuncIgnore);
 		if (found.Key == PointHash.InvalidValue())
 		{
 			return -1;
