Geometry updates. Add MeshTangents that calculates vertex tangent spaces based on DynamicMesh UV/Normal overlays.
MeshDescriptionBuilder::AppendMesh() now can copy per-vertex colors and normals.
Add FMeshDescriptionToDynamicMesh::CopyTangents.
Add FMathUtil::Clamp, VectorUtil::BinormalSign, VectorUtil::Binormal.
#rnx

diff --git a/Engine/Plugins/Experimental/GeometryProcessing/Source/DynamicMesh/Private/MeshTangents.cpp b/Engine/Plugins/Experimental/GeometryProcessing/Source/DynamicMesh/Private/MeshTangents.cpp
new file mode 100644
--- /dev/null
+++ b/Engine/Plugins/Experimental/GeometryProcessing/Source/DynamicMesh/Private/MeshTangents.cpp
@@ -0,0 +1,76 @@
+// Copyright 1998-2019 Epic Games, Inc. All Rights Reserved.
+
+#include "MeshTangents.h"
+#include "Async/ParallelFor.h"
+
+
+template<typename RealType>
+void FMeshTangents<RealType>::SetTangentCount(int Count, bool bClearToZero)
+{
+	if (Tangents.Num() < Count)
+	{
+		Tangents.SetNum(Count);
+	}
+	if (Bitangents.Num() < Count)
+	{
+		Bitangents.SetNum(Count);
+	}
+	if (bClearToZero)
+	{
+		for (int i = 0; i < Count; ++i)
+		{
+			Tangents[i] = FVector3<RealType>::Zero();
+			Bitangents[i] = FVector3<RealType>::Zero();
+		}
+	}
+}
+
+
+
+template<typename RealType>
+void FMeshTangents<RealType>::Internal_ComputePerTriangleTangents(const FDynamicMeshNormalOverlay* NormalOverlay, const FDynamicMeshUVOverlay* UVOverlay)
+{
+	int MaxTriangleID = Mesh->MaxTriangleID();
+	InitializePerTriangleTangents(false);
+
+	ParallelFor(MaxTriangleID, [this, NormalOverlay, UVOverlay](int TriangleID)
+	{
+		if (Mesh->IsTriangle(TriangleID) == false)
+		{
+			return;
+		}
+
+		FVector3d TriVertices[3];
+		Mesh->GetTriVertices(TriangleID, TriVertices[0], TriVertices[1], TriVertices[2]);
+		FVector2f TriUVs[3];
+		UVOverlay->GetTriElements(TriangleID, TriUVs[0], TriUVs[1], TriUVs[2]);
+
+		for (int j = 0; j < 3; ++j) 
+		{
+			FVector3d DPosition1 = TriVertices[(j+1)%3] - TriVertices[j];
+			FVector3d DPosition2 = TriVertices[(j+2)%3] - TriVertices[j];
+			FVector2d DUV1 = (FVector2d)TriUVs[(j+1)%3] - (FVector2d)TriUVs[j];
+			FVector2d DUV2 = (FVector2d)TriUVs[(j+2)%3] - (FVector2d)TriUVs[j];
+
+			//@todo handle degenerate edges
+
+			double DetUV = DUV1.Cross(DUV2);
+			double InvDetUV = (DetUV == 0.0f) ? 0.0f : 1.0 / DetUV;
+			FVector3d Tangent = (DPosition1 * DUV2.Y - DPosition2 * DUV1.Y) * InvDetUV;
+			Tangent.Normalize();
+			Tangents[3*TriangleID + j] = (FVector3<RealType>)Tangent;
+			
+			FVector3d Bitangent = (DPosition2 * DUV1.X - DPosition1 * DUV2.X) * InvDetUV;
+			Bitangent.Normalize();
+			Bitangents[3*TriangleID + j] = (FVector3<RealType>)Bitangent;
+		}
+	});
+}
+
+
+template class DYNAMICMESH_API FMeshTangents<float>;
+template class DYNAMICMESH_API FMeshTangents<double>;
+
+
+
+
diff --git a/Engine/Plugins/Experimental/GeometryProcessing/Source/DynamicMesh/Public/MeshTangents.h b/Engine/Plugins/Experimental/GeometryProcessing/Source/DynamicMesh/Public/MeshTangents.h
new file mode 100644
--- /dev/null
+++ b/Engine/Plugins/Experimental/GeometryProcessing/Source/DynamicMesh/Public/MeshTangents.h
@@ -0,0 +1,103 @@
+// Copyright 1998-2019 Epic Games, Inc. All Rights Reserved.
+
+#pragma once
+
+#include "DynamicMesh3.h"
+#include "DynamicMeshAttributeSet.h"
+
+/**
+ * FMeshTangents is a utility class that can calculate and store various types of
+ * tangent vectors for a FDynamicMesh.
+ */
+template<typename RealType>
+class DYNAMICMESH_API FMeshTangents
+{
+protected:
+	/** Target Mesh */
+	const FDynamicMesh3* Mesh;
+	/** Set of computed tangents */
+	TArray<FVector3<RealType>> Tangents;
+	/** Set of computed bitangents */
+	TArray<FVector3<RealType>> Bitangents;
+
+public:
+	FMeshTangents()
+	{
+		Mesh = nullptr;
+	}
+
+	FMeshTangents(const FDynamicMesh3* Mesh)
+	{
+		SetMesh(Mesh);
+	}
+
+	void SetMesh(const FDynamicMesh3* MeshIn)
+	{
+		this->Mesh = MeshIn;
+	}
+
+	const TArray<FVector3<RealType>>& GetTangents() const { return Tangents; }
+
+	const TArray<FVector3<RealType>>& GetBitangents() const { return Bitangents; }
+
+
+	/**
+	 * Return tangent and bitangent at a vertex of triangle for per-triangle computed tangents
+	 * @param TriangleID triangle index in mesh
+	 * @param TriVertIdx vertex index in range 0,1,2
+	 */
+	void GetPerTriangleTangent(int TriangleID, int TriVertIdx, FVector3<RealType>& TangentOut, FVector3<RealType>& BitangentOut)
+	{
+		int k = TriangleID * 3 + TriVertIdx;
+		TangentOut = Tangents[k];
+		BitangentOut = Bitangents[k];
+	}
+
+
+	/**
+	 * Set tangent and bitangent at a vertex of triangle for per-triangle computed tangents.
+	 * @param TriangleID triangle index in mesh
+	 * @param TriVertIdx vertex index in range 0,1,2
+	 */
+	void SetPerTriangleTangent(int TriangleID, int TriVertIdx, const FVector3<RealType>& Tangent, const FVector3<RealType>& Bitangent)
+	{
+		int k = TriangleID * 3 + TriVertIdx;
+		Tangents[k] = Tangent;
+		Bitangents[k] = Bitangent;
+	}
+
+
+	/**
+	 * Calculate per-triangle tangent spaces based on the given per-triangle normal and UV overlays.
+	 * In this mode there is no averaging of tangents across triangles. So if we have N triangles
+	 * in the mesh, then 3*N tangents are generated. These tangents are computed in parallel.
+	 */
+	void ComputePerTriangleTangents(const FDynamicMeshNormalOverlay* NormalOverlay, const FDynamicMeshUVOverlay* UVOverlay)
+	{
+		Internal_ComputePerTriangleTangents(NormalOverlay, UVOverlay);
+	}
+
+
+	/**
+	 * Set internal buffer sizes suitable for calculating per-triangle tangents.
+	 * This is intended to be used if you wish to calculate your own tangents and use SetPerTriangleTangent()
+	 */
+	void InitializePerTriangleTangents(bool bClearToZero)
+	{
+		SetTangentCount(Mesh->MaxTriangleID() * 3, bClearToZero);
+	}
+
+
+
+protected:
+	/**
+	 * Set the size of the Tangents array to Count, and optionally clear all values to (0,0,0)
+	 */
+	void SetTangentCount(int Count, bool bClearToZero);
+
+	// Calculate per-triangle tangents
+	void Internal_ComputePerTriangleTangents(const FDynamicMeshNormalOverlay* NormalOverlay, const FDynamicMeshUVOverlay* UVOverlay);
+};
+
+typedef FMeshTangents<float> FMeshTangentsf;
+typedef FMeshTangents<double> FMeshTangentsd;
\ No newline at end of file
diff --git a/Engine/Plugins/Experimental/GeometryProcessing/Source/GeometricObjects/Public/MathUtil.h b/Engine/Plugins/Experimental/GeometryProcessing/Source/GeometricObjects/Public/MathUtil.h
--- a/Engine/Plugins/Experimental/GeometryProcessing/Source/GeometricObjects/Public/MathUtil.h
+++ b/Engine/Plugins/Experimental/GeometryProcessing/Source/GeometricObjects/Public/MathUtil.h
@@ -49,6 +49,7 @@ public:
 	static inline bool IsNaN(const RealType Value);
 	static inline bool IsFinite(const RealType Value);
 	static inline RealType Abs(const RealType Value);
+	static inline RealType Clamp(const RealType Value, const RealType ClampMin, const RealType ClampMax);
 	static inline RealType Sign(const RealType Value);
 	static inline RealType SignNonZero(const RealType Value);
 	static inline RealType Max(const RealType A, const RealType B);
@@ -95,6 +96,13 @@ RealType TMathUtil<RealType>::Abs(const RealType Value)
 }
 
 
+template<typename RealType>
+RealType TMathUtil<RealType>::Clamp(const RealType Value, const RealType ClampMin, const RealType ClampMax)
+{
+	return (Value < ClampMin) ? ClampMin : ((Value > ClampMax) ? ClampMax : Value);
+}
+
+
 template<typename RealType>
 RealType TMathUtil<RealType>::Sign(const RealType Value)
 {
diff --git a/Engine/Plugins/Experimental/GeometryProcessing/Source/GeometricObjects/Public/VectorTypes.h b/Engine/Plugins/Experimental/GeometryProcessing/Source/GeometricObjects/Public/VectorTypes.h
--- a/Engine/Plugins/Experimental/GeometryProcessing/Source/GeometricObjects/Public/VectorTypes.h
+++ b/Engine/Plugins/Experimental/GeometryProcessing/Source/GeometricObjects/Public/VectorTypes.h
@@ -58,9 +58,23 @@ struct FVector3
 	}
 
 
-	inline operator FColor() const
+	explicit inline operator FLinearColor() const
 	{
-		return FColor((float)X, (float)Y, (float)Z);
+		return FLinearColor((float)X, (float)Y, (float)Z);
+	}
+	inline FVector3(const FLinearColor& Color)
+	{
+		X = (T)Color.R;
+		Y = (T)Color.G;
+		Z = (T)Color.B;
+	}
+
+	explicit inline operator FColor() const
+	{
+		return FColor(
+			FMathf::Clamp((int)((float)X*255.0f), 0, 255),
+			FMathf::Clamp((int)((float)Y*255.0f), 0, 255),
+			FMathf::Clamp((int)((float)Z*255.0f), 0, 255));
 	}
 
 
diff --git a/Engine/Plugins/Experimental/GeometryProcessing/Source/GeometricObjects/Public/VectorUtil.h b/Engine/Plugins/Experimental/GeometryProcessing/Source/GeometricObjects/Public/VectorUtil.h
--- a/Engine/Plugins/Experimental/GeometryProcessing/Source/GeometricObjects/Public/VectorUtil.h
+++ b/Engine/Plugins/Experimental/GeometryProcessing/Source/GeometricObjects/Public/VectorUtil.h
@@ -186,6 +186,32 @@ namespace VectorUtil
 		}
 	}
 
+	/**
+	 * Calculates one vector perpendicular to input Normal, as efficiently as possible.
+	 */
+	template <typename RealType>
+	inline void MakePerpVector(const FVector3<RealType>& Normal, FVector3<RealType>& OutPerp1)
+	{
+		// Duff et al method, from https://graphics.pixar.com/library/OrthonormalB/paper.pdf
+		if (Normal.Z < (RealType)0)
+		{
+			RealType A = (RealType)1 / ((RealType)1 - Normal.Z);
+			RealType B = Normal.X * Normal.Y * A;
+			OutPerp1.X = (RealType)1 - Normal.X * Normal.X * A;
+			OutPerp1.Y = -B;
+			OutPerp1.Z = Normal.X;
+		}
+		else
+		{
+			RealType A = (RealType)1 / ((RealType)1 + Normal.Z);
+			RealType B = -Normal.X * Normal.Y * A;
+			OutPerp1.X = (RealType)1 - Normal.X * Normal.X * A;
+			OutPerp1.Y = B;
+			OutPerp1.Z = -Normal.X;
+		}
+	}
+
+
 	/**
 	 * Calculates angle between VFrom and VTo after projection onto plane with normal defined by PlaneN
 	 * @return angle in degrees
@@ -272,4 +298,35 @@ namespace VectorUtil
 		return A.AngleD(B);
 	}
 
+
+
+	/**
+	 * @return sign of Binormal/Bitangent relative to Normal and Tangent
+	 */
+	template<typename RealType>
+	inline RealType BinormalSign(const FVector3<RealType>& NormalIn, const FVector3<RealType>& TangentIn, const FVector3<RealType>& BinormalIn)
+	{
+		// following math from RenderUtils.h::GetBasisDeterminantSign()
+		RealType Cross00 = BinormalIn.Y*NormalIn.Z - BinormalIn.Z*NormalIn.Y;
+		RealType Cross10 = BinormalIn.Z*NormalIn.X - BinormalIn.X*NormalIn.Z;
+		RealType Cross20 = BinormalIn.X*NormalIn.Y - BinormalIn.Y*NormalIn.X;
+		RealType Determinant = TangentIn.X*Cross00 + TangentIn.Y*Cross10 + TangentIn.Z*Cross20;
+		return (Determinant < 0) ? (RealType)-1 : (RealType)1;
+	}
+
+	/**
+	 * @return Binormal vector based on given Normal, Tangent, and Sign value (+1/-1)
+	 */
+	template<typename RealType>
+	inline FVector3<RealType> Binormal(const FVector3<RealType>& NormalIn, const FVector3<RealType>& TangentIn, RealType BinormalSign)
+	{
+		return BinormalSign * FVector3<RealType>(
+			NormalIn.Y*TangentIn.Z - NormalIn.Z*TangentIn.Y,
+			NormalIn.Z*TangentIn.X - NormalIn.X*TangentIn.Z,
+			NormalIn.X*TangentIn.Y - NormalIn.Y*TangentIn.X);
+	}
+
+
+
+
 }; // namespace VectorUtil
diff --git a/Engine/Plugins/Experimental/GeometryProcessing/Source/MeshConversion/Private/MeshDescriptionBuilder.cpp b/Engine/Plugins/Experimental/GeometryProcessing/Source/MeshConversion/Private/MeshDescriptionBuilder.cpp
--- a/Engine/Plugins/Experimental/GeometryProcessing/Source/MeshConversion/Private/MeshDescriptionBuilder.cpp
+++ b/Engine/Plugins/Experimental/GeometryProcessing/Source/MeshConversion/Private/MeshDescriptionBuilder.cpp
@@ -24,6 +24,8 @@ void FMeshDescriptionBuilder::SetMeshDescription(FMeshDescription* Description)
 		MeshDescription->VertexInstanceAttributes().GetAttributesRef<FVector2D>(MeshAttribute::VertexInstance::TextureCoordinate);
 	this->InstanceNormals =
 		MeshDescription->VertexInstanceAttributes().GetAttributesRef<FVector>(MeshAttribute::VertexInstance::Normal);
+	this->InstanceColors =
+		MeshDescription->VertexInstanceAttributes().GetAttributesRef<FVector4>(MeshAttribute::VertexInstance::Color);
 }
 
 
@@ -103,6 +105,22 @@ void FMeshDescriptionBuilder::SetInstance(const FVertexInstanceID& InstanceID, c
 }
 
 
+void FMeshDescriptionBuilder::SetInstanceNormal(const FVertexInstanceID& InstanceID, const FVector& Normal)
+{
+	if (InstanceNormals.IsValid())
+	{
+		InstanceNormals.Set(InstanceID, Normal);
+	}
+}
+
+
+void FMeshDescriptionBuilder::SetInstanceColor(const FVertexInstanceID& InstanceID, const FVector4& Color)
+{
+	if (InstanceColors.IsValid())
+	{
+		InstanceColors.Set(InstanceID, Color);
+	}
+}
 
 
 FPolygonID FMeshDescriptionBuilder::AppendTriangle(const FVertexID* Triangle, const FPolygonGroupID& PolygonGroup, 
@@ -220,6 +238,8 @@ void FMeshDescriptionBuilder::AppendMesh(const FDynamicMesh3* Mesh, bool bSetPol
 
 	FPolygonGroupID AllGroupID = AppendPolygonGroup();
 
+	bool bHasVertexColors = Mesh->HasVertexColors();
+	bool bHasVertexNormals = Mesh->HasVertexNormals();
 	const FDynamicMeshUVOverlay* UVOverlay = Mesh->HasAttributes() ? Mesh->Attributes()->PrimaryUV() : nullptr;
 	const FDynamicMeshNormalOverlay* NormalOverlay = Mesh->HasAttributes() ? Mesh->Attributes()->PrimaryNormals() : nullptr;
 	//FDynamicMeshUVOverlay* UVOverlay = nullptr;
@@ -254,6 +274,22 @@ void FMeshDescriptionBuilder::AppendMesh(const FDynamicMesh3* Mesh, bool bSetPol
 		{
 			SetPolyGroupID(NewPolyID, Mesh->GetTriangleGroup(TriID));
 		}
+
+		if (bHasVertexColors)
+		{
+			for (int j = 0; j < 3; ++j)
+			{
+				FVector3f Color = Mesh->GetVertexColor(Triangle[j]);
+				FVector4 Color4(Color.X, Color.Y, Color.Z, 1.0);
+				SetInstanceColor(InstanceTri[j], Color4);
+			}
+		}
+		if (NormalOverlay == nullptr && bHasVertexNormals)
+		{
+			SetInstanceNormal(InstanceTri[0], Mesh->GetVertexNormal(Triangle[0]));
+			SetInstanceNormal(InstanceTri[1], Mesh->GetVertexNormal(Triangle[1]));
+			SetInstanceNormal(InstanceTri[2], Mesh->GetVertexNormal(Triangle[2]));
+		}
 	}
 }
 
diff --git a/Engine/Plugins/Experimental/GeometryProcessing/Source/MeshConversion/Private/MeshDescriptionToDynamicMesh.cpp b/Engine/Plugins/Experimental/GeometryProcessing/Source/MeshConversion/Private/MeshDescriptionToDynamicMesh.cpp
--- a/Engine/Plugins/Experimental/GeometryProcessing/Source/MeshConversion/Private/MeshDescriptionToDynamicMesh.cpp
+++ b/Engine/Plugins/Experimental/GeometryProcessing/Source/MeshConversion/Private/MeshDescriptionToDynamicMesh.cpp
@@ -173,8 +173,11 @@ void FMeshDescriptionToDynamicMesh::Convert(const FMeshDescription* MeshIn, FDyn
 		FPolygonGroupID PolygonGroupID = MeshIn->GetPolygonPolygonGroup(PolygonID);
 
 		const TArray<FMeshTriangle>& Triangles = MeshIn->GetPolygonTriangles(PolygonID);
-		for (FMeshTriangle Triangle : Triangles)
+		int NumTriangles = Triangles.Num();
+		for ( int TriIdx = 0; TriIdx < NumTriangles; ++TriIdx)
 		{
+			const FMeshTriangle& Triangle = Triangles[TriIdx];
+				
 			FVertexInstanceID InstanceTri[3];
 			InstanceTri[0] = Triangle.VertexInstanceID0; 
 			InstanceTri[1] = Triangle.VertexInstanceID1; 
@@ -205,6 +208,11 @@ void FMeshDescriptionToDynamicMesh::Convert(const FMeshDescription* MeshIn, FDyn
 			int NewTriangleID = MeshOut.AppendTriangle(VertexID0, VertexID1, VertexID2, GroupID);
 			FIndex3i Tri(VertexID0, VertexID1, VertexID2);
 
+			if (bCalculateMaps)
+			{
+				TriToPolyTriMap.Insert(FIndex2i(PolygonID.GetValue(), TriIdx), NewTriangleID);
+			}
+
 			if (UVOverlay != nullptr)
 			{
 				FIndex3i TriUV;
@@ -238,3 +246,47 @@ void FMeshDescriptionToDynamicMesh::Convert(const FMeshDescription* MeshIn, FDyn
 	}
 
 }
+
+
+
+
+
+void FMeshDescriptionToDynamicMesh::CopyTangents(const FMeshDescription* SourceMesh, const FDynamicMesh3* TargetMesh, FMeshTangentsf& TangentsOut)
+{
+	check(bCalculateMaps == true);
+	check(TriToPolyTriMap.Num() == TargetMesh->TriangleCount());
+
+	TVertexInstanceAttributesConstRef<FVector> InstanceNormals =
+		SourceMesh->VertexInstanceAttributes().GetAttributesRef<FVector>(MeshAttribute::VertexInstance::Normal);
+	TVertexInstanceAttributesConstRef<FVector> InstanceTangents =
+		SourceMesh->VertexInstanceAttributes().GetAttributesRef<FVector>(MeshAttribute::VertexInstance::Tangent);
+	TVertexInstanceAttributesConstRef<float> InstanceSigns =
+		SourceMesh->VertexInstanceAttributes().GetAttributesRef<float>(MeshAttribute::VertexInstance::BinormalSign);
+	check(InstanceNormals.IsValid());
+	check(InstanceTangents.IsValid());
+	check(InstanceSigns.IsValid());
+
+	TangentsOut.SetMesh(TargetMesh);
+	TangentsOut.InitializePerTriangleTangents(false);
+	
+	for (int TriID : TargetMesh->TriangleIndicesItr())
+	{
+		FIndex2i PolyTriIdx = TriToPolyTriMap[TriID];
+		const TArray<FMeshTriangle>& Triangles = SourceMesh->GetPolygonTriangles( FPolygonID(PolyTriIdx.A) );
+		const FMeshTriangle& Triangle = Triangles[PolyTriIdx.B];
+		FVertexInstanceID InstanceTri[3];
+		InstanceTri[0] = Triangle.VertexInstanceID0;
+		InstanceTri[1] = Triangle.VertexInstanceID1;
+		InstanceTri[2] = Triangle.VertexInstanceID2;
+		for (int j = 0; j < 3; ++j)
+		{
+			FVector Normal = InstanceNormals.Get(InstanceTri[j], 0);
+			FVector Tangent = InstanceTangents.Get(InstanceTri[j], 0);
+			float BinormalSign = InstanceSigns.Get(InstanceTri[j], 0);
+			FVector3f Bitangent = VectorUtil::Binormal((FVector3f)Normal, (FVector3f)Tangent, (float)BinormalSign);
+			Tangent.Normalize(); Bitangent.Normalize();
+			TangentsOut.SetPerTriangleTangent(TriID, j, Tangent, Bitangent);
+		}
+	}
+
+}
\ No newline at end of file
diff --git a/Engine/Plugins/Experimental/GeometryProcessing/Source/MeshConversion/Public/MeshDescriptionBuilder.h b/Engine/Plugins/Experimental/GeometryProcessing/Source/MeshConversion/Public/MeshDescriptionBuilder.h
--- a/Engine/Plugins/Experimental/GeometryProcessing/Source/MeshConversion/Public/MeshDescriptionBuilder.h
+++ b/Engine/Plugins/Experimental/GeometryProcessing/Source/MeshConversion/Public/MeshDescriptionBuilder.h
@@ -45,7 +45,11 @@ public:
 	/** Set the UV and Normal of a vertex instance*/
 	void SetInstance(const FVertexInstanceID& InstanceID, const FVector2D& InstanceUV, const FVector& InstanceNormal);
 
+	/** Set the Normal of a vertex instance*/
+	void SetInstanceNormal(const FVertexInstanceID& InstanceID, const FVector& Normal);
 
+	/** Set the Color of a vertex instance*/
+	void SetInstanceColor(const FVertexInstanceID& InstanceID, const FVector4& Color);
 
 	/** Enable per-triangle integer attribute named PolyTriGroups */
 	void EnablePolyGroups();
@@ -104,6 +108,7 @@ protected:
 	TVertexAttributesRef<FVector> VertexPositions;
 	TVertexInstanceAttributesRef<FVector2D> InstanceUVs;
 	TVertexInstanceAttributesRef<FVector> InstanceNormals;
+	TVertexInstanceAttributesRef<FVector4> InstanceColors;
 	TArray<FVertexID> TempBuffer;
 	TArray<FVector2D> UVBuffer;
 	TArray<FVector> NormalBuffer;
diff --git a/Engine/Plugins/Experimental/GeometryProcessing/Source/MeshConversion/Public/MeshDescriptionToDynamicMesh.h b/Engine/Plugins/Experimental/GeometryProcessing/Source/MeshConversion/Public/MeshDescriptionToDynamicMesh.h
--- a/Engine/Plugins/Experimental/GeometryProcessing/Source/MeshConversion/Public/MeshDescriptionToDynamicMesh.h
+++ b/Engine/Plugins/Experimental/GeometryProcessing/Source/MeshConversion/Public/MeshDescriptionToDynamicMesh.h
@@ -4,6 +4,7 @@
 
 #include "CoreMinimal.h"
 #include "DynamicMesh3.h"
+#include "MeshTangents.h"
 #include "MeshDescription.h"
 
 /**
@@ -23,6 +24,13 @@ public:
 	/** Should we initialize triangle groups on output mesh */
 	bool bEnableOutputGroups = true;
 
+	/** Should we calculate conversion index maps */
+	bool bCalculateMaps = true;
+
+	/** map from triangle ID to (polygon,triangle) pair */
+	TArray<FIndex2i> TriToPolyTriMap;
+
+
 	/**
 	 * Various modes can be used to create output triangle groups
 	 */
@@ -43,4 +51,10 @@ public:
 	 * Default conversion of MeshDescription to DynamicMesh
 	 */
 	void Convert(const FMeshDescription* MeshIn, FDynamicMesh3& MeshOut);
+
+	/**
+	 * Copy tangents from MeshDescription to a FMeshTangents instance.
+	 * @warning Convert() must have been used to create the TargetMesh before calling this function
+	 */
+	void CopyTangents(const FMeshDescription* SourceMesh, const FDynamicMesh3* TargetMesh, FMeshTangentsf& TangentsOut);
 };
\ No newline at end of file
