Copying Arrangement2d robustness fixes from //UE4/Dev-Physics to Dev-Editor (//UE4/Dev-Editor)

diff --git a/Engine/Plugins/Experimental/GeometryProcessing/Source/GeometryAlgorithms/Public/Arrangement2d.h b/Engine/Plugins/Experimental/GeometryProcessing/Source/GeometryAlgorithms/Public/Arrangement2d.h
--- a/Engine/Plugins/Experimental/GeometryProcessing/Source/GeometryAlgorithms/Public/Arrangement2d.h
+++ b/Engine/Plugins/Experimental/GeometryProcessing/Source/GeometryAlgorithms/Public/Arrangement2d.h
@@ -386,12 +386,13 @@ protected:
 		}
 		else
 		{
+			FVector2d Pt = seg.PointAt(T);
+			// ensure(find_existing_vertex(Pt) == -1); // TODO: handle this case!?
 			FDynamicGraph2d::FEdgeSplitInfo splitInfo;
 			EMeshResult result = Graph.SplitEdge(EID, splitInfo);
 			ensureMsgf(result == EMeshResult::Ok, TEXT("insert_into_segment: edge split failed?"));
 			use_vid = splitInfo.VNew;
 			new_eid = splitInfo.ENewBN;
-			FVector2d Pt = seg.PointAt(T);
 			Graph.SetVertex(use_vid, Pt);
 			PointHash.InsertPointUnsafe(splitInfo.VNew, Pt);
 		}
@@ -452,11 +453,16 @@ protected:
 	{
 		int num_hits = 0;
 		FVector2d x = FVector2d::Zero(), y = FVector2d::Zero();
+		FVector2d EPerp = (B - A).Perp();
+		EPerp.Normalize();
 		for (int EID : Graph.EdgeIndices())
 		{
 			Graph.GetEdgeV(EID, x, y);
-			int SideX = FSegment2d::WhichSide(A, B, x, Tol);
-			int SideY = FSegment2d::WhichSide(A, B, y, Tol);
+			// inlined version of WhichSide with pre-normalized EPerp, to ensure Tolerance is consistent for different edge lengths
+			double SignX = EPerp.Dot(x - A);
+			double SignY = EPerp.Dot(y - A);
+			int SideX = (SignX > Tol ? +1 : (SignX < -Tol ? -1 : 0)); 
+			int SideY = (SignY > Tol ? +1 : (SignY < -Tol ? -1 : 0)); 
 			if (SideX == SideY && SideX != 0)
 			{
 				continue; // both pts on same side
@@ -464,6 +470,10 @@ protected:
 
 			FIntrSegment2Segment2d Intr(FSegment2d(x, y), FSegment2d(A, B));
 			Intr.SetIntervalThreshold(Tol);
+			// set a loose DotThreshold as well so almost-parallel segments are treated as parallel;
+			//  otherwise we're more likely to hit later problems when an edge intersects near-overlapping edges at almost the same point
+			// (TODO: detect + handle that case!)
+			Intr.SetDotThreshold(1e-4);
 			if (Intr.Find())
 			{
 				Hits.Add(FIntersection{EID, SideX, SideY, Intr});
