Mesh Simplification improvements
6623697 Modify mesh simplification to template on quadric type and add attribute aware mesh quadric.
6620269 Fallback to duplicating vertices when nonmanifold geometry is encountered in MeshDescriptionToDynamicMesh  (see original CL for more details)
#rnx

diff --git a/Engine/Plugins/Experimental/GeometryProcessing/Source/DynamicMesh/Private/MeshSimplification.cpp b/Engine/Plugins/Experimental/GeometryProcessing/Source/DynamicMesh/Private/MeshSimplification.cpp
--- a/Engine/Plugins/Experimental/GeometryProcessing/Source/DynamicMesh/Private/MeshSimplification.cpp
+++ b/Engine/Plugins/Experimental/GeometryProcessing/Source/DynamicMesh/Private/MeshSimplification.cpp
@@ -6,41 +6,91 @@
 
 
 
+template <typename QuadricErrorType>
+QuadricErrorType TMeshSimplification<QuadricErrorType>::ComputeFaceQuadric(const int tid, FVector3d& nface, FVector3d& c, double& Area) const
+{
+	// compute the new quadric for this tri.
+	Mesh->GetTriInfo(tid, nface, Area, c);
+
+	return FQuadricErrorType(nface, c);
+}
+
+
+// Face Quadric Error computation specialized for FAttrBasedQuadricErrord
+template<>
+FAttrBasedQuadricErrord TMeshSimplification<FAttrBasedQuadricErrord>::ComputeFaceQuadric(const int tid, FVector3d& nface, FVector3d& c, double& Area) const
+{
+	// compute the new quadric for this tri.
+	Mesh->GetTriInfo(tid, nface, Area, c);
+
+	FVector3f n0; FVector3f n1; FVector3f n2;
+
+	if (NormalOverlay != nullptr)
+	{
+		NormalOverlay->GetTriElements(tid, n0, n1, n2);
+	}
+	else
+	{
+		FIndex3i vids = Mesh->GetTriangle(tid);
+		n0 = Mesh->GetVertexNormal(vids[0]);
+		n1 = Mesh->GetVertexNormal(vids[1]);
+		n2 = Mesh->GetVertexNormal(vids[2]);
+	}
 
-void FMeshSimplification::InitializeVertexQuadrics()
+
+	FVector3d p0, p1, p2;
+	Mesh->GetTriVertices(tid, p0, p1, p2);
+
+	FVector3d n0d(n0.X, n0.Y, n0.Z);
+	FVector3d n1d(n1.X, n1.Y, n1.Z);
+	FVector3d n2d(n2.X, n2.Y, n2.Z);
+
+	double attrweight = 1.;
+	return FQuadricErrorType(p0, p1, p2, n0d, n1d, n2d, nface, c, attrweight);
+}
+
+template <typename QuadricErrorType>
+void TMeshSimplification<QuadricErrorType>::InitializeTriQuadrics()
 {
-	int NT = Mesh->MaxTriangleID();
-	TArray<FQuadricErrord> triQuadrics;
+	const int NT = Mesh->MaxTriangleID();
 	triQuadrics.SetNum(NT);
-	TArray<double> triAreas;
 	triAreas.SetNum(NT);
+
+	// tested with ParallelFor - no measurable benifit
 	//@todo parallel version
 	//gParallel.BlockStartEnd(0, Mesh->MaxTriangleID - 1, (start_tid, end_tid) = > {
-	FVector3d c, n;
+	FVector3d n, c;
 	for (int tid : Mesh->TriangleIndicesItr())
 	{
-		Mesh->GetTriInfo(tid, n, triAreas[tid], c);
-		triQuadrics[tid] = FQuadricErrord(n, c);
+		triQuadrics[tid] = ComputeFaceQuadric(tid, n, c, triAreas[tid]);
 	}
 
+}
+
+template <typename QuadricErrorType>
+void TMeshSimplification<QuadricErrorType>::InitializeVertexQuadrics()
+{
 
 	int NV = Mesh->MaxVertexID();
 	vertQuadrics.SetNum(NV);
+	// tested with ParallelFor - no measurable benifit 
 	//gParallel.BlockStartEnd(0, Mesh->MaxVertexID - 1, (start_vid, end_vid) = > {
 	for (int vid : Mesh->VertexIndicesItr())
 	{
-		vertQuadrics[vid] = FQuadricErrord::Zero();
+		vertQuadrics[vid] = FQuadricErrorType::Zero();
 		for (int tid : Mesh->VtxTrianglesItr(vid))
 		{
 			vertQuadrics[vid].Add(triAreas[tid], triQuadrics[tid]);
 		}
 		//check(TMathUtil.EpsilonEqual(0, vertQuadrics[i].Evaluate(Mesh->GetVertex(i)), TMathUtil.Epsilon * 10));
 	}
+
 }
 
 
 
-void FMeshSimplification::InitializeQueue()
+template <typename QuadricErrorType>
+void TMeshSimplification<QuadricErrorType>::InitializeQueue()
 {
 	int NE = Mesh->EdgeCount();
 	int MaxEID = Mesh->MaxEdgeID();
@@ -56,7 +106,7 @@ void FMeshSimplification::InitializeQueue()
 	for (int eid : Mesh->EdgeIndicesItr())
 	{
 		FIndex2i ev = Mesh->GetEdgeV(eid);
-		FQuadricErrord Q(vertQuadrics[ev.A], vertQuadrics[ev.B]);
+		FQuadricErrorType Q(vertQuadrics[ev.A], vertQuadrics[ev.B]);
 		FVector3d opt = OptimalPoint(eid, Q, ev.A, ev.B);
 		EdgeErrors[eid] = { (float)Q.Evaluate(opt), eid };
 		EdgeQuadrics[eid] = QEdge(eid, Q, opt);
@@ -67,7 +117,7 @@ void FMeshSimplification::InitializeQueue()
 
 	// now do inserts
 	int N = EdgeErrors.Num();
-	for (int i = 0; i < N; ++i) 
+	for (int i = 0; i < N; ++i)
 	{
 		int eid = EdgeErrors[i].eid;
 		if (Mesh->IsEdge(eid))
@@ -97,8 +147,8 @@ void FMeshSimplification::InitializeQueue()
 
 
 
-
-FVector3d FMeshSimplification::OptimalPoint(int eid, const FQuadricErrord& q, int ea, int eb)
+template <typename QuadricErrorType>
+FVector3d TMeshSimplification<QuadricErrorType>::OptimalPoint(int eid, const FQuadricErrorType& q, int ea, int eb)
 {
 	// if we would like to preserve boundary, we need to know that here
 	// so that we properly score these edges
@@ -160,7 +210,8 @@ FVector3d FMeshSimplification::OptimalPoint(int eid, const FQuadricErrord& q, in
 
 
 // update queue weight for each edge in vertex one-ring
-void FMeshSimplification::UpdateNeighbours(int vid)
+template <>
+void DYNAMICMESH_API TMeshSimplification<FQuadricErrord>::UpdateNeighbours(int vid, FIndex2i removedTris, FIndex2i opposingVerts)
 {
 	for (int eid : Mesh->VtxEdgesItr(vid))
 	{
@@ -180,10 +231,117 @@ void FMeshSimplification::UpdateNeighbours(int vid)
 	}
 }
 
+// update queue weight for each edge in vertex one-ring.  Memoryless
+template <typename QuadricErrorType>
+void TMeshSimplification<QuadricErrorType>::UpdateNeighbours(int vid, FIndex2i removedTris, FIndex2i opposingVerts)
+{
+
+
+	// This is the faster version that selectively updates the one-ring
+	{
+
+		// compute the change in affected face quadrics, and then propagate 
+		// that change to the face adjacent verts.
+		FVector3d n, c;
+		double NewtriArea;
+
+		// Update the triangle areas and quadrics that will have changed
+		for (int tid : Mesh->VtxTrianglesItr(vid))
+		{
+
+			const double OldtriArea = triAreas[tid];
+			const FQuadricErrorType OldtriQuadric = triQuadrics[tid];
+
+
+			// compute the new quadric for this tri.
+			FQuadricErrorType NewtriQuadric = ComputeFaceQuadric(tid, n, c, NewtriArea);
+
+			// update the arrays that hold the current face area & quadrics
+			triAreas[tid] = NewtriArea;
+			triQuadrics[tid] = NewtriQuadric;
+
+			FIndex3i tri_vids = Mesh->GetTriangle(tid);
+
+			// update the vert quadrics that are adjacent to vid.
+			for (int32 i = 0; i < 3; ++i)
+			{
+				if (tri_vids[i] == vid) continue;
+
+				// correct the adjacent vertQuadrics
+				vertQuadrics[tri_vids[i]].Add(-OldtriArea, OldtriQuadric); // subtract old quadric
+				vertQuadrics[tri_vids[i]].Add(NewtriArea, NewtriQuadric); // add new quadric
+			}
+		}
+
+		// remove the influence of the dead tris from the two verts that were opposing the collapsed edge
+		{
+			for (int i = 0; i < 2; ++i)
+			{
+				if (removedTris[i] != FDynamicMesh3::InvalidID)
+				{
+					const double   oldArea = triAreas[removedTris[i]];
+					FQuadricErrorType oldQuadric = triQuadrics[removedTris[i]];
+
+					triAreas[removedTris[i]] = 0.;
+
+					// subtract the quadric from the opposing vert
+					vertQuadrics[opposingVerts[i]].Add(-oldArea, oldQuadric);
+				}
+			}
+		}
+		// Rebuild the quadric for the vert that was retained during the collapse.
+		// NB: in the version with memory this quadric took the value of the edge quadric that collapsed.
+		{
+			FQuadricErrorType vertQuadric;
+			for (int tid : Mesh->VtxTrianglesItr(vid))
+			{
+				vertQuadric.Add(triAreas[tid], triQuadrics[tid]);
+			}
+			vertQuadrics[vid] = vertQuadric;
+		}
+	}
+
+	// Update all the edges
+	{
+		TArray<int, TInlineAllocator<64>> EdgesToUpdate;
+		for (int adjvid : Mesh->VtxVerticesItr(vid))
+		{
+			for (int eid : Mesh->VtxEdgesItr(adjvid))
+			{
+				EdgesToUpdate.AddUnique(eid);
+			}
+		}
+
+		for (int eid : EdgesToUpdate)
+		{
+			// The volume conservation plane data held in the 
+			// vertex quadrics will have duplicates for 
+			// the two face adjacent to the edge.
+
+			const FIndex4i edgeData = Mesh->GetEdge(eid);
+			FQuadricErrorType Q(vertQuadrics[edgeData[0]], vertQuadrics[edgeData[1]]);
+
+			FVector3d opt = OptimalPoint(eid, Q, edgeData[0], edgeData[1]);
+			double err = Q.Evaluate(opt);
+			EdgeQuadrics[eid] = QEdge(eid, Q, opt);
+			if (EdgeQueue.Contains(eid))
+			{
+				EdgeQueue.Update(eid, (float)err);
+			}
+			else
+			{
+				EdgeQueue.Insert(eid, (float)err);
+			}
+		}
+	}
+}
+
+
 
 
 
-void FMeshSimplification::Precompute(bool bMeshIsClosed)
+template <typename QuadricErrorType>
+void TMeshSimplification<QuadricErrorType>::Precompute(bool bMeshIsClosed)
 {
 	bHaveBoundary = false;
 	IsBoundaryVtxCache.SetNum(Mesh->MaxVertexID());
@@ -202,7 +360,8 @@ void FMeshSimplification::Precompute(bool bMeshIsClosed)
 
 
 
-void FMeshSimplification::DoSimplify()
+template <typename QuadricErrorType>
+void TMeshSimplification<QuadricErrorType>::DoSimplify()
 {
 	if (Mesh->TriangleCount() == 0)    // badness if we don't catch this...
 	{
@@ -217,6 +376,11 @@ void FMeshSimplification::DoSimplify()
 	{
 		return;
 	}
+	InitializeTriQuadrics();
+	if (Cancelled())
+	{
+		return;
+	}
 	InitializeVertexQuadrics();
 	if (Cancelled())
 	{
@@ -232,7 +396,7 @@ void FMeshSimplification::DoSimplify()
 	ProfileBeginOps();
 
 	ProfileBeginCollapse();
-	while (EdgeQueue.GetCount() > 0) 
+	while (EdgeQueue.GetCount() > 0)
 	{
 		// termination criteria
 		if (SimplifyMode == ETargetModes::VertexCount)
@@ -242,7 +406,7 @@ void FMeshSimplification::DoSimplify()
 				break;
 			}
 		}
-		else 
+		else
 		{
 			if (Mesh->TriangleCount() <= TargetCount)
 			{
@@ -261,12 +425,18 @@ void FMeshSimplification::DoSimplify()
 			return;
 		}
 
+		// find triangles adjacent to the target edge
+		// and the verts opposite the edge.
+		FIndex2i targetTris = Mesh->GetEdgeT(eid);
+		FIndex2i targetVrts = Mesh->GetEdgeOpposingV(eid);
+
+
 		int vKeptID;
-		EProcessResult result = CollapseEdge(eid, EdgeQuadrics[eid].collapse_pt, vKeptID);
-		if (result == EProcessResult::Ok_Collapsed) 
+		ESimplificationResult result = CollapseEdge(eid, EdgeQuadrics[eid].collapse_pt, vKeptID);
+		if (result == ESimplificationResult::Ok_Collapsed)
 		{
 			vertQuadrics[vKeptID] = EdgeQuadrics[eid].q;
-			UpdateNeighbours(vKeptID);
+			UpdateNeighbours(vKeptID, targetTris, targetVrts);
 		}
 	}
 	ProfileEndCollapse();
@@ -283,7 +453,8 @@ void FMeshSimplification::DoSimplify()
 }
 
 
-void FMeshSimplification::SimplifyToTriangleCount(int nCount)
+template <typename QuadricErrorType>
+void TMeshSimplification<QuadricErrorType>::SimplifyToTriangleCount(int nCount)
 {
 	SimplifyMode = ETargetModes::TriangleCount;
 	TargetCount = FMath::Max(1, nCount);
@@ -291,7 +462,8 @@ void FMeshSimplification::SimplifyToTriangleCount(int nCount)
 	DoSimplify();
 }
 
-void FMeshSimplification::SimplifyToVertexCount(int nCount)
+template <typename QuadricErrorType>
+void TMeshSimplification<QuadricErrorType>::SimplifyToVertexCount(int nCount)
 {
 	SimplifyMode = ETargetModes::VertexCount;
 	TargetCount = FMath::Max(3, nCount);
@@ -299,7 +471,8 @@ void FMeshSimplification::SimplifyToVertexCount(int nCount)
 	DoSimplify();
 }
 
-void FMeshSimplification::SimplifyToEdgeLength(double minEdgeLen)
+template <typename QuadricErrorType>
+void TMeshSimplification<QuadricErrorType>::SimplifyToEdgeLength(double minEdgeLen)
 {
 	SimplifyMode = ETargetModes::MinEdgeLength;
 	TargetCount = 1;
@@ -309,8 +482,8 @@ void FMeshSimplification::SimplifyToEdgeLength(double minEdgeLen)
 
 
 
-
-void FMeshSimplification::FastCollapsePass(double fMinEdgeLength, int nRounds, bool MeshIsClosedHint)
+template <typename QuadricErrorType>
+void TMeshSimplification<QuadricErrorType>::FastCollapsePass(double fMinEdgeLength, int nRounds, bool MeshIsClosedHint)
 {
 	if (Mesh->TriangleCount() == 0)    // badness if we don't catch this...
 	{
@@ -339,14 +512,14 @@ void FMeshSimplification::FastCollapsePass(double fMinEdgeLength, int nRounds, b
 
 	int N = Mesh->MaxEdgeID();
 	int num_last_pass = 0;
-	for (int ri = 0; ri < nRounds; ++ri) 
+	for (int ri = 0; ri < nRounds; ++ri)
 	{
 		num_last_pass = 0;
 
 		FVector3d va = FVector3d::Zero(), vb = FVector3d::Zero();
-		for (int eid = 0; eid < N; ++eid) 
+		for (int eid = 0; eid < N; ++eid)
 		{
-			if ( ( ! Mesh->IsEdge(eid) ) || Mesh->IsBoundaryEdge(eid) )
+			if ((!Mesh->IsEdge(eid)) || Mesh->IsBoundaryEdge(eid))
 			{
 				continue;
 			}
@@ -365,8 +538,8 @@ void FMeshSimplification::FastCollapsePass(double fMinEdgeLength, int nRounds, b
 
 			FVector3d midpoint = (va + vb) * 0.5;
 			int vKeptID;
-			EProcessResult result = CollapseEdge(eid, midpoint, vKeptID);
-			if (result == EProcessResult::Ok_Collapsed) 
+			ESimplificationResult result = CollapseEdge(eid, midpoint, vKeptID);
+			if (result == ESimplificationResult::Ok_Collapsed)
 			{
 				++num_last_pass;
 			}
@@ -403,7 +576,8 @@ void FMeshSimplification::FastCollapsePass(double fMinEdgeLength, int nRounds, b
 
 
 
-FMeshSimplification::EProcessResult FMeshSimplification::CollapseEdge(int edgeID, FVector3d vNewPos, int& collapseToV)
+template <typename QuadricErrorType>
+ESimplificationResult TMeshSimplification<QuadricErrorType>::CollapseEdge(int edgeID, FVector3d vNewPos, int& collapseToV)
 {
 	collapseToV = FDynamicMesh3::InvalidID;
 	RuntimeDebugCheck(edgeID);
@@ -412,18 +586,18 @@ FMeshSimplification::EProcessResult FMeshSimplification::CollapseEdge(int edgeID
 		(Constraints == nullptr) ? FEdgeConstraint::Unconstrained() : Constraints->GetEdgeConstraint(edgeID);
 	if (constraint.NoModifications())
 	{
-		return EProcessResult::Ignored_EdgeIsFullyConstrained;
+		return ESimplificationResult::Ignored_EdgeIsFullyConstrained;
 	}
 	if (constraint.CanCollapse() == false)
 	{
-		return EProcessResult::Ignored_EdgeIsFullyConstrained;
+		return ESimplificationResult::Ignored_EdgeIsFullyConstrained;
 	}
 
 
 	// look up verts and tris for this edge
 	if (Mesh->IsEdge(edgeID) == false)
 	{
-		return EProcessResult::Failed_NotAnEdge;
+		return ESimplificationResult::Failed_NotAnEdge;
 	}
 	FIndex4i edgeInfo = Mesh->GetEdge(edgeID);
 	int a = edgeInfo.A, b = edgeInfo.B, t0 = edgeInfo.C, t1 = edgeInfo.D;
@@ -440,7 +614,7 @@ FMeshSimplification::EProcessResult FMeshSimplification::CollapseEdge(int edgeID
 	double edge_len_sqr = (vA - vB).SquaredLength();
 	if (edge_len_sqr > MinEdgeLength * MinEdgeLength)
 	{
-		return EProcessResult::Ignored_EdgeTooLong;
+		return ESimplificationResult::Ignored_EdgeTooLong;
 	}
 
 	ProfileBeginCollapse();
@@ -451,18 +625,18 @@ FMeshSimplification::EProcessResult FMeshSimplification::CollapseEdge(int edgeID
 	bool bCanCollapse = CanCollapseEdge(edgeID, a, b, c, d, t0, t1, collapse_to);
 	if (bCanCollapse == false)
 	{
-		return EProcessResult::Ignored_Constrained;
+		return ESimplificationResult::Ignored_Constrained;
 	}
 
 	// if we have a boundary, we want to collapse to boundary
-	if (bPreserveBoundaryShape && bHaveBoundary) 
+	if (bPreserveBoundaryShape && bHaveBoundary)
 	{
-		if (collapse_to != -1) 
+		if (collapse_to != -1)
 		{
 			if ((IsBoundaryVertex(b) && collapse_to != b) ||
 				(IsBoundaryVertex(a) && collapse_to != a))
 			{
-				return EProcessResult::Ignored_Constrained;
+				return ESimplificationResult::Ignored_Constrained;
 			}
 		}
 		if (IsBoundaryVertex(b))
@@ -478,18 +652,18 @@ FMeshSimplification::EProcessResult FMeshSimplification::CollapseEdge(int edgeID
 	// optimization: if edge cd exists, we cannot collapse or flip. look that up here?
 	//  funcs will do it internally...
 	//  (or maybe we can collapse if cd exists? edge-collapse doesn't check for it explicitly...)
-	EProcessResult retVal = EProcessResult::Failed_OpNotSuccessful;
+	ESimplificationResult retVal = ESimplificationResult::Failed_OpNotSuccessful;
 
 	int iKeep = b, iCollapse = a;
 
 	// if either vtx is fixed, collapse to that position
 	double collapse_t = 0;
-	if (collapse_to == b) 
+	if (collapse_to == b)
 	{
 		vNewPos = vB;
 		collapse_t = 0;
 	}
-	else if (collapse_to == a) 
+	else if (collapse_to == a)
 	{
 		iKeep = a; iCollapse = b;
 		vNewPos = vA;
@@ -506,10 +680,10 @@ FMeshSimplification::EProcessResult FMeshSimplification::CollapseEdge(int edgeID
 	// check if this collapse will create a normal flip. Also checks
 	// for invalid collapse nbrhood, since we are doing one-ring iter anyway.
 	// [TODO] could we skip this one-ring check in CollapseEdge? pass in hints?
-	if (CheckIfCollapseCreatesFlipOrInvalid(a, b, vNewPos, t0, t1) || CheckIfCollapseCreatesFlipOrInvalid(b, a, vNewPos, t0, t1)) 
+	if (CheckIfCollapseCreatesFlipOrInvalid(a, b, vNewPos, t0, t1) || CheckIfCollapseCreatesFlipOrInvalid(b, a, vNewPos, t0, t1))
 	{
 		ProfileEndCollapse();
-		return EProcessResult::Ignored_CreatesFlip;
+		return ESimplificationResult::Ignored_CreatesFlip;
 	}
 
 	// lots of cases where we cannot collapse, but we should just let
@@ -517,11 +691,11 @@ FMeshSimplification::EProcessResult FMeshSimplification::CollapseEdge(int edgeID
 	COUNT_COLLAPSES++;
 	FDynamicMesh3::FEdgeCollapseInfo collapseInfo;
 	EMeshResult result = Mesh->CollapseEdge(iKeep, iCollapse, collapse_t, collapseInfo);
-	if (result == EMeshResult::Ok) 
+	if (result == EMeshResult::Ok)
 	{
 		collapseToV = iKeep;
 		Mesh->SetVertex(iKeep, vNewPos);
-		if (Constraints != nullptr) 
+		if (Constraints != nullptr)
 		{
 			Constraints->ClearEdgeConstraint(edgeID);
 			Constraints->ClearEdgeConstraint(collapseInfo.RemovedEdges.A);
@@ -534,7 +708,7 @@ FMeshSimplification::EProcessResult FMeshSimplification::CollapseEdge(int edgeID
 		OnEdgeCollapse(edgeID, iKeep, iCollapse, collapseInfo);
 		DoDebugChecks();
 
-		retVal = EProcessResult::Ok_Collapsed;
+		retVal = ESimplificationResult::Ok_Collapsed;
 	}
 
 	ProfileEndCollapse();
@@ -550,9 +724,10 @@ FMeshSimplification::EProcessResult FMeshSimplification::CollapseEdge(int edgeID
 
 // Project vertices onto projection target. 
 // We can do projection in parallel if we have .net 
-void FMeshSimplification::FullProjectionPass()
+template <typename QuadricErrorType>
+void TMeshSimplification<QuadricErrorType>::FullProjectionPass()
 {
-	auto project = [&](int vID) 
+	auto project = [&](int vID)
 	{
 		if (IsVertexConstrained(vID))
 		{
@@ -578,8 +753,8 @@ void FMeshSimplification::FullProjectionPass()
 	//}
 }
 
-
-void FMeshSimplification::ApplyToProjectVertices(const TFunction<void(int)>& apply_f)
+template <typename QuadricErrorType>
+void TMeshSimplification<QuadricErrorType>::ApplyToProjectVertices(const TFunction<void(int)>& apply_f)
 {
 	for (int vid : Mesh->VertexIndicesItr())
 	{
@@ -587,8 +762,8 @@ void FMeshSimplification::ApplyToProjectVertices(const TFunction<void(int)>& app
 	}
 }
 
-
-void FMeshSimplification::ProjectVertex(int vID, IProjectionTarget* targetIn)
+template <typename QuadricErrorType>
+void TMeshSimplification<QuadricErrorType>::ProjectVertex(int vID, IProjectionTarget* targetIn)
 {
 	FVector3d curpos = Mesh->GetVertex(vID);
 	FVector3d projected = targetIn->Project(curpos, vID);
@@ -596,9 +771,10 @@ void FMeshSimplification::ProjectVertex(int vID, IProjectionTarget* targetIn)
 }
 
 // used by collapse-edge to get projected position for new vertex
-FVector3d FMeshSimplification::GetProjectedCollapsePosition(int vid, const FVector3d& vNewPos)
+template <typename QuadricErrorType>
+FVector3d TMeshSimplification<QuadricErrorType>::GetProjectedCollapsePosition(int vid, const FVector3d& vNewPos)
 {
-	if (Constraints != nullptr) 
+	if (Constraints != nullptr)
 	{
 		FVertexConstraint vc = Constraints->GetVertexConstraint(vid);
 		if (vc.Target != nullptr)
@@ -621,3 +797,46 @@ FVector3d FMeshSimplification::GetProjectedCollapsePosition(int vid, const FVect
 	return vNewPos;
 }
 
+
+// Custom behavior for FAttrBasedQuadric simplifier.
+template<>
+void TMeshSimplification<FAttrBasedQuadricErrord>::OnEdgeCollapse(int edgeID, int va, int vb, const FDynamicMesh3::FEdgeCollapseInfo& collapseInfo)
+{
+
+	// Update the normal
+	FAttrBasedQuadricErrord& Quadric = EdgeQuadrics[edgeID].q;
+	FVector3d collapse_pt = EdgeQuadrics[edgeID].collapse_pt;
+
+	FVector3d UpdatedNormald;
+	Quadric.ComputeAttributes(collapse_pt, UpdatedNormald);
+
+	FVector3f UpdatedNormal(UpdatedNormald.X, UpdatedNormald.Y, UpdatedNormald.Z);
+	UpdatedNormal.Normalize();
+
+	if (NormalOverlay != nullptr)
+	{
+		// Get all the elements associated with this vertex (could be more than one to account for split vertex data)
+		TArray<int> ElementIdArray;
+		NormalOverlay->GetVertexElements(va, ElementIdArray);
+
+		// update everyone with the same normal.
+		for (int ElementId : ElementIdArray)
+		{
+			NormalOverlay->SetElement(ElementId, UpdatedNormal);
+		}
+	}
+	else
+	{
+		Mesh->SetVertexNormal(va, UpdatedNormal);
+	}
+
+}
+
+
+
+// These are explicit instantiations of the templates that are exported from the shared lib.
+// Only these instantiations of the template can be used.
+// This is necessary because we have placed most of the templated functions in this .cpp file, instead of the header.
+template class DYNAMICMESH_API TMeshSimplification< FAttrBasedQuadricErrord >;
+template class DYNAMICMESH_API TMeshSimplification< FVolPresQuadricErrord >;
+template class DYNAMICMESH_API TMeshSimplification< FQuadricErrord >;
diff --git a/Engine/Plugins/Experimental/GeometryProcessing/Source/DynamicMesh/Public/MeshSimplification.h b/Engine/Plugins/Experimental/GeometryProcessing/Source/DynamicMesh/Public/MeshSimplification.h
--- a/Engine/Plugins/Experimental/GeometryProcessing/Source/DynamicMesh/Public/MeshSimplification.h
+++ b/Engine/Plugins/Experimental/GeometryProcessing/Source/DynamicMesh/Public/MeshSimplification.h
@@ -8,13 +8,31 @@
 #include "QuadricError.h"
 #include "Util/IndexPriorityQueue.h"
 
+#include "DynamicMeshAttributeSet.h"
+
+
+enum class ESimplificationResult
+{
+	Ok_Collapsed = 0,
+	Ignored_CannotCollapse = 1,
+	Ignored_EdgeIsFullyConstrained = 2,
+	Ignored_EdgeTooLong = 3,
+	Ignored_Constrained = 4,
+	Ignored_CreatesFlip = 5,
+	Failed_OpNotSuccessful = 6,
+	Failed_NotAnEdge = 7
+};
+
 /**
  * Implementation of Garland & Heckbert Quadric Error Metric (QEM) Triangle Mesh Simplification
  */
-class DYNAMICMESH_API FMeshSimplification : public FMeshRefinerBase
+template <typename QuadricErrorType>
+class TMeshSimplification : public FMeshRefinerBase
 {
 public:
 
+	typedef QuadricErrorType  FQuadricErrorType;
+
 	/**
 	 * If true, we try to find position for collapsed vertices that minimizes quadric error.
 	 * If false we just use midpoints, which is actually significantly slower, because it results
@@ -27,8 +45,13 @@ public:
 
 
 
-	FMeshSimplification(FDynamicMesh3* m) : FMeshRefinerBase(m)
+	TMeshSimplification(FDynamicMesh3* m) : FMeshRefinerBase(m)
 	{
+		NormalOverlay = nullptr;
+		if (m->Attributes())
+		{
+			NormalOverlay = m->Attributes()->PrimaryNormals();
+		}
 	}
 
 
@@ -64,7 +87,7 @@ public:
 
 protected:
 
-	FMeshSimplification()		// for subclasses that extend our behavior
+	TMeshSimplification()		// for subclasses that extend our behavior
 	{
 	}
 
@@ -112,20 +135,27 @@ protected:
 
 
 
-	TArray<FQuadricErrord> vertQuadrics;
+	TArray<FQuadricErrorType> vertQuadrics;
 	virtual void InitializeVertexQuadrics();
 
+	TArray<double> triAreas;
+	TArray<FQuadricErrorType> triQuadrics;
+	virtual void InitializeTriQuadrics();
 
+	FDynamicMeshNormalOverlay* NormalOverlay;
+
+	FQuadricErrorType ComputeFaceQuadric(const int tid, FVector3d& nface, FVector3d& c, double& Area) const;
+	
 	// internal class for priority queue
 	struct QEdge 
 	{
 		int eid;
-		FQuadricErrord q;
+		FQuadricErrorType q;
 		FVector3d collapse_pt;
 
 		QEdge() { eid = 0; }
 
-		QEdge(int edge_id, const FQuadricErrord& qin, const FVector3d& pt) 
+		QEdge(int edge_id, const FQuadricErrorType& qin, const FVector3d& pt) 
 		{
 			eid = edge_id;
 			q = qin;
@@ -149,7 +179,7 @@ protected:
 	virtual void InitializeQueue();
 
 	// return point that minimizes quadric error for edge [ea,eb]
-	FVector3d OptimalPoint(int eid, const FQuadricErrord& q, int ea, int eb);
+	FVector3d OptimalPoint(int eid, const FQuadricErrorType& q, int ea, int eb);
 	
 	FVector3d GetProjectedPoint(const FVector3d& pos)
 	{
@@ -162,7 +192,7 @@ protected:
 
 
 	// update queue weight for each edge in vertex one-ring
-	virtual void UpdateNeighbours(int vid);
+	virtual void UpdateNeighbours(int vid, FIndex2i removedTris, FIndex2i opposingVerts);
 
 
 	virtual void Reproject() 
@@ -192,18 +222,8 @@ protected:
 
 
 
-	enum class EProcessResult 
-	{
-		Ok_Collapsed = 0,
-		Ignored_CannotCollapse = 1,
-		Ignored_EdgeIsFullyConstrained = 2,
-		Ignored_EdgeTooLong = 3,
-		Ignored_Constrained = 4,
-		Ignored_CreatesFlip = 5,
-		Failed_OpNotSuccessful = 6,
-		Failed_NotAnEdge = 7
-	};
-	EProcessResult CollapseEdge(int edgeID, FVector3d vNewPos, int& collapseToV);
+
+	ESimplificationResult CollapseEdge(int edgeID, FVector3d vNewPos, int& collapseToV);
 
 
 
@@ -304,4 +324,10 @@ protected:
 	}
 
 
-};
\ No newline at end of file
+};
+
+
+// The simplifier
+typedef TMeshSimplification< FAttrBasedQuadricErrord >  FAttrMeshSimplification;
+typedef TMeshSimplification < FVolPresQuadricErrord >   FVolPresMeshSimplification;
+typedef TMeshSimplification< FQuadricErrord >           FQEMSimplification;
\ No newline at end of file
diff --git a/Engine/Plugins/Experimental/GeometryProcessing/Source/GeometricObjects/Public/QuadricError.h b/Engine/Plugins/Experimental/GeometryProcessing/Source/GeometricObjects/Public/QuadricError.h
--- a/Engine/Plugins/Experimental/GeometryProcessing/Source/GeometricObjects/Public/QuadricError.h
+++ b/Engine/Plugins/Experimental/GeometryProcessing/Source/GeometricObjects/Public/QuadricError.h
@@ -4,6 +4,7 @@
 
 #include "MathUtil.h"
 #include "VectorTypes.h"
+#include "MatrixTypes.h"
 
 
 /**
@@ -15,8 +16,8 @@ template<typename RealType>
 struct TQuadricError 
 {
 	RealType Axx, Axy, Axz, Ayy, Ayz, Azz;
-	RealType bx, by, bz;
-	RealType c;
+	RealType bx, by, bz; //d = -normal.dot(c); b = d*normal
+	RealType c; // d*d
 
 	inline static TQuadricError Zero() { return TQuadricError(); };
 
@@ -38,8 +39,8 @@ struct TQuadricError
 		Azz = Normal.Z * Normal.Z;
 		bx = by = bz = c = 0;
 		FVector3<RealType> v = MultiplyA(Point);
-		bx = -v.X; by = -v.Y; bz = -v.Z;
-		c = Point.Dot(v);
+		bx = -v.X; by = -v.Y; bz = -v.Z; // b = -Normal.Dot(Point) Normal
+		c = Point.Dot(v); // note this is the same as Normal.Dot(Point) ^2
 	}
 
 	/**
@@ -77,6 +78,35 @@ struct TQuadricError
 		c += w * b.c;
 	}
 
+	void Add(const TQuadricError& b)
+	{
+		Axx += b.Axx;
+		Axy += b.Axy;
+		Axz += b.Axz;
+		Ayy += b.Ayy;
+		Ayz += b.Ayz;
+		Azz += b.Azz;
+		bx += b.bx;
+		by += b.by;
+		bz += b.bz;
+		c +=  b.c;
+	}
+
+	TQuadricError& operator=(const TQuadricError& b)
+	{
+		Axx = b.Axx;
+		Axy = b.Axy;
+		Axz = b.Axz;
+		Ayy = b.Ayy;
+		Ayz = b.Ayz;
+		Azz = b.Azz;
+		bx = b.bx;
+		by = b.by;
+		bz = b.bz;
+		c = b.c;
+
+		return *this;
+	}
 
 	/**
 	 * Evaluates p*A*p + 2*dot(p,b) + c
@@ -103,8 +133,7 @@ struct TQuadricError
 	}
 
 
-
-	bool OptimalPoint(FVector3<RealType>& OutResult, RealType minThresh = 1000.0*TMathUtil<RealType>::Epsilon ) const
+	bool SolveAxEqualsb(FVector3<RealType>& OutResult, const RealType bvecx, const RealType bvecy, const RealType bvecz, const RealType minThresh = 1000.0*TMathUtil<RealType>::Epsilon) const 
 	{
 		RealType a11 = Azz * Ayy - Ayz * Ayz;
 		RealType a12 = Axz * Ayz - Azz * Axy;
@@ -116,26 +145,544 @@ struct TQuadricError
 
 		// [RMS] not sure what we should be using for this threshold...have seen
 		//  det less than 10^-9 on "normal" meshes.
-		if ( (RealType)fabs(det) > minThresh)
+		if (FMath::Abs(det) > minThresh)
 		{
 			det = 1.0 / det;
 			a11 *= det; a12 *= det; a13 *= det;
 			a22 *= det; a23 *= det; a33 *= det;
-			RealType x = a11 * bx + a12 * by + a13 * bz;
-			RealType y = a12 * bx + a22 * by + a23 * bz;
-			RealType z = a13 * bx + a23 * by + a33 * bz;
-			OutResult = FVector3<RealType>(-x, -y, -z);
+			RealType x = a11 * bvecx + a12 * bvecy + a13 * bvecz;
+			RealType y = a12 * bvecx + a22 * bvecy + a23 * bvecz;
+			RealType z = a13 * bvecx + a23 * bvecy + a33 * bvecz;
+			OutResult = FVector3<RealType>(x, y, z);
 			return true;
 		}
 		else 
 		{
 			return false;
 		}
+	}
+
+	static bool InvertSymmetricMatrix(const RealType SM[6], RealType InvSM[6], RealType minThresh = 1000.0*TMathUtil<RealType>::Epsilon)
+	{
+		bool Result = false;
+		// Axx = SM[0];  Axy = SM[1]; Axz = SM[2]
+		//               Ayy = SM[3]; Ayz = SM[4]
+		//                            Azz = SM[5]
 
+		InvSM[0] = SM[5] * SM[3] - SM[4] * SM[4]; //s11
+		InvSM[1] = SM[2] * SM[4] - SM[5] * SM[1]; //s12
+		InvSM[2] = SM[1] * SM[4] - SM[2] * SM[3]; //s13
+		InvSM[3] = SM[5] * SM[0] - SM[2] * SM[2]; //s22
+		InvSM[4] = SM[1] * SM[2] - SM[0] * SM[4]; //s23
+		InvSM[5] = SM[0] * SM[3] - SM[1] * SM[1]; //s33
+		RealType Det = (SM[0] * InvSM[0]) + (SM[1] * InvSM[1]) + (SM[2] * InvSM[2]);
+		if (FMath::Abs(Det) > minThresh)
+		{
+			RealType InvDet = RealType(1) / Det;
+			InvSM[0] *= InvDet; InvSM[1] *= InvDet; InvSM[2] *= InvDet; 
+			                    InvSM[3] *= InvDet; InvSM[4] *= InvDet;  
+							                        InvSM[5] *= InvDet;
+			Result = true;
+	}
+		return Result;
+	}
+
+	static FVector3<RealType> MultiplySymmetricMatrix(const RealType SM[6], const RealType vec[3])
+	{
+		
+		RealType a =  SM[0] * vec[0] + SM[1] * vec[1] + SM[2] * vec[2];
+		RealType b =  SM[1] * vec[0] + SM[3] * vec[1] + SM[4] * vec[2];
+		RealType c =  SM[2] * vec[0] + SM[4] * vec[1] + SM[5] * vec[2];
+
+		return FVector3<RealType>(a, b, c);
+	}
+	static FVector3<RealType> MultiplySymmetricMatrix(const RealType SM[6], const FVector3<RealType>& vec)
+	{
+		RealType vectmp[3] = { vec.X, vec.Y, vec.Z };
+		return MultiplySymmetricMatrix(SM, vectmp);
+	}
+
+
+	bool OptimalPoint(FVector3<RealType>& OutResult, RealType minThresh = 1000.0*TMathUtil<RealType>::Epsilon ) const
+	{
+		return SolveAxEqualsb(OutResult, -bx, -by, -bz, minThresh);
 	}
 
 };
 
 
 typedef TQuadricError<float> FQuadricErrorf;
-typedef TQuadricError<double> FQuadricErrord;
\ No newline at end of file
+typedef TQuadricError<double> FQuadricErrord;
+
+/**
+* Quadric Error type for use in memory-less simplification with volume preservation constraints.
+* 
+* See: http://hhoppe.com/newqem.pdf or https://www.cc.gatech.edu/~turk/my_papers/memless_vis98.pdf
+* for information about the volume preservation. 
+*/
+template<typename RealType>
+class TVolPresQuadricError : public TQuadricError<RealType>
+{
+public:
+
+	typedef TQuadricError<RealType>  BaseStruct;
+
+	struct FPlaneData
+	{
+		FPlaneData(const FVector3<RealType>& Normal, const FVector3<RealType>& Point)
+		{
+			N    = Normal;
+			Dist = -Normal.Dot(Point);
+		}
+
+		FPlaneData(RealType w, const FPlaneData& other)
+		{
+			N    = w * other.N;
+			Dist = w * other.Dist;
+		}
+
+		FPlaneData(const FPlaneData& other)
+		{
+			N    =  other.N;
+			Dist =  other.Dist;
+		}
+
+		FPlaneData() 
+		{
+			N    = FVector3<RealType>::Zero(); 
+			Dist = RealType(0);
+		}
+
+		void Add(const FPlaneData& other)
+		{
+			N    += other.N;
+			Dist += other.Dist;
+		}
+
+		void Add(RealType w, const FPlaneData& other)
+		{
+			N    += w * other.N;
+			Dist += w * other.Dist;
+		}
+
+		FPlaneData& operator=(const FPlaneData& other)
+		{
+			N    = other.N;
+			Dist = other.Dist;
+
+			return *this;
+		}
+
+		FVector3<RealType> N;
+		RealType Dist;
+	};
+
+	FPlaneData PlaneData;
+
+public:
+
+	inline static TVolPresQuadricError Zero() { return TVolPresQuadricError(); };
+	
+	TVolPresQuadricError() : BaseStruct()
+	{
+	}
+
+	TVolPresQuadricError(const FVector3<RealType>& Normal, const FVector3<RealType>& Point)
+		: BaseStruct(Normal, Point)
+		, PlaneData(Normal, Point)
+	{ }
+	
+
+	TVolPresQuadricError(const TVolPresQuadricError& a, const TVolPresQuadricError& b)
+		: BaseStruct(a, b)
+		, PlaneData(a.PlaneData)
+	{
+		PlaneData.Add(b.PlaneData);
+	}
+
+	TVolPresQuadricError(const TVolPresQuadricError& a, const TVolPresQuadricError& b, const FPlaneData& DuplicatePlaneData)
+		: BaseStruct(a, b)
+		, PlaneData(a.PlaneData)
+	{
+		PlaneData.Add(b.PlaneData);
+
+		// Subtract a single copy of the duplicate plane data.
+		PlaneData.Add(-1., DuplicatePlaneData);
+	}
+
+	void Add(RealType w, const TVolPresQuadricError& b)
+	{
+		BaseStruct::Add(w, b);
+
+		PlaneData.Add(w, b.PlaneData);
+	}
+
+	void Add(const TVolPresQuadricError& b)
+	{
+		BaseStruct::Add(b);
+		PlaneData.Add(b.PlaneData);
+	}
+
+	TVolPresQuadricError& operator=(const TVolPresQuadricError& other)
+	{
+		BaseStruct::operator=(other);
+		PlaneData = other.PlaneData;
+		return *this;
+	}
+
+	bool OptimalPoint(FVector3<RealType>& OutResult, RealType minThresh = 1000.0*TMathUtil<RealType>::Epsilon) const
+	{
+		// Compute the unconstrained optimal point
+		bool bValid = BaseStruct::OptimalPoint(OutResult, minThresh);
+
+		// if it failed ( the A matrix wasn't invertible) we early out
+		if (!bValid)
+		{
+			return bValid;
+		}
+
+		// Adjust with the volume constraint.  
+		// See: http://hhoppe.com/newqem.pdf or https://www.cc.gatech.edu/~turk/my_papers/memless_vis98.pdf
+
+		FPlaneData gvol(1. / 3., PlaneData);
+
+		// Compute the effect of the volumetric constraint.
+		RealType Tol = minThresh; //(1.e-7);  NB: using minThresh here to better compare with the attribute version..
+		//constexpr RealType Tol = (1.e-7);  //NB: using minThresh here to better compare with the attribute version..
+		
+		FVector3<RealType> Ainv_g;
+		if (BaseStruct::SolveAxEqualsb(Ainv_g, gvol.N.X, gvol.N.Y, gvol.N.Z, Tol))
+		{
+
+			RealType gt_Ainv_g = gvol.N.Dot(Ainv_g);
+
+			// move the point to the constrained optimal
+			RealType gt_unopt = gvol.N.Dot(OutResult);
+			RealType lambda = (gvol.Dist + gt_unopt) / gt_Ainv_g;
+
+			// Check if the constraint failed.
+
+			if (FMath::Abs(lambda) > 1.e5)
+			{
+				return bValid;
+			}
+
+			OutResult -= lambda * Ainv_g;
+		}
+
+		return bValid;
+
+
+	}
+};
+
+
+typedef TVolPresQuadricError<float>  FVolPresQuadricErrorf;
+typedef TVolPresQuadricError<double> FVolPresQuadricErrord;
+
+/**
+* Quadric Error type for use in volume memory-less simplification with volume preservation constraints.
+* using the normal as three additional attributes to contribute to the quadric error.
+* See: http://hhoppe.com/newqem.pdf 
+*/
+template<typename RealType>
+class TAttrBasedQuadricError : public TVolPresQuadricError<RealType>
+{
+public:
+	typedef TVolPresQuadricError<RealType>  BaseClass;
+
+
+	// Triangle Quadric constructor.  Take vertex locations, vertex normals, face normal, and center of face.
+	TAttrBasedQuadricError(const FVector3<RealType>& P0,    const FVector3<RealType>& P1, const FVector3<RealType>& P2,
+		                   const FVector3<RealType>& N0,    const FVector3<RealType>& N1, const FVector3<RealType>& N2,
+		                   const FVector3<RealType>& NFace, const FVector3<RealType>& CenterPoint, RealType AttrWeight)
+		: BaseClass(NFace, CenterPoint), a(0)
+	{
+		/**
+		* Given scalar attribute 'a' defined at the vertices e.g. a0, a1, a2
+		*
+		* solve 4x4 system:
+		*       (p0^t   1) (g[0])  = (a0)
+		*       (p1^t   1) (g[1])    (a1)
+		*       (p2^t   1) (g[2])    (a2)
+		*       (n^t    0) (  d )    (0 )
+		*
+		* for the interpolating gradient 'g'   -- note this is really p0.dot(g) + d = a0  and n.dot(g) = 0
+		*
+		* in practice we actually use a Schur complement approach that only requires inverting a 3x3 matrix.
+		*
+		* For this quadric, each component of the normal is treated as separate attribute.
+		*/
+		FMatrix3<RealType> Matrix(P0, P1, P2, true); // row-based matrix constructor.
+
+		bool bSolved = false;
+		double det = Matrix.Determinant();
+		if (FMath::Abs(det) > 1.e-5)
+		{
+			FMatrix3d InvMatrix = Matrix.Inverse();
+			FVector3d NFaceT_InvMatrix;
+			for (int i = 0; i < 3; ++i)
+			{
+				NFaceT_InvMatrix[i] = NFace[0] * InvMatrix.Row0[i] + NFace[1] * InvMatrix.Row1[i] + NFace[2] * InvMatrix.Row2[i];
+			}
+
+			double NFaceDotInvMOne = NFaceT_InvMatrix.Dot(FVector3d(1, 1, 1));
+
+			if (FMath::Abs(NFaceDotInvMOne) > 1.e-5)
+			{
+				bSolved = true;
+				for (int i = 0; i < 3; ++i)
+				{
+					FVector3d AttrVec(N0[i], N1[i], N2[i]);
+					AttrVec *= AttrWeight;
+
+					d[i] = NFaceT_InvMatrix.Dot(AttrVec) / NFaceDotInvMOne;
+					grad[i] = InvMatrix * (AttrVec - FVector3d(d[i], d[i], d[i]));
+
+					Axx += grad[i].X * grad[i].X; Axy += grad[i].X * grad[i].Y;  Axz += grad[i].X * grad[i].Z;
+					                              Ayy += grad[i].Y * grad[i].Y;  Ayz += grad[i].Y * grad[i].Z;
+					                                                             Azz += grad[i].Z * grad[i].Z;
+
+					
+					
+					
+
+					bx += d[i] * grad[i].X;
+					by += d[i] * grad[i].Y;
+					bz += d[i] * grad[i].Z;
+
+					c += d[i] * d[i];
+				}
+			}
+
+		}
+
+		if (!bSolved)
+		{
+			for (int i = 0; i < 3; ++i)
+			{
+				grad[i] = FVector3<RealType>::Zero();
+				d[i] = (N0[i] + N1[i] + N2[i]) / 3.; // just average the attribute.
+			}
+		}
+	}
+
+	
+	TAttrBasedQuadricError()
+		:BaseClass()
+	{
+		a = 0.;
+		const RealType zero[3] = { RealType(0), RealType(0), RealType(0) };
+		grad[0] = FVector3<RealType>(zero);
+		grad[1] = FVector3<RealType>(zero);
+		grad[2] = FVector3<RealType>(zero);
+
+		d[0] = RealType(0);
+		d[1] = RealType(0);
+		d[2] = RealType(0);
+	}
+
+	TAttrBasedQuadricError(const TAttrBasedQuadricError& Aother, const TAttrBasedQuadricError& Bother)
+		:BaseClass(Aother, Bother)
+	{
+		a = Aother.a + Bother.a;
+
+		grad[0] = Aother.grad[0] + Bother.grad[0];
+		grad[1] = Aother.grad[1] + Bother.grad[1];
+		grad[2] = Aother.grad[2] + Bother.grad[2];
+
+		d[0] = Aother.d[0] + Bother.d[0];
+		d[1] = Aother.d[1] + Bother.d[1];
+		d[2] = Aother.d[2] + Bother.d[2];
+	}
+
+	static TAttrBasedQuadricError Zero() { return TAttrBasedQuadricError(); }
+
+	void Add(RealType w, const TAttrBasedQuadricError& other)
+	{
+		BaseClass::Add(w, other);
+		
+		a += w; // accumulate weight
+
+		grad[0] += w * other.grad[0];
+		grad[1] += w * other.grad[1];
+		grad[2] += w * other.grad[2];
+
+		d[0] += w * other.d[0];
+		d[1] += w * other.d[1];
+		d[2] += w * other.d[2];
+	}
+
+	void Add(const TAttrBasedQuadricError& other)
+	{
+		BaseClass::Add(other);
+		a += other.a;
+
+		grad[0] += other.grad[0];
+		grad[1] += other.grad[1];
+		grad[2] += other.grad[2];
+
+		d[0] += other.d[0];
+		d[1] += other.d[1];
+		d[2] += other.d[2];
+	}
+
+	TAttrBasedQuadricError& operator=(const TAttrBasedQuadricError& other)
+	{
+		BaseClass::operator=(other);
+
+		a = other.a;
+
+		grad[0] = other.grad[0];
+		grad[1] = other.grad[1];
+		grad[2] = other.grad[2];
+
+		d[0] = other.d[0];
+		d[1] = other.d[1];
+		d[2] = other.d[2];
+
+		return *this;
+	}
+
+	// The optimal point not counting volume conservation.
+	bool OptimalPoint(FVector3<RealType>& OptPoint, RealType minThresh = 1000.0*TMathUtil<RealType>::Epsilon) const
+	{
+		// Generate symmetric matrix
+		RealType SM[6] = { 0., 0., 0., 0., 0., 0. };
+
+
+		for (int i = 0; i < 3; ++i)
+		{
+			SM[0] -= grad[i].X * grad[i].X;  SM[1] -= grad[i].X * grad[i].Y;  SM[2] -= grad[i].X * grad[i].Z;
+			                                 SM[3] -= grad[i].Y * grad[i].Y;  SM[4] -= grad[i].Y * grad[i].Z;
+			                                                                  SM[5] -= grad[i].Z * grad[i].Z;
+		}
+
+		RealType InvA = (FMath::Abs(a) > 1.e-7) ?  RealType(1) / a : 0.;
+
+		SM[0] *= InvA;  SM[1] *= InvA; SM[2] *= InvA;
+		                SM[3] *= InvA; SM[4] *= InvA;
+						               SM[5] *= InvA;
+
+		// A - (grad_attr0, grad_attr1.. grad_attrn) . (grad_attr0, grad_attr1.. grad_attrn)^T
+
+		SM[0] += Axx; SM[1] += Axy; SM[2] += Axz;
+		              SM[3] += Ayy; SM[4] += Ayz;
+		                            SM[5] += Azz;
+
+
+		RealType InvSM[6];
+		bool bValid = InvertSymmetricMatrix(SM, InvSM, minThresh);
+
+		RealType b[3] = { -bx, -by, -bz };
+		b[0] += InvA * (grad[0].X * d[0] + grad[1].X * d[1] + grad[2].X * d[2]);
+		b[1] += InvA * (grad[0].Y * d[0] + grad[1].Y * d[1] + grad[2].Y * d[2]);
+		b[2] += InvA * (grad[0].Z * d[0] + grad[1].Z * d[1] + grad[2].Z * d[2]);
+
+		// add volume constraint 
+		FPlaneData gvol(1. / 3., PlaneData);
+
+		if (bValid)
+		{
+			// optimal point prior to volume correction
+
+			OptPoint = MultiplySymmetricMatrix(InvSM, b);
+			//SolveAxEqualsb(OptPoint, b[0], b[1], b[2]);
+
+			FVector3<RealType> InvSMgvol = MultiplySymmetricMatrix(InvSM, gvol.N);
+			//SolveAxEqualsb(InvSMgvol, gvol.N.X, gvol.N.Y, gvol.N.Z);
+
+			RealType gvolDotInvSMgvol = gvol.N.Dot(InvSMgvol);
+
+			// move the point to the constrained optimal
+			RealType gvolDotuncopt = gvol.N.Dot(OptPoint);
+			RealType lambda = (gvol.Dist + gvolDotuncopt) / gvolDotInvSMgvol;
+
+			// Check if the constraint failed.
+
+			if (FMath::Abs(lambda) > 1.e5)
+			{
+				return bValid;
+			}
+
+			OptPoint -= lambda * InvSMgvol;
+		}
+
+		return bValid;
+		
+	}
+
+	RealType Evaluate(const FVector3<RealType>& point, const FVector3<RealType>& attr) const
+	{
+		// 6x6 symmetric matrix   (    A      -grad[0], -grad[1], -grad[2] )
+		//                        ( -grad[0]^T                             ) 
+		//                        ( -grad[1]^T             aI              )
+		//                        ( -grad[2]^T                             )          
+		//
+		//  aI is Identity matrix scaled by the accumulated area 'a'
+		//
+		// extended state vector  v = ( point )
+		//                            ( attr  )
+		//
+
+
+		RealType ptAp = point.Dot(MultiplyA(point));
+		RealType attrDotattr = attr.Dot(attr);
+		FVector3<RealType> Gradattr(grad[0].X * attr[0] + grad[1].X * attr[1] + grad[2].X * attr[2],
+			                        grad[0].Y * attr[0] + grad[1].Y * attr[1] + grad[2].Y * attr[2],
+			                        grad[0].Z * attr[0] + grad[1].Z * attr[1] + grad[2].Z * attr[2]);
+		RealType ptGradattr = point.Dot(Gradattr);
+
+		RealType vtSv = ptAp + a * attrDotattr - RealType(2) * ptGradattr;
+
+		// extended 'b' vector
+		//                     ( b )
+		//                     (-d )
+		// compute 2<v|b>
+		RealType vtb = point.X  * bx + point.Y * by + point.Z * bz - (d[0] * attr[0] + d[1] * attr[1] + d[2] * attr[2]);
+
+		RealType QuadricError =  vtSv + RealType(2) * vtb + c;
+
+		return QuadricError;
+	}
+
+	void ComputeAttributes(const FVector3<RealType>& point, FVector3<RealType>& attr) const
+	{
+		if (FMath::Abs(a) > 1.e-5)
+		{
+			RealType aInv = RealType(1) / a;
+
+			attr.X = grad[0].Dot(point) + d[0];
+			attr.Y = grad[1].Dot(point) + d[1];
+			attr.Z = grad[2].Dot(point) + d[2];
+
+			attr *= aInv;
+		} 
+		else
+		{
+			attr.X = RealType(0);
+			attr.Y = RealType(0);
+			attr.Z = RealType(0);
+		}
+	}
+
+	RealType Evaluate(const FVector3<RealType>& point) const
+	{
+		FVector3<RealType> attr;
+		ComputeAttributes(point, attr);
+
+		return Evaluate(point, attr);
+	}
+
+public:
+	
+	RealType           a;
+	// Additional planes for the attributes.
+	FVector3<RealType> grad[3];
+	RealType           d[3];
+	
+
+};
+
+typedef TAttrBasedQuadricError<float>  FAttrBasedQuadricErrorf;
+typedef TAttrBasedQuadricError<double> FAttrBasedQuadricErrord;
\ No newline at end of file
diff --git a/Engine/Plugins/Experimental/GeometryProcessing/Source/MeshConversion/Private/MeshDescriptionToDynamicMesh.cpp b/Engine/Plugins/Experimental/GeometryProcessing/Source/MeshConversion/Private/MeshDescriptionToDynamicMesh.cpp
--- a/Engine/Plugins/Experimental/GeometryProcessing/Source/MeshConversion/Private/MeshDescriptionToDynamicMesh.cpp
+++ b/Engine/Plugins/Experimental/GeometryProcessing/Source/MeshConversion/Private/MeshDescriptionToDynamicMesh.cpp
@@ -206,6 +206,57 @@ void FMeshDescriptionToDynamicMesh::Convert(const FMeshDescription* MeshIn, FDyn
 			int VertexID1 = MeshIn->GetVertexInstance(Triangle.VertexInstanceID1).VertexID.GetValue();
 			int VertexID2 = MeshIn->GetVertexInstance(Triangle.VertexInstanceID2).VertexID.GetValue();
 			int NewTriangleID = MeshOut.AppendTriangle(VertexID0, VertexID1, VertexID2, GroupID);
+			
+			// if append failed due to non-manifold, duplicate verts
+			if (NewTriangleID == FDynamicMesh3::NonManifoldID)
+			{
+				int e0 = MeshOut.FindEdge(VertexID0, VertexID1);
+				int e1 = MeshOut.FindEdge(VertexID1, VertexID2);
+				int e2 = MeshOut.FindEdge(VertexID2, VertexID0);
+
+				// determine which verts need to be duplicated
+				bool bDuplicate[3] = { false, false, false };
+				if (e0 != FDynamicMesh3::InvalidID && MeshOut.IsBoundaryEdge(e0) == false)
+				{
+					bDuplicate[0] = true;
+					bDuplicate[1] = true;
+				}
+				if (e1 != FDynamicMesh3::InvalidID && MeshOut.IsBoundaryEdge(e1) == false)
+				{
+					bDuplicate[1] = true;
+					bDuplicate[2] = true;
+				}
+				if (e2 != FDynamicMesh3::InvalidID && MeshOut.IsBoundaryEdge(e2) == false)
+				{
+					bDuplicate[2] = true;
+					bDuplicate[0] = true;
+				}
+				if (bDuplicate[0])
+				{
+					FVertexID VertexID = MeshIn->GetVertexInstance(Triangle.VertexInstanceID0).VertexID;
+					const FVector Position = VertexPositions.Get(VertexID);
+					int NewVertIdx = MeshOut.AppendVertex(Position);
+					VertexID0 = NewVertIdx;
+				}
+				if (bDuplicate[1])
+				{
+					FVertexID VertexID = MeshIn->GetVertexInstance(Triangle.VertexInstanceID1).VertexID;
+					const FVector Position = VertexPositions.Get(VertexID);
+					int NewVertIdx = MeshOut.AppendVertex(Position);
+					VertexID1 = NewVertIdx;
+				}
+				if (bDuplicate[2])
+				{
+					FVertexID VertexID = MeshIn->GetVertexInstance(Triangle.VertexInstanceID2).VertexID;
+					const FVector Position = VertexPositions.Get(VertexID);
+					int NewVertIdx = MeshOut.AppendVertex(Position);
+					VertexID2 = NewVertIdx;
+				}
+
+				NewTriangleID = MeshOut.AppendTriangle(VertexID0, VertexID1, VertexID2, GroupID);
+				checkSlow(NewTriangleID != FDynamicMesh3::NonManifoldID);
+			}
+
 			FIndex3i Tri(VertexID0, VertexID1, VertexID2);
 
 			if (bCalculateMaps)
