Copy robustness fixes from //UE4/Dev-Physics (prevent near-infinite loop case of nearly zero cellsize)

diff --git a/Engine/Plugins/Experimental/GeometryProcessing/Source/GeometricObjects/Public/Spatial/PointHashGrid2.h b/Engine/Plugins/Experimental/GeometryProcessing/Source/GeometricObjects/Public/Spatial/PointHashGrid2.h
--- a/Engine/Plugins/Experimental/GeometryProcessing/Source/GeometricObjects/Public/Spatial/PointHashGrid2.h
+++ b/Engine/Plugins/Experimental/GeometryProcessing/Source/GeometricObjects/Public/Spatial/PointHashGrid2.h
@@ -165,6 +165,11 @@ public:
 		TFunction<double(const PointDataType&)> DistanceSqFunc,
 		TFunction<bool(const PointDataType&)> IgnoreFunc = [](const PointDataType& data) { return false; }) const
 	{
+		if (!Hash.Num())
+		{
+			return TPair<PointDataType, double>(InvalidValue(), TNumericLimits<double>::Max());
+		}
+
 		FVector2i min_idx = Indexer.ToGrid(QueryPoint - Radius * FVector2d::One());
 		FVector2i max_idx = Indexer.ToGrid(QueryPoint + Radius * FVector2d::One());
 
diff --git a/Engine/Plugins/Experimental/GeometryProcessing/Source/GeometricObjects/Public/Util/GridIndexing2.h b/Engine/Plugins/Experimental/GeometryProcessing/Source/GeometricObjects/Public/Util/GridIndexing2.h
--- a/Engine/Plugins/Experimental/GeometryProcessing/Source/GeometricObjects/Public/Util/GridIndexing2.h
+++ b/Engine/Plugins/Experimental/GeometryProcessing/Source/GeometricObjects/Public/Util/GridIndexing2.h
@@ -16,7 +16,10 @@ struct TScaleGridIndexer2
 	/** Real-valued size of an integer grid cell */
 	RealType CellSize;
 
-	TScaleGridIndexer2(RealType CellSize) : CellSize(CellSize) {}
+	TScaleGridIndexer2(RealType CellSize) : CellSize(CellSize)
+	{
+		ensure(CellSize >= TMathUtil<RealType>::ZeroTolerance);
+	}
 
 	/** Convert real-valued point to integer grid coordinates */
 	inline FVector2i ToGrid(const FVector2<RealType>& P) const
@@ -49,6 +52,7 @@ struct TShiftGridIndexer2
 	{
 		Origin = origin;
 		CellSize = cellSize;
+		ensure(CellSize >= TMathUtil<RealType>::ZeroTolerance);
 	}
 
 	/** Convert real-valued point to integer grid coordinates */
diff --git a/Engine/Plugins/Experimental/GeometryProcessing/Source/GeometryAlgorithms/Public/Arrangement2d.h b/Engine/Plugins/Experimental/GeometryProcessing/Source/GeometryAlgorithms/Public/Arrangement2d.h
--- a/Engine/Plugins/Experimental/GeometryProcessing/Source/GeometryAlgorithms/Public/Arrangement2d.h
+++ b/Engine/Plugins/Experimental/GeometryProcessing/Source/GeometryAlgorithms/Public/Arrangement2d.h
@@ -39,7 +39,7 @@ struct FArrangement2d
 	double VertexSnapTol = 0.00001;
 
 	FArrangement2d(const FAxisAlignedBox2d& BoundsHint)
-		: PointHash(BoundsHint.MaxDim() / 64, -1)
+		: PointHash(FMath::Max(1e-4, BoundsHint.MaxDim() / 64), -1)
 	{
 	}
 
