Add DynamicMeshOverlay::ClearElements(), discards elements and sets triangles to InvalidID. Also default constructor.
#rnx

diff --git a/Engine/Plugins/Experimental/GeometryProcessing/Source/DynamicMesh/Private/DynamicMeshOverlay.cpp b/Engine/Plugins/Experimental/GeometryProcessing/Source/DynamicMesh/Private/DynamicMeshOverlay.cpp
--- a/Engine/Plugins/Experimental/GeometryProcessing/Source/DynamicMesh/Private/DynamicMeshOverlay.cpp
+++ b/Engine/Plugins/Experimental/GeometryProcessing/Source/DynamicMesh/Private/DynamicMeshOverlay.cpp
@@ -5,6 +5,18 @@
 
 
 
+
+template<typename RealType, int ElementSize>
+void TDynamicMeshOverlay<RealType, ElementSize>::ClearElements()
+{
+	Elements.Clear();
+	ElementsRefCounts = FRefCountVector();
+	ParentVertices.Clear();
+	InitializeTriangles(ParentMesh->MaxTriangleID());
+}
+
+
+
 template<typename RealType, int ElementSize>
 int TDynamicMeshOverlay<RealType, ElementSize>::AppendElement(RealType ConstantValue, int SourceVertex)
 {
diff --git a/Engine/Plugins/Experimental/GeometryProcessing/Source/DynamicMesh/Public/DynamicMeshOverlay.h b/Engine/Plugins/Experimental/GeometryProcessing/Source/DynamicMesh/Public/DynamicMeshOverlay.h
--- a/Engine/Plugins/Experimental/GeometryProcessing/Source/DynamicMesh/Public/DynamicMeshOverlay.h
+++ b/Engine/Plugins/Experimental/GeometryProcessing/Source/DynamicMesh/Public/DynamicMeshOverlay.h
@@ -49,6 +49,12 @@ protected:
 	friend class FDynamicMeshAttributeSet;
 
 public:
+	/** Create an empty overlay */
+	TDynamicMeshOverlay()
+	{
+		ParentMesh = nullptr;
+	}
+
 	/** Create an overlay for the given parent mesh */
 	TDynamicMeshOverlay(FDynamicMesh3* ParentMeshIn)
 	{
@@ -69,6 +75,9 @@ public:
 		ElementTriangles = Copy.ElementTriangles;
 	}
 
+	/** Discard current set of elements, but keep triangles */
+	void ClearElements();
+
 	/** @return the number of in-use Elements in the overlay */
 	int ElementCount() const { return (int)ElementsRefCounts.GetCount(); }
 	/** @return the maximum element index in the overlay. This may be larger than the count if Elements have been deleted. */
@@ -213,6 +222,11 @@ class TDynamicMeshVectorOverlay : public TDynamicMeshOverlay<RealType, ElementSi
 public:
 	using BaseType = TDynamicMeshOverlay<RealType, ElementSize>;
 
+	TDynamicMeshVectorOverlay()
+		: TDynamicMeshOverlay<RealType, ElementSize>()
+	{
+	}
+
 	TDynamicMeshVectorOverlay(FDynamicMesh3* parentMesh) 
 		: TDynamicMeshOverlay<RealType, ElementSize>(parentMesh)
 	{
diff --git a/Engine/Plugins/Experimental/GeometryProcessing/Source/DynamicMesh/Public/MeshNormals.h b/Engine/Plugins/Experimental/GeometryProcessing/Source/DynamicMesh/Public/MeshNormals.h
--- a/Engine/Plugins/Experimental/GeometryProcessing/Source/DynamicMesh/Public/MeshNormals.h
+++ b/Engine/Plugins/Experimental/GeometryProcessing/Source/DynamicMesh/Public/MeshNormals.h
@@ -79,14 +79,14 @@ public:
 	 * @warning assumes that the computed normals are vertex normals
 	 * @param bInvert if true, normals are flipped
 	 */
-	void CopyToVertexNormals(FDynamicMesh3* SetMesh, bool bInvert) const;
+	void CopyToVertexNormals(FDynamicMesh3* SetMesh, bool bInvert = false) const;
 
 	/**
 	 * Copy the current set of normals to the NormalOverlay attribute layer
 	 * @warning assumes that the computed normals are attribute normals
 	 * @param bInvert if true, normals are flipped
 	 */
-	void CopyToOverlay(FDynamicMeshNormalOverlay* NormalOverlay, bool bInvert) const;
+	void CopyToOverlay(FDynamicMeshNormalOverlay* NormalOverlay, bool bInvert = false) const;
 
 
 
@@ -94,7 +94,7 @@ public:
 	 * Compute per-vertex normals for the given Mesh
 	 * @param bInvert if true, normals are flipped
 	 */
-	static void QuickComputeVertexNormals(FDynamicMesh3& Mesh, bool bInvert);
+	static void QuickComputeVertexNormals(FDynamicMesh3& Mesh, bool bInvert = false);
 
 
 	/**
