fix linux compile errors and static analysis warnings
#rnx

diff --git a/Engine/Plugins/Experimental/GeometryProcessing/Source/GeometricObjects/Public/QuadricError.h b/Engine/Plugins/Experimental/GeometryProcessing/Source/GeometricObjects/Public/QuadricError.h
--- a/Engine/Plugins/Experimental/GeometryProcessing/Source/GeometricObjects/Public/QuadricError.h
+++ b/Engine/Plugins/Experimental/GeometryProcessing/Source/GeometricObjects/Public/QuadricError.h
@@ -390,7 +390,7 @@ class TAttrBasedQuadricError : public TVolPresQuadricError<RealType>
 {
 public:
 	typedef TVolPresQuadricError<RealType>  BaseClass;
-
+	typedef TQuadricError<RealType>  BaseStruct;
 
 	// Triangle Quadric constructor.  Take vertex locations, vertex normals, face normal, and center of face.
 	TAttrBasedQuadricError(const FVector3<RealType>& P0,    const FVector3<RealType>& P1, const FVector3<RealType>& P2,
@@ -439,19 +439,20 @@ public:
 					d[i] = NFaceT_InvMatrix.Dot(AttrVec) / NFaceDotInvMOne;
 					grad[i] = InvMatrix * (AttrVec - FVector3d(d[i], d[i], d[i]));
 
-					Axx += grad[i].X * grad[i].X; Axy += grad[i].X * grad[i].Y;  Axz += grad[i].X * grad[i].Z;
-					                              Ayy += grad[i].Y * grad[i].Y;  Ayz += grad[i].Y * grad[i].Z;
-					                                                             Azz += grad[i].Z * grad[i].Z;
-
-					
-					
+					BaseStruct::Axx += grad[i].X * grad[i].X; 
 					
+					BaseStruct::Axy += grad[i].X * grad[i].Y;  
+					BaseStruct::Axz += grad[i].X * grad[i].Z;
+
+					BaseStruct::Ayy += grad[i].Y * grad[i].Y;  
+					BaseStruct::Ayz += grad[i].Y * grad[i].Z;
+					BaseStruct::Azz += grad[i].Z * grad[i].Z;
 
-					bx += d[i] * grad[i].X;
-					by += d[i] * grad[i].Y;
-					bz += d[i] * grad[i].Z;
+					BaseStruct::bx += d[i] * grad[i].X;
+					BaseStruct::by += d[i] * grad[i].Y;
+					BaseStruct::bz += d[i] * grad[i].Z;
 
-					c += d[i] * d[i];
+					BaseStruct::c += d[i] * d[i];
 				}
 			}
 
@@ -566,30 +567,30 @@ public:
 
 		// A - (grad_attr0, grad_attr1.. grad_attrn) . (grad_attr0, grad_attr1.. grad_attrn)^T
 
-		SM[0] += Axx; SM[1] += Axy; SM[2] += Axz;
-		              SM[3] += Ayy; SM[4] += Ayz;
-		                            SM[5] += Azz;
+		SM[0] += BaseStruct::Axx;	SM[1] += BaseStruct::Axy;	SM[2] += BaseStruct::Axz;
+									SM[3] += BaseStruct::Ayy;	SM[4] += BaseStruct::Ayz;
+																SM[5] += BaseStruct::Azz;
 
 
 		RealType InvSM[6];
-		bool bValid = InvertSymmetricMatrix(SM, InvSM, minThresh);
+		bool bValid = BaseStruct::InvertSymmetricMatrix(SM, InvSM, minThresh);
 
-		RealType b[3] = { -bx, -by, -bz };
+		RealType b[3] = { -BaseStruct::bx, -BaseStruct::by, -BaseStruct::bz };
 		b[0] += InvA * (grad[0].X * d[0] + grad[1].X * d[1] + grad[2].X * d[2]);
 		b[1] += InvA * (grad[0].Y * d[0] + grad[1].Y * d[1] + grad[2].Y * d[2]);
 		b[2] += InvA * (grad[0].Z * d[0] + grad[1].Z * d[1] + grad[2].Z * d[2]);
 
 		// add volume constraint 
-		FPlaneData gvol(1. / 3., PlaneData);
+		typename BaseClass::FPlaneData gvol(1. / 3., BaseClass::PlaneData);
 
 		if (bValid)
 		{
 			// optimal point prior to volume correction
 
-			OptPoint = MultiplySymmetricMatrix(InvSM, b);
+			OptPoint = BaseStruct::MultiplySymmetricMatrix(InvSM, b);
 			//SolveAxEqualsb(OptPoint, b[0], b[1], b[2]);
 
-			FVector3<RealType> InvSMgvol = MultiplySymmetricMatrix(InvSM, gvol.N);
+			FVector3<RealType> InvSMgvol = BaseStruct::MultiplySymmetricMatrix(InvSM, gvol.N);
 			//SolveAxEqualsb(InvSMgvol, gvol.N.X, gvol.N.Y, gvol.N.Z);
 
 			RealType gvolDotInvSMgvol = gvol.N.Dot(InvSMgvol);
@@ -626,7 +627,7 @@ public:
 		//
 
 
-		RealType ptAp = point.Dot(MultiplyA(point));
+		RealType ptAp = point.Dot(BaseStruct::MultiplyA(point));
 		RealType attrDotattr = attr.Dot(attr);
 		FVector3<RealType> Gradattr(grad[0].X * attr[0] + grad[1].X * attr[1] + grad[2].X * attr[2],
 			                        grad[0].Y * attr[0] + grad[1].Y * attr[1] + grad[2].Y * attr[2],
@@ -639,9 +640,10 @@ public:
 		//                     ( b )
 		//                     (-d )
 		// compute 2<v|b>
-		RealType vtb = point.X  * bx + point.Y * by + point.Z * bz - (d[0] * attr[0] + d[1] * attr[1] + d[2] * attr[2]);
+		RealType vtb = point.X  * BaseStruct::bx + point.Y * BaseStruct::by + point.Z * BaseStruct::bz 
+			- (d[0] * attr[0] + d[1] * attr[1] + d[2] * attr[2]);
 
-		RealType QuadricError =  vtSv + RealType(2) * vtb + c;
+		RealType QuadricError =  vtSv + RealType(2) * vtb + BaseStruct::c;
 
 		return QuadricError;
 	}
diff --git a/Engine/Plugins/Experimental/GeometryProcessing/Source/GeometryAlgorithms/Public/Arrangement2d.h b/Engine/Plugins/Experimental/GeometryProcessing/Source/GeometryAlgorithms/Public/Arrangement2d.h
--- a/Engine/Plugins/Experimental/GeometryProcessing/Source/GeometryAlgorithms/Public/Arrangement2d.h
+++ b/Engine/Plugins/Experimental/GeometryProcessing/Source/GeometryAlgorithms/Public/Arrangement2d.h
@@ -39,7 +39,7 @@ struct FArrangement2d
 	double VertexSnapTol = 0.00001;
 
 	FArrangement2d(const FAxisAlignedBox2d& BoundsHint)
-		: PointHash(FMath::Max(1e-4, BoundsHint.MaxDim() / 64), -1)
+		: PointHash(BoundsHint.MaxDim() / 64, -1)
 	{
 	}
 
@@ -501,6 +501,7 @@ protected:
 			if (DSQ < Tol*Tol)
 			{
 				Hits.Add(FSegmentPoint{ T, VID });
+				num_hits++;
 			}
 		}
 
