Add Frame3::ConstrainedAlignPerpAxes(), computes best-alignment of frame tangent-plane axes with world axes
#rnx

diff --git a/Engine/Plugins/Experimental/GeometryProcessing/Source/GeometricObjects/Public/FrameTypes.h b/Engine/Plugins/Experimental/GeometryProcessing/Source/GeometricObjects/Public/FrameTypes.h
--- a/Engine/Plugins/Experimental/GeometryProcessing/Source/GeometricObjects/Public/FrameTypes.h
+++ b/Engine/Plugins/Experimental/GeometryProcessing/Source/GeometricObjects/Public/FrameTypes.h
@@ -318,6 +318,38 @@ struct FFrame3
 	}
 
 
+	/**
+	 * Compute rotation around NormalAxis that best-aligns one of the other two frame axes with either given UpAxis or FallbackAxis
+	 * (FallbackAxis is required if Dot(NormalAxis,UpAxis) > UpDotTolerance, ie if the Normal and Up directions are too closely aligned.
+	 * Basically this divides direction-sphere into three regions - polar caps with size defined by UpDotTolerance, and
+	 * a wide equator band covering the rest. When crossing between these regions the alignment has a discontinuity.
+	 * It is impossible to avoid this discontinuity because it is impossible to comb a sphere.
+	 * @param PerpAxis1 Index of first axis orthogonal to NormalAxis
+	 * @param PerpAxis2 Index of second axis orthogonal to NormalAxis
+	 * @param NormalAxis Axis of frame to rotate around
+	 * @param UpAxis Target axis in equator region, defaults to UnitZ
+	 * @param FallbackAxis Target axis in polar region, defaults to UnitX
+	 * @param UpDotTolerance defaults to cos(45), ie flip between regions happens roughly half way to poles
+	 */
+	void ConstrainedAlignPerpAxes(int PerpAxis1 = 0, int PerpAxis2 = 1, int NormalAxis = 2, 
+		const FVector3<RealType>& UpAxis = FVector3<RealType>::UnitZ(),
+		const FVector3<RealType>& FallbackAxis = FVector3<RealType>::UnitX(),
+		RealType UpDotTolerance = (RealType)0.707)
+	{
+		check(PerpAxis1 != PerpAxis2 && PerpAxis1 != NormalAxis && PerpAxis2 != NormalAxis);
+		const FVector3<RealType> NormalVec = GetAxis(NormalAxis);
+
+		// decide if we should use Fallback (polar-cap) axis or main (equator-region) axis
+		const FVector3<RealType>& TargetAxis =
+			(TMathUtil<RealType>::Abs(NormalVec.Dot(UpAxis)) > UpDotTolerance) ?
+			FallbackAxis : UpAxis;
+
+		// figure out which PerpAxis is closer to target, and align that one to +/- TargetAxis (whichever is smaller rotation)
+		FVector2<RealType> Dots(GetAxis(PerpAxis1).Dot(TargetAxis), GetAxis(PerpAxis2).Dot(TargetAxis));
+		int UseAxis = (TMathUtil<RealType>::Abs(Dots.X) > TMathUtil<RealType>::Abs(Dots.Y)) ? 0 : 1;
+		RealType UseSign = Dots[UseAxis] < 0 ? -1 : 1;
+		ConstrainedAlignAxis(UseAxis, UseSign*TargetAxis, NormalVec);
+	}
 
 
 	/**
