add mesh distance query to AABB tree

diff --git a/Engine/Plugins/Experimental/GeometryProcessing/Source/GeometricObjects/Public/BoxTypes.h b/Engine/Plugins/Experimental/GeometryProcessing/Source/GeometricObjects/Public/BoxTypes.h
--- a/Engine/Plugins/Experimental/GeometryProcessing/Source/GeometricObjects/Public/BoxTypes.h
+++ b/Engine/Plugins/Experimental/GeometryProcessing/Source/GeometricObjects/Public/BoxTypes.h
@@ -192,6 +192,23 @@ struct FAxisAlignedBox3
 		this->Min = Min;
 		this->Max = Max;
 	}
+	FAxisAlignedBox3(const FAxisAlignedBox3& Box, const TFunction<FVector3<RealType>(const FVector3<RealType>&)> TransformF)
+	{
+		if (TransformF == nullptr)
+		{
+			Min = Box.Min;
+			Max = Box.Max;
+			return;
+		}
+
+		FVector3<RealType> C0 = TransformF(Box.GetCorner(0));
+		Min = C0;
+		Max = C0;
+		for (int i = 1; i < 8; ++i)
+		{
+			Contain(TransformF(Box.GetCorner(i)));
+		}
+	}
 
 	operator FBox() const
 	{
@@ -205,6 +222,19 @@ struct FAxisAlignedBox3
 		Max = Box.Max;
 	}
 
+	/**
+	* @param Index corner index in range 0-7
+	* @return Corner point on the box identified by the given index. See diagram in OrientedBoxTypes.h for index/corner mapping.
+	*/
+	FVector3<RealType> GetCorner(int Index) const
+	{
+		check(Index >= 0 && Index <= 7);
+		double X = (((Index & 1) != 0) ^ ((Index & 2) != 0)) ? (Max.X) : (Min.X);
+		double Y = ((Index / 2) % 2 == 0) ? (Min.Y) : (Max.Y);
+		double Z = (Index < 4) ? (Min.Z) : (Max.Z);
+		return FVector3<RealType>(X, Y, Z);
+	}
+
 	static FAxisAlignedBox3<RealType> Empty()
 	{
 		return FAxisAlignedBox3(
@@ -293,6 +323,30 @@ struct FAxisAlignedBox3
 		return dx * dx + dy * dy + dz * dz;
 	}
 
+	RealType DistanceSquared(const FAxisAlignedBox3<RealType>& Box)
+	{
+		// compute lensqr( max(0, abs(center1-center2) - (extent1+extent2)) )
+		RealType delta_x = TMathUtil<RealType>::Abs((Box.Min.X + Box.Max.X) - (Min.X + Max.X))
+			- ((Max.X - Min.X) + (Box.Max.X - Box.Min.X));
+		if (delta_x < 0)
+		{
+			delta_x = 0;
+		}
+		RealType delta_y = TMathUtil<RealType>::Abs((Box.Min.Y + Box.Max.Y) - (Min.Y + Max.Y))
+			- ((Max.Y - Min.Y) + (Box.Max.Y - Box.Min.Y));
+		if (delta_y < 0)
+		{
+			delta_y = 0;
+		}
+		RealType delta_z = TMathUtil<RealType>::Abs((Box.Min.Z + Box.Max.Z) - (Min.Z + Max.Z))
+			- ((Max.Z - Min.Z) + (Box.Max.Z - Box.Min.Z));
+		if (delta_z < 0)
+		{
+			delta_z = 0;
+		}
+		return (RealType)0.25 * (delta_x * delta_x + delta_y * delta_y + delta_z * delta_z);
+	}
+
 	RealType Width() const
 	{
 		return TMathUtil<RealType>::Max(Max.X - Min.X, (RealType)0);
diff --git a/Engine/Plugins/Experimental/GeometryProcessing/Source/GeometricObjects/Public/Distance/DistLine3Segment3.h b/Engine/Plugins/Experimental/GeometryProcessing/Source/GeometricObjects/Public/Distance/DistLine3Segment3.h
new file mode 100644
--- /dev/null
+++ b/Engine/Plugins/Experimental/GeometryProcessing/Source/GeometricObjects/Public/Distance/DistLine3Segment3.h
@@ -0,0 +1,122 @@
+// Copyright 1998-2019 Epic Games, Inc. All Rights Reserved.
+
+// Port of geometry3Sharp DistLine3Segment3
+// which was ported from WildMagic 5 
+
+#pragma once
+
+#include "VectorTypes.h"
+#include "SegmentTypes.h"
+#include "LineTypes.h"
+
+
+/**
+* Compute unsigned distance between 3D line and 3D segment
+*/
+template <typename Real>
+class FDistLine3Segment3
+{
+public:
+	// Input
+	FLine3<Real> Line;
+	FSegment3<Real> Segment;
+
+	// Output
+	Real DistanceSquared = -1.0;
+	Real LineParameter, SegmentParameter;
+	FVector3<Real> LineClosest, SegmentClosest;
+
+
+	FDistLine3Segment3(const FLine3<Real>& LineIn, const FSegment3<Real>& SegmentIn) : Line(LineIn), Segment(SegmentIn)
+	{
+	}
+
+	Real Get()
+	{
+		return TMathUtil<Real>::Sqrt(ComputeResult());
+	}
+	Real GetSquared()
+	{
+		return ComputeResult();
+	}
+
+	Real ComputeResult()
+	{
+		if (DistanceSquared >= 0)
+		{
+			return DistanceSquared;
+		}
+
+		FVector3<Real> diff = Line.Origin - Segment.Center;
+		Real a01 = -Line.Direction.Dot(Segment.Direction);
+		Real b0 = diff.Dot(Line.Direction);
+		Real c = diff.SquaredLength();
+		Real det = TMathUtil<Real>::Abs(1 - a01 * a01);
+		Real b1, s0, s1, sqrDist, extDet;
+
+		if (det >= TMathUtil<Real>::ZeroTolerance)
+		{
+			// The Line and Segment are not parallel.
+			b1 = -diff.Dot(Segment.Direction);
+			s1 = a01 * b0 - b1;
+			extDet = Segment.Extent * det;
+
+			if (s1 >= -extDet)
+			{
+				if (s1 <= extDet)
+				{
+					// Two interior points are closest, one on the Line and one
+					// on the Segment.
+					Real invDet = (1) / det;
+					s0 = (a01 * b1 - b0) * invDet;
+					s1 *= invDet;
+					sqrDist = s0 * (s0 + a01 * s1 + (2) * b0) +
+						s1 * (a01 * s0 + s1 + (2) * b1) + c;
+				}
+				else
+				{
+					// The endpoint e1 of the Segment and an interior point of
+					// the Line are closest.
+					s1 = Segment.Extent;
+					s0 = -(a01 * s1 + b0);
+					sqrDist = -s0 * s0 + s1 * (s1 + (2) * b1) + c;
+				}
+			}
+			else
+			{
+				// The end point e0 of the Segment and an interior point of the
+				// Line are closest.
+				s1 = -Segment.Extent;
+				s0 = -(a01 * s1 + b0);
+				sqrDist = -s0 * s0 + s1 * (s1 + (2) * b1) + c;
+			}
+		}
+		else
+		{
+			// The Line and Segment are parallel.  Choose the closest pair so that
+			// one point is at Segment center.
+			s1 = 0;
+			s0 = -b0;
+			sqrDist = b0 * s0 + c;
+		}
+
+		LineClosest = Line.Origin + s0 * Line.Direction;
+		SegmentClosest = Segment.Center + s1 * Segment.Direction;
+		LineParameter = s0;
+		SegmentParameter = s1;
+
+		// Account for numerical round-off errors.
+		if (sqrDist < 0)
+		{
+			sqrDist = 0;
+		}
+
+		DistanceSquared = sqrDist;
+
+		return DistanceSquared;
+	}
+};
+
+typedef FDistLine3Segment3<float> FDistLine3Segment3f;
+typedef FDistLine3Segment3<double> FDistLine3Segment3d;
+
diff --git a/Engine/Plugins/Experimental/GeometryProcessing/Source/GeometricObjects/Public/Distance/DistLine3Triangle3.h b/Engine/Plugins/Experimental/GeometryProcessing/Source/GeometricObjects/Public/Distance/DistLine3Triangle3.h
new file mode 100644
--- /dev/null
+++ b/Engine/Plugins/Experimental/GeometryProcessing/Source/GeometricObjects/Public/Distance/DistLine3Triangle3.h
@@ -0,0 +1,127 @@
+// Copyright 1998-2019 Epic Games, Inc. All Rights Reserved.
+
+// Port of geometry3Sharp DistLine3Triangle3
+// which was ported from WildMagic 5 
+
+#pragma once
+
+#include "VectorTypes.h"
+#include "TriangleTypes.h"
+#include "SegmentTypes.h"
+#include "DistLine3Segment3.h"
+
+
+/**
+* Compute unsigned distance between 3D line and 3D triangle
+*/
+template <typename Real>
+class FDistLine3Triangle3
+{
+public:
+	// Input
+	FLine3<Real> Line;
+	FTriangle3<Real> Triangle;
+
+	// Output
+	Real DistanceSquared = -1.0;
+	Real LineParam;
+	FVector3<Real> LineClosest, TriangleClosest, TriangleBaryCoords;
+
+
+	FDistLine3Triangle3(const FLine3<Real>& LineIn, const FTriangle3<Real>& TriangleIn) : Line(LineIn), Triangle(TriangleIn)
+	{
+	}
+
+	Real Get()
+	{
+		return TMathUtil<Real>::Sqrt(ComputeResult());
+	}
+	Real GetSquared()
+	{
+		return ComputeResult();
+	}
+
+	Real ComputeResult()
+	{
+		if (DistanceSquared >= 0)
+		{
+			return DistanceSquared;
+		}
+
+		// Test if Line intersects Triangle.  If so, the squared distance is zero.
+		FVector3<Real> edge0 = Triangle.V[1] - Triangle.V[0];
+		FVector3<Real> edge1 = Triangle.V[2] - Triangle.V[0];
+		FVector3<Real> normal = edge0.Cross(edge1);
+		normal.Normalize();
+		Real NdD = normal.Dot(Line.Direction);
+		if (TMathUtil<Real>::Abs(NdD) > TMathUtil<Real>::ZeroTolerance)
+		{
+			// The line and triangle are not parallel, so the line intersects
+			// the plane of the triangle.
+			FVector3<Real> diff = Line.Origin - Triangle.V[0];
+			FVector3<Real> U, V;
+			VectorUtil::MakePerpVectors(Line.Direction, U, V);
+			Real UdE0 = U.Dot(edge0);
+			Real UdE1 = U.Dot(edge1);
+			Real UdDiff = U.Dot(diff);
+			Real VdE0 = V.Dot(edge0);
+			Real VdE1 = V.Dot(edge1);
+			Real VdDiff = V.Dot(diff);
+			Real invDet = (1) / (UdE0 * VdE1 - UdE1 * VdE0);
+
+			// Barycentric coordinates for the point of intersection.
+			Real b1 = (VdE1 * UdDiff - UdE1 * VdDiff) * invDet;
+			Real b2 = (UdE0 * VdDiff - VdE0 * UdDiff) * invDet;
+			Real b0 = 1 - b1 - b2;
+
+			if (b0 >= 0 && b1 >= 0 && b2 >= 0)
+			{
+				// Line parameter for the point of intersection.
+				Real DdE0 = Line.Direction.Dot(edge0);
+				Real DdE1 = Line.Direction.Dot(edge1);
+				Real DdDiff = Line.Direction.Dot(diff);
+				LineParam = b1 * DdE0 + b2 * DdE1 - DdDiff;
+
+				// Barycentric coordinates for the point of intersection.
+				TriangleBaryCoords = FVector3<Real>(b0, b1, b2);
+
+				// The intersection point is inside or on the Triangle.
+				LineClosest = Line.Origin + LineParam * Line.Direction;
+				TriangleClosest = Triangle.V[0] + b1 * edge0 + b2 * edge1;
+				DistanceSquared = 0;
+				return 0;
+			}
+		}
+
+		// Either (1) the Line is not parallel to the Triangle and the point of
+		// intersection of the Line and the plane of the Triangle is outside the
+		// Triangle or (2) the Line and Triangle are parallel.  Regardless, the
+		// closest point on the Triangle is on an edge of the Triangle.  Compare
+		// the Line to all three edges of the Triangle.
+		Real sqrDist = TMathUtil<Real>::MaxReal;
+		for (int i0 = 2, i1 = 0; i1 < 3; i0 = i1++)
+		{
+			FSegment3<Real> segment(Triangle.V[i0], Triangle.V[i1]);
+			FDistLine3Segment3<Real> queryLS(Line, segment);
+			Real sqrDistTmp = queryLS.GetSquared();
+			if (sqrDistTmp < sqrDist)
+			{
+				LineClosest = queryLS.LineClosest;
+				TriangleClosest = queryLS.SegmentClosest;
+				sqrDist = sqrDistTmp;
+				LineParam = queryLS.LineParameter;
+				Real ratio = queryLS.SegmentParameter / segment.Extent;
+				TriangleBaryCoords[i0] = (0.5) * (1 - ratio);
+				TriangleBaryCoords[i1] = 1 - TriangleBaryCoords[i0];
+				TriangleBaryCoords[3 - i0 - i1] = 0;
+			}
+		}
+
+		DistanceSquared = sqrDist;
+		return DistanceSquared;
+	}
+};
+
+typedef FDistLine3Triangle3<float> FDistLine3Triangle3f;
+typedef FDistLine3Triangle3<double> FDistLine3Triangle3d;
+
diff --git a/Engine/Plugins/Experimental/GeometryProcessing/Source/GeometricObjects/Public/Distance/DistSegment3Triangle3.h b/Engine/Plugins/Experimental/GeometryProcessing/Source/GeometricObjects/Public/Distance/DistSegment3Triangle3.h
new file mode 100644
--- /dev/null
+++ b/Engine/Plugins/Experimental/GeometryProcessing/Source/GeometricObjects/Public/Distance/DistSegment3Triangle3.h
@@ -0,0 +1,86 @@
+// Copyright 1998-2019 Epic Games, Inc. All Rights Reserved.
+
+// Port of geometry3Sharp DistSegment3Triangle3
+// which was ported from WildMagic 5 
+
+#pragma once
+
+#include "VectorTypes.h"
+#include "TriangleTypes.h"
+#include "SegmentTypes.h"
+#include "DistLine3Triangle3.h"
+#include "DistPoint3Triangle3.h"
+
+/**
+* Compute unsigned distance between 3D segment and 3D triangle
+*/
+template <typename Real>
+class FDistSegment3Triangle3
+{
+public:
+	// Input
+	FSegment3<Real> Segment;
+	FTriangle3<Real> Triangle;
+
+	// Output
+	Real DistanceSquared = -1.0;
+	Real SegmentParameter;
+	FVector3<Real> TriangleClosest, TriangleBaryCoords, SegmentClosest;
+	
+	FDistSegment3Triangle3(const FSegment3<Real>& SegmentIn, const FTriangle3<Real>& TriangleIn) : Segment(SegmentIn), Triangle(TriangleIn)
+	{
+	}
+
+	Real Get() 
+	{
+		return TMathUtil<Real>::Sqrt(ComputeResult());
+	}
+	Real GetSquared()
+	{
+		return ComputeResult();
+	}
+
+	Real ComputeResult()
+	{
+		if (DistanceSquared >= 0)
+		{
+			return DistanceSquared;
+		}
+
+		FLine3<Real> line(Segment.Center, Segment.Direction);
+		FDistLine3Triangle3<Real> queryLT(line, Triangle);
+		double sqrDist = queryLT.GetSquared();
+		SegmentParameter = queryLT.LineParam;
+
+		if (SegmentParameter >= -Segment.Extent) {
+			if (SegmentParameter <= Segment.Extent) {
+				SegmentClosest = queryLT.LineClosest;
+				TriangleClosest = queryLT.TriangleClosest;
+				TriangleBaryCoords = queryLT.TriangleBaryCoords;
+			}
+			else {
+				SegmentClosest = Segment.EndPoint();
+				FDistPoint3Triangle3<Real> queryPT(SegmentClosest, Triangle);
+				sqrDist = queryPT.GetSquared();
+				TriangleClosest = queryPT.ClosestTrianglePoint;
+				SegmentParameter = Segment.Extent;
+				TriangleBaryCoords = queryPT.TriangleBaryCoords;
+			}
+		}
+		else {
+			SegmentClosest = Segment.StartPoint();
+			FDistPoint3Triangle3<Real> queryPT(SegmentClosest, Triangle);
+			sqrDist = queryPT.GetSquared();
+			TriangleClosest = queryPT.ClosestTrianglePoint;
+			SegmentParameter = -Segment.Extent;
+			TriangleBaryCoords = queryPT.TriangleBaryCoords;
+		}
+
+		DistanceSquared = sqrDist;
+		return DistanceSquared;
+	}
+};
+
+typedef FDistSegment3Triangle3<float> FDistSegment3Triangle3f;
+typedef FDistSegment3Triangle3<double> FDistSegment3Triangle3d;
+
diff --git a/Engine/Plugins/Experimental/GeometryProcessing/Source/GeometricObjects/Public/Distance/DistTriangle3Triangle3.h b/Engine/Plugins/Experimental/GeometryProcessing/Source/GeometricObjects/Public/Distance/DistTriangle3Triangle3.h
new file mode 100644
--- /dev/null
+++ b/Engine/Plugins/Experimental/GeometryProcessing/Source/GeometricObjects/Public/Distance/DistTriangle3Triangle3.h
@@ -0,0 +1,122 @@
+// Copyright 1998-2019 Epic Games, Inc. All Rights Reserved.
+
+// Port of geometry3Sharp DistTriangle3Triangle3
+// which was ported from WildMagic 5 
+
+
+#pragma once
+
+#include "VectorTypes.h"
+#include "TriangleTypes.h"
+#include "SegmentTypes.h"
+#include "DistSegment3Triangle3.h"
+#include "DistPoint3Triangle3.h"
+
+/**
+* Compute unsigned distance between 3D segment and 3D triangle
+*/
+template <typename Real>
+class FDistTriangle3Triangle3
+{
+public:
+	// Input
+	FTriangle3<Real> Triangle[2];
+
+	// Output
+	Real DistanceSquared = -1.0;
+	FVector3<Real> TriangleClosest[2];
+	FVector3<Real> TriangleBaryCoords[2];
+
+
+	FDistTriangle3Triangle3()
+	{
+	}
+	FDistTriangle3Triangle3(const FTriangle3<Real>& TriangleA, const FTriangle3<Real>& TriangleB)
+	{
+		Triangle[0] = TriangleA;
+		Triangle[1] = TriangleB;
+	}
+
+	Real Get()
+	{
+		return TMathUtil<Real>::Sqrt(ComputeResult());
+	}
+	Real GetSquared()
+	{
+		return ComputeResult();
+	}
+
+	Real ComputeResult()
+	{
+		if (DistanceSquared >= 0)
+		{
+			return DistanceSquared;
+		}
+
+
+		// Compare edges of Triangle[0] to the interior of Triangle[1].
+		Real sqrDist = TMathUtil<Real>::MaxReal, sqrDistTmp;
+		Real ratio;
+		int i0, i1;
+		for (i0 = 2, i1 = 0; i1 < 3; i0 = i1++)
+		{
+			FSegment3<Real> edge(Triangle[0].V[i0], Triangle[0].V[i1]);
+
+			FDistSegment3Triangle3<Real> queryST(edge, Triangle[1]);
+			sqrDistTmp = queryST.GetSquared();
+			if (sqrDistTmp < sqrDist)
+			{
+				TriangleClosest[0] = queryST.SegmentClosest;
+				TriangleClosest[1] = queryST.TriangleClosest;
+				sqrDist = sqrDistTmp;
+
+				ratio = queryST.SegmentParameter / edge.Extent;
+				TriangleBaryCoords[0][i0] = (0.5) * (1 - ratio);
+				TriangleBaryCoords[0][i1] = 1 - TriangleBaryCoords[0][i0];
+				TriangleBaryCoords[0][3 - i0 - i1] = 0;
+				TriangleBaryCoords[1] = queryST.TriangleBaryCoords;
+
+				if (sqrDist <= TMathUtil<Real>::ZeroTolerance)
+				{
+					DistanceSquared = 0;
+					return 0;
+				}
+			}
+		}
+
+		// Compare edges of Triangle[1] to the interior of Triangle[0].
+		for (i0 = 2, i1 = 0; i1 < 3; i0 = i1++)
+		{
+			FSegment3<Real> edge(Triangle[1].V[i0], Triangle[1].V[i1]);
+
+			FDistSegment3Triangle3<Real> queryST(edge, Triangle[0]);
+			sqrDistTmp = queryST.GetSquared();
+			if (sqrDistTmp < sqrDist)
+			{
+				TriangleClosest[0] = queryST.SegmentClosest;
+				TriangleClosest[1] = queryST.TriangleClosest;
+				sqrDist = sqrDistTmp;
+
+				ratio = queryST.SegmentParameter / edge.Extent;
+				TriangleBaryCoords[1][i0] = (0.5) * (1 - ratio);
+				TriangleBaryCoords[1][i1] = 1 - TriangleBaryCoords[1][i0];
+				TriangleBaryCoords[1][3 - i0 - i1] = 0;
+				TriangleBaryCoords[0] = queryST.TriangleBaryCoords;
+
+				if (sqrDist <= TMathUtil<Real>::ZeroTolerance)
+				{
+					DistanceSquared = 0;
+					return 0;
+				}
+			}
+		}
+
+
+		DistanceSquared = sqrDist;
+		return DistanceSquared;
+	}
+};
+
+typedef FDistTriangle3Triangle3<float> FDistTriangle3Triangle3f;
+typedef FDistTriangle3Triangle3<double> FDistTriangle3Triangle3d;
+
diff --git a/Engine/Plugins/Experimental/GeometryProcessing/Source/GeometricObjects/Public/Spatial/MeshAABBTree3.h b/Engine/Plugins/Experimental/GeometryProcessing/Source/GeometricObjects/Public/Spatial/MeshAABBTree3.h
--- a/Engine/Plugins/Experimental/GeometryProcessing/Source/GeometricObjects/Public/Spatial/MeshAABBTree3.h
+++ b/Engine/Plugins/Experimental/GeometryProcessing/Source/GeometricObjects/Public/Spatial/MeshAABBTree3.h
@@ -8,6 +8,7 @@
 #include "Intersection/IntrRay3AxisAlignedBox3.h"
 #include "MeshQueries.h"
 #include "Spatial/SpatialInterfaces.h"
+#include "Distance/DistTriangle3Triangle3.h"
 
 template <class TriangleMeshType>
 class TFastWindingTree;
@@ -237,6 +238,30 @@ public:
 		}
 	}
 
+	/**
+	 * Find nearest pair of triangles on this tree with otherTree, within max_dist.
+	 * TransformF transforms vertices of otherTree into our coordinates. can be null.
+	 * returns triangle-id pair (my_tri,other_tri), or Index2i.Max if not found within max_dist
+	 * Use MeshQueries.TrianglesDistance() to get more information
+	 */
+	virtual FIndex2i FindNearestTriangles(TMeshAABBTree3& OtherTree, TFunction<FVector3d(const FVector3d&)> TransformF, double& Distance, double MaxDist = FMathd::MaxReal)
+	{
+		check(MeshTimestamp == Mesh->GetShapeTimestamp());
+
+		double NearestSqr = FMathd::MaxReal;
+		if (MaxDist < FMathd::MaxReal)
+		{
+			NearestSqr = MaxDist * MaxDist;
+		}
+		FIndex2i NearestPair = FIndex2i::Max();
+
+		find_nearest_triangles(RootIndex, OtherTree, TransformF, OtherTree.RootIndex, 0, NearestSqr, NearestPair);
+		Distance = (NearestSqr < FMathd::MaxReal) ? FMathd::Sqrt(NearestSqr) : FMathd::MaxReal;
+		return NearestPair;
+	}
+
+
+
 	virtual bool SupportsPointContainment() override
 	{
 		return false;
@@ -325,6 +350,18 @@ protected:
 		FVector3d Min = c - e, Max = c + e;
 		return FAxisAlignedBox3d(Min, Max);
 	}
+	FAxisAlignedBox3d GetBox(int iBox, TFunction<FVector3d(const FVector3d&)> TransformF)
+	{
+		if (TransformF != nullptr)
+		{
+			FAxisAlignedBox3d box = GetBox(iBox);
+			return FAxisAlignedBox3d(box, TransformF);
+		}
+		else
+		{
+			return GetBox(iBox);
+		}
+	}
 
 	FAxisAlignedBox3d GetBoxEps(int IBox, double Epsilon = FMathd::ZeroTolerance) const
 	{
@@ -585,6 +622,178 @@ protected:
 		return IBox;
 	}
 
+
+	void find_nearest_triangles(int iBox, TMeshAABBTree3& otherTree, TFunction<FVector3d(const FVector3d&)> TransformF, int oBox, int depth, double &nearest_sqr, FIndex2i &nearest_pair)
+	{
+		int idx = BoxToIndex[iBox];
+		int odx = otherTree.BoxToIndex[oBox];
+
+		if (idx < TrianglesEnd && odx < otherTree.TrianglesEnd)
+		{
+			// ok we are at triangles for both trees, do triangle-level testing
+			FTriangle3d tri, otri;
+			int num_tris = IndexList[idx], onum_tris = otherTree.IndexList[odx];
+
+			FDistTriangle3Triangle3d dist;
+
+			// outer iteration is "other" tris that need to be transformed (more expensive)
+			for (int j = 1; j <= onum_tris; ++j)
+			{
+				int tj = otherTree.IndexList[odx + j];
+				if (otherTree.TriangleFilterF != nullptr && otherTree.TriangleFilterF(tj) == false)
+				{
+					continue;
+				}
+				otherTree.Mesh->GetTriVertices(tj, otri.V[0], otri.V[1], otri.V[2]);
+				if (TransformF != nullptr)
+				{
+					otri.V[0] = TransformF(otri.V[0]);
+					otri.V[1] = TransformF(otri.V[1]);
+					otri.V[2] = TransformF(otri.V[2]);
+				}
+				dist.Triangle[0] = otri;
+
+				// inner iteration over "our" triangles
+				for (int i = 1; i <= num_tris; ++i)
+				{
+					int ti = IndexList[idx + i];
+					if (TriangleFilterF != nullptr && TriangleFilterF(ti) == false)
+					{
+						continue;
+					}
+					Mesh->GetTriVertices(ti, tri.V[0], tri.V[1], tri.V[2]);
+					dist.Triangle[1] = tri;
+					double dist_sqr = dist.GetSquared();
+					if (dist_sqr < nearest_sqr)
+					{
+						nearest_sqr = dist_sqr;
+						nearest_pair = FIndex2i(ti, tj);
+					}
+				}
+			}
+
+			return;
+		}
+
+		// we either descend "our" tree or the other tree
+		//   - if we have hit triangles on "our" tree, we have to descend other
+		//   - if we hit triangles on "other", we have to descend ours
+		//   - otherwise, we alternate at each depth. This produces wider
+		//     branching but is significantly faster (~10x) for both hits and misses
+		bool bDescendOther = (idx < TrianglesEnd || depth % 2 == 0);
+		if (bDescendOther && odx < otherTree.TrianglesEnd)
+		{
+			bDescendOther = false;      // can't
+		}
+
+		if (bDescendOther)
+		{
+			// ok we reached triangles on our side but we need to still reach triangles on
+			// the other side, so we descend "their" children
+			FAxisAlignedBox3d bounds = GetBox(iBox);
+
+			int oChild1 = otherTree.IndexList[odx];
+			if (oChild1 < 0)		// 1 child, descend if nearer than cur min-dist
+			{
+				oChild1 = (-oChild1) - 1;
+				FAxisAlignedBox3d oChild1Box = otherTree.GetBox(oChild1, TransformF);
+				if (oChild1Box.DistanceSquared(bounds) < nearest_sqr)
+				{
+					find_nearest_triangles(iBox, otherTree, TransformF, oChild1, depth + 1, nearest_sqr, nearest_pair);
+				}
+			}
+			else                            // 2 children
+			{
+				oChild1 = oChild1 - 1;
+				int oChild2 = otherTree.IndexList[odx + 1] - 1;
+
+				FAxisAlignedBox3d oChild1Box = otherTree.GetBox(oChild1, TransformF);
+				FAxisAlignedBox3d oChild2Box = otherTree.GetBox(oChild2, TransformF);
+
+				// descend closer box first
+				double d1Sqr = oChild1Box.DistanceSquared(bounds);
+				double d2Sqr = oChild2Box.DistanceSquared(bounds);
+				if (d2Sqr < d1Sqr)
+				{
+					if (d2Sqr < nearest_sqr)
+					{
+						find_nearest_triangles(iBox, otherTree, TransformF, oChild2, depth + 1, nearest_sqr, nearest_pair);
+					}
+					if (d1Sqr < nearest_sqr)
+					{
+						find_nearest_triangles(iBox, otherTree, TransformF, oChild1, depth + 1, nearest_sqr, nearest_pair);
+					}
+				}
+				else
+				{
+					if (d1Sqr < nearest_sqr)
+					{
+						find_nearest_triangles(iBox, otherTree, TransformF, oChild1, depth + 1, nearest_sqr, nearest_pair);
+					}
+					if (d2Sqr < nearest_sqr)
+					{
+						find_nearest_triangles(iBox, otherTree, TransformF, oChild2, depth + 1, nearest_sqr, nearest_pair);
+					}
+				}
+			}
+		}
+		else
+		{
+			// descend our tree nodes if they intersect w/ current bounds of other tree
+			FAxisAlignedBox3d oBounds = otherTree.GetBox(oBox, TransformF);
+
+			int iChild1 = IndexList[idx];
+			if (iChild1 < 0)                  // 1 child, descend if nearer than cur min-dist
+			{
+				iChild1 = (-iChild1) - 1;
+				if (box_box_distsqr(iChild1, oBounds) < nearest_sqr)
+				{
+					find_nearest_triangles(iChild1, otherTree, TransformF, oBox, depth + 1, nearest_sqr, nearest_pair);
+				}
+			}
+			else                             // 2 children
+			{
+				iChild1 = iChild1 - 1;
+				int iChild2 = IndexList[idx + 1] - 1;
+
+				// descend closer box first
+				double d1Sqr = box_box_distsqr(iChild1, oBounds);
+				double d2Sqr = box_box_distsqr(iChild2, oBounds);
+				if (d2Sqr < d1Sqr)
+				{
+					if (d2Sqr < nearest_sqr)
+					{
+						find_nearest_triangles(iChild2, otherTree, TransformF, oBox, depth + 1, nearest_sqr, nearest_pair);
+					}
+					if (d1Sqr < nearest_sqr)
+					{
+						find_nearest_triangles(iChild1, otherTree, TransformF, oBox, depth + 1, nearest_sqr, nearest_pair);
+					}
+				}
+				else
+				{
+					if (d1Sqr < nearest_sqr)
+					{
+						find_nearest_triangles(iChild1, otherTree, TransformF, oBox, depth + 1, nearest_sqr, nearest_pair);
+					}
+					if (d2Sqr < nearest_sqr)
+					{
+						find_nearest_triangles(iChild2, otherTree, TransformF, oBox, depth + 1, nearest_sqr, nearest_pair);
+					}
+				}
+			}
+		}
+	}
+
+	double box_box_distsqr(int iBox, const FAxisAlignedBox3d& testBox)
+	{
+		// [TODO] could compute this w/o constructing box
+		FAxisAlignedBox3d box = GetBoxEps(iBox, BoxEps);
+		return box.DistanceSquared(testBox);
+	}
+
+
+
 public:
 	// 1) make sure we can reach every tri in Mesh through tree (also demo of how to traverse tree...)
 	// 2) make sure that Triangles are contained in parent boxes
