Minor GeometryProcessing changes: VectorUtil::Min3Index, FTriangle::Expand, Segment23::ProjectUnitRange().
Add TPolyline3, add EdgeSpan::GetPolyline
Add DistRay3Segment3
#rnx

diff --git a/Engine/Plugins/Experimental/GeometryProcessing/Source/DynamicMesh/Private/EdgeSpan.cpp b/Engine/Plugins/Experimental/GeometryProcessing/Source/DynamicMesh/Private/EdgeSpan.cpp
--- a/Engine/Plugins/Experimental/GeometryProcessing/Source/DynamicMesh/Private/EdgeSpan.cpp
+++ b/Engine/Plugins/Experimental/GeometryProcessing/Source/DynamicMesh/Private/EdgeSpan.cpp
@@ -102,6 +102,15 @@ FAxisAlignedBox3d FEdgeSpan::GetBounds() const
 
 
 
+void FEdgeSpan::GetPolyline(FPolyline3d& PolylineOut) const
+{
+	for (int i = 0; i < Vertices.Num(); ++i)
+	{
+		PolylineOut.AppendVertex(Mesh->GetVertex(Vertices[i]));
+	}
+}
+
+
 bool FEdgeSpan::SetCorrectOrientation()
 {
 	int NumEdges = Edges.Num();
diff --git a/Engine/Plugins/Experimental/GeometryProcessing/Source/DynamicMesh/Public/EdgeSpan.h b/Engine/Plugins/Experimental/GeometryProcessing/Source/DynamicMesh/Public/EdgeSpan.h
--- a/Engine/Plugins/Experimental/GeometryProcessing/Source/DynamicMesh/Public/EdgeSpan.h
+++ b/Engine/Plugins/Experimental/GeometryProcessing/Source/DynamicMesh/Public/EdgeSpan.h
@@ -6,6 +6,7 @@
 
 #include "DynamicMesh3.h"
 #include "Util/IndexUtil.h"
+#include "Polyline3.h"
 
 
 /**
@@ -142,6 +143,11 @@ public:
 	FAxisAlignedBox3d GetBounds() const;
 
 
+	/**
+	 * Extract Polyline from Mesh based on vertex list
+	 */
+	void GetPolyline(FPolyline3d& PolylineOut) const;
+
 	/**
 	 * If any edges if the span are on a mesh boundary, we can check that the span is
 	 * oriented such that it corresponds with the boundary edges, and if not, reverse it.
diff --git a/Engine/Plugins/Experimental/GeometryProcessing/Source/GeometricObjects/Public/Distance/DistRay3Segment3.h b/Engine/Plugins/Experimental/GeometryProcessing/Source/GeometricObjects/Public/Distance/DistRay3Segment3.h
new file mode 100644
--- /dev/null
+++ b/Engine/Plugins/Experimental/GeometryProcessing/Source/GeometricObjects/Public/Distance/DistRay3Segment3.h
@@ -0,0 +1,389 @@
+// Copyright 1998-2019 Epic Games, Inc. All Rights Reserved.
+
+// Port of geometry3Sharp DistRay3Segment3
+
+#pragma once
+
+#include "VectorTypes.h"
+#include "SegmentTypes.h"
+#include "RayTypes.h"
+
+/**
+ * Compute distance between 3D ray and 3D segment
+ */
+template <typename Real>
+class TDistRay3Segment3
+{
+public:
+	// Input
+	FRay3<Real> Ray;
+	FSegment3<Real> Segment;
+
+	// Results
+	Real DistanceSquared = -1.0;
+
+	FVector3<Real> RayClosestPoint;
+	Real RayParameter;
+	FVector3<Real> SegmentClosestPoint;
+	Real SegmentParameter;
+
+
+	TDistRay3Segment3(const FRay3<Real>& RayIn, const FSegment3<Real>& SegmentIn)
+	{
+		Ray = RayIn;
+		Segment = SegmentIn;
+	}
+
+	Real Get()
+	{
+		return (Real)sqrt(ComputeResult());
+	}
+	Real GetSquared()
+	{
+		return ComputeResult();
+	}
+
+	Real ComputeResult()
+	{
+		if (DistanceSquared >= 0)
+		{
+			return DistanceSquared;
+		}
+		FVector3d diff = Ray.Origin - Segment.Center;
+		double a01 = -Ray.Direction.Dot(Segment.Direction);
+		double b0 = diff.Dot(Ray.Direction);
+		double b1 = -diff.Dot(Segment.Direction);
+		double c = diff.SquaredLength();
+		double det = TMathUtil::Abs(1 - a01 * a01);
+		double s0, s1, sqrDist, extDet;
+
+		if (det >= TMathUtil<Real>::ZeroTolerance ) 
+		{
+			// The Ray and Segment are not parallel.
+			s0 = a01 * b1 - b0;
+			s1 = a01 * b0 - b1;
+			extDet = Segment.Extent * det;
+
+			if (s0 >= 0) 
+			{
+				if (s1 >= -extDet) 
+				{
+					if (s1 <= extDet)  // region 0
+					{
+						// Minimum at interior points of Ray and Segment.
+						double invDet = (1) / det;
+						s0 *= invDet;
+						s1 *= invDet;
+						sqrDist = s0 * (s0 + a01 * s1 + (2) * b0) +
+							s1 * (a01 * s0 + s1 + (2) * b1) + c;
+					}
+					else  // region 1
+					{
+						s1 = Segment.Extent;
+						s0 = -(a01 * s1 + b0);
+						if (s0 > 0) 
+						{
+							sqrDist = -s0 * s0 + s1 * (s1 + (2) * b1) + c;
+						}
+						else 
+						{
+							s0 = 0;
+							sqrDist = s1 * (s1 + (2) * b1) + c;
+						}
+					}
+				}
+				else  // region 5
+				{
+					s1 = -Segment.Extent;
+					s0 = -(a01 * s1 + b0);
+					if (s0 > 0) 
+					{
+						sqrDist = -s0 * s0 + s1 * (s1 + (2) * b1) + c;
+					}
+					else 
+					{
+						s0 = 0;
+						sqrDist = s1 * (s1 + (2) * b1) + c;
+					}
+				}
+			}
+			else 
+			{
+				if (s1 <= -extDet)  // region 4
+				{
+					s0 = -(-a01 * Segment.Extent + b0);
+					if (s0 > 0) 
+					{
+						s1 = -Segment.Extent;
+						sqrDist = -s0 * s0 + s1 * (s1 + (2) * b1) + c;
+					}
+					else 
+					{
+						s0 = 0;
+						s1 = -b1;
+						if (s1 < -Segment.Extent) 
+						{
+							s1 = -Segment.Extent;
+						}
+						else if (s1 > Segment.Extent) 
+						{
+							s1 = Segment.Extent;
+						}
+						sqrDist = s1 * (s1 + (2) * b1) + c;
+					}
+				}
+				else if (s1 <= extDet)  // region 3
+				{
+					s0 = 0;
+					s1 = -b1;
+					if (s1 < -Segment.Extent) 
+					{
+						s1 = -Segment.Extent;
+					}
+					else if (s1 > Segment.Extent) 
+					{
+						s1 = Segment.Extent;
+					}
+					sqrDist = s1 * (s1 + (2) * b1) + c;
+				}
+				else  // region 2
+				{
+					s0 = -(a01 * Segment.Extent + b0);
+					if (s0 > 0) 
+					{
+						s1 = Segment.Extent;
+						sqrDist = -s0 * s0 + s1 * (s1 + (2) * b1) + c;
+					}
+					else 
+					{
+						s0 = 0;
+						s1 = -b1;
+						if (s1 < -Segment.Extent) 
+						{
+							s1 = -Segment.Extent;
+						}
+						else if (s1 > Segment.Extent) 
+						{
+							s1 = Segment.Extent;
+						}
+						sqrDist = s1 * (s1 + (2) * b1) + c;
+					}
+				}
+			}
+		}
+		else 
+		{
+			// Ray and Segment are parallel.
+			if (a01 > 0) 
+			{
+				// Opposite direction vectors.
+				s1 = -Segment.Extent;
+			}
+			else 
+			{
+				// Same direction vectors.
+				s1 = Segment.Extent;
+			}
+
+			s0 = -(a01 * s1 + b0);
+			if (s0 > 0) 
+			{
+				sqrDist = -s0 * s0 + s1 * (s1 + (2) * b1) + c;
+			}
+			else 
+			{
+				s0 = 0;
+				sqrDist = s1 * (s1 + (2) * b1) + c;
+			}
+		}
+
+		RayClosestPoint = Ray.Origin + s0 * Ray.Direction;
+		SegmentClosestPoint = Segment.Center + s1 * Segment.Direction;
+		RayParameter = s0;
+		SegmentParameter = s1;
+
+		// Account for numerical round-off errors.
+		if (sqrDist < (Real)0) 
+		{
+			sqrDist = (Real)0;
+		}
+		DistanceSquared = sqrDist;
+
+		return sqrDist;
+	}
+
+
+
+
+
+	static double SquaredDistance(const FRay3<Real>& Ray, const FSegment3<Real>& Segment,
+		Real& RayParam, Real& SegParam)
+	{
+		FVector3<Real> diff = Ray.Origin - Segment.Center;
+		double a01 = -Ray.Direction.Dot(Segment.Direction);
+		double b0 = diff.Dot(Ray.Direction);
+		double b1 = -diff.Dot(Segment.Direction);
+		double c = diff.SquaredLength();
+		double det = TMathUtil<Real>::Abs(1 - a01 * a01);
+		double s0, s1, sqrDist, extDet;
+
+		if (det >= TMathUtil<Real>::ZeroTolerance) 
+		{
+			// The Ray and Segment are not parallel.
+			s0 = a01 * b1 - b0;
+			s1 = a01 * b0 - b1;
+			extDet = Segment.Extent * det;
+
+			if (s0 >= 0) 
+			{
+				if (s1 >= -extDet) 
+				{
+					if (s1 <= extDet)  // region 0
+					{
+						// Minimum at interior points of Ray and Segment.
+						double invDet = (1) / det;
+						s0 *= invDet;
+						s1 *= invDet;
+						sqrDist = s0 * (s0 + a01 * s1 + (2) * b0) +
+							s1 * (a01 * s0 + s1 + (2) * b1) + c;
+					}
+					else  // region 1
+					{
+						s1 = Segment.Extent;
+						s0 = -(a01 * s1 + b0);
+						if (s0 > 0) 
+						{
+							sqrDist = -s0 * s0 + s1 * (s1 + (2) * b1) + c;
+						}
+						else 
+						{
+							s0 = 0;
+							sqrDist = s1 * (s1 + (2) * b1) + c;
+						}
+					}
+				}
+				else  // region 5
+				{
+					s1 = -Segment.Extent;
+					s0 = -(a01 * s1 + b0);
+					if (s0 > 0) 
+					{
+						sqrDist = -s0 * s0 + s1 * (s1 + (2) * b1) + c;
+					}
+					else 
+					{
+						s0 = 0;
+						sqrDist = s1 * (s1 + (2) * b1) + c;
+					}
+				}
+			}
+			else 
+			{
+				if (s1 <= -extDet)  // region 4
+				{
+					s0 = -(-a01 * Segment.Extent + b0);
+					if (s0 > 0) 
+					{
+						s1 = -Segment.Extent;
+						sqrDist = -s0 * s0 + s1 * (s1 + (2) * b1) + c;
+					}
+					else 
+					{
+						s0 = 0;
+						s1 = -b1;
+						if (s1 < -Segment.Extent) 
+						{
+							s1 = -Segment.Extent;
+						}
+						else if (s1 > Segment.Extent) 
+						{
+							s1 = Segment.Extent;
+						}
+						sqrDist = s1 * (s1 + (2) * b1) + c;
+					}
+				}
+				else if (s1 <= extDet)  // region 3
+				{
+					s0 = 0;
+					s1 = -b1;
+					if (s1 < -Segment.Extent) 
+					{
+						s1 = -Segment.Extent;
+					}
+					else if (s1 > Segment.Extent) 
+					{
+						s1 = Segment.Extent;
+					}
+					sqrDist = s1 * (s1 + (2) * b1) + c;
+				}
+				else  // region 2
+				{
+					s0 = -(a01 * Segment.Extent + b0);
+					if (s0 > 0) 
+					{
+						s1 = Segment.Extent;
+						sqrDist = -s0 * s0 + s1 * (s1 + (2) * b1) + c;
+					}
+					else 
+					{
+						s0 = 0;
+						s1 = -b1;
+						if (s1 < -Segment.Extent) 
+						{
+							s1 = -Segment.Extent;
+						}
+						else if (s1 > Segment.Extent) 
+						{
+							s1 = Segment.Extent;
+						}
+						sqrDist = s1 * (s1 + (2) * b1) + c;
+					}
+				}
+			}
+		}
+		else {
+			// Ray and Segment are parallel.
+			if (a01 > 0) 
+			{
+				// Opposite direction vectors.
+				s1 = -Segment.Extent;
+			}
+			else 
+			{
+				// Same direction vectors.
+				s1 = Segment.Extent;
+			}
+
+			s0 = -(a01 * s1 + b0);
+			if (s0 > 0) 
+			{
+				sqrDist = -s0 * s0 + s1 * (s1 + (2) * b1) + c;
+			}
+			else 
+			{
+				s0 = 0;
+				sqrDist = s1 * (s1 + (2) * b1) + c;
+			}
+		}
+
+		RayParam = s0;
+		SegParam = s1;
+
+		// Account for numerical round-off errors.
+		if (sqrDist < 0)
+		{
+			sqrDist = 0;
+		}
+		return sqrDist;
+	}
+
+
+
+
+
+
+
+};
+
+typedef TDistRay3Segment3<float> FDistRay3Segment3f;
+typedef TDistRay3Segment3<double> FDistRay3Segment3d;
\ No newline at end of file
diff --git a/Engine/Plugins/Experimental/GeometryProcessing/Source/GeometricObjects/Public/Polyline3.h b/Engine/Plugins/Experimental/GeometryProcessing/Source/GeometricObjects/Public/Polyline3.h
new file mode 100644
--- /dev/null
+++ b/Engine/Plugins/Experimental/GeometryProcessing/Source/GeometricObjects/Public/Polyline3.h
@@ -0,0 +1,399 @@
+// Copyright 1998-2019 Epic Games, Inc. All Rights Reserved.
+
+#pragma once
+
+#include "VectorTypes.h"
+#include "SegmentTypes.h"
+#include "LineTypes.h"
+#include "RayTypes.h"
+
+/**
+ * TPolyline3 represents a 3D polyline stored as a list of Vertices.
+ *
+ * @todo move operators
+ */
+template<typename T>
+class TPolyline3
+{
+protected:
+	/** The list of vertices of the polyline */
+	TArray<FVector3<T>> Vertices;
+
+	/** A counter that is incremented every time the polyline vertices are modified */
+	int Timestamp;
+
+public:
+
+
+	TPolyline3() : Timestamp(0)
+	{
+	}
+
+	/**
+	 * Construct polyline that is a copy of another polyline
+	 */
+	TPolyline3(const TPolyline3& Copy) : Vertices(Copy.Vertices), Timestamp(Copy.Timestamp)
+	{
+	}
+
+	/**
+	 * Construct polyline with given list of vertices
+	 */
+	TPolyline3(const TArray<FVector3<T>>& VertexList) : Vertices(VertexList), Timestamp(0)
+	{
+	}
+
+	/** @return the Timestamp for the polyline, which is updated every time the polyline is modified */
+	int GetTimestamp() const
+	{
+		return Timestamp;
+	}
+
+	/** Explicitly increment the Timestamp */
+	void IncrementTimestamp()
+	{
+		Timestamp++;
+	}
+
+	/**
+	 * Get the vertex at a given index
+	 */
+	const FVector3<T>& operator[](int Index) const
+	{
+		return Vertices[Index];
+	}
+
+	/**
+	 * Get the vertex at a given index
+	 * @warning changing the vertex via this operator does not update Timestamp!
+	 */
+	FVector3<T>& operator[](int Index)
+	{
+		return Vertices[Index];
+	}
+
+
+	/**
+	 * @return first vertex of polyline
+	 */
+	const FVector3<T>& Start() const
+	{
+		return Vertices[0];
+	}
+
+	/**
+	 * @return last vertex of polyline
+	 */
+	const FVector3<T>& End() const
+	{
+		return Vertices[Vertices.Num()-1];
+	}
+
+
+	/**
+	 * @return list of Vertices of polyline
+	 */
+	const TArray<FVector3<T>>& GetVertices() const
+	{
+		return Vertices;
+	}
+
+	/**
+	 * @return number of Vertices in polyline
+	 */
+	int VertexCount() const
+	{
+		return Vertices.Num();
+	}
+
+	/**
+	 * @return number of segments in polyline
+	 */
+	int SegmentCount() const
+	{
+		return Vertices.Num()-1;
+	}
+
+
+	/**
+	 * Add a vertex to the polyline
+	 */
+	void AppendVertex(const FVector3<T>& Position)
+	{
+		Vertices.Add(Position);
+		Timestamp++;
+	}
+
+	/**
+	 * Add a list of Vertices to the polyline
+	 */
+	void AppendVertices(const TArray<FVector3<T>>& NewVertices)
+	{
+		Vertices.Append(NewVertices);
+		Timestamp++;
+	}
+
+	/**
+	 * Set vertex at given index to a new Position
+	 */
+	void Set(int VertexIndex, const FVector3<T>& Position)
+	{
+		Vertices[VertexIndex] = Position;
+		Timestamp++;
+	}
+
+	/**
+	 * Remove a vertex of the polyline (existing Vertices are shifted)
+	 */
+	void RemoveVertex(int VertexIndex)
+	{
+		Vertices.RemoveAt(VertexIndex);
+		Timestamp++;
+	}
+
+	/**
+	 * Replace the list of Vertices with a new list
+	 */
+	void SetVertices(const TArray<FVector3<T>>& NewVertices)
+	{
+		int NumVerts = NewVertices.Num();
+		Vertices.SetNum(NumVerts, false);
+		for (int k = 0; k < NumVerts; ++k)
+		{
+			Vertices[k] = NewVertices[k];
+		}
+		Timestamp++;
+	}
+
+
+	/**
+	 * Reverse the order of the Vertices in the polyline (ie switch between Clockwise and CounterClockwise)
+	 */
+	void Reverse()
+	{
+		int32 j = Vertices.Num() - 1;
+		for (int32 VertexIndex = 0; VertexIndex < j; VertexIndex++, j--)
+		{
+			Swap(Vertices[VertexIndex], Vertices[j]);
+		}
+		Timestamp++;
+	}
+
+
+	/**
+	 * Get the tangent vector at a vertex of the polyline, which is the normalized
+	 * vector from the previous vertex to the next vertex
+	 */
+	FVector3<T> GetTangent(int VertexIndex) const
+	{
+		if (VertexIndex == 0)
+		{
+			return (Vertices[1] - Vertices[0]).Normalized();
+		} 
+		int NumVerts = Vertices.Num();
+		if (VertexIndex == NumVerts - 1)
+		{
+			return (Vertices[NumVerts-1] - Vertices[NumVerts-2]).Normalized();
+		}
+		return (Vertices[VertexIndex+1] - Vertices[VertexIndex-1]).Normalized()
+	}
+
+
+
+	/**
+	 * @return edge of the poyline starting at vertex SegmentIndex
+	 */
+	FSegment3<T> GetSegment(int SegmentIndex) const
+	{
+		return FSegment3<T>(Vertices[SegmentIndex], Vertices[SegmentIndex+1]);
+	}
+
+
+	/**
+	 * @param SegmentIndex index of first vertex of the edge
+	 * @param SegmentParam parameter in range [0,1] along segment
+	 * @return point on the segment at the given parameter value
+	 */
+	FVector3<T> PointAt(int SegmentIndex, double SegmentParam) const
+	{
+		FSegment3<T> seg(Vertices[SegmentIndex], Vertices[SegmentIndex + 1]);
+		return seg.PointAt(SegmentParam);
+	}
+
+
+
+
+	/**
+	 * @return the bounding box of the polyline Vertices
+	 */
+	FAxisAlignedBox3<T> GetBounds() const
+	{
+		FAxisAlignedBox3<T> box = FAxisAlignedBox3<T>::Empty();
+		int NumVertices = Vertices.Num();
+		for (int k = 0; k < NumVertices; ++k)
+		{
+			box.Contain(Vertices[k]);
+		}
+		return box;
+	}
+
+
+
+	/**
+	 * @return the total perimeter length of the Polygon
+	 */
+	double Length() const
+	{
+		double length = 0;
+		int N = SegmentCount();
+		for (int i = 0; i < N; ++i)
+		{
+			length += Vertices[i].Distance(Vertices[i+1]);
+		}
+		return length;
+	}
+
+
+
+	/**
+	 * SegmentIterator is used to iterate over the FSegment3<T> segments of the polyline
+	 */
+	class SegmentIterator
+	{
+	public:
+		inline bool operator!()
+		{
+			return i < SegmentCount();
+		}
+		inline FSegment3<T> operator*() const
+		{
+			check(Polyline != nullptr && i < SegmentCount());
+			return FSegment3<T>(Polyline->Vertices[i], Polyline->Vertices[i+1]);
+		}
+		//inline FSegment3<T> & operator*();
+		inline SegmentIterator & operator++() 		// prefix
+		{
+			i++;
+			return *this;
+		}
+		inline SegmentIterator operator++(int) 		// postfix
+		{
+			SegmentIterator copy(*this);
+			i++;
+			return copy;
+		}
+		inline bool operator==(const SegmentIterator & i3) { return i3.Polyline == Polyline && i3.i == i; }
+		inline bool operator!=(const SegmentIterator & i3) { return i3.Polyline != Polyline || i3.i != i; }
+	protected:
+		const TPolyline3 * Polyline;
+		int i;
+		inline SegmentIterator(const TPolyline3 * p, int iCur) : Polyline(p), i(iCur) {}
+		friend class TPolyline3;
+	};
+	friend class SegmentIterator;
+
+	SegmentIterator SegmentItr() const
+	{
+		return SegmentIterator(this, 0);
+	}
+
+	/**
+	 * Wrapper around SegmentIterator that has begin() and end() suitable for range-based for loop
+	 */
+	class SegmentEnumerable
+	{
+	public:
+		const TPolyline3<T>* Polyline;
+		SegmentEnumerable() : Polyline(nullptr) {}
+		SegmentEnumerable(const TPolyline3<T> * p) : Polyline(p) {}
+		SegmentIterator begin() { return Polyline->SegmentItr(); }
+		SegmentIterator end() { return SegmentIterator(Polyline, SegmentCount()); }
+	};
+
+	/**
+	 * @return an object that can be used in a range-based for loop to iterate over the Segments of the Polyline
+	 */
+	SegmentEnumerable Segments() const
+	{
+		return SegmentEnumerable(this);
+	}
+
+
+
+
+
+
+	/**
+	 * Calculate the squared distance from a point to the polyline
+	 * @param QueryPoint the query point
+	 * @param NearestSegIndexOut The index of the nearest segment
+	 * @param NearestSegParamOut the parameter value of the nearest point on the segment
+	 * @return squared distance to the polyline
+	 */
+	double DistanceSquared(const FVector3<T>& QueryPoint, int& NearestSegIndexOut, double& NearestSegParamOut) const
+	{
+		NearestSegIndexOut = -1;
+		NearestSegParamOut = TNumericLimits<double>::Max();
+		double dist = TNumericLimits<double>::Max();
+		int N = SegmentCount();
+		for (int vi = 0; vi < N; ++vi)
+		{
+			// @todo can't we just use segment function here now?
+			FSegment3<T> seg = FSegment3<T>(Vertices[vi], Vertices[vi+1]);
+			double t = (QueryPoint - seg.Center).Dot(seg.Direction);
+			double d = TNumericLimits<double>::Max();
+			if (t >= seg.Extent)
+			{
+				d = seg.EndPoint().DistanceSquared(QueryPoint);
+			}
+			else if (t <= -seg.Extent)
+			{
+				d = seg.StartPoint().DistanceSquared(QueryPoint);
+			}
+			else
+			{
+				d = (seg.PointAt(t) - QueryPoint).SquaredLength();
+			}
+			if (d < dist)
+			{
+				dist = d;
+				NearestSegIndexOut = vi;
+				NearestSegParamOut = t;
+			}
+		}
+		return dist;
+	}
+
+
+
+	/**
+	 * Calculate the squared distance from a point to the polyline
+	 * @param QueryPoint the query point
+	 * @return squared distance to the polyline
+	 */
+	double DistanceSquared(const FVector3<T>& QueryPoint) const
+	{
+		int seg; double segt;
+		return DistanceSquared(QueryPoint, seg, segt);
+	}
+
+
+
+
+	/**
+	 * @return average edge length of all the edges of the Polygon
+	 */
+	double AverageEdgeLength() const
+	{
+		double avg = 0; int N = Vertices.Num();
+		for (int i = 1; i < N; ++i) {
+			avg += Vertices[i].Distance(Vertices[i - 1]);
+		}
+		return avg / (double)(N-1);
+	}
+
+
+};
+
+typedef TPolyline3<double> FPolyline3d;
+typedef TPolyline3<float> FPolyline3f;
diff --git a/Engine/Plugins/Experimental/GeometryProcessing/Source/GeometricObjects/Public/SegmentTypes.h b/Engine/Plugins/Experimental/GeometryProcessing/Source/GeometricObjects/Public/SegmentTypes.h
--- a/Engine/Plugins/Experimental/GeometryProcessing/Source/GeometricObjects/Public/SegmentTypes.h
+++ b/Engine/Plugins/Experimental/GeometryProcessing/Source/GeometricObjects/Public/SegmentTypes.h
@@ -169,6 +169,18 @@ public:
 	}
 
 
+	/**
+	 * @return scalar projection of QueryPoint onto line of Segment, mapped to [0,1] range along segment
+	 */
+	inline T ProjectUnitRange(const FVector2<T>& QueryPoint) const
+	{
+		T ProjT = (QueryPoint - Center).Dot(Direction);
+		T Alpha = ((ProjT / Extent) + (T)1) * (T)0.5;
+		return TMathUtil<T>::Clamp(Alpha, (T)0, (T)1);
+	}
+
+
+
 	/**
 	 * Determine which side of the segment the query point lies on
 	 * @param QueryPoint test point
@@ -468,6 +480,16 @@ public:
 	}
 
 
+	/**
+	 * @return scalar projection of QueryPoint onto line of Segment, mapped to [0,1] range along segment
+	 */
+	inline T ProjectUnitRange(const FVector3<T>& QueryPoint) const
+	{
+		T ProjT = (QueryPoint - Center).Dot(Direction);
+		T Alpha = ((ProjT / Extent) + (T)1) * (T)0.5;
+		return TMathUtil<T>::Clamp(Alpha, (T)0, (T)1);
+	}
+
 
 protected:
 
diff --git a/Engine/Plugins/Experimental/GeometryProcessing/Source/GeometricObjects/Public/TriangleTypes.h b/Engine/Plugins/Experimental/GeometryProcessing/Source/GeometricObjects/Public/TriangleTypes.h
--- a/Engine/Plugins/Experimental/GeometryProcessing/Source/GeometricObjects/Public/TriangleTypes.h
+++ b/Engine/Plugins/Experimental/GeometryProcessing/Source/GeometricObjects/Public/TriangleTypes.h
@@ -157,6 +157,15 @@ struct FTriangle3
 			(V[0].Z + V[1].Z + V[2].Z) * f
 		);
 	}
+
+	/** grow the triangle around the centroid */
+	void Expand(RealType Delta)
+	{
+		FVector3<RealType> Centroid(Centroid());
+		V[0] += Delta * ((V[0] - Centroid).Normalized());
+		V[1] += Delta * ((V[1] - Centroid).Normalized());
+		V[2] += Delta * ((V[2] - Centroid).Normalized());
+	}
 };
 
 typedef FTriangle3<float> FTriangle3f;
diff --git a/Engine/Plugins/Experimental/GeometryProcessing/Source/GeometricObjects/Public/VectorUtil.h b/Engine/Plugins/Experimental/GeometryProcessing/Source/GeometricObjects/Public/VectorUtil.h
--- a/Engine/Plugins/Experimental/GeometryProcessing/Source/GeometricObjects/Public/VectorUtil.h
+++ b/Engine/Plugins/Experimental/GeometryProcessing/Source/GeometricObjects/Public/VectorUtil.h
@@ -155,6 +155,22 @@ namespace VectorUtil
 		return EpsilonEqual(V0.X, V1.X, Epsilon) && EpsilonEqual(V0.Y, V1.Y, Epsilon) && EpsilonEqual(V0.Z, V1.Z, Epsilon);
 	}
 
+
+	/** @return 0/1/2 index of smallest value in Vector3 */
+	template <typename ValueVecType>
+	inline int Min3Index(const ValueVecType& Vector3)
+	{
+		if (Vector3[0] <= Vector3[1])
+		{
+			return Vector3[0] <= Vector3[2] ? 0 : 2;
+		}
+		else
+		{
+			return (Vector3[1] <= Vector3[2]) ? 1 : 2;
+		}
+	}
+
+
 	/**
 	 * Calculates two vectors perpendicular to input Normal, as efficiently as possible.
 	 */
