fix porting errors
#rnx

diff --git a/Engine/Plugins/Experimental/GeometryProcessing/Source/DynamicMesh/Private/EdgeSpan.cpp b/Engine/Plugins/Experimental/GeometryProcessing/Source/DynamicMesh/Private/EdgeSpan.cpp
--- a/Engine/Plugins/Experimental/GeometryProcessing/Source/DynamicMesh/Private/EdgeSpan.cpp
+++ b/Engine/Plugins/Experimental/GeometryProcessing/Source/DynamicMesh/Private/EdgeSpan.cpp
@@ -104,6 +104,7 @@ FAxisAlignedBox3d FEdgeSpan::GetBounds() const
 
 void FEdgeSpan::GetPolyline(FPolyline3d& PolylineOut) const
 {
+	PolylineOut.Clear();
 	for (int i = 0; i < Vertices.Num(); ++i)
 	{
 		PolylineOut.AppendVertex(Mesh->GetVertex(Vertices[i]));
diff --git a/Engine/Plugins/Experimental/GeometryProcessing/Source/GeometricObjects/Public/Polygon2.h b/Engine/Plugins/Experimental/GeometryProcessing/Source/GeometricObjects/Public/Polygon2.h
--- a/Engine/Plugins/Experimental/GeometryProcessing/Source/GeometricObjects/Public/Polygon2.h
+++ b/Engine/Plugins/Experimental/GeometryProcessing/Source/GeometricObjects/Public/Polygon2.h
@@ -192,7 +192,7 @@ public:
 		prev -= Vertices[VertexIndex]; prev.Normalize();
 
 		FVector2<T> n = (next.Perp() - prev.Perp());
-		double len = n.Normalize();
+		T len = n.Normalize();
 		if (len == 0) 
 		{
 			return (next + prev).Normalized();   // this gives right direction for degenerate angle
@@ -290,9 +290,9 @@ public:
 	/**
 	 * @return the signed area of the Polygon
 	 */
-	double SignedArea() const
+	T SignedArea() const
 	{
-		double fArea = 0;
+		T fArea = 0;
 		int N = Vertices.Num();
 		if (N == 0)
 		{
@@ -310,7 +310,7 @@ public:
 	/**
 	 * @return the unsigned area of the Polygon
 	 */
-	double Area() const
+	T Area() const
 	{
 		return TMathUtil<T>::Abs(SignedArea());
 	}
@@ -318,9 +318,9 @@ public:
 	/**
 	 * @return the total perimeter length of the Polygon
 	 */
-	double Perimeter() const
+	T Perimeter() const
 	{
-		double fPerim = 0;
+		T fPerim = 0;
 		int N = Vertices.Num();
 		for (int i = 0; i < N; ++i)
 		{
@@ -360,7 +360,7 @@ public:
 	/**
 	 * @return the opening angle in degrees at a vertex of the Polygon
 	 */
-	double OpeningAngleDeg(int iVertex) const
+	T OpeningAngleDeg(int iVertex) const
 	{
 		FVector2<T> e0, e1;
 		NeighbourVectors(iVertex, e0, e1, true);
@@ -371,9 +371,9 @@ public:
 	/**
 	 * @return analytic winding integral for this Polygon at an arbitrary point
 	 */
-	double WindingIntegral(const FVector2<T>& QueryPoint) const
+	T WindingIntegral(const FVector2<T>& QueryPoint) const
 	{
-		double sum = 0;
+		T sum = 0;
 		int N = Vertices.Num();
 		FVector2<T> a = Vertices[0] - QueryPoint, b = FVector2<T>::Zero();
 		for (int i = 0; i < N; ++i) 
@@ -571,7 +571,7 @@ public:
 	 * @param SegmentParam parameter in range [0,1] along segment
 	 * @return point on the segment at the given parameter value
 	 */
-	FVector2<T> PointAt(int SegmentIndex, double SegmentParam) const
+	FVector2<T> PointAt(int SegmentIndex, T SegmentParam) const
 	{
 		FSegment2<T> seg(Vertices[SegmentIndex], Vertices[(SegmentIndex + 1) % Vertices.Num()]);
 		return seg.PointAt(SegmentParam);
@@ -583,7 +583,7 @@ public:
 	 * @param SegmentParam parameter in range [0,1] along segment
 	 * @return interpolated normal to the segment at the given parameter value
 	 */
-	FVector2<T> GetNormal(int iSeg, double SegmentParam) const
+	FVector2<T> GetNormal(int iSeg, T SegmentParam) const
 	{
 		FSegment2<T> seg(Vertices[iSeg], Vertices[(iSeg + 1) % Vertices.Num()]);
 		T t = ((SegmentParam / seg.Extent) + 1.0) / 2.0;
@@ -602,18 +602,18 @@ public:
 	 * @param NearestSegParamOut the parameter value of the nearest point on the segment
 	 * @return squared distance to the polygon
 	 */
-	double DistanceSquared(const FVector2<T>& QueryPoint, int& NearestSegIndexOut, double& NearestSegParamOut) const
+	T DistanceSquared(const FVector2<T>& QueryPoint, int& NearestSegIndexOut, T& NearestSegParamOut) const
 	{
 		NearestSegIndexOut = -1;
-		NearestSegParamOut = TNumericLimits<double>::Max();
-		double dist = TNumericLimits<double>::Max();
+		NearestSegParamOut = TNumericLimits<T>::Max();
+		T dist = TNumericLimits<T>::Max();
 		int N = Vertices.Num();
 		for (int vi = 0; vi < N; ++vi) 
 		{
 			// @todo can't we just use segment function here now?
 			FSegment2<T> seg = FSegment2<T>(Vertices[vi], Vertices[(vi + 1) % N]);
-			double t = (QueryPoint - seg.Center).Dot(seg.Direction);
-			double d = TNumericLimits<double>::Max();
+			T t = (QueryPoint - seg.Center).Dot(seg.Direction);
+			T d = TNumericLimits<T>::Max();
 			if (t >= seg.Extent)
 			{
 				d = seg.EndPoint().DistanceSquared(QueryPoint);
@@ -630,7 +630,7 @@ public:
 			{
 				dist = d;
 				NearestSegIndexOut = vi;
-				NearestSegParamOut = t;
+				NearestSegParamOut = TMathUtil<T>::Clamp(t, -seg.Extent, seg.Extent);
 			}
 		}
 		return dist;
@@ -642,9 +642,9 @@ public:
 	 * @param QueryPoint the query point
 	 * @return squared distance to the polygon
 	 */
-	double DistanceSquared(const FVector2<T>& QueryPoint) const
+	T DistanceSquared(const FVector2<T>& QueryPoint) const
 	{
-		int seg; double segt;
+		int seg; T segt;
 		return DistanceSquared(QueryPoint, seg, segt);
 	}
 
@@ -652,9 +652,9 @@ public:
 	/**
 	 * @return average edge length of all the edges of the Polygon
 	 */
-	double AverageEdgeLength() const
+	T AverageEdgeLength() const
 	{
-		double avg = 0; int N = Vertices.Num();
+		T avg = 0; int N = Vertices.Num();
 		for (int i = 1; i < N; ++i) {
 			avg += Vertices[i].Distance(Vertices[i - 1]);
 		}
@@ -790,7 +790,7 @@ private:
 	//            v[] = polyline array of vertex points
 	//            j,k = indices for the subchain v[j] to v[k]
 	//    Output: mk[] = array of markers matching vertex array v[]
-	static void SimplifyDouglasPeucker(double Tolerance, const TArray<FVector2<T>>& Vertices, int j, int k, TArray<bool>& Marked)
+	static void SimplifyDouglasPeucker(T Tolerance, const TArray<FVector2<T>>& Vertices, int j, int k, TArray<bool>& Marked)
 	{
 		Marked.SetNum(Vertices.Num());
 		if (k <= j + 1) // there is nothing to simplify
@@ -798,15 +798,15 @@ private:
 
 		// check for adequate approximation by segment S from v[j] to v[k]
 		int maxi = j;          // index of vertex farthest from S
-		double maxd2 = 0;         // distance squared of farthest vertex
-		double tol2 = Tolerance * Tolerance;  // tolerance squared
+		T maxd2 = 0;         // distance squared of farthest vertex
+		T tol2 = Tolerance * Tolerance;  // tolerance squared
 		FSegment2<T> S = FSegment2<T>(Vertices[j], Vertices[k]);    // segment from v[j] to v[k]
 
 		// test each vertex v[i] for max distance from S
 		// Note: this works in any dimension (2D, 3D, ...)
 		for (int i = j + 1; i < k; i++)
 		{
-			double dv2 = S.DistanceSquared(Vertices[i]);
+			T dv2 = S.DistanceSquared(Vertices[i]);
 			if (dv2 <= maxd2)
 				continue;
 			// v[i] is a max vertex
@@ -833,7 +833,7 @@ public:
 	 * @param ClusterTolerance Vertices closer than this distance will be merged into a single vertex
 	 * @param LineDeviationTolerance Vertices are allowed to deviate this much from the input polygon lines
 	 */
-	void Simplify(double ClusterTolerance = 0.0001, double LineDeviationTolerance = 0.01)
+	void Simplify(T ClusterTolerance = 0.0001, T LineDeviationTolerance = 0.01)
 	{
 		int n = Vertices.Num();
 		if (n < 3)
@@ -852,7 +852,7 @@ public:
 		}
 
 		// STAGE 1.  Vertex Reduction within tolerance of prior vertex cluster
-		double clusterTol2 = ClusterTolerance * ClusterTolerance;
+		T clusterTol2 = ClusterTolerance * ClusterTolerance;
 		NewVertices[0] = Vertices[0];              // start at the beginning
 		for (i = 1, k = 1, pv = 0; i < n; i++) 
 		{
@@ -945,7 +945,7 @@ public:
 	 * Chamfer each vertex corner of the Polygon
 	 * @param ChamferDist offset distance from corner that we cut at
 	 */
-	void Chamfer(double ChamferDist, double MinConvexAngleDeg = 30, double MinConcaveAngleDeg = 30)
+	void Chamfer(T ChamferDist, T MinConvexAngleDeg = 30, T MinConcaveAngleDeg = 30)
 	{
 		check(IsClockwise());
 
@@ -963,9 +963,9 @@ public:
 			FVector2<T> next = OldV[iNext];
 
 			FVector2<T> cp = prev - center;
-			double cpdist = cp.Normalize();
+			T cpdist = cp.Normalize();
 			FVector2<T> cn = next - center;
-			double cndist = cn.Normalize();
+			T cndist = cn.Normalize();
 
 			// if degenerate, skip this vert
 			if (cpdist < TMathUtil<T>::ZeroTolerance || cndist < TMathUtil<T>::ZeroTolerance) 
@@ -974,13 +974,13 @@ public:
 				continue;
 			}
 
-			double angle = cp.AngleD(cn);
+			T angle = cp.AngleD(cn);
 			// TODO document what this means sign-wise
 			// TODO re-test post Unreal port that this DotPerp is doing the right thing
-			double sign = cn.DotPerp(cp);
+			T sign = cn.DotPerp(cp);
 			bool bConcave = (sign > 0);
 
-			double thresh = (bConcave) ? MinConcaveAngleDeg : MinConvexAngleDeg;
+			T thresh = (bConcave) ? MinConcaveAngleDeg : MinConvexAngleDeg;
 
 			// ok not too sharp
 			if (angle > thresh) 
@@ -991,9 +991,9 @@ public:
 			}
 
 
-			double prev_cut_dist = TMathUtil<T>::Min(ChamferDist, cpdist*0.5);
+			T prev_cut_dist = TMathUtil<T>::Min(ChamferDist, cpdist*0.5);
 			FVector2<T> prev_cut = center + cp * prev_cut_dist;
-			double next_cut_dist = TMathUtil<T>::Min(ChamferDist, cndist * 0.5);
+			T next_cut_dist = TMathUtil<T>::Min(ChamferDist, cndist * 0.5);
 			FVector2<T> next_cut = center + cn * next_cut_dist;
 
 			Vertices.Add(prev_cut);
@@ -1010,7 +1010,7 @@ public:
 	/**
 	 * Construct a four-vertex rectangle Polygon
 	 */
-	static FPolygon2<T> MakeRectangle(const FVector2<T>& Center, double Width, double Height)
+	static FPolygon2<T> MakeRectangle(const FVector2<T>& Center, T Width, T Height)
 	{
 		FPolygon2<T> Rectangle;
 		Rectangle.Vertices.SetNumUninitialized(4);
@@ -1025,15 +1025,15 @@ public:
 	/**
 	 * Construct a circular Polygon
 	 */
-	static FPolygon2<T> MakeCircle(double Radius, int Steps, double AngleShiftRadians = 0)
+	static FPolygon2<T> MakeCircle(T Radius, int Steps, T AngleShiftRadians = 0)
 	{
 		FPolygon2<T> Circle;
 		Circle.Vertices.SetNumUninitialized(Steps);
 
 		for (int i = 0; i < Steps; ++i)
 		{
-			double t = (double)i / (double)Steps;
-			double a = TMathUtil<T>::TwoPi * t + AngleShiftRadians;
+			T t = (T)i / (T)Steps;
+			T a = TMathUtil<T>::TwoPi * t + AngleShiftRadians;
 			Circle.Set(i, FVector2<T>(Radius * TMathUtil<T>::Cos(a), Radius * TMathUtil<T>::Sin(a)));
 		}
 
diff --git a/Engine/Plugins/Experimental/GeometryProcessing/Source/GeometricObjects/Public/Polyline3.h b/Engine/Plugins/Experimental/GeometryProcessing/Source/GeometricObjects/Public/Polyline3.h
--- a/Engine/Plugins/Experimental/GeometryProcessing/Source/GeometricObjects/Public/Polyline3.h
+++ b/Engine/Plugins/Experimental/GeometryProcessing/Source/GeometricObjects/Public/Polyline3.h
@@ -115,6 +115,13 @@ public:
 	}
 
 
+	/** Discard all vertices of polyline */
+	void Clear()
+	{
+		Vertices.Reset();
+		Timestamp++;
+	}
+
 	/**
 	 * Add a vertex to the polyline
 	 */
@@ -214,7 +221,7 @@ public:
 	 * @param SegmentParam parameter in range [0,1] along segment
 	 * @return point on the segment at the given parameter value
 	 */
-	FVector3<T> PointAt(int SegmentIndex, double SegmentParam) const
+	FVector3<T> PointAt(int SegmentIndex, T SegmentParam) const
 	{
 		FSegment3<T> seg(Vertices[SegmentIndex], Vertices[SegmentIndex + 1]);
 		return seg.PointAt(SegmentParam);
@@ -242,9 +249,9 @@ public:
 	/**
 	 * @return the total perimeter length of the Polygon
 	 */
-	double Length() const
+	T Length() const
 	{
-		double length = 0;
+		T length = 0;
 		int N = SegmentCount();
 		for (int i = 0; i < N; ++i)
 		{
@@ -270,7 +277,6 @@ public:
 			check(Polyline != nullptr && i < SegmentCount());
 			return FSegment3<T>(Polyline->Vertices[i], Polyline->Vertices[i+1]);
 		}
-		//inline FSegment3<T> & operator*();
 		inline SegmentIterator & operator++() 		// prefix
 		{
 			i++;
@@ -330,18 +336,18 @@ public:
 	 * @param NearestSegParamOut the parameter value of the nearest point on the segment
 	 * @return squared distance to the polyline
 	 */
-	double DistanceSquared(const FVector3<T>& QueryPoint, int& NearestSegIndexOut, double& NearestSegParamOut) const
+	T DistanceSquared(const FVector3<T>& QueryPoint, int& NearestSegIndexOut, T& NearestSegParamOut) const
 	{
 		NearestSegIndexOut = -1;
-		NearestSegParamOut = TNumericLimits<double>::Max();
-		double dist = TNumericLimits<double>::Max();
+		NearestSegParamOut = TNumericLimits<T>::Max();
+		T dist = TNumericLimits<T>::Max();
 		int N = SegmentCount();
 		for (int vi = 0; vi < N; ++vi)
 		{
 			// @todo can't we just use segment function here now?
 			FSegment3<T> seg = FSegment3<T>(Vertices[vi], Vertices[vi+1]);
-			double t = (QueryPoint - seg.Center).Dot(seg.Direction);
-			double d = TNumericLimits<double>::Max();
+			T t = (QueryPoint - seg.Center).Dot(seg.Direction);
+			T d = TNumericLimits<T>::Max();
 			if (t >= seg.Extent)
 			{
 				d = seg.EndPoint().DistanceSquared(QueryPoint);
@@ -358,7 +364,7 @@ public:
 			{
 				dist = d;
 				NearestSegIndexOut = vi;
-				NearestSegParamOut = t;
+				NearestSegParamOut = TMathUtil<T>::Clamp(t, -seg.Extent, seg.Extent);
 			}
 		}
 		return dist;
@@ -371,9 +377,9 @@ public:
 	 * @param QueryPoint the query point
 	 * @return squared distance to the polyline
 	 */
-	double DistanceSquared(const FVector3<T>& QueryPoint) const
+	T DistanceSquared(const FVector3<T>& QueryPoint) const
 	{
-		int seg; double segt;
+		int seg; T segt;
 		return DistanceSquared(QueryPoint, seg, segt);
 	}
 
@@ -383,13 +389,13 @@ public:
 	/**
 	 * @return average edge length of all the edges of the Polygon
 	 */
-	double AverageEdgeLength() const
+	T AverageEdgeLength() const
 	{
-		double avg = 0; int N = Vertices.Num();
+		T avg = 0; int N = Vertices.Num();
 		for (int i = 1; i < N; ++i) {
 			avg += Vertices[i].Distance(Vertices[i - 1]);
 		}
-		return avg / (double)(N-1);
+		return avg / (T)(N-1);
 	}
 
 
