Add FGroupTopology: extracts topological graph of triangle group face/edge/corner elements, supports various queries
Add FGroupTopologyDeformer: base class for doing mesh deformations based on group topology. Solves for vertex locations when pulling on group face/edge/corner.
add FGeometrySet3: stores point and polycurve elements, supports nearest-point-on-ray queries on these geometries (work in progress)
#rnx

diff --git a/Engine/Plugins/Experimental/GeometryProcessing/Source/DynamicMesh/Private/GroupTopology.cpp b/Engine/Plugins/Experimental/GeometryProcessing/Source/DynamicMesh/Private/GroupTopology.cpp
new file mode 100644
--- /dev/null
+++ b/Engine/Plugins/Experimental/GeometryProcessing/Source/DynamicMesh/Private/GroupTopology.cpp
@@ -0,0 +1,479 @@
+// Copyright 1998-2019 Epic Games, Inc. All Rights Reserved.
+
+#include "GroupTopology.h"
+#include "MeshRegionBoundaryLoops.h"
+
+
+
+
+
+bool FGroupTopology::FGroupEdge::IsConnectedToVertices(const TSet<int>& Vertices) const
+{
+	for (int VertexID : Span.Vertices)
+	{
+		if (Vertices.Contains(VertexID))
+		{
+			return true;
+		}
+	}
+	return false;
+}
+
+
+
+FGroupTopology::FGroupTopology(const FDynamicMesh3* MeshIn, bool bAutoBuild)
+{
+	this->Mesh = MeshIn;
+	if (bAutoBuild)
+	{
+		RebuildTopology();
+	}
+}
+
+
+void FGroupTopology::RebuildTopology()
+{
+	Groups.Reset();
+	Edges.Reset();
+	Corners.Reset();
+
+	// initialize groups map first to avoid resizes
+	GroupIDToGroupIndexMap.Reset();
+	GroupIDToGroupIndexMap.Init(-1, Mesh->MaxGroupID());
+	TArray<int> GroupFaceCounts;
+	GroupFaceCounts.Init(0, Mesh->MaxGroupID());
+	for (int tid : Mesh->TriangleIndicesItr())
+	{
+		int GroupID = GetGroupID(tid);
+		if (GroupIDToGroupIndexMap[GroupID] == -1)
+		{
+			FGroup NewGroup;
+			NewGroup.GroupID = GroupID;
+			GroupIDToGroupIndexMap[GroupID] = Groups.Add(NewGroup);
+		}
+		GroupFaceCounts[GroupID]++;
+	}
+	for (FGroup& Group : Groups)
+	{
+		Group.Faces.Reserve(GroupFaceCounts[Group.GroupID]);
+	}
+
+
+	// sort faces into groups
+	for (int tid : Mesh->TriangleIndicesItr())
+	{
+		int GroupID = GetGroupID(tid);
+		Groups[GroupIDToGroupIndexMap[GroupID]].Faces.Add(tid);
+	}
+
+	// precompute junction vertices set
+	JunctionVertices.Init(false, Mesh->MaxVertexID());
+	for (int vid : Mesh->VertexIndicesItr())
+	{
+		if (IsJunctionVertex(vid))
+		{
+			JunctionVertices[vid] = true;
+			FCorner Corner = { vid };
+			Corners.Add(Corner);
+		}
+	}
+	for (FCorner& Corner : Corners)
+	{
+		Mesh->GetAllVertexGroups(Corner.VertexID, Corner.NeighbourGroupIDs);
+	}
+
+
+	// construct boundary loops
+	for (FGroup& Group : Groups)
+	{
+		ExtractGroupEdges(Group);
+
+		for (FGroupBoundary& Boundary : Group.Boundaries)
+		{
+			Boundary.bIsOnBoundary = false;
+			for (int EdgeIndex : Boundary.GroupEdges)
+			{
+				FGroupEdge& Edge = Edges[EdgeIndex];
+
+				int OtherGroupID = (Edge.Groups.A == Group.GroupID) ? Edge.Groups.B : Edge.Groups.A;
+				if (OtherGroupID != FDynamicMesh3::InvalidID)
+				{
+					Boundary.NeighbourGroupIDs.AddUnique(OtherGroupID);
+				}
+				else
+				{
+					Boundary.bIsOnBoundary = true;
+				}
+			}
+		}
+
+		// make all-neighbour-groups list at group level
+		for (FGroupBoundary& Boundary : Group.Boundaries)
+		{
+			for (int NbrGroupID : Boundary.NeighbourGroupIDs)
+			{
+				Group.NeighbourGroupIDs.AddUnique(NbrGroupID);
+			}
+		}
+	}
+
+}
+
+
+/** @return true if given mesh vertex is a Corner vertex */
+bool FGroupTopology::IsJunctionVertex(int VertexID) const
+{
+	FIndex3i UniqueGroups;
+	int UniqueCount = 0;
+	for (int tid : Mesh->VtxTrianglesItr(VertexID))
+	{
+		int GroupID = GetGroupID(tid);
+		if (UniqueCount == 0)
+		{
+			UniqueGroups[0] = GroupID;
+			UniqueCount++;
+		}
+		else if (UniqueCount == 1 && GroupID != UniqueGroups[0])
+		{
+			UniqueGroups[1] = GroupID;
+			UniqueCount++;
+		}
+		else if (UniqueCount == 2 && GroupID != UniqueGroups[0] && GroupID != UniqueGroups[1])
+		{
+			return true;
+		}
+	}
+	return false;
+}
+
+
+
+
+int FGroupTopology::GetCornerVertexID(int CornerID) const
+{
+	check(CornerID >= 0 && CornerID < Corners.Num());
+	return Corners[CornerID].VertexID;
+}
+
+
+const FGroupTopology::FGroup* FGroupTopology::FindGroupByID(int GroupID) const
+{
+	if (GroupID < 0 || GroupID >= GroupIDToGroupIndexMap.Num() || GroupIDToGroupIndexMap[GroupID] == -1)
+	{
+		return nullptr;
+	}
+	return &Groups[GroupIDToGroupIndexMap[GroupID]];
+}
+
+
+const TArray<int>& FGroupTopology::GetGroupFaces(int GroupID) const
+{
+	const FGroup* Found = FindGroupByID(GroupID);
+	ensure(Found != nullptr);
+	return (Found != nullptr) ? Found->Faces : EmptyArray;
+}
+
+const TArray<int>& FGroupTopology::GetGroupNbrGroups(int GroupID) const
+{
+	const FGroup* Found = FindGroupByID(GroupID);
+	ensure(Found != nullptr);
+	return (Found != nullptr) ? Found->NeighbourGroupIDs : EmptyArray;
+}
+
+
+
+int FGroupTopology::FindGroupEdgeID(int MeshEdgeID) const
+{
+	int GroupID = GetGroupID(Mesh->GetEdgeT(MeshEdgeID).A);
+	const FGroup* Group = FindGroupByID(GroupID);
+	ensure(Group != nullptr);
+	if (Group != nullptr)
+	{
+		for (const FGroupBoundary& Boundary : Group->Boundaries)
+		{
+			for (int EdgeID : Boundary.GroupEdges)
+			{
+				const FGroupEdge& Edge = Edges[EdgeID];
+				if (Edge.Span.Edges.Contains(MeshEdgeID))
+				{
+					return EdgeID;
+				}
+			}
+		}
+	}
+	return -1;
+}
+
+
+const TArray<int>& FGroupTopology::GetGroupEdgeVertices(int GroupEdgeID) const
+{
+	check(GroupEdgeID >= 0 && GroupEdgeID < Edges.Num());
+	return Edges[GroupEdgeID].Span.Vertices;
+}
+
+
+
+void FGroupTopology::FindEdgeNbrGroups(int GroupEdgeID, TArray<int>& GroupsOut) const
+{
+	check(GroupEdgeID >= 0 && GroupEdgeID < Edges.Num());
+	const TArray<int> & Vertices = GetGroupEdgeVertices(GroupEdgeID);
+	FindVertexNbrGroups(Vertices[0], GroupsOut);
+	FindVertexNbrGroups(Vertices[Vertices.Num() - 1], GroupsOut);
+}
+
+void FGroupTopology::FindEdgeNbrGroups(const TArray<int>& GroupEdgeIDs, TArray<int>& GroupsOut) const
+{
+	for (int GroupEdgeID : GroupEdgeIDs)
+	{
+		FindEdgeNbrGroups(GroupEdgeID, GroupsOut);
+	}
+}
+
+
+
+void FGroupTopology::FindCornerNbrGroups(int CornerID, TArray<int>& GroupsOut) const
+{
+	check(CornerID >= 0 && CornerID < Corners.Num());
+	for (int GroupID : Corners[CornerID].NeighbourGroupIDs)
+	{
+		GroupsOut.AddUnique(GroupID);
+	}
+}
+
+void FGroupTopology::FindCornerNbrGroups(const TArray<int>& CornerIDs, TArray<int>& GroupsOut) const
+{
+	for (int cid : CornerIDs)
+	{
+		FindCornerNbrGroups(cid, GroupsOut);
+	}
+}
+
+
+
+
+void FGroupTopology::FindVertexNbrGroups(int VertexID, TArray<int>& GroupsOut) const
+{
+	for (int tid : Mesh->VtxTrianglesItr(VertexID))
+	{
+		int GroupID = GetGroupID(tid);
+		GroupsOut.AddUnique(GroupID);
+	}
+}
+
+void FGroupTopology::FindVertexNbrGroups(const TArray<int>& VertexIDs, TArray<int>& GroupsOut) const
+{
+	for (int vid : VertexIDs)
+	{
+		for (int tid : Mesh->VtxTrianglesItr(vid))
+		{
+			int GroupID = GetGroupID(tid);
+			GroupsOut.AddUnique(GroupID);
+		}
+	}
+}
+
+
+
+void FGroupTopology::CollectGroupVertices(int GroupID, TSet<int>& Vertices) const
+{
+	const FGroup* Found = FindGroupByID(GroupID);
+	ensure(Found != nullptr);
+	if (Found != nullptr)
+	{
+		for (int TriID : Found->Faces)
+		{
+			FIndex3i TriVerts = Mesh->GetTriangle(TriID);
+			Vertices.Add(TriVerts.A);
+			Vertices.Add(TriVerts.B);
+			Vertices.Add(TriVerts.C);
+		}
+	}
+}
+
+
+
+void FGroupTopology::CollectGroupBoundaryVertices(int GroupID, TSet<int>& Vertices) const
+{
+	const FGroup* Group = FindGroupByID(GroupID);
+	ensure(Group != nullptr);
+	if (Group != nullptr)
+	{
+		for (const FGroupBoundary& Boundary : Group->Boundaries)
+		{
+			for (int EdgeIndex : Boundary.GroupEdges)
+			{
+				const FGroupEdge& Edge = Edges[EdgeIndex];
+				for (int vid : Edge.Span.Vertices)
+				{
+					Vertices.Add(vid);
+				}
+			}
+		}
+	}
+}
+
+
+
+
+
+void FGroupTopology::ForGroupEdges(int GroupID,
+	const TFunction<void(const FGroupEdge& Edge, int EdgeIndex)>& EdgeFunc) const
+{
+	const FGroup* Group = FindGroupByID(GroupID);
+	ensure(Group != nullptr);
+	if (Group != nullptr)
+	{
+		for (const FGroupBoundary& Boundary : Group->Boundaries)
+		{
+			for (int EdgeIndex : Boundary.GroupEdges)
+			{
+				EdgeFunc(Edges[EdgeIndex], EdgeIndex);
+			}
+		}
+	}
+}
+
+
+
+void FGroupTopology::ForGroupSetEdges(const TArray<int>& GroupIDs,
+	const TFunction<void(const FGroupEdge& Edge, int EdgeIndex)>& EdgeFunc) const
+{
+	TArray<int> DoneEdges;
+
+	for (int GroupID : GroupIDs)
+	{
+		const FGroup* Group = FindGroupByID(GroupID);
+		ensure(Group != nullptr);
+		if (Group != nullptr)
+		{
+			for (const FGroupBoundary& Boundary : Group->Boundaries)
+			{
+				for (int EdgeIndex : Boundary.GroupEdges)
+				{
+					if (DoneEdges.Contains(EdgeIndex) == false)
+					{
+						EdgeFunc(Edges[EdgeIndex], EdgeIndex);
+						DoneEdges.Add(EdgeIndex);
+					}
+				}
+			}
+		}
+	}
+}
+
+
+
+
+void FGroupTopology::ExtractGroupEdges(FGroup& Group)
+{
+	FMeshRegionBoundaryLoops BdryLoops(Mesh, Group.Faces, true);
+	int NumLoops = BdryLoops.Loops.Num();
+
+	Group.Boundaries.SetNum(NumLoops);
+	for ( int li = 0; li < NumLoops; ++li )
+	{
+		FEdgeLoop& Loop = BdryLoops.Loops[li];
+		FGroupBoundary& Boundary = Group.Boundaries[li];
+
+		// find indices of corners of group polygon
+		TArray<int> JunctionIndices;
+		int NumV = Loop.Vertices.Num();
+		for (int i = 0; i < NumV; ++i)
+		{
+			if (JunctionVertices[Loop.Vertices[i]])
+			{
+				JunctionIndices.Add(i);
+			}
+		}
+
+		// if we had no indices then this is like the cap of a cylinder, just one single long edge
+		if ( JunctionIndices.Num() == 0 )
+		{ 
+			FIndex2i EdgeID = MakeEdgeID(Loop.Edges[0]);
+			int OtherGroupID = (EdgeID.A == Group.GroupID) ? EdgeID.B : EdgeID.A;
+			int EdgeIndex = FindExistingGroupEdge(Group.GroupID, OtherGroupID, Loop.Vertices[0]);
+			if (EdgeIndex == -1)
+			{
+				FGroupEdge Edge = { EdgeID };
+				Edge.Span = FEdgeSpan(Mesh);
+				Edge.Span.InitializeFromEdges(Loop.Edges);
+				EdgeIndex = Edges.Add(Edge);
+			}
+			Boundary.GroupEdges.Add(EdgeIndex);
+			continue;
+		}
+
+		// duplicate first junction vertex so that we can just loop back around to it w/ modulo count
+		int NumSpans = JunctionIndices.Num();
+		int FirstIdx = JunctionIndices[0];
+		JunctionIndices.Add(FirstIdx);
+
+		// add each span
+		for (int k = 0; k < NumSpans; ++k)
+		{
+			int i0 = JunctionIndices[k];
+
+			FIndex2i EdgeID = MakeEdgeID(Loop.Edges[i0]);
+			int OtherGroupID = (EdgeID.A == Group.GroupID) ? EdgeID.B : EdgeID.A;
+			int EdgeIndex = FindExistingGroupEdge(Group.GroupID, OtherGroupID, Loop.Vertices[i0]);
+			if (EdgeIndex != -1)
+			{
+				FGroupEdge& Existing = Edges[EdgeIndex];
+				Boundary.GroupEdges.Add(EdgeIndex);
+				continue;
+			}
+
+			FGroupEdge Edge = { EdgeID };
+
+			int i1 = JunctionIndices[k+1];
+			check(i0 != i1);	// this case will happen on a cylinder w/ a cut side, I think...
+								// (but that case won't exist because that's wouldn't be a group boundary edge!)
+			TArray<int> SpanVertices;
+			while (i0 != i1)
+			{
+				SpanVertices.Add(Loop.Vertices[i0]);
+				i0 = (i0 + 1) % NumV;
+			}
+			SpanVertices.Add(Loop.Vertices[i0]);
+
+			Edge.Span = FEdgeSpan(Mesh);
+			Edge.Span.InitializeFromVertices(SpanVertices);
+			EdgeIndex = Edges.Add(Edge);
+			Boundary.GroupEdges.Add(EdgeIndex);
+		}
+	}
+}
+
+
+
+
+
+
+int FGroupTopology::FindExistingGroupEdge(int GroupID, int OtherGroupID, int FirstVertexID)
+{
+	if (OtherGroupID < 0)
+	{
+		return -1;
+	}
+
+	FGroup& OtherGroup = Groups[GroupIDToGroupIndexMap[OtherGroupID]];
+	FIndex2i EdgeID = MakeEdgeID(GroupID, OtherGroupID);
+
+	for (FGroupBoundary& Boundary : OtherGroup.Boundaries)
+	{
+		for (int EdgeIndex : Boundary.GroupEdges)
+		{
+			if (Edges[EdgeIndex].Groups == EdgeID)
+			{
+				TArray<int>& Vertices = Edges[EdgeIndex].Span.Vertices;
+				if (Vertices[0] == FirstVertexID || Vertices[Vertices.Num() - 1] == FirstVertexID)
+				{
+					return EdgeIndex;
+				}
+			}
+		}
+	}
+	return -1;
+}
+
+
diff --git a/Engine/Plugins/Experimental/GeometryProcessing/Source/DynamicMesh/Private/Operations/GroupTopologyDeformer.cpp b/Engine/Plugins/Experimental/GeometryProcessing/Source/DynamicMesh/Private/Operations/GroupTopologyDeformer.cpp
new file mode 100644
--- /dev/null
+++ b/Engine/Plugins/Experimental/GeometryProcessing/Source/DynamicMesh/Private/Operations/GroupTopologyDeformer.cpp
@@ -0,0 +1,309 @@
+// Copyright 1998-2019 Epic Games, Inc. All Rights Reserved.
+
+#include "Operations/GroupTopologyDeformer.h"
+#include "DynamicMesh3.h"
+#include "GroupTopology.h"
+#include "SegmentTypes.h"
+#include "Async/ParallelFor.h"
+#include "Containers/BitArray.h"
+
+
+void FGroupTopologyDeformer::Initialize(const FDynamicMesh3* MeshIn, const FGroupTopology* TopologyIn)
+{
+	Mesh = MeshIn;
+	Topology = TopologyIn;
+}
+
+
+
+void FGroupTopologyDeformer::Reset()
+{
+	HandleVertices.Reset();
+	HandleBoundaryVertices.Reset();
+	FixedBoundaryVertices.Reset();
+	ROIEdgeVertices.Reset();
+	ROIFaces.Reset();
+	InitialPositions.Reset();
+	ROIEdges.Reset();
+
+	EdgeEncodings.Reset();
+	FaceEncodings.Reset();
+}
+
+
+void FGroupTopologyDeformer::SetActiveHandleFaces(const TArray<int>& FaceGroupIDs)
+{
+	Reset();
+
+	check(FaceGroupIDs.Num() == 1);   // multi-face not supported yet
+	int GroupID = FaceGroupIDs[0];
+
+	// find set of vertices in handle 
+	Topology->CollectGroupVertices(GroupID, HandleVertices);
+	Topology->CollectGroupBoundaryVertices(GroupID, HandleBoundaryVertices);
+	ModifiedVertices = HandleVertices;
+
+	// find neighbour group set
+	TArray<int> HandleGroups = FaceGroupIDs;
+	const TArray<int>& GroupNbrGroups = Topology->GetGroupNbrGroups(GroupID);
+	CalculateROI(HandleGroups, GroupNbrGroups);
+
+	SaveInitialPositions();
+
+	ComputeEncoding();
+}
+
+
+
+void FGroupTopologyDeformer::SetActiveHandleEdges(const TArray<int>& TopologyEdgeIDs)
+{
+	Reset();
+
+	for (int EdgeID : TopologyEdgeIDs)
+	{
+		const TArray<int>& EdgeVerts = Topology->GetGroupEdgeVertices(EdgeID);
+		for (int VertID : EdgeVerts)
+		{
+			HandleVertices.Add(VertID);
+		}
+	}
+	HandleBoundaryVertices = HandleVertices;
+	ModifiedVertices = HandleVertices;
+
+	TArray<int> HandleGroups;
+	TArray<int> NbrGroups;
+	Topology->FindEdgeNbrGroups(TopologyEdgeIDs, NbrGroups);
+	CalculateROI(HandleGroups, NbrGroups);
+
+	SaveInitialPositions();
+
+	ComputeEncoding();
+}
+
+
+void FGroupTopologyDeformer::SetActiveHandleCorners(const TArray<int>& CornerIDs)
+{
+	Reset();
+
+	for (int CornerID : CornerIDs)
+	{
+		int VertID = Topology->GetCornerVertexID(CornerID);
+		if (VertID >= 0)
+		{
+			HandleVertices.Add(VertID);
+		}
+	}
+	HandleBoundaryVertices = HandleVertices;
+	ModifiedVertices = HandleVertices;
+
+	TArray<int> HandleGroups;
+	TArray<int> NbrGroups;
+	Topology->FindCornerNbrGroups(CornerIDs, NbrGroups);
+	CalculateROI(HandleGroups, NbrGroups);
+
+	SaveInitialPositions();
+
+	ComputeEncoding();
+}
+
+
+
+
+void FGroupTopologyDeformer::SaveInitialPositions()
+{
+	for (int VertID : ModifiedVertices)
+	{
+		InitialPositions.AddVertex(Mesh, VertID);
+	}
+}
+
+
+
+void FGroupTopologyDeformer::CalculateROI(const TArray<int>& HandleGroups, const TArray<int>& ROIGroups)
+{
+	// sort ROI edges into various sets
+	//   HandleBoundaryVertices: all vertices on border of handle
+	//   FixedBoundaryVertices: all vertices on fixed border (ie part of edges not connected to handle)
+	//   ROIEdges: list of edge data structures for edges we will deform
+	ROIEdges.Reserve(ROIGroups.Num() * 5);  // guesstimate
+	Topology->ForGroupSetEdges(ROIGroups, [this, &HandleGroups](const FGroupTopology::FGroupEdge& Edge, int EdgeIndex)
+	{
+		if (HandleGroups.Contains(Edge.Groups.A) || HandleGroups.Contains(Edge.Groups.B))
+		{
+			return;		// this is a Handle boundary edge
+		}
+
+		bool bIsConnectedToHandle = Edge.IsConnectedToVertices(HandleBoundaryVertices);
+		if (bIsConnectedToHandle)
+		{
+			for (int VertID : Edge.Span.Vertices)
+			{
+				ROIEdgeVertices.Add(VertID);
+			}
+			FROIEdge& ROIEdge = ROIEdges[ROIEdges.Add(FROIEdge())];
+			ROIEdge.EdgeIndex = EdgeIndex;
+			ROIEdge.Span = Edge.Span;
+			if (HandleBoundaryVertices.Contains(ROIEdge.Span.Vertices[0]) == false)
+			{
+				ROIEdge.Span.Reverse();
+			}
+		}
+		else
+		{
+			for (int VertID : Edge.Span.Vertices)
+			{
+				FixedBoundaryVertices.Add(VertID);
+			}
+		}
+	});
+
+
+	// create ROI faces
+	ROIFaces.SetNum(ROIGroups.Num());
+	int FaceIdx = 0;
+	for (int NbrGroupID : ROIGroups)
+	{
+		FaceVertsTemp.Reset();
+		FaceBoundaryVertsTemp.Reset();
+		Topology->CollectGroupVertices(NbrGroupID, FaceVertsTemp);
+		Topology->CollectGroupBoundaryVertices(NbrGroupID, FaceBoundaryVertsTemp);
+
+		FROIFace& Face = ROIFaces[FaceIdx++];
+		for (int vid : FaceVertsTemp)
+		{
+			TArray<int>& AddTo = (FaceBoundaryVertsTemp.Contains(vid)) ? Face.BoundaryVerts : Face.InteriorVerts;
+			AddTo.Add(vid);
+			ModifiedVertices.Add(vid);
+		}
+	}
+}
+
+
+
+
+
+
+void FGroupTopologyDeformer::ComputeEncoding()
+{
+	// encode ROI edges
+	int NumROIEdges = ROIEdges.Num();
+	EdgeEncodings.SetNum(NumROIEdges);
+	for (int ei = 0; ei < NumROIEdges; ei++)
+	{
+		FEdgeSpan& Span = ROIEdges[ei].Span;
+		int NumVerts = Span.Vertices.Num();
+
+		FEdgeEncoding& Encoding = EdgeEncodings[ei];
+		Encoding.Vertices.SetNum(NumVerts);
+
+		FVector3d StartV = Mesh->GetVertex(Span.Vertices[0]);
+		FVector3d EndV = Mesh->GetVertex(Span.Vertices[NumVerts - 1]);
+		FSegment3d Seg(StartV, EndV);
+
+		for (int k = 1; k < NumVerts - 1; ++k)
+		{
+			FVector3d Pos = Mesh->GetVertex(Span.Vertices[k]);
+			FEdgeVertexEncoding& Enc = Encoding.Vertices[k];
+			Enc.T = Seg.ProjectUnitRange(Pos);
+			//Enc.T *= Enc.T;
+			//Enc.T = FMathd::Sqrt(Enc.T);
+			Enc.Delta = Pos - Seg.PointBetween(Enc.T);
+		}
+	}
+
+	// encode ROI faces
+	int NumROIFaces = ROIFaces.Num();
+	FaceEncodings.SetNum(NumROIFaces);
+	for (int fi = 0; fi < NumROIFaces; ++fi)
+	{
+		FROIFace& Face = ROIFaces[fi];
+		int NumBoundaryV = Face.BoundaryVerts.Num();
+		int NumInteriorV = Face.InteriorVerts.Num();
+		FFaceEncoding& Encoding = FaceEncodings[fi];
+		Encoding.Vertices.SetNum(NumInteriorV);
+		for (int vi = 0; vi < NumInteriorV; ++vi)
+		{
+			FVector3d Pos = Mesh->GetVertex(Face.InteriorVerts[vi]);
+			FFaceVertexEncoding& VtxEncoding = Encoding.Vertices[vi];
+			VtxEncoding.Weights.SetNum(NumBoundaryV);
+			VtxEncoding.Deltas.SetNum(NumBoundaryV);
+			double WeightSum = 0;
+			for (int k = 0; k < NumBoundaryV; ++k)
+			{
+				FVector3d BorderPos = Mesh->GetVertex(Face.BoundaryVerts[k]);
+				FVector3d DeltaVec = Pos - BorderPos;
+				double Weight = 1.0 / DeltaVec.SquaredLength();
+				VtxEncoding.Deltas[k] = DeltaVec;
+				VtxEncoding.Weights[k] = Weight;
+				WeightSum += Weight;
+			}
+			FVector3d Reconstruct(0, 0, 0);
+			for (int k = 0; k < NumBoundaryV; ++k)
+			{
+				VtxEncoding.Weights[k] /= WeightSum;
+				Reconstruct += VtxEncoding.Weights[k] * (Mesh->GetVertex(Face.BoundaryVerts[k]) + VtxEncoding.Deltas[k]);
+			}
+			check(Reconstruct.Distance(Pos) < 0.0001);
+		}
+	}
+}
+
+
+
+void FGroupTopologyDeformer::ClearSolution(FDynamicMesh3* TargetMesh)
+{
+	InitialPositions.SetPositions(TargetMesh);
+}
+
+
+void FGroupTopologyDeformer::UpdateSolution(FDynamicMesh3* TargetMesh, const TFunction<FVector3d(FDynamicMesh3* Mesh, int)>& HandleVertexDeformFunc)
+{
+	InitialPositions.SetPositions(TargetMesh);
+
+	for (int VertIdx : HandleVertices)
+	{
+		FVector3d DeformPos = HandleVertexDeformFunc(TargetMesh, VertIdx);
+		TargetMesh->SetVertex(VertIdx, DeformPos);
+	}
+
+	// reconstruct edges
+	int NumEdges = ROIEdges.Num();
+	for (int ei = 0; ei < NumEdges; ++ei)
+	{
+		const FEdgeSpan& Span = ROIEdges[ei].Span;
+		const FEdgeEncoding& Encoding = EdgeEncodings[ei];
+		int NumVerts = Span.Vertices.Num();
+		FVector3d A = TargetMesh->GetVertex(Span.Vertices[0]);
+		FVector3d B = TargetMesh->GetVertex(Span.Vertices[NumVerts - 1]);
+		FSegment3d Seg(A, B);
+		for (int k = 1; k < NumVerts - 1; ++k)
+		{
+			FVector3d NewPos = Seg.PointBetween(Encoding.Vertices[k].T);
+			NewPos += Encoding.Vertices[k].Delta;
+			TargetMesh->SetVertex(Span.Vertices[k], NewPos);
+		}
+	}
+
+	// reconstruct faces
+	int NumFaces = ROIFaces.Num();
+	ParallelFor(NumFaces, [this, &TargetMesh](int FaceIdx) {
+		const FROIFace& Face = ROIFaces[FaceIdx];
+		const FFaceEncoding& Encoding = FaceEncodings[FaceIdx];
+		int NumBorder = Face.BoundaryVerts.Num();
+		int NumVerts = Face.InteriorVerts.Num();
+		for (int vi = 0; vi < NumVerts; ++vi)
+		{
+			const FFaceVertexEncoding& VtxEncoding = Encoding.Vertices[vi];
+			FVector3d Sum(0, 0, 0);
+			for (int k = 0; k < NumBorder; ++k)
+			{
+				Sum += VtxEncoding.Weights[k] *
+					(TargetMesh->GetVertex(Face.BoundaryVerts[k]) + VtxEncoding.Deltas[k]);
+			}
+			TargetMesh->SetVertex(Face.InteriorVerts[vi], Sum);
+		}
+	});
+}
+
+
+
diff --git a/Engine/Plugins/Experimental/GeometryProcessing/Source/DynamicMesh/Public/GroupTopology.h b/Engine/Plugins/Experimental/GeometryProcessing/Source/DynamicMesh/Public/GroupTopology.h
new file mode 100644
--- /dev/null
+++ b/Engine/Plugins/Experimental/GeometryProcessing/Source/DynamicMesh/Public/GroupTopology.h
@@ -0,0 +1,191 @@
+// Copyright 1998-2019 Epic Games, Inc. All Rights Reserved.
+
+#pragma once
+
+#include "DynamicMesh3.h"
+#include "EdgeSpan.h"
+#include "Util/IndexUtil.h"
+#include "Containers/BitArray.h"
+
+
+/**
+ * Given a per-triangle integer ("group"), FGroupTopology extracts a
+ * group-level topological graph from an input Mesh. The graph consists
+ * of three elements:
+ *   Corners: there is a corner at each vertex where 3 or more groups meet (ie 3 or more groups in one-ring)
+ *   Edges: a group edge is a list of one or more connected edges that have the same pair of different groups on either side
+ *   Group: a group is a set of connected faces with the same GroupID
+ *   
+ * By default, the triangle group attribute of the input Mesh is used.
+ * You can override GetGroupID to provide your own grouping.
+ * 
+ * Various query functions are provided to allow group topology to be interrogated.
+ * Note that these functions refer to "GroupID", "CornerID", and "GroupEdgeID", 
+ * these are simply indices into the internal .Groups, .Corners, and .Edges arrays
+ */
+class DYNAMICMESH_API FGroupTopology
+{
+public:
+	FGroupTopology() {}
+	FGroupTopology(const FDynamicMesh3* Mesh, bool bAutoBuild);
+
+	virtual ~FGroupTopology() {}
+
+
+	/**
+	 * Build the group topology graph.
+	 */
+	void RebuildTopology();
+
+	/** 
+	 * Adjacency of Per-Triangle integers are what define the triangle groups.
+	 * Override this function to provide an alternate group definition.
+	 * @return group id integer for given TriangleID 
+	 */
+	virtual int GetGroupID(int TriangleID) const
+	{
+		return Mesh->GetTriangleGroup(TriangleID);
+	}
+
+
+	/**
+	 * FCorner is a "corner" in the group topology, IE a vertex where 3+ groups meet.
+	 */
+	struct DYNAMICMESH_API FCorner
+	{
+		/** Mesh Vertex ID */
+		int VertexID;		
+		/** List of IDs of groups connected to this Corner */
+		TArray<int> NeighbourGroupIDs;
+	};
+
+	/** List of Corners in the topology graph (ie these are the nodes/vertices of the graph) */
+	TArray<FCorner> Corners;
+
+	/** A Group is bounded by closed loops of FGroupEdge elements. A FGroupBoundary is one such loop. */
+	struct DYNAMICMESH_API FGroupBoundary
+	{
+		/** Ordered list of edges forming this boundary */
+		TArray<int> GroupEdges;
+		/** List of IDs of groups on the "other side" of this boundary (this GroupBoundary is owned by a particular FGroup) */
+		TArray<int> NeighbourGroupIDs;
+		/** true if one or more edges in GroupEdges is on the mesh boundary */
+		bool bIsOnBoundary;
+	};
+
+	/** FGroup is a set of connected triangles with the same GroupID */
+	struct DYNAMICMESH_API FGroup
+	{
+		/** GroupID for this group */
+		int GroupID;
+		/** List of triangles forming this group */
+		TArray<int> Faces;
+		
+		/** List of boundaries of this group (may be empty, eg on a closed component with only one group) */
+		TArray<FGroupBoundary> Boundaries;
+		/** List of groups that are adjacent to this group */
+		TArray<int> NeighbourGroupIDs;
+	};
+
+	/** List of Groups in the topology graph (ie faces) */
+	TArray<FGroup> Groups;
+
+
+	/**
+	 * FGroupEdge is a sequence of group-boundary-edges where the two groups on either side of each edge
+	 * are the same. The sequence is stored an FEdgeSpan.
+	 * 
+	 * FGroupEdge instances are *shared* between the two FGroups on either side, via FGroupBoundary.
+	 */
+	struct DYNAMICMESH_API FGroupEdge
+	{
+		/** Identifier for this edge, as a pair of groups, sorted by increasing value */
+		FIndex2i Groups;
+		/** Edge span for this edge */
+		FEdgeSpan Span;
+
+		/** @return the member of .Groups that is not GroupID */
+		int OtherGroupID(int GroupID) const 
+		{ 
+			check(Groups.A == GroupID || Groups.B == GroupID);
+			return (Groups.A == GroupID) ? Groups.B : Groups.A; 
+		}
+
+		/** @return true if any vertex in the Span is in the Vertices set*/
+		bool IsConnectedToVertices(const TSet<int>& Vertices) const;
+	};
+
+	/** List of Edges in the topology graph (each edge connects two corners/nodes) */
+	TArray<FGroupEdge> Edges;
+
+
+
+	/** @return the mesh vertex ID for the given Corner ID */
+	int GetCornerVertexID(int CornerID) const;
+
+	/** @return the FGroup for the given GroupID, or nullptr if not found */
+	const FGroup* FindGroupByID(int GroupID) const;
+	/** @return the list of triangles in the given GroupID, or empty list if not found */
+	const TArray<int>& GetGroupFaces(int GroupID) const;
+	/** @return the list of neighbour GroupIDs for the given GroupID, or empty list if not found */
+	const TArray<int>& GetGroupNbrGroups(int GroupID) const;
+
+	/** @return the ID of the FGroupEdge that contains the given Mesh Edge ID*/
+	int FindGroupEdgeID(int MeshEdgeID) const;
+	/** @return the list of vertices of a FGroupEdge identified by the GroupEdgeID */
+	const TArray<int>& GetGroupEdgeVertices(int GroupEdgeID) const;
+
+	/** Add the groups connected to the given GroupEdgeID to the GroupsOut list. This is not the either-side pair, but the set of groups on the one-ring of each connected corner. */
+	void FindEdgeNbrGroups(int GroupEdgeID, TArray<int>& GroupsOut) const;
+	/** Add the groups connected to all the GroupEdgeIDs to the GroupsOut list. This is not the either-side pair, but the set of groups on the one-ring of each connected corner. */
+	void FindEdgeNbrGroups(const TArray<int>& GroupEdgeIDs, TArray<int>& GroupsOut) const;
+
+	/** Add all the groups connected to the given Corner to the GroupsOut list */
+	void FindCornerNbrGroups(int CornerID, TArray<int>& GroupsOut) const;
+	/** Add all the groups connected to the given Corners to the GroupsOut list */
+	void FindCornerNbrGroups(const TArray<int>& CornerIDs, TArray<int>& GroupsOut) const;
+
+	/** Add all the groups connected to the given Mesh Vertex to the GroupsOut list */
+	void FindVertexNbrGroups(int VertexID, TArray<int>& GroupsOut ) const;
+	/** Add all the groups connected to the given Mesh Vertices to the GroupsOut list */
+	void FindVertexNbrGroups(const TArray<int>& VertexIDs, TArray<int>& GroupsOut) const;
+
+	/** Call EdgeFunc for each boundary edge of the given Group  (no order defined) */
+	void ForGroupEdges(int GroupID, 
+		const TFunction<void(const FGroupEdge& Edge, int EdgeIndex)>& EdgeFunc) const;
+
+	/** Call EdgeFunc for each boundary edge of each of the given Groups (no order defined) */
+	void ForGroupSetEdges(const TArray<int>& GroupIDs,
+		const TFunction<void(const FGroupEdge& Edge, int EdgeIndex)>& EdgeFunc) const;
+
+	/** Add all the vertices of the given GroupID to the Vertices set */
+	void CollectGroupVertices(int GroupID, TSet<int>& Vertices) const;
+	/** Add all the group boundary vertices of the given GroupID to the Vertices set */
+	void CollectGroupBoundaryVertices(int GroupID, TSet<int>& Vertices) const;
+
+
+protected:
+	const FDynamicMesh3* Mesh = nullptr;
+
+	TArray<int> GroupIDToGroupIndexMap;		// allow fast lookup of index in .Groups, given GroupID
+	TBitArray<> JunctionVertices;		// bit array of corners for fast testing in ExtractGroupEdges
+	TArray<int> EmptyArray;
+
+	/** @return true if given mesh vertex is a Corner vertex */
+	virtual bool IsJunctionVertex(int VertexID) const;
+
+	void ExtractGroupEdges(FGroup& Group);
+
+	FIndex2i MakeEdgeID(int MeshEdgeID)
+	{
+		FIndex2i EdgeTris = Mesh->GetEdgeT(MeshEdgeID);
+		return MakeEdgeID(GetGroupID(EdgeTris.A), GetGroupID(EdgeTris.B));
+	}
+	FIndex2i MakeEdgeID(int Group1, int Group2)
+	{
+		check(Group1 != Group2);
+		return (Group1 < Group2) ? FIndex2i(Group1, Group2) : FIndex2i(Group2, Group1);
+	}
+
+	int FindExistingGroupEdge(int GroupID, int OtherGroupID, int FirstVertexID);
+};
\ No newline at end of file
diff --git a/Engine/Plugins/Experimental/GeometryProcessing/Source/DynamicMesh/Public/Operations/GroupTopologyDeformer.h b/Engine/Plugins/Experimental/GeometryProcessing/Source/DynamicMesh/Public/Operations/GroupTopologyDeformer.h
new file mode 100644
--- /dev/null
+++ b/Engine/Plugins/Experimental/GeometryProcessing/Source/DynamicMesh/Public/Operations/GroupTopologyDeformer.h
@@ -0,0 +1,199 @@
+// Copyright 1998-2019 Epic Games, Inc. All Rights Reserved.
+
+#pragma once
+
+#include "EdgeSpan.h"
+
+class FDynamicMesh3;
+class FGroupTopology;
+
+
+
+/**
+ * Basic cache of vertex positions.
+ */
+class DYNAMICMESH_API FVertexPositionCache
+{
+public:
+	TArray<int> Vertices;
+	TArray<FVector3d> Positions;
+
+	void Reset()
+	{
+		Vertices.Reset();
+		Positions.Reset();
+	}
+
+	/** Save this vertex. Does not check that VertexID hasn't already been added. */
+	void AddVertex(const FDynamicMesh3* Mesh, int VertexID)
+	{
+		Vertices.Add(VertexID);
+		Positions.Add(Mesh->GetVertex(VertexID));
+	}
+
+	/** Apply saved positions to Mesh */
+	void SetPositions(FDynamicMesh3* Mesh)
+	{
+		int Num = Vertices.Num();
+		for (int k = 0; k < Num; ++k)
+		{
+			Mesh->SetVertex(Vertices[k], Positions[k]);
+		}
+	}
+};
+
+
+
+/**
+ * FGroupTopologyDeformer supports deforming a Mesh based on an overlaid FGroupTopology.
+ * First the client defines a set of "Handle" elements (Faces/Corners/Edges) using SetActiveHandleX(). 
+ * The client will provide new vertex positions for the vertices of these elements via the 
+ * HandleVertexDeformFunc argument to UpdateSolution(). Once the Handle vertices have been
+ * updated, the deformer solves for updated vertex positions in the GroupTopology Faces
+ * that are adjacent to the handles. This region is referred to as the "ROI" (Region-of-Interest)
+ * 
+ * The default deformation is to first solve for the updated edges, and then
+ * solve for updated faces. This is done via linear encoding of the edge and face vertices
+ * relative to their boundaries (edge boundary is endpoint corners, face boundary is edges).
+ * 
+ * Various functions can be overrided to customize behavior.
+ * 
+ */
+class DYNAMICMESH_API FGroupTopologyDeformer
+{
+public:
+	virtual ~FGroupTopologyDeformer() {}
+
+	/** Set the Mesh and Topology to use for the deformation */
+	void Initialize(const FDynamicMesh3* Mesh, const FGroupTopology* Topology);
+	
+	const FDynamicMesh3* GetMesh() const { return Mesh; }
+	const FGroupTopology* GetTopology() const { return Topology; }
+
+	//
+	// Handle setup/configuraiton
+	// 
+
+	/**
+	 * Set the active handle to the given Faces
+	 */
+	virtual void SetActiveHandleFaces(const TArray<int>& FaceGroupIDs);
+
+	/**
+	 * Set the active handle to the given Corners
+	 */
+	virtual void SetActiveHandleCorners(const TArray<int>& TopologyCornerIDs);
+
+	/**
+	 * Set the active handle to the given Edges
+	 */
+	virtual void SetActiveHandleEdges(const TArray<int>& TopologyEdgeIDs);
+
+	//
+	// Solving
+	// 
+
+	/**
+	 * Update TargetMesh by first calling HandleVertexDeformFunc() to get new
+	 * handle vertex positions, then solving for new ROI vertex positions.
+	 */
+	virtual void UpdateSolution(FDynamicMesh3* TargetMesh, const TFunction<FVector3d(FDynamicMesh3* Mesh, int)>& HandleVertexDeformFunc);
+
+	/**
+	 * Restore the Handle and ROI vertex positions to their initial state
+	 */
+	virtual void ClearSolution(FDynamicMesh3* TargetMesh);
+
+
+
+	/** @return the set of handle vertices */
+	const TSet<int>& GetHandleVertices() const { return HandleVertices; }
+	/** @return the set of all vertices that will be modified by the deformation*/
+	const TSet<int>& GetModifiedVertices() const { return ModifiedVertices; }
+
+
+protected:
+	//
+	// These are functions that subclasses may wish to override
+	// 
+
+	/** reset all internal data structures, eg when changing handle */
+	virtual void Reset();
+
+	/**
+	 * Populate the internal ROI data structures based on the given HandleGroups and ROIGroups.
+	 * HandleGroups will be empty if the Handle is a set of Vertices or Edges.
+	 */
+	virtual void CalculateROI(const TArray<int>& HandleGroups, const TArray<int>& ROIGroups);
+
+	/**
+	 * Save the positions of all vertices that will be modified
+	 */
+	virtual void SaveInitialPositions();
+
+	/**
+	 * Precompute the representation of the ROI vertices at the initial positions.
+	 */
+	virtual void ComputeEncoding();
+
+protected:
+	const FDynamicMesh3* Mesh = nullptr;
+	const FGroupTopology* Topology = nullptr;
+
+	FVertexPositionCache InitialPositions;
+	TSet<int> ModifiedVertices;
+	TSet<int> HandleVertices;
+	TSet<int> HandleBoundaryVertices;
+	TSet<int> FixedBoundaryVertices;
+	TSet<int> ROIEdgeVertices;
+	TSet<int> FaceVertsTemp;
+	TSet<int> FaceBoundaryVertsTemp;
+
+	struct FROIEdge
+	{
+		int EdgeIndex;
+		FEdgeSpan Span;
+	};
+	TArray<FROIEdge> ROIEdges;
+
+	struct FROIFace
+	{
+		TArray<int> BoundaryVerts;
+		TArray<int> InteriorVerts;
+	};
+	TArray<FROIFace> ROIFaces;
+
+
+	//
+	// Deformation strategy (should be in subclass?)
+	// 
+
+	struct FEdgeVertexEncoding
+	{
+		double T;
+		FVector3d Delta;
+		FEdgeVertexEncoding() { T = 0; Delta = FVector3d::Zero(); }
+	};
+	struct FEdgeEncoding
+	{
+		TArray<FEdgeVertexEncoding> Vertices;
+	};
+	TArray<FEdgeEncoding> EdgeEncodings;
+
+
+
+	struct FFaceVertexEncoding
+	{
+		TArray<double> Weights;
+		TArray<FVector3d> Deltas;
+	};
+	struct FFaceEncoding
+	{
+		TArray<FFaceVertexEncoding> Vertices;
+	};
+	TArray<FFaceEncoding> FaceEncodings;
+
+
+
+
+};
diff --git a/Engine/Plugins/Experimental/GeometryProcessing/Source/GeometricObjects/Private/Spatial/GeometrySet3.cpp b/Engine/Plugins/Experimental/GeometryProcessing/Source/GeometricObjects/Private/Spatial/GeometrySet3.cpp
new file mode 100644
--- /dev/null
+++ b/Engine/Plugins/Experimental/GeometryProcessing/Source/GeometricObjects/Private/Spatial/GeometrySet3.cpp
@@ -0,0 +1,169 @@
+// Copyright 1998-2019 Epic Games, Inc. All Rights Reserved.
+
+#include "Spatial/GeometrySet3.h"
+#include "Distance/DistRay3Segment3.h"
+#include "Async/ParallelFor.h"
+
+
+void FGeometrySet3::Reset(bool bPoints, bool bCurves)
+{
+	if (bPoints)
+	{
+		Points.Reset();
+	}
+	if (bCurves)
+	{
+		Curves.Reset();
+	}
+}
+
+
+
+void FGeometrySet3::AddPoint(int PointID, const FVector3d& Point)
+{
+	check(PointIDToIndex.Contains(PointID) == false);
+	FPoint NewPoint = { PointID, Point };
+	int NewIndex = Points.Add(NewPoint);
+	PointIDToIndex.Add(PointID, NewIndex);
+}
+
+void FGeometrySet3::AddCurve(int CurveID, const FPolyline3d& Polyline)
+{
+	check(CurveIDToIndex.Contains(CurveID) == false);
+	int NewIndex = Curves.Add(FCurve());
+	FCurve& NewCurve = Curves[NewIndex];
+	NewCurve.ID = CurveID;
+	NewCurve.Geometry = Polyline;
+	NewCurve.Bounds = Polyline.GetBounds();
+	CurveIDToIndex.Add(CurveID, NewIndex);
+}
+
+
+void FGeometrySet3::UpdatePoint(int PointID, const FVector3d& Point)
+{
+	const int* Index = PointIDToIndex.Find(PointID);
+	check(Index != nullptr);
+	Points[*Index].Position = Point;
+}
+
+void FGeometrySet3::UpdateCurve(int CurveID, const FPolyline3d& Polyline)
+{
+	const int* Index = CurveIDToIndex.Find(CurveID);
+	check(Index != nullptr);
+	Curves[*Index].Geometry.SetVertices(Polyline.GetVertices());
+	Curves[*Index].Bounds = Polyline.GetBounds();
+}
+
+
+
+bool FGeometrySet3::FindNearestPointToRay(const FRay3d& Ray, FNearest& ResultOut,
+	TFunction<bool(const FVector3d&, const FVector3d&)> PointWithinToleranceTest) const
+{
+	double MinRayParamT = TNumericLimits<double>::Max();
+	int NearestID = -1;
+	int NearestIndex = -1;
+
+	FCriticalSection Critical;
+
+	int NumPoints = Points.Num();
+	ParallelFor(NumPoints, [&](int pi)
+	{
+		const FPoint& Point = Points[pi];
+		double RayT = Ray.Project(Point.Position);
+		FVector3d RayPt = Ray.NearestPoint(Point.Position);
+		if (PointWithinToleranceTest(RayPt, Point.Position))
+		{
+			Critical.Lock();
+			if (RayT < MinRayParamT)
+			{
+				MinRayParamT = RayT;
+				NearestIndex = pi;
+				NearestID = Points[pi].ID;
+			}
+			Critical.Unlock();
+		}
+	});
+
+	if (NearestID != -1)
+	{
+		ResultOut.ID = NearestID;
+		ResultOut.bIsPoint = true;
+		ResultOut.NearestRayPoint = Ray.PointAt(MinRayParamT);
+		ResultOut.NearestGeoPoint = Points[NearestIndex].Position;
+		ResultOut.RayParam = MinRayParamT;
+	}
+
+	return NearestID != -1;
+}
+
+
+
+
+
+bool FGeometrySet3::FindNearestCurveToRay(const FRay3d& Ray, FNearest& ResultOut,
+	TFunction<bool(const FVector3d&, const FVector3d&)> PointWithinToleranceTest) const
+{
+	double MinRayParamT = TNumericLimits<double>::Max();
+	int NearestID = -1;
+	int NearestIndex = -1;
+	int NearestSegmentIdx = -1;
+	double NearestSegmentParam = 0;
+
+	FCriticalSection Critical;
+
+	int NumCurves = Curves.Num();
+	ParallelFor(NumCurves, [&](int ci)
+	{
+		const FCurve& Curve = Curves[ci];
+		const FPolyline3d& Polyline = Curve.Geometry;
+		int NumSegments = Polyline.SegmentCount();
+
+		double CurveMinRayParamT = TNumericLimits<double>::Max();
+		int CurveNearestSegmentIdx = -1;
+		double CurveNearestSegmentParam = 0;
+		for (int si = 0; si < NumSegments; ++si)
+		{
+			double SegRayParam; double SegSegParam;
+			double SegDistSqr = FDistRay3Segment3d::SquaredDistance(Ray, Polyline.GetSegment(si), SegRayParam, SegSegParam);
+			if (SegRayParam < CurveMinRayParamT)
+			{
+				FVector3d RayPosition = Ray.PointAt(SegRayParam);
+				FVector3d CurvePosition = Polyline.PointAt(si, SegSegParam);
+				if (PointWithinToleranceTest(RayPosition, CurvePosition))
+				{
+					CurveMinRayParamT = SegRayParam;
+					CurveNearestSegmentIdx = si;
+					CurveNearestSegmentParam = SegSegParam;
+				}
+			}
+		}
+
+		// if we found a possible point, lock outer structures and merge in
+		if (CurveMinRayParamT < TNumericLimits<double>::Max())
+		{
+			Critical.Lock();
+			if (CurveMinRayParamT < MinRayParamT)
+			{
+				MinRayParamT = CurveMinRayParamT;
+				NearestIndex = ci;
+				NearestID = Curve.ID;
+				NearestSegmentIdx = CurveNearestSegmentIdx;
+				NearestSegmentParam = CurveNearestSegmentParam;
+			}
+			Critical.Unlock();
+		}
+	});   // end ParallelFor
+
+	if (NearestID != -1)
+	{
+		ResultOut.ID = NearestID;
+		ResultOut.bIsPoint = false;
+		ResultOut.NearestRayPoint = Ray.PointAt(MinRayParamT);
+		ResultOut.NearestGeoPoint = Curves[NearestIndex].Geometry.PointAt(NearestSegmentIdx, NearestSegmentParam);
+		ResultOut.RayParam = MinRayParamT;
+		ResultOut.PolySegmentIdx = NearestSegmentIdx;
+		ResultOut.PolySegmentParam = NearestSegmentParam;
+	}
+
+	return NearestID != -1;
+}
diff --git a/Engine/Plugins/Experimental/GeometryProcessing/Source/GeometricObjects/Public/Spatial/GeometrySet3.h b/Engine/Plugins/Experimental/GeometryProcessing/Source/GeometricObjects/Public/Spatial/GeometrySet3.h
new file mode 100644
--- /dev/null
+++ b/Engine/Plugins/Experimental/GeometryProcessing/Source/GeometricObjects/Public/Spatial/GeometrySet3.h
@@ -0,0 +1,98 @@
+// Copyright 1998-2019 Epic Games, Inc. All Rights Reserved.
+
+#pragma once
+
+#include "BoxTypes.h"
+#include "RayTypes.h"
+#include "Polyline3.h"
+
+/**
+ * FGeometrySet3 stores a set of 3D Points and Polyline curves,
+ * and supports spatial queries against these sets. 
+ * 
+ * Since Points and Curves have no area to hit, hit-tests are done via nearest-point-on-ray.
+ */
+class GEOMETRICOBJECTS_API FGeometrySet3
+{
+public:
+	/**
+	 * @param bPoints if true, discard all points
+	 * @param bCurves if true, discard all polycurves
+	 */
+	void Reset(bool bPoints = true, bool bCurves = true);
+
+	/** Add a point with given PointID at the given Position*/
+	void AddPoint(int PointID, const FVector3d& Position);
+	/** Add a polycurve with given CurveID and the give Polyline */
+	void AddCurve(int CurveID, const FPolyline3d& Polyline);
+
+	/** Update the Position of previously-added PointID */
+	void UpdatePoint(int PointID, const FVector3d& Position);
+	/** Update the Polyline of previously-added CurveID */
+	void UpdateCurve(int CurveID, const FPolyline3d& Polyline);
+	
+	/**
+	 * FNearest is returned by nearest-point queries
+	 */
+	struct GEOMETRICOBJECTS_API FNearest
+	{
+		/** ID of point or curve */
+		int ID;
+		/** true for point, false for polyline curve*/
+		bool bIsPoint;
+
+		/** Nearest point on ray */
+		FVector3d NearestRayPoint;
+		/** Nearest point on geometry (ie the point, or point on curve)*/
+		FVector3d NearestGeoPoint;
+
+		/** parameter of nearest point on ray (equivalent to NearestRayPoint) */
+		double RayParam;
+
+		/** if bIsPoint=false, index of nearest segment on polyline curve */
+		int PolySegmentIdx;
+		/** if bIsPoint=false, parameter of NearestGeoPoint along segment defined by PolySegmentIdx */
+		double PolySegmentParam;
+	};
+
+
+	/**
+	 * @param Ray query ray
+	 * @param ResultOut populated with information about successful nearest point result
+	 * @param PointWithinToleranceTest should return true if two 3D points are "close enough" to be considered a hit
+	 * @return true if the nearest point on Ray to some point in the set passed the PointWithinToleranceTest.
+	 */
+	bool FindNearestPointToRay(const FRay3d& Ray, FNearest& ResultOut,
+		TFunction<bool(const FVector3d&, const FVector3d&)> PointWithinToleranceTest) const;
+
+	/**
+	 * @param Ray query ray
+	 * @param ResultOut populated with information about successful nearest point result
+	 * @param PointWithinToleranceTest should return true if two 3D points are "close enough" to be considered a hit
+	 * @return true if the nearest point on Ray to some curve in the set passed the PointWithinToleranceTest.
+	 */
+	bool FindNearestCurveToRay(const FRay3d& Ray, FNearest& ResultOut,
+		TFunction<bool(const FVector3d&, const FVector3d&)> PointWithinToleranceTest) const;
+
+
+protected:
+
+	struct FPoint
+	{
+		int ID;
+		FVector3d Position;
+	};
+	TArray<FPoint> Points;
+	TMap<int, int> PointIDToIndex;
+
+	struct FCurve
+	{
+		int ID;
+		FPolyline3d Geometry;
+		
+		FAxisAlignedBox3d Bounds;
+	};
+	TArray<FCurve> Curves;
+	TMap<int, int> CurveIDToIndex;
+
+};
\ No newline at end of file
