Add constructor for arrangment2d to just directly specify cellsize, to simplify safety logic of planarcuts code

diff --git a/Engine/Plugins/Experimental/GeometryProcessing/Source/GeometryAlgorithms/Public/Arrangement2d.h b/Engine/Plugins/Experimental/GeometryProcessing/Source/GeometryAlgorithms/Public/Arrangement2d.h
--- a/Engine/Plugins/Experimental/GeometryProcessing/Source/GeometryAlgorithms/Public/Arrangement2d.h
+++ b/Engine/Plugins/Experimental/GeometryProcessing/Source/GeometryAlgorithms/Public/Arrangement2d.h
@@ -43,6 +43,11 @@ struct FArrangement2d
 	{
 	}
 
+	FArrangement2d(double PointHashCellSize)
+		: PointHash(PointHashCellSize, -1)
+	{
+	}
+
 	/**
 	 * Attempts to triangulates the arrangement with a constrained delaunay triangulation
 	 * NOTE: Not robust; generates invalid triangulations and fails to insert edges sometimes, even if the arrangement has no self-intersections
