change current FIndex2i::Max usages to FIndex2i::Invalid

diff --git a/Engine/Plugins/Experimental/GeometryProcessing/Source/DynamicMesh/Private/MeshBoundaryLoops.cpp b/Engine/Plugins/Experimental/GeometryProcessing/Source/DynamicMesh/Private/MeshBoundaryLoops.cpp
--- a/Engine/Plugins/Experimental/GeometryProcessing/Source/DynamicMesh/Private/MeshBoundaryLoops.cpp
+++ b/Engine/Plugins/Experimental/GeometryProcessing/Source/DynamicMesh/Private/MeshBoundaryLoops.cpp
@@ -30,7 +30,7 @@ FIndex2i FMeshBoundaryLoops::FindVertexInLoop(int VertexID) const
 			return FIndex2i(li, idx);
 		}
 	}
-	return FIndex2i::Max();
+	return FIndex2i::Invalid();
 }
 
 
diff --git a/Engine/Plugins/Experimental/GeometryProcessing/Source/DynamicMesh/Public/MeshBoundaryLoops.h b/Engine/Plugins/Experimental/GeometryProcessing/Source/DynamicMesh/Public/MeshBoundaryLoops.h
--- a/Engine/Plugins/Experimental/GeometryProcessing/Source/DynamicMesh/Public/MeshBoundaryLoops.h
+++ b/Engine/Plugins/Experimental/GeometryProcessing/Source/DynamicMesh/Public/MeshBoundaryLoops.h
@@ -99,7 +99,7 @@ public:
 	int GetMaxVerticesLoopIndex() const;
 
 	/**
-	 * @return pair (LoopIndex,VertexIndexInLoop) of VertexID in EdgeLoops, or FIndex2i::Max if not found
+	 * @return pair (LoopIndex,VertexIndexInLoop) of VertexID in EdgeLoops, or FIndex2i::Invalid if not found
 	 */
 	FIndex2i FindVertexInLoop(int VertexID) const;
 
diff --git a/Engine/Plugins/Experimental/GeometryProcessing/Source/GeometricObjects/Public/Spatial/MeshAABBTree3.h b/Engine/Plugins/Experimental/GeometryProcessing/Source/GeometricObjects/Public/Spatial/MeshAABBTree3.h
--- a/Engine/Plugins/Experimental/GeometryProcessing/Source/GeometricObjects/Public/Spatial/MeshAABBTree3.h
+++ b/Engine/Plugins/Experimental/GeometryProcessing/Source/GeometricObjects/Public/Spatial/MeshAABBTree3.h
@@ -241,7 +241,7 @@ public:
 	/**
 	 * Find nearest pair of triangles on this tree with otherTree, within max_dist.
 	 * TransformF transforms vertices of otherTree into our coordinates. can be null.
-	 * returns triangle-id pair (my_tri,other_tri), or Index2i.Max if not found within max_dist
+	 * returns triangle-id pair (my_tri,other_tri), or Index2i::Invalid if not found within max_dist
 	 * Use MeshQueries.TrianglesDistance() to get more information
 	 */
 	virtual FIndex2i FindNearestTriangles(TMeshAABBTree3& OtherTree, TFunction<FVector3d(const FVector3d&)> TransformF, double& Distance, double MaxDist = FMathd::MaxReal)
@@ -253,7 +253,7 @@ public:
 		{
 			NearestSqr = MaxDist * MaxDist;
 		}
-		FIndex2i NearestPair = FIndex2i::Max();
+		FIndex2i NearestPair = FIndex2i::Invalid();
 
 		find_nearest_triangles(RootIndex, OtherTree, TransformF, OtherTree.RootIndex, 0, NearestSqr, NearestPair);
 		Distance = (NearestSqr < FMathd::MaxReal) ? FMathd::Sqrt(NearestSqr) : FMathd::MaxReal;
