fix linux and clang compile errors, improve naming/comments
#rnx

diff --git a/Engine/Plugins/Experimental/GeometryProcessing/Source/DynamicMesh/Private/GroupTopology.cpp b/Engine/Plugins/Experimental/GeometryProcessing/Source/DynamicMesh/Private/GroupTopology.cpp
--- a/Engine/Plugins/Experimental/GeometryProcessing/Source/DynamicMesh/Private/GroupTopology.cpp
+++ b/Engine/Plugins/Experimental/GeometryProcessing/Source/DynamicMesh/Private/GroupTopology.cpp
@@ -67,12 +67,12 @@ void FGroupTopology::RebuildTopology()
 	}
 
 	// precompute junction vertices set
-	JunctionVertices.Init(false, Mesh->MaxVertexID());
+	CornerVerticesFlags.Init(false, Mesh->MaxVertexID());
 	for (int vid : Mesh->VertexIndicesItr())
 	{
-		if (IsJunctionVertex(vid))
+		if (IsCornerVertex(vid))
 		{
-			JunctionVertices[vid] = true;
+			CornerVerticesFlags[vid] = true;
 			FCorner Corner = { vid };
 			Corners.Add(Corner);
 		}
@@ -86,8 +86,10 @@ void FGroupTopology::RebuildTopology()
 	// construct boundary loops
 	for (FGroup& Group : Groups)
 	{
+		// finds FGroupEdges and uses to populate Group.Boundaries
 		ExtractGroupEdges(Group);
 
+		// collect up .NeighbourGroupIDs and set .bIsOnBoundary
 		for (FGroupBoundary& Boundary : Group.Boundaries)
 		{
 			Boundary.bIsOnBoundary = false;
@@ -120,8 +122,7 @@ void FGroupTopology::RebuildTopology()
 }
 
 
-/** @return true if given mesh vertex is a Corner vertex */
-bool FGroupTopology::IsJunctionVertex(int VertexID) const
+bool FGroupTopology::IsCornerVertex(int VertexID) const
 {
 	FIndex3i UniqueGroups;
 	int UniqueCount = 0;
@@ -376,18 +377,18 @@ void FGroupTopology::ExtractGroupEdges(FGroup& Group)
 		FGroupBoundary& Boundary = Group.Boundaries[li];
 
 		// find indices of corners of group polygon
-		TArray<int> JunctionIndices;
+		TArray<int> CornerIndices;
 		int NumV = Loop.Vertices.Num();
 		for (int i = 0; i < NumV; ++i)
 		{
-			if (JunctionVertices[Loop.Vertices[i]])
+			if (CornerVerticesFlags[Loop.Vertices[i]])
 			{
-				JunctionIndices.Add(i);
+				CornerIndices.Add(i);
 			}
 		}
 
 		// if we had no indices then this is like the cap of a cylinder, just one single long edge
-		if ( JunctionIndices.Num() == 0 )
+		if ( CornerIndices.Num() == 0 )
 		{ 
 			FIndex2i EdgeID = MakeEdgeID(Loop.Edges[0]);
 			int OtherGroupID = (EdgeID.A == Group.GroupID) ? EdgeID.B : EdgeID.A;
@@ -403,15 +404,15 @@ void FGroupTopology::ExtractGroupEdges(FGroup& Group)
 			continue;
 		}
 
-		// duplicate first junction vertex so that we can just loop back around to it w/ modulo count
-		int NumSpans = JunctionIndices.Num();
-		int FirstIdx = JunctionIndices[0];
-		JunctionIndices.Add(FirstIdx);
+		// duplicate first corner vertex so that we can just loop back around to it w/ modulo count
+		int NumSpans = CornerIndices.Num();
+		int FirstIdx = CornerIndices[0];
+		CornerIndices.Add(FirstIdx);
 
 		// add each span
 		for (int k = 0; k < NumSpans; ++k)
 		{
-			int i0 = JunctionIndices[k];
+			int i0 = CornerIndices[k];
 
 			FIndex2i EdgeID = MakeEdgeID(Loop.Edges[i0]);
 			int OtherGroupID = (EdgeID.A == Group.GroupID) ? EdgeID.B : EdgeID.A;
@@ -425,7 +426,7 @@ void FGroupTopology::ExtractGroupEdges(FGroup& Group)
 
 			FGroupEdge Edge = { EdgeID };
 
-			int i1 = JunctionIndices[k+1];
+			int i1 = CornerIndices[k+1];
 			check(i0 != i1);	// this case will happen on a cylinder w/ a cut side, I think...
 								// (but that case won't exist because that's wouldn't be a group boundary edge!)
 			TArray<int> SpanVertices;
@@ -458,13 +459,16 @@ int FGroupTopology::FindExistingGroupEdge(int GroupID, int OtherGroupID, int Fir
 
 	FGroup& OtherGroup = Groups[GroupIDToGroupIndexMap[OtherGroupID]];
 	FIndex2i EdgeID = MakeEdgeID(GroupID, OtherGroupID);
-
+	
 	for (FGroupBoundary& Boundary : OtherGroup.Boundaries)
 	{
 		for (int EdgeIndex : Boundary.GroupEdges)
 		{
 			if (Edges[EdgeIndex].Groups == EdgeID)
 			{
+				// same EdgeID pair may occur multiple times in the same boundary loop! 
+				// need to check that at least one endpoint is the same vertex
+
 				TArray<int>& Vertices = Edges[EdgeIndex].Span.Vertices;
 				if (Vertices[0] == FirstVertexID || Vertices[Vertices.Num() - 1] == FirstVertexID)
 				{
diff --git a/Engine/Plugins/Experimental/GeometryProcessing/Source/DynamicMesh/Public/GroupTopology.h b/Engine/Plugins/Experimental/GeometryProcessing/Source/DynamicMesh/Public/GroupTopology.h
--- a/Engine/Plugins/Experimental/GeometryProcessing/Source/DynamicMesh/Public/GroupTopology.h
+++ b/Engine/Plugins/Experimental/GeometryProcessing/Source/DynamicMesh/Public/GroupTopology.h
@@ -168,11 +168,11 @@ protected:
 	const FDynamicMesh3* Mesh = nullptr;
 
 	TArray<int> GroupIDToGroupIndexMap;		// allow fast lookup of index in .Groups, given GroupID
-	TBitArray<> JunctionVertices;		// bit array of corners for fast testing in ExtractGroupEdges
+	TBitArray<> CornerVerticesFlags;		// bit array of corners for fast testing in ExtractGroupEdges
 	TArray<int> EmptyArray;
 
 	/** @return true if given mesh vertex is a Corner vertex */
-	virtual bool IsJunctionVertex(int VertexID) const;
+	virtual bool IsCornerVertex(int VertexID) const;
 
 	void ExtractGroupEdges(FGroup& Group);
 
diff --git a/Engine/Plugins/Experimental/GeometryProcessing/Source/GeometricObjects/Public/Distance/DistRay3Segment3.h b/Engine/Plugins/Experimental/GeometryProcessing/Source/GeometricObjects/Public/Distance/DistRay3Segment3.h
--- a/Engine/Plugins/Experimental/GeometryProcessing/Source/GeometricObjects/Public/Distance/DistRay3Segment3.h
+++ b/Engine/Plugins/Experimental/GeometryProcessing/Source/GeometricObjects/Public/Distance/DistRay3Segment3.h
@@ -54,7 +54,7 @@ public:
 		double b0 = diff.Dot(Ray.Direction);
 		double b1 = -diff.Dot(Segment.Direction);
 		double c = diff.SquaredLength();
-		double det = TMathUtil::Abs(1 - a01 * a01);
+		double det = TMathUtil<Real>::Abs(1 - a01 * a01);
 		double s0, s1, sqrDist, extDet;
 
 		if (det >= TMathUtil<Real>::ZeroTolerance ) 
diff --git a/Engine/Plugins/Experimental/GeometryProcessing/Source/GeometricObjects/Public/Polyline3.h b/Engine/Plugins/Experimental/GeometryProcessing/Source/GeometricObjects/Public/Polyline3.h
--- a/Engine/Plugins/Experimental/GeometryProcessing/Source/GeometricObjects/Public/Polyline3.h
+++ b/Engine/Plugins/Experimental/GeometryProcessing/Source/GeometricObjects/Public/Polyline3.h
@@ -202,7 +202,7 @@ public:
 		{
 			return (Vertices[NumVerts-1] - Vertices[NumVerts-2]).Normalized();
 		}
-		return (Vertices[VertexIndex+1] - Vertices[VertexIndex-1]).Normalized()
+		return (Vertices[VertexIndex+1] - Vertices[VertexIndex-1]).Normalized();
 	}
 
 
@@ -270,11 +270,11 @@ public:
 	public:
 		inline bool operator!()
 		{
-			return i < SegmentCount();
+			return i < Polyline->SegmentCount();
 		}
 		inline FSegment3<T> operator*() const
 		{
-			check(Polyline != nullptr && i < SegmentCount());
+			check(Polyline != nullptr && i < Polyline->SegmentCount());
 			return FSegment3<T>(Polyline->Vertices[i], Polyline->Vertices[i+1]);
 		}
 		inline SegmentIterator & operator++() 		// prefix
@@ -313,7 +313,7 @@ public:
 		SegmentEnumerable() : Polyline(nullptr) {}
 		SegmentEnumerable(const TPolyline3<T> * p) : Polyline(p) {}
 		SegmentIterator begin() { return Polyline->SegmentItr(); }
-		SegmentIterator end() { return SegmentIterator(Polyline, SegmentCount()); }
+		SegmentIterator end() { return SegmentIterator(Polyline, Polyline->SegmentCount()); }
 	};
 
 	/**
diff --git a/Engine/Plugins/Experimental/GeometryProcessing/Source/GeometricObjects/Public/Spatial/GeometrySet3.h b/Engine/Plugins/Experimental/GeometryProcessing/Source/GeometricObjects/Public/Spatial/GeometrySet3.h
--- a/Engine/Plugins/Experimental/GeometryProcessing/Source/GeometricObjects/Public/Spatial/GeometrySet3.h
+++ b/Engine/Plugins/Experimental/GeometryProcessing/Source/GeometricObjects/Public/Spatial/GeometrySet3.h
@@ -61,6 +61,7 @@ public:
 	 * @param ResultOut populated with information about successful nearest point result
 	 * @param PointWithinToleranceTest should return true if two 3D points are "close enough" to be considered a hit
 	 * @return true if the nearest point on Ray to some point in the set passed the PointWithinToleranceTest.
+	 * @warning PointWithinToleranceTest is called in parallel and hence must be thread-safe/re-entrant!
 	 */
 	bool FindNearestPointToRay(const FRay3d& Ray, FNearest& ResultOut,
 		TFunction<bool(const FVector3d&, const FVector3d&)> PointWithinToleranceTest) const;
@@ -70,6 +71,7 @@ public:
 	 * @param ResultOut populated with information about successful nearest point result
 	 * @param PointWithinToleranceTest should return true if two 3D points are "close enough" to be considered a hit
 	 * @return true if the nearest point on Ray to some curve in the set passed the PointWithinToleranceTest.
+	 * @warning PointWithinToleranceTest is called in parallel and hence must be thread-safe/re-entrant!
 	 */
 	bool FindNearestCurveToRay(const FRay3d& Ray, FNearest& ResultOut,
 		TFunction<bool(const FVector3d&, const FVector3d&)> PointWithinToleranceTest) const;
diff --git a/Engine/Plugins/Experimental/ProxyLODPlugin/Source/ProxyLOD/Public/ProxyLODVolume.h b/Engine/Plugins/Experimental/ProxyLODPlugin/Source/ProxyLOD/Public/ProxyLODVolume.h
--- a/Engine/Plugins/Experimental/ProxyLODPlugin/Source/ProxyLOD/Public/ProxyLODVolume.h
+++ b/Engine/Plugins/Experimental/ProxyLODPlugin/Source/ProxyLOD/Public/ProxyLODVolume.h
@@ -65,6 +65,8 @@ class PROXYLODMESHREDUCTION_API IVoxelBasedCSG
 public :
 	static TUniquePtr<IVoxelBasedCSG> CreateCSGTool(float VoxelSize);
 
+	virtual ~IVoxelBasedCSG() {}
+
 	class FPlacedMesh
 	{
 	public:
