Add some tests for FVector3 using the Automation Spec framework

#rnx

diff --git a/Engine/Plugins/Experimental/GeometryProcessing/Source/GeometricObjects/Private/Tests/VectorTypes.spec.cpp b/Engine/Plugins/Experimental/GeometryProcessing/Source/GeometricObjects/Private/Tests/VectorTypes.spec.cpp
new file mode 100644
--- /dev/null
+++ b/Engine/Plugins/Experimental/GeometryProcessing/Source/GeometricObjects/Private/Tests/VectorTypes.spec.cpp
@@ -0,0 +1,190 @@
+#include "Misc/AutomationTest.h"
+#include "VectorTypes.h"
+
+#if WITH_DEV_AUTOMATION_TESTS
+
+BEGIN_DEFINE_SPEC(FVectorTypesSpec,
+				  "GeometryProcessing.Unit",
+				  EAutomationTestFlags::ProductFilter | EAutomationTestFlags::ApplicationContextMask)
+END_DEFINE_SPEC(FVectorTypesSpec)
+
+void FVectorTypesSpec::Define()
+{
+	using Vec = FVector3<float>;
+	Describe("FVector3<float>", [this]()
+	{
+		Describe("Constructors", [this]()
+		{
+			constexpr float initialValues[3]{1.f, 2.f, 3.f};
+			It("FVector3<float>()", [this]()
+			{
+				Vec v{};
+			});
+			It("FVector3<float>(float, float, float)", [this, initialValues]()
+			{
+				Vec v{initialValues[0], initialValues[1], initialValues[2]};
+				TestEqual("v.X", v.X, initialValues[0]);
+				TestEqual("v.Y", v.Y, initialValues[1]);
+				TestEqual("v.Z", v.Z, initialValues[2]);
+			});
+			It("FVector3<float>(const float*)", [this, initialValues]()
+			{
+				Vec v{initialValues};
+				TestEqual("v.X", v.X, initialValues[0]);
+				TestEqual("v.Y", v.Y, initialValues[1]);
+				TestEqual("v.Z", v.Z, initialValues[2]);
+			});
+			It("FVector3<float>(const FVector&)", [this, initialValues]()
+			{
+				FVector fVector{initialValues[0], initialValues[1],
+								initialValues[2]};
+				Vec fVector3{fVector};
+				TestEqual("v.X", fVector.X, fVector3.X);
+				TestEqual("v.Y", fVector.Y, fVector3.Y);
+				TestEqual("v.Z", fVector.Z, fVector3.Z);
+			});
+			It("FVector3<float>(const FLinearColor&)", [this, initialValues]()
+			{
+				FLinearColor fLinearColor{initialValues[0], initialValues[1],
+										  initialValues[2]};
+				Vec fVector3{FVector{fLinearColor}};
+				TestEqual("v.X", fLinearColor.R, fVector3.X);
+				TestEqual("v.Y", fLinearColor.G, fVector3.Y);
+				TestEqual("v.Z", fLinearColor.B, fVector3.Z);
+			});
+		});
+		Describe("Conversion", [this]()
+		{
+			It("(const float*)", [this]()
+			{
+				constexpr float initialValues[3]{1.f, 2.f, 3.f};
+				Vec testVec{initialValues};
+				const float *converted = static_cast<const float *>(testVec);
+				for (int index : {0, 1, 2}) {
+					TestEqual(FString::Printf(TEXT("Component %d"), index),
+							  initialValues[index], converted[index]);
+				}
+			});
+			It("(float*)", [this]()
+			{
+				constexpr float initialValues[3]{1.f, 2.f, 3.f};
+				Vec testVec{initialValues};
+				float *converted = static_cast<float *>(testVec);
+				for (int index : {0, 1, 2}) {
+					TestEqual(FString::Printf(TEXT("Component %d"), index),
+							  initialValues[index], converted[index]);
+				}
+				for (int index : {0, 1, 2}) {
+					converted[index] *= 2.f;
+				}
+				TestEqual("v.X", testVec.X, 2.f * initialValues[0]);
+				TestEqual("v.Y", testVec.Y, 2.f * initialValues[1]);
+				TestEqual("v.Z", testVec.Z, 2.f * initialValues[2]);
+			});
+			It("(FVector)", [this]()
+			{
+				constexpr float initialValues[3]{1.f, 2.f, 3.f};
+				Vec testVec{initialValues};
+				FVector converted = static_cast<FVector>(testVec);
+				for (int index : {0, 1, 2}) {
+					TestEqual(FString::Printf(TEXT("Component %d"), index),
+							  testVec[index], converted[index]);
+				}
+			});
+			It("(FLinearColor)", [this]()
+			{
+				constexpr float initialValues[3]{0.5f, 0.5f, 0.5f};
+				Vec testVec{initialValues};
+				FLinearColor converted = {testVec};
+				TestEqual("Color.R", testVec.X, converted.R);
+				TestEqual("Color.G", testVec.Y, converted.G);
+				TestEqual("Color.B", testVec.Z, converted.B);
+				TestEqual("Color.A", 1.f, converted.A);
+			});
+		});
+		Describe("Assignment Operator", [this]()
+		{
+			It("assigns every component", [this]()
+			{
+				constexpr float initialValues[3]{1.f, 2.f, 3.f};
+				Vec testVec{initialValues};
+				Vec testVecCopy{};
+				testVecCopy = testVec;
+				TestEqual("v.X", testVec.X, testVecCopy.X);
+				TestEqual("v.Y", testVec.Y, testVecCopy.Y);
+				TestEqual("v.Z", testVec.Z, testVecCopy.Z);
+			});
+		});
+		Describe("Element Access", [this]()
+		{
+			It("can be const", [this]()
+			{
+				constexpr float initialValues[3]{1.f, 2.f, 3.f};
+				Vec testVec{initialValues};
+				for ( int i : { 0, 1, 2 } )
+				{
+					const float& val = testVec[i];
+					TestEqual(FString::Printf(TEXT("Component %d"), i), val, initialValues[i]);
+				}
+			});
+			It("can be mutable", [this]()
+			{
+				constexpr float initialValues[3]{1.f, 2.f, 3.f};
+				Vec testVec{initialValues};
+				for ( int i : { 0, 1, 2 } )
+				{
+					float& val = testVec[i];
+					TestEqual(FString::Printf(TEXT("Component %d"), i), val, testVec[i]);
+					val *= 2.f;
+					TestEqual(FString::Printf(TEXT("Component %d"), i), testVec[i], initialValues[i] * 2.f);
+				}
+			});
+		});
+		Describe("Length", [this]()
+		{
+			It("of Zero vector is 0.f", [this]()
+			{
+				TestEqual("Zero.Length()", Vec::Zero().Length(), 0.f);
+			});
+			It("of One vector is Sqrt(3.f)", [this]()
+			{
+				TestEqual("One.Length()", Vec::One().Length(), FMathf::Sqrt(3.f));
+			});
+			It("of Unit vectors is 1.f", [this]()
+			{
+				TestEqual("UnitX.Length()", Vec::UnitX().Length(), 1.f);
+				TestEqual("UnitY.Length()", Vec::UnitY().Length(), 1.f);
+				TestEqual("UnitZ.Length()", Vec::UnitZ().Length(), 1.f);
+			});
+		});
+		Describe("SquaredLength", [this]()
+		{
+			It("of Zero vector is 0.f", [this]()
+			{
+				TestEqual("Zero.SquaredLength()", Vec::Zero().SquaredLength(), 0.f);
+			});
+			It("of One vector is 3.f", [this]()
+			{
+				TestEqual("One.SquaredLength()", Vec::One().SquaredLength(), 3.f);
+			});
+			It("of Unit vectors is 1.f", [this]()
+			{
+				TestEqual("UnitX.SquaredLength()", Vec::UnitX().SquaredLength(), 1.f);
+				TestEqual("UnitY.SquaredLength()", Vec::UnitY().SquaredLength(), 1.f);
+				TestEqual("UnitZ.SquaredLength()", Vec::UnitZ().SquaredLength(), 1.f);
+			});
+		});
+		Describe("Distance", [this]()
+		{
+			It("from vector to itself is 0.f", [this]()
+			{
+				const Vec testVecs[] = {Vec::Zero(), Vec::One(), Vec::UnitX(), Vec::UnitY(), Vec::UnitZ(), Vec(1.0, 2.0, 3.0)};
+				for (const auto &testVec : testVecs)
+				{
+					TestEqual("Self Distance", testVec.Distance(testVec), 0.f);
+				}
+			});
+		});
+	});
+}
+#endif
