rename most GeomProc templates to have T prefix
#rnx

diff --git a/Engine/Plugins/Experimental/GeometryProcessing/Source/DynamicMesh/Private/DynamicMesh3.cpp b/Engine/Plugins/Experimental/GeometryProcessing/Source/DynamicMesh/Private/DynamicMesh3.cpp
--- a/Engine/Plugins/Experimental/GeometryProcessing/Source/DynamicMesh/Private/DynamicMesh3.cpp
+++ b/Engine/Plugins/Experimental/GeometryProcessing/Source/DynamicMesh/Private/DynamicMesh3.cpp
@@ -39,22 +39,22 @@ FDynamicMesh3::~FDynamicMesh3()
 
 FDynamicMesh3::FDynamicMesh3(bool bWantNormals, bool bWantColors, bool bWantUVs, bool bWantTriGroups)
 {
-	Vertices = FDynamicVector<double>();
-	VertexNormals = (bWantNormals) ? new FDynamicVector<float>() : nullptr;
-	VertexColors = (bWantColors) ? new FDynamicVector<float>() : nullptr;
-	VertexUVs = (bWantUVs) ? new FDynamicVector<float>() : nullptr;
+	Vertices = TDynamicVector<double>();
+	VertexNormals = (bWantNormals) ? new TDynamicVector<float>() : nullptr;
+	VertexColors = (bWantColors) ? new TDynamicVector<float>() : nullptr;
+	VertexUVs = (bWantUVs) ? new TDynamicVector<float>() : nullptr;
 
 	VertexEdgeLists = FSmallListSet();
 
 	VertexRefCounts = FRefCountVector();
 
-	Triangles = FDynamicVector<int>();
-	TriangleEdges = FDynamicVector<int>();
+	Triangles = TDynamicVector<int>();
+	TriangleEdges = TDynamicVector<int>();
 	TriangleRefCounts = FRefCountVector();
-	TriangleGroups = (bWantTriGroups) ? new FDynamicVector<int>() : nullptr;
+	TriangleGroups = (bWantTriGroups) ? new TDynamicVector<int>() : nullptr;
 	GroupIDCounter = 0;
 
-	Edges = FDynamicVector<int>();
+	Edges = TDynamicVector<int>();
 	EdgeRefCounts = FRefCountVector();
 
 	AttributeSet = nullptr;
@@ -106,15 +106,15 @@ void FDynamicMesh3::Copy(const FMeshShapeGenerator* Generator)
 {
 	Clear();
 
-	Vertices = FDynamicVector<double>();
+	Vertices = TDynamicVector<double>();
 	VertexEdgeLists = FSmallListSet();
 	VertexRefCounts = FRefCountVector();
-	Triangles = FDynamicVector<int>();
-	TriangleEdges = FDynamicVector<int>();
+	Triangles = TDynamicVector<int>();
+	TriangleEdges = TDynamicVector<int>();
 	TriangleRefCounts = FRefCountVector();
-	TriangleGroups = new FDynamicVector<int>();
+	TriangleGroups = new TDynamicVector<int>();
 	GroupIDCounter = 0;
-	Edges = FDynamicVector<int>();
+	Edges = TDynamicVector<int>();
 	EdgeRefCounts = FRefCountVector();
 
 	EnableAttributes();
@@ -160,23 +160,23 @@ void FDynamicMesh3::Copy(const FDynamicMesh3& copy, bool bNormals, bool bColors,
 	DiscardTriangleGroups();
 	DiscardAttributes();
 
-	Vertices = FDynamicVector<double>(copy.Vertices);
+	Vertices = TDynamicVector<double>(copy.Vertices);
 
-	VertexNormals = (bNormals && copy.HasVertexNormals()) ? new FDynamicVector<float>(*copy.VertexNormals) : nullptr;
-	VertexColors = (bColors && copy.HasVertexColors()) ? new FDynamicVector<float>(*copy.VertexColors) : nullptr;
-	VertexUVs = (bUVs && copy.HasVertexUVs()) ? new FDynamicVector<float>(*copy.VertexUVs) : nullptr;
+	VertexNormals = (bNormals && copy.HasVertexNormals()) ? new TDynamicVector<float>(*copy.VertexNormals) : nullptr;
+	VertexColors = (bColors && copy.HasVertexColors()) ? new TDynamicVector<float>(*copy.VertexColors) : nullptr;
+	VertexUVs = (bUVs && copy.HasVertexUVs()) ? new TDynamicVector<float>(*copy.VertexUVs) : nullptr;
 
 	VertexRefCounts = FRefCountVector(copy.VertexRefCounts);
 
 	VertexEdgeLists = FSmallListSet(copy.VertexEdgeLists);
 
-	Triangles = FDynamicVector<int>(copy.Triangles);
-	TriangleEdges = FDynamicVector<int>(copy.TriangleEdges);
+	Triangles = TDynamicVector<int>(copy.Triangles);
+	TriangleEdges = TDynamicVector<int>(copy.TriangleEdges);
 	TriangleRefCounts = FRefCountVector(copy.TriangleRefCounts);
-	TriangleGroups = (copy.HasTriangleGroups()) ? new FDynamicVector<int>(*copy.TriangleGroups) : nullptr;
+	TriangleGroups = (copy.HasTriangleGroups()) ? new TDynamicVector<int>(*copy.TriangleGroups) : nullptr;
 	GroupIDCounter = copy.GroupIDCounter;
 
-	Edges = FDynamicVector<int>(copy.Edges);
+	Edges = TDynamicVector<int>(copy.Edges);
 	EdgeRefCounts = FRefCountVector(copy.EdgeRefCounts);
 
 	if (bAttributes && copy.HasAttributes())
@@ -212,17 +212,17 @@ void FDynamicMesh3::CompactCopy(const FDynamicMesh3& copy, bool bNormals, bool b
 	DiscardVertexUVs();
 	DiscardTriangleGroups();
 
-	Vertices = FDynamicVector<double>();
+	Vertices = TDynamicVector<double>();
 	VertexEdgeLists = FSmallListSet();
 	VertexRefCounts = FRefCountVector();
-	Triangles = FDynamicVector<int>();
-	TriangleEdges = FDynamicVector<int>();
+	Triangles = TDynamicVector<int>();
+	TriangleEdges = TDynamicVector<int>();
 	TriangleRefCounts = FRefCountVector();
-	Edges = FDynamicVector<int>();
+	Edges = TDynamicVector<int>();
 	EdgeRefCounts = FRefCountVector();
 	GroupIDCounter = 0;
 
-	TriangleGroups = (copy.HasTriangleGroups()) ? new FDynamicVector<int>(*copy.TriangleGroups) : nullptr;
+	TriangleGroups = (copy.HasTriangleGroups()) ? new TDynamicVector<int>(*copy.TriangleGroups) : nullptr;
 
 	// [TODO] if we knew some of these were dense we could copy directly...
 
@@ -327,7 +327,7 @@ void FDynamicMesh3::EnableVertexNormals(const FVector3f& InitialNormal)
 	{
 		return;
 	}
-	VertexNormals = new FDynamicVector<float>();
+	VertexNormals = new TDynamicVector<float>();
 	int NV = MaxVertexID();
 	VertexNormals->Resize(3 * NV);
 	for (int i = 0; i < NV; ++i)
@@ -358,7 +358,7 @@ void FDynamicMesh3::EnableVertexColors(const FVector3f& InitialColor)
 	{
 		return;
 	}
-	VertexColors = new FDynamicVector<float>();
+	VertexColors = new TDynamicVector<float>();
 	int NV = MaxVertexID();
 	VertexColors->Resize(3 * NV);
 	for (int i = 0; i < NV; ++i)
@@ -388,7 +388,7 @@ void FDynamicMesh3::EnableVertexUVs(const FVector2f& InitialUV)
 	{
 		return;
 	}
-	VertexUVs = new FDynamicVector<float>();
+	VertexUVs = new TDynamicVector<float>();
 	int NV = MaxVertexID();
 	VertexUVs->Resize(2 * NV);
 	for (int i = 0; i < NV; ++i)
@@ -416,7 +416,7 @@ void FDynamicMesh3::EnableTriangleGroups(int InitialGroup)
 	{
 		return;
 	}
-	TriangleGroups = new FDynamicVector<int>();
+	TriangleGroups = new TDynamicVector<int>();
 	int NT = MaxTriangleID();
 	TriangleGroups->Resize(NT);
 	for (int i = 0; i < NT; ++i)
diff --git a/Engine/Plugins/Experimental/GeometryProcessing/Source/DynamicMesh/Private/DynamicMesh3_Edits.cpp b/Engine/Plugins/Experimental/GeometryProcessing/Source/DynamicMesh/Private/DynamicMesh3_Edits.cpp
--- a/Engine/Plugins/Experimental/GeometryProcessing/Source/DynamicMesh/Private/DynamicMesh3_Edits.cpp
+++ b/Engine/Plugins/Experimental/GeometryProcessing/Source/DynamicMesh/Private/DynamicMesh3_Edits.cpp
@@ -316,7 +316,7 @@ void FDynamicMesh3::CompactInPlace(FCompactMaps* CompactInfo)
 		iCurV++;
 	}
 
-	FDynamicVector<short> &vref = VertexRefCounts.GetRawRefCountsUnsafe();
+	TDynamicVector<short> &vref = VertexRefCounts.GetRawRefCountsUnsafe();
 
 	while (iCurV < iLastV)
 	{
@@ -411,7 +411,7 @@ void FDynamicMesh3::CompactInPlace(FCompactMaps* CompactInfo)
 		iCurT++;
 	}
 
-	FDynamicVector<short> &tref = TriangleRefCounts.GetRawRefCountsUnsafe();
+	TDynamicVector<short> &tref = TriangleRefCounts.GetRawRefCountsUnsafe();
 
 	while (iCurT < iLastT)
 	{
@@ -473,7 +473,7 @@ void FDynamicMesh3::CompactInPlace(FCompactMaps* CompactInfo)
 		iCurE++;
 	}
 
-	FDynamicVector<short> &eref = EdgeRefCounts.GetRawRefCountsUnsafe();
+	TDynamicVector<short> &eref = EdgeRefCounts.GetRawRefCountsUnsafe();
 
 	while (iCurE < iLastE) 
 	{
diff --git a/Engine/Plugins/Experimental/GeometryProcessing/Source/DynamicMesh/Private/DynamicMesh3_Queries.cpp b/Engine/Plugins/Experimental/GeometryProcessing/Source/DynamicMesh/Private/DynamicMesh3_Queries.cpp
--- a/Engine/Plugins/Experimental/GeometryProcessing/Source/DynamicMesh/Private/DynamicMesh3_Queries.cpp
+++ b/Engine/Plugins/Experimental/GeometryProcessing/Source/DynamicMesh/Private/DynamicMesh3_Queries.cpp
@@ -678,7 +678,7 @@ FVector3d FDynamicMesh3::GetTriBaryNormal(int tID, double bary0, double bary1, d
 	int ai = 3 * Triangles[3 * tID],
 		bi = 3 * Triangles[3 * tID + 1],
 		ci = 3 * Triangles[3 * tID + 2];
-	const FDynamicVector<float>& normalsR = *(this->VertexNormals);
+	const TDynamicVector<float>& normalsR = *(this->VertexNormals);
 	FVector3d n = FVector3d(
 		(bary0*normalsR[ai] + bary1 * normalsR[bi] + bary2 * normalsR[ci]),
 		(bary0*normalsR[ai + 1] + bary1 * normalsR[bi + 1] + bary2 * normalsR[ci + 1]),
@@ -713,7 +713,7 @@ void FDynamicMesh3::GetTriBaryPoint(int tID, double bary0, double bary1, double
 	vinfo.bHaveN = HasVertexNormals();
 	if (vinfo.bHaveN) 
 	{
-		FDynamicVector<float>& normalsR = *(this->VertexNormals);
+		TDynamicVector<float>& normalsR = *(this->VertexNormals);
 		vinfo.Normal = FVector3f(
 			(float)(bary0 * normalsR[ai] + bary1 * normalsR[bi] + bary2 * normalsR[ci]),
 			(float)(bary0 * normalsR[ai + 1] + bary1 * normalsR[bi + 1] + bary2 * normalsR[ci + 1]),
@@ -723,7 +723,7 @@ void FDynamicMesh3::GetTriBaryPoint(int tID, double bary0, double bary1, double
 	vinfo.bHaveC = HasVertexColors();
 	if (vinfo.bHaveC) 
 	{
-		FDynamicVector<float>& colorsR = *(this->VertexColors);
+		TDynamicVector<float>& colorsR = *(this->VertexColors);
 		vinfo.Color = FVector3f(
 			(float)(bary0 * colorsR[ai] + bary1 * colorsR[bi] + bary2 * colorsR[ci]),
 			(float)(bary0 * colorsR[ai + 1] + bary1 * colorsR[bi + 1] + bary2 * colorsR[ci + 1]),
@@ -732,7 +732,7 @@ void FDynamicMesh3::GetTriBaryPoint(int tID, double bary0, double bary1, double
 	vinfo.bHaveUV = HasVertexUVs();
 	if (vinfo.bHaveUV) 
 	{
-		FDynamicVector<float>& uvR = *(this->VertexUVs);
+		TDynamicVector<float>& uvR = *(this->VertexUVs);
 		ai = 2 * Triangles[3 * tID];
 		bi = 2 * Triangles[3 * tID + 1];
 		ci = 2 * Triangles[3 * tID + 2];
diff --git a/Engine/Plugins/Experimental/GeometryProcessing/Source/DynamicMesh/Private/MeshTangents.cpp b/Engine/Plugins/Experimental/GeometryProcessing/Source/DynamicMesh/Private/MeshTangents.cpp
--- a/Engine/Plugins/Experimental/GeometryProcessing/Source/DynamicMesh/Private/MeshTangents.cpp
+++ b/Engine/Plugins/Experimental/GeometryProcessing/Source/DynamicMesh/Private/MeshTangents.cpp
@@ -5,7 +5,7 @@
 
 
 template<typename RealType>
-void FMeshTangents<RealType>::SetTangentCount(int Count, bool bClearToZero)
+void TMeshTangents<RealType>::SetTangentCount(int Count, bool bClearToZero)
 {
 	if (Tangents.Num() < Count)
 	{
@@ -28,7 +28,7 @@ void FMeshTangents<RealType>::SetTangentCount(int Count, bool bClearToZero)
 
 
 template<typename RealType>
-void FMeshTangents<RealType>::Internal_ComputePerTriangleTangents(const FDynamicMeshNormalOverlay* NormalOverlay, const FDynamicMeshUVOverlay* UVOverlay)
+void TMeshTangents<RealType>::Internal_ComputePerTriangleTangents(const FDynamicMeshNormalOverlay* NormalOverlay, const FDynamicMeshUVOverlay* UVOverlay)
 {
 	int MaxTriangleID = Mesh->MaxTriangleID();
 	InitializePerTriangleTangents(false);
@@ -68,8 +68,8 @@ void FMeshTangents<RealType>::Internal_ComputePerTriangleTangents(const FDynamic
 }
 
 
-template class DYNAMICMESH_API FMeshTangents<float>;
-template class DYNAMICMESH_API FMeshTangents<double>;
+template class DYNAMICMESH_API TMeshTangents<float>;
+template class DYNAMICMESH_API TMeshTangents<double>;
 
 
 
diff --git a/Engine/Plugins/Experimental/GeometryProcessing/Source/DynamicMesh/Public/DynamicMesh3.h b/Engine/Plugins/Experimental/GeometryProcessing/Source/DynamicMesh/Public/DynamicMesh3.h
--- a/Engine/Plugins/Experimental/GeometryProcessing/Source/DynamicMesh/Public/DynamicMesh3.h
+++ b/Engine/Plugins/Experimental/GeometryProcessing/Source/DynamicMesh/Public/DynamicMesh3.h
@@ -89,7 +89,7 @@ struct FVertexInfo
 *
 * internally, all data is stored in POD-type buffers, except for the vertex->edge
 * links, which are stored as List<int>'s. The arrays of POD data are stored in
-* FDynamicVector's, so they grow in chunks, which is relatively efficient. The actual
+* TDynamicVector's, so they grow in chunks, which is relatively efficient. The actual
 * blocks are arrays, so they can be efficiently mem-copied into larger buffers
 * if necessary.
 *
@@ -130,7 +130,7 @@ struct FVertexInfo
 *  - additional Topology timestamp?
 *  - CompactInPlace() does not compact VertexEdgeLists
 *  - CompactCopy does not support per-vertex or extended attributes
-*  ? FDynamicVector w/ 'stride' option, so that we can guarantee that tuples are in single block.
+*  ? TDynamicVector w/ 'stride' option, so that we can guarantee that tuples are in single block.
 *    The can have custom accessor that looks up entire tuple
 */
 class DYNAMICMESH_API FDynamicMesh3
@@ -160,13 +160,13 @@ protected:
 	/** Reference counts of vertex indices. Iterate over this to find out which vertex indices are valid. */
 	FRefCountVector VertexRefCounts;
 	/** List of vertex positions */
-	FDynamicVector<double> Vertices;
+	TDynamicVector<double> Vertices;
 	/** (optional) List of per-vertex normals */
-	FDynamicVector<float>* VertexNormals = nullptr;
+	TDynamicVector<float>* VertexNormals = nullptr;
 	/** (optional) List of per-vertex colors */
-	FDynamicVector<float>* VertexColors = nullptr;
+	TDynamicVector<float>* VertexColors = nullptr;
 	/** (optional) List of per-vertex uvs */
-	FDynamicVector<float>* VertexUVs = nullptr;
+	TDynamicVector<float>* VertexUVs = nullptr;
 
 	/** List of per-vertex edge one-rings */
 	FSmallListSet VertexEdgeLists;
@@ -174,16 +174,16 @@ protected:
 	/** Reference counts of triangle indices. Iterate over this to find out which triangle indices are valid. */
 	FRefCountVector TriangleRefCounts;
 	/** List of triangle vertex-index triplets [Vert0 Vert1 Vert2]*/
-	FDynamicVector<int> Triangles;
+	TDynamicVector<int> Triangles;
 	/** List of triangle edge triplets [Edge0 Edge1 Edge2] */
-	FDynamicVector<int> TriangleEdges;
+	TDynamicVector<int> TriangleEdges;
 	/** (optional) List of per-triangle group identifiers */
-	FDynamicVector<int>* TriangleGroups = nullptr;
+	TDynamicVector<int>* TriangleGroups = nullptr;
 
 	/** Reference counts of edge indices. Iterate over this to find out which edge indices are valid. */
 	FRefCountVector EdgeRefCounts;
 	/** List of edge elements. An edge is four elements [VertA, VertB, Tri0, Tri1], where VertA < VertB, and Tri1 may be InvalidID (if the edge is a boundary edge) */
-	FDynamicVector<int> Edges;
+	TDynamicVector<int> Edges;
 
 	FDynamicMeshAttributeSet* AttributeSet = nullptr;
 
@@ -967,7 +967,7 @@ public:
 	// direct buffer access
 	//
 public:
-	const FDynamicVector<double>& GetVerticesBuffer()
+	const TDynamicVector<double>& GetVerticesBuffer()
 	{
 		return Vertices;
 	}
@@ -975,20 +975,20 @@ public:
 	{
 		return VertexRefCounts;
 	}
-	const FDynamicVector<float>* GetNormalsBuffer()
+	const TDynamicVector<float>* GetNormalsBuffer()
 	{
 		return VertexNormals;
 	}
-	const FDynamicVector<float>* GetColorsBuffer()
+	const TDynamicVector<float>* GetColorsBuffer()
 	{
 		return VertexColors;
 	}
-	const FDynamicVector<float>* GetUVBuffer()
+	const TDynamicVector<float>* GetUVBuffer()
 	{
 		return VertexUVs;
 	}
 
-	const FDynamicVector<int>& GetTrianglesBuffer()
+	const TDynamicVector<int>& GetTrianglesBuffer()
 	{
 		return Triangles;
 	}
@@ -996,12 +996,12 @@ public:
 	{
 		return TriangleRefCounts;
 	}
-	const FDynamicVector<int>* GetTriangleGroupsBuffer()
+	const TDynamicVector<int>* GetTriangleGroupsBuffer()
 	{
 		return TriangleGroups;
 	}
 
-	const FDynamicVector<int>& GetEdgesBuffer()
+	const TDynamicVector<int>& GetEdgesBuffer()
 	{
 		return Edges;
 	}
diff --git a/Engine/Plugins/Experimental/GeometryProcessing/Source/DynamicMesh/Public/DynamicMeshOverlay.h b/Engine/Plugins/Experimental/GeometryProcessing/Source/DynamicMesh/Public/DynamicMeshOverlay.h
--- a/Engine/Plugins/Experimental/GeometryProcessing/Source/DynamicMesh/Public/DynamicMeshOverlay.h
+++ b/Engine/Plugins/Experimental/GeometryProcessing/Source/DynamicMesh/Public/DynamicMeshOverlay.h
@@ -38,12 +38,12 @@ protected:
 	/** Reference counts of element indices. Iterate over this to find out which elements are valid. */
 	FRefCountVector ElementsRefCounts;
 	/** List of element values */
-	FDynamicVector<RealType> Elements;
+	TDynamicVector<RealType> Elements;
 	/** List of parent vertex indices, one per element */
-	FDynamicVector<int> ParentVertices;
+	TDynamicVector<int> ParentVertices;
 
 	/** List of triangle element-index triplets [Elem0 Elem1 Elem2]*/
-	FDynamicVector<int> ElementTriangles;
+	TDynamicVector<int> ElementTriangles;
 
 	friend class FDynamicMesh3;
 	friend class FDynamicMeshAttributeSet;
diff --git a/Engine/Plugins/Experimental/GeometryProcessing/Source/DynamicMesh/Public/MeshTangents.h b/Engine/Plugins/Experimental/GeometryProcessing/Source/DynamicMesh/Public/MeshTangents.h
--- a/Engine/Plugins/Experimental/GeometryProcessing/Source/DynamicMesh/Public/MeshTangents.h
+++ b/Engine/Plugins/Experimental/GeometryProcessing/Source/DynamicMesh/Public/MeshTangents.h
@@ -6,11 +6,11 @@
 #include "DynamicMeshAttributeSet.h"
 
 /**
- * FMeshTangents is a utility class that can calculate and store various types of
+ * TMeshTangents is a utility class that can calculate and store various types of
  * tangent vectors for a FDynamicMesh.
  */
 template<typename RealType>
-class DYNAMICMESH_API FMeshTangents
+class DYNAMICMESH_API TMeshTangents
 {
 protected:
 	/** Target Mesh */
@@ -21,12 +21,12 @@ protected:
 	TArray<FVector3<RealType>> Bitangents;
 
 public:
-	FMeshTangents()
+	TMeshTangents()
 	{
 		Mesh = nullptr;
 	}
 
-	FMeshTangents(const FDynamicMesh3* Mesh)
+	TMeshTangents(const FDynamicMesh3* Mesh)
 	{
 		SetMesh(Mesh);
 	}
@@ -99,5 +99,5 @@ protected:
 	void Internal_ComputePerTriangleTangents(const FDynamicMeshNormalOverlay* NormalOverlay, const FDynamicMeshUVOverlay* UVOverlay);
 };
 
-typedef FMeshTangents<float> FMeshTangentsf;
-typedef FMeshTangents<double> FMeshTangentsd;
\ No newline at end of file
+typedef TMeshTangents<float> FMeshTangentsf;
+typedef TMeshTangents<double> FMeshTangentsd;
\ No newline at end of file
diff --git a/Engine/Plugins/Experimental/GeometryProcessing/Source/DynamicMesh/Public/Remesher.h b/Engine/Plugins/Experimental/GeometryProcessing/Source/DynamicMesh/Public/Remesher.h
--- a/Engine/Plugins/Experimental/GeometryProcessing/Source/DynamicMesh/Public/Remesher.h
+++ b/Engine/Plugins/Experimental/GeometryProcessing/Source/DynamicMesh/Public/Remesher.h
@@ -218,7 +218,7 @@ protected:
 
 
 
-	FDynamicVector<FVector3d> TempPosBuffer;	// this is a temporary buffer used by smoothing and projection
+	TDynamicVector<FVector3d> TempPosBuffer;	// this is a temporary buffer used by smoothing and projection
 	TArray<bool> TempFlagBuffer;				// list of which indices in TempPosBuffer were modified
 
 	virtual void InitializeVertexBufferForPass();
diff --git a/Engine/Plugins/Experimental/GeometryProcessing/Source/GeometricObjects/Private/CompGeom/PolygonTriangulation.cpp b/Engine/Plugins/Experimental/GeometryProcessing/Source/GeometricObjects/Private/CompGeom/PolygonTriangulation.cpp
--- a/Engine/Plugins/Experimental/GeometryProcessing/Source/GeometricObjects/Private/CompGeom/PolygonTriangulation.cpp
+++ b/Engine/Plugins/Experimental/GeometryProcessing/Source/GeometricObjects/Private/CompGeom/PolygonTriangulation.cpp
@@ -25,7 +25,7 @@ void PolygonTriangulation::TriangulateSimplePolygon(const TArray<FVector2<T>>& V
 	{
 		static inline bool IsTriangleFlipped(T OrientationSign, const FVector2<T>& VertexPositionA, const FVector2<T>& VertexPositionB, const FVector2<T>& VertexPositionC)
 		{
-			T TriSignedArea = FTriangle2<T>::SignedArea(VertexPositionA, VertexPositionB, VertexPositionC);
+			T TriSignedArea = TTriangle2<T>::SignedArea(VertexPositionA, VertexPositionB, VertexPositionC);
 			return TriSignedArea * OrientationSign < 0;
 		}
 	};
@@ -102,7 +102,7 @@ void PolygonTriangulation::TriangulateSimplePolygon(const TArray<FVector2<T>>& V
 					// Test every other remaining vertex to make sure that it doesn't lie inside our potential ear
 					// triangle.  If we find a vertex that's inside the triangle, then it cannot actually be an ear.
 					const FVector2<T>& TestVertexPosition = VertexPositions[TestVertexNumber];
-					if (FTriangle2<T>::IsInside(PrevVertexPosition, EarVertexPosition, NextVertexPosition, TestVertexPosition))
+					if (TTriangle2<T>::IsInside(PrevVertexPosition, EarVertexPosition, NextVertexPosition, TestVertexPosition))
 					{
 						bIsEar = false;
 						break;
diff --git a/Engine/Plugins/Experimental/GeometryProcessing/Source/GeometricObjects/Private/Util/SmallListSet.cpp b/Engine/Plugins/Experimental/GeometryProcessing/Source/GeometricObjects/Private/Util/SmallListSet.cpp
--- a/Engine/Plugins/Experimental/GeometryProcessing/Source/GeometricObjects/Private/Util/SmallListSet.cpp
+++ b/Engine/Plugins/Experimental/GeometryProcessing/Source/GeometricObjects/Private/Util/SmallListSet.cpp
@@ -7,21 +7,21 @@ const int32 FSmallListSet::NullValue = -1;
 
 FSmallListSet::FSmallListSet()
 {
-	ListHeads = FDynamicVector<int32>();
-	LinkedListElements = FDynamicVector<int32>();
+	ListHeads = TDynamicVector<int32>();
+	LinkedListElements = TDynamicVector<int32>();
 	FreeHeadIndex = NullValue;
-	ListBlocks = FDynamicVector<int32>();
-	FreeBlocks = FDynamicVector<int32>();
+	ListBlocks = TDynamicVector<int32>();
+	FreeBlocks = TDynamicVector<int32>();
 }
 
 
 FSmallListSet::FSmallListSet(const FSmallListSet& copy)
 {
-	LinkedListElements = FDynamicVector<int32>(copy.LinkedListElements);
+	LinkedListElements = TDynamicVector<int32>(copy.LinkedListElements);
 	FreeHeadIndex = copy.FreeHeadIndex;
-	ListHeads = FDynamicVector<int32>(copy.ListHeads);
-	ListBlocks = FDynamicVector<int32>(copy.ListBlocks);
-	FreeBlocks = FDynamicVector<int32>(copy.FreeBlocks);
+	ListHeads = TDynamicVector<int32>(copy.ListHeads);
+	ListBlocks = TDynamicVector<int32>(copy.ListBlocks);
+	FreeBlocks = TDynamicVector<int32>(copy.FreeBlocks);
 }
 
 
diff --git a/Engine/Plugins/Experimental/GeometryProcessing/Source/GeometricObjects/Public/BoxTypes.h b/Engine/Plugins/Experimental/GeometryProcessing/Source/GeometricObjects/Public/BoxTypes.h
--- a/Engine/Plugins/Experimental/GeometryProcessing/Source/GeometricObjects/Public/BoxTypes.h
+++ b/Engine/Plugins/Experimental/GeometryProcessing/Source/GeometricObjects/Public/BoxTypes.h
@@ -7,23 +7,23 @@
 #include "VectorTypes.h"
 
 template <typename RealType>
-struct FInterval1
+struct TInterval1
 {
 	RealType Min;
 	RealType Max;
 
-	FInterval1()
+	TInterval1()
 	{
 	}
-	FInterval1(const RealType& Min, const RealType& Max)
+	TInterval1(const RealType& Min, const RealType& Max)
 	{
 		this->Min = Min;
 		this->Max = Max;
 	}
 
-	static FInterval1<RealType> Empty()
+	static TInterval1<RealType> Empty()
 	{
-		return FInterval1(TNumericLimits<RealType>::Max(), -TNumericLimits<RealType>::Max());
+		return TInterval1(TNumericLimits<RealType>::Max(), -TNumericLimits<RealType>::Max());
 	}
 
 	RealType Center() const
@@ -57,12 +57,12 @@ struct FInterval1
 		return D >= Min && D <= Max;
 	}
 
-	bool Overlaps(const FInterval1<RealType>& O) const
+	bool Overlaps(const TInterval1<RealType>& O) const
 	{
 		return !(O.Min > Max || O.Max < Min);
 	}
 
-	RealType SquaredDist(const FInterval1<RealType>& O) const
+	RealType SquaredDist(const TInterval1<RealType>& O) const
 	{
 		if (Max < O.Min)
 		{
@@ -77,7 +77,7 @@ struct FInterval1
 			return 0;
 		}
 	}
-	RealType Dist(const FInterval1<RealType>& O) const
+	RealType Dist(const TInterval1<RealType>& O) const
 	{
 		if (Max < O.Min)
 		{
@@ -93,13 +93,13 @@ struct FInterval1
 		}
 	}
 
-	FInterval1<RealType> IntersectionWith(const FInterval1<RealType>& O) const
+	TInterval1<RealType> IntersectionWith(const TInterval1<RealType>& O) const
 	{
 		if (O.Min > Max || O.Max < Min)
 		{
-			return FInterval1<RealType>::Empty();
+			return TInterval1<RealType>::Empty();
 		}
-		return FInterval1<RealType>(TMathUtil<RealType>::Max(Min, O.Min), TMathUtil<RealType>::Min(Max, O.Max));
+		return TInterval1<RealType>(TMathUtil<RealType>::Max(Min, O.Min), TMathUtil<RealType>::Min(Max, O.Max));
 	}
 
 	/**
@@ -141,7 +141,7 @@ struct FInterval1
 		}
 	}
 
-	void Set(FInterval1 O)
+	void Set(TInterval1 O)
 	{
 		Min = O.Min;
 		Max = O.Max;
@@ -153,46 +153,48 @@ struct FInterval1
 		Max = B;
 	}
 
-	FInterval1 operator-(FInterval1 V) const
+	TInterval1 operator-(TInterval1 V) const
 	{
-		return FInterval1(-V.Min, -V.Max);
+		return TInterval1(-V.Min, -V.Max);
 	}
 
-	FInterval1 operator+(RealType f) const
+	TInterval1 operator+(RealType f) const
 	{
-		return FInterval1(Min + f, Max + f);
+		return TInterval1(Min + f, Max + f);
 	}
 
-	FInterval1 operator-(RealType f) const
+	TInterval1 operator-(RealType f) const
 	{
-		return FInterval1(Min - f, Max - f);
+		return TInterval1(Min - f, Max - f);
 	}
 
-	FInterval1 operator*(RealType f) const
+	TInterval1 operator*(RealType f) const
 	{
-		return FInterval1(Min * f, Max * f);
+		return TInterval1(Min * f, Max * f);
 	}
 };
 
-typedef FInterval1<float> FInterval1f;
-typedef FInterval1<double> FInterval1d;
-typedef FInterval1<int> FInterval1i;
+typedef TInterval1<float> FInterval1f;
+typedef TInterval1<double> FInterval1d;
+typedef TInterval1<int> FInterval1i;
+
+
 
 template <typename RealType>
-struct FAxisAlignedBox3
+struct TAxisAlignedBox3
 {
 	FVector3<RealType> Min;
 	FVector3<RealType> Max;
 
-	FAxisAlignedBox3()
+	TAxisAlignedBox3()
 	{
 	}
-	FAxisAlignedBox3(const FVector3<RealType>& Min, const FVector3<RealType>& Max)
+	TAxisAlignedBox3(const FVector3<RealType>& Min, const FVector3<RealType>& Max)
 	{
 		this->Min = Min;
 		this->Max = Max;
 	}
-	FAxisAlignedBox3(const FAxisAlignedBox3& Box, const TFunction<FVector3<RealType>(const FVector3<RealType>&)> TransformF)
+	TAxisAlignedBox3(const TAxisAlignedBox3& Box, const TFunction<FVector3<RealType>(const FVector3<RealType>&)> TransformF)
 	{
 		if (TransformF == nullptr)
 		{
@@ -216,7 +218,7 @@ struct FAxisAlignedBox3
 		FVector MaxV((float)Max.X, (float)Max.Y, (float)Max.Z);
 		return FBox(MinV, MaxV);
 	}
-	FAxisAlignedBox3(const FBox& Box)
+	TAxisAlignedBox3(const FBox& Box)
 	{
 		Min = Box.Min;
 		Max = Box.Max;
@@ -235,9 +237,9 @@ struct FAxisAlignedBox3
 		return FVector3<RealType>(X, Y, Z);
 	}
 
-	static FAxisAlignedBox3<RealType> Empty()
+	static TAxisAlignedBox3<RealType> Empty()
 	{
-		return FAxisAlignedBox3(
+		return TAxisAlignedBox3(
 			FVector3<RealType>(TNumericLimits<RealType>::Max(), TNumericLimits<RealType>::Max(), TNumericLimits<RealType>::Max()),
 			FVector3<RealType>(-TNumericLimits<RealType>::Max(), -TNumericLimits<RealType>::Max(), -TNumericLimits<RealType>::Max()));
 	}
@@ -283,7 +285,7 @@ struct FAxisAlignedBox3
 		}
 	}
 
-	void Contain(const FAxisAlignedBox3<RealType>& Other)
+	void Contain(const TAxisAlignedBox3<RealType>& Other)
 	{
 		// todo: can be optimized
 		Contain(Other.Min);
@@ -295,14 +297,14 @@ struct FAxisAlignedBox3
 		return (Min.X <= V.X) && (Min.Y <= V.Y) && (Min.Z <= V.Z) && (Max.X >= V.X) && (Max.Y >= V.Y) && (Max.Z >= V.Z);
 	}
 
-	FAxisAlignedBox3<RealType> Intersect(const FAxisAlignedBox3<RealType>& Box) const
+	TAxisAlignedBox3<RealType> Intersect(const TAxisAlignedBox3<RealType>& Box) const
 	{
-		FAxisAlignedBox3<RealType> Intersection(
+		TAxisAlignedBox3<RealType> Intersection(
 			FVector3<RealType>(TMathUtil<RealType>::Max(Min.X, Box.Min.X), TMathUtil<RealType>::Max(Min.Y, Box.Min.Y), TMathUtil<RealType>::Max(Min.Z, Box.Min.Z)),
 			FVector3<RealType>(TMathUtil<RealType>::Min(Max.X, Box.Max.X), TMathUtil<RealType>::Min(Max.Y, Box.Max.Y), TMathUtil<RealType>::Min(Max.Z, Box.Max.Z)));
 		if (Intersection.Height() <= 0 || Intersection.Width() <= 0 || Intersection.Depth() <= 0)
 		{
-			return FAxisAlignedBox3<RealType>::Empty();
+			return TAxisAlignedBox3<RealType>::Empty();
 		}
 		else
 		{
@@ -310,7 +312,7 @@ struct FAxisAlignedBox3
 		}
 	}
 
-	bool Intersects(FAxisAlignedBox3 Box) const
+	bool Intersects(TAxisAlignedBox3 Box) const
 	{
 		return !((Box.Max.X <= Min.X) || (Box.Min.X >= Max.X) || (Box.Max.Y <= Min.Y) || (Box.Min.Y >= Max.Y) || (Box.Max.Z <= Min.Z) || (Box.Min.Z >= Max.Z));
 	}
@@ -323,7 +325,7 @@ struct FAxisAlignedBox3
 		return dx * dx + dy * dy + dz * dz;
 	}
 
-	RealType DistanceSquared(const FAxisAlignedBox3<RealType>& Box)
+	RealType DistanceSquared(const TAxisAlignedBox3<RealType>& Box)
 	{
 		// compute lensqr( max(0, abs(center1-center2) - (extent1+extent2)) )
 		RealType delta_x = TMathUtil<RealType>::Abs((Box.Min.X + Box.Max.X) - (Min.X + Max.X))
@@ -384,23 +386,23 @@ struct FAxisAlignedBox3
 };
 
 template <typename RealType>
-struct FAxisAlignedBox2
+struct TAxisAlignedBox2
 {
 	FVector2<RealType> Min;
 	FVector2<RealType> Max;
 
-	FAxisAlignedBox2()
+	TAxisAlignedBox2()
 	{
 	}
-	FAxisAlignedBox2(const FVector2<RealType>& Min, const FVector2<RealType>& Max)
+	TAxisAlignedBox2(const FVector2<RealType>& Min, const FVector2<RealType>& Max)
 		: Min(Min), Max(Max)
 	{
 	}
-	FAxisAlignedBox2(RealType SquareSize)
+	TAxisAlignedBox2(RealType SquareSize)
 		: Min((RealType)0, (RealType)0), Max(SquareSize, SquareSize)
 	{
 	}
-	FAxisAlignedBox2(RealType Width, RealType Height)
+	TAxisAlignedBox2(RealType Width, RealType Height)
 		: Min((RealType)0, (RealType)0), Max(Width, Height)
 	{
 	}
@@ -411,15 +413,15 @@ struct FAxisAlignedBox2
 		FVector2D MaxV((float)Max.X, (float)Max.Y);
 		return FBox2D(MinV, MaxV);
 	}
-	FAxisAlignedBox2(const FBox2D& Box)
+	TAxisAlignedBox2(const FBox2D& Box)
 	{
 		Min = Box.Min;
 		Max = Box.Max;
 	}
 
-	static FAxisAlignedBox2<RealType> Empty()
+	static TAxisAlignedBox2<RealType> Empty()
 	{
-		return FAxisAlignedBox2(
+		return TAxisAlignedBox2(
 			FVector2<RealType>(TNumericLimits<RealType>::Max(), TNumericLimits<RealType>::Max()),
 			FVector2<RealType>(-TNumericLimits<RealType>::Max(), -TNumericLimits<RealType>::Max()));
 	}
@@ -456,7 +458,7 @@ struct FAxisAlignedBox2
 		}
 	}
 
-	inline void Contain(const FAxisAlignedBox2<RealType>& Other)
+	inline void Contain(const TAxisAlignedBox2<RealType>& Other)
 	{
 		// todo: can be optimized
 		Contain(Other.Min);
@@ -476,19 +478,19 @@ struct FAxisAlignedBox2
 		return (Min.X <= V.X) && (Min.Y <= V.Y) && (Max.X >= V.X) && (Max.Y >= V.Y);
 	}
 
-	bool Intersects(const FAxisAlignedBox2<RealType>& Box) const
+	bool Intersects(const TAxisAlignedBox2<RealType>& Box) const
 	{
 		return !((Box.Max.X < Min.X) || (Box.Min.X > Max.X) || (Box.Max.Y < Min.Y) || (Box.Min.Y > Max.Y));
 	}
 
-	FAxisAlignedBox2<RealType> Intersect(const FAxisAlignedBox2<RealType> &Box) const
+	TAxisAlignedBox2<RealType> Intersect(const TAxisAlignedBox2<RealType> &Box) const
 	{
-		FAxisAlignedBox2<RealType> Intersection(
+		TAxisAlignedBox2<RealType> Intersection(
 			FVector2<RealType>(TMathUtil<RealType>::Max(Min.X, Box.Min.X), TMathUtil<RealType>::Max(Min.Y, Box.Min.Y)),
 			FVector2<RealType>(TMathUtil<RealType>::Min(Max.X, Box.Max.X), TMathUtil<RealType>::Min(Max.Y, Box.Max.Y)));
 		if (Intersection.Height() <= 0 || Intersection.Width() <= 0)
 		{
-			return FAxisAlignedBox2<RealType>::Empty();
+			return TAxisAlignedBox2<RealType>::Empty();
 		}
 		else
 		{
@@ -534,9 +536,9 @@ struct FAxisAlignedBox2
 	}
 };
 
-typedef FAxisAlignedBox2<float> FAxisAlignedBox2f;
-typedef FAxisAlignedBox2<double> FAxisAlignedBox2d;
-typedef FAxisAlignedBox2<int> FAxisAlignedBox2i;
-typedef FAxisAlignedBox3<float> FAxisAlignedBox3f;
-typedef FAxisAlignedBox3<double> FAxisAlignedBox3d;
-typedef FAxisAlignedBox3<int> FAxisAlignedBox3i;
+typedef TAxisAlignedBox2<float> FAxisAlignedBox2f;
+typedef TAxisAlignedBox2<double> FAxisAlignedBox2d;
+typedef TAxisAlignedBox2<int> FAxisAlignedBox2i;
+typedef TAxisAlignedBox3<float> FAxisAlignedBox3f;
+typedef TAxisAlignedBox3<double> FAxisAlignedBox3d;
+typedef TAxisAlignedBox3<int> FAxisAlignedBox3i;
diff --git a/Engine/Plugins/Experimental/GeometryProcessing/Source/GeometricObjects/Public/Curve/DynamicGraph.h b/Engine/Plugins/Experimental/GeometryProcessing/Source/GeometricObjects/Public/Curve/DynamicGraph.h
--- a/Engine/Plugins/Experimental/GeometryProcessing/Source/GeometricObjects/Public/Curve/DynamicGraph.h
+++ b/Engine/Plugins/Experimental/GeometryProcessing/Source/GeometricObjects/Public/Curve/DynamicGraph.h
@@ -41,7 +41,7 @@ protected:
 	FSmallListSet vertex_edges;
 
 	FRefCountVector edges_refcount;
-	FDynamicVector<FEdge> edges; // each edge is a tuple (v0,v0,GroupID)
+	TDynamicVector<FEdge> edges; // each edge is a tuple (v0,v0,GroupID)
 
 	int timestamp = 0;
 	int shape_timestamp = 0;
diff --git a/Engine/Plugins/Experimental/GeometryProcessing/Source/GeometricObjects/Public/Curve/DynamicGraph2.h b/Engine/Plugins/Experimental/GeometryProcessing/Source/GeometricObjects/Public/Curve/DynamicGraph2.h
--- a/Engine/Plugins/Experimental/GeometryProcessing/Source/GeometricObjects/Public/Curve/DynamicGraph2.h
+++ b/Engine/Plugins/Experimental/GeometryProcessing/Source/GeometricObjects/Public/Curve/DynamicGraph2.h
@@ -18,10 +18,10 @@
 template <typename T>
 class FDynamicGraph2 : public FDynamicGraph
 {
-	FAxisAlignedBox2<T> cached_bounds;
+	TAxisAlignedBox2<T> cached_bounds;
 	int cached_bounds_timestamp = -1;
 
-	FDynamicVectorN<T, 2> vertices;
+	TDynamicVectorN<T, 2> vertices;
 
 public:
 	static FVector2<T> InvalidVertex()
@@ -56,11 +56,11 @@ public:
 		return false;
 	}
 
-	FSegment2<T> GetEdgeSegment(int EID) const
+	TSegment2<T> GetEdgeSegment(int EID) const
 	{
 		checkf(edges_refcount.IsValid(EID), TEXT("FDynamicGraph2.GetEdgeSegment: invalid segment with id %d"), EID);
 		const FEdge& e = edges[EID];
-		return FSegment2<T>(
+		return TSegment2<T>(
 			vertices.AsVector2(e.A),
 			vertices.AsVector2(e.B));
 	}
@@ -162,7 +162,7 @@ public:
 	// compute vertex bounding box
 	FAxisAlignedBox2d GetBounds() const
 	{
-		FAxisAlignedBox2<T> AABB;
+		TAxisAlignedBox2<T> AABB;
 		for (const FVector2<T>& V : Vertices())
 		{
 			AABB.Contain(V);
@@ -171,7 +171,7 @@ public:
 	}
 
 	//! cached bounding box, lazily re-computed on access if mesh has changed
-	FAxisAlignedBox2<T> CachedBounds()
+	TAxisAlignedBox2<T> CachedBounds()
 	{
 		if (cached_bounds_timestamp != Timestamp())
 		{
diff --git a/Engine/Plugins/Experimental/GeometryProcessing/Source/GeometricObjects/Public/Distance/DistLine3Ray3.h b/Engine/Plugins/Experimental/GeometryProcessing/Source/GeometricObjects/Public/Distance/DistLine3Ray3.h
--- a/Engine/Plugins/Experimental/GeometryProcessing/Source/GeometricObjects/Public/Distance/DistLine3Ray3.h
+++ b/Engine/Plugins/Experimental/GeometryProcessing/Source/GeometricObjects/Public/Distance/DistLine3Ray3.h
@@ -12,12 +12,12 @@
  * Compute distance between 3D line and 3D ray
  */
 template <typename Real>
-class FDistLine3Ray3
+class TDistLine3Ray3
 {
 public:
 	// Input
-	FLine3<Real> Line;
-	FRay3<Real> Ray;
+	TLine3<Real> Line;
+	TRay3<Real> Ray;
 
 	// Results
 	Real DistanceSquared = -1.0;
@@ -28,7 +28,7 @@ public:
 	Real RayParameter;
 
 
-	FDistLine3Ray3(const FLine3<Real>& LineIn, const FRay3<Real>& RayIn)
+	TDistLine3Ray3(const TLine3<Real>& LineIn, const TRay3<Real>& RayIn)
 	{
 		Line = LineIn;
 		Ray = RayIn;
@@ -103,5 +103,5 @@ public:
 	}
 };
 
-typedef FDistLine3Ray3<float> FDistLine3Ray3f;
-typedef FDistLine3Ray3<double> FDistLine3Ray3d;
\ No newline at end of file
+typedef TDistLine3Ray3<float> FDistLine3Ray3f;
+typedef TDistLine3Ray3<double> FDistLine3Ray3d;
\ No newline at end of file
diff --git a/Engine/Plugins/Experimental/GeometryProcessing/Source/GeometricObjects/Public/Distance/DistLine3Segment3.h b/Engine/Plugins/Experimental/GeometryProcessing/Source/GeometricObjects/Public/Distance/DistLine3Segment3.h
--- a/Engine/Plugins/Experimental/GeometryProcessing/Source/GeometricObjects/Public/Distance/DistLine3Segment3.h
+++ b/Engine/Plugins/Experimental/GeometryProcessing/Source/GeometricObjects/Public/Distance/DistLine3Segment3.h
@@ -14,12 +14,12 @@
 * Compute unsigned distance between 3D line and 3D segment
 */
 template <typename Real>
-class FDistLine3Segment3
+class TDistLine3Segment3
 {
 public:
 	// Input
-	FLine3<Real> Line;
-	FSegment3<Real> Segment;
+	TLine3<Real> Line;
+	TSegment3<Real> Segment;
 
 	// Output
 	Real DistanceSquared = -1.0;
@@ -27,7 +27,7 @@ public:
 	FVector3<Real> LineClosest, SegmentClosest;
 
 
-	FDistLine3Segment3(const FLine3<Real>& LineIn, const FSegment3<Real>& SegmentIn) : Line(LineIn), Segment(SegmentIn)
+	TDistLine3Segment3(const TLine3<Real>& LineIn, const TSegment3<Real>& SegmentIn) : Line(LineIn), Segment(SegmentIn)
 	{
 	}
 
@@ -117,6 +117,6 @@ public:
 	}
 };
 
-typedef FDistLine3Segment3<float> FDistLine3Segment3f;
-typedef FDistLine3Segment3<double> FDistLine3Segment3d;
+typedef TDistLine3Segment3<float> FDistLine3Segment3f;
+typedef TDistLine3Segment3<double> FDistLine3Segment3d;
 
diff --git a/Engine/Plugins/Experimental/GeometryProcessing/Source/GeometricObjects/Public/Distance/DistLine3Triangle3.h b/Engine/Plugins/Experimental/GeometryProcessing/Source/GeometricObjects/Public/Distance/DistLine3Triangle3.h
--- a/Engine/Plugins/Experimental/GeometryProcessing/Source/GeometricObjects/Public/Distance/DistLine3Triangle3.h
+++ b/Engine/Plugins/Experimental/GeometryProcessing/Source/GeometricObjects/Public/Distance/DistLine3Triangle3.h
@@ -15,12 +15,12 @@
 * Compute unsigned distance between 3D line and 3D triangle
 */
 template <typename Real>
-class FDistLine3Triangle3
+class TDistLine3Triangle3
 {
 public:
 	// Input
-	FLine3<Real> Line;
-	FTriangle3<Real> Triangle;
+	TLine3<Real> Line;
+	TTriangle3<Real> Triangle;
 
 	// Output
 	Real DistanceSquared = -1.0;
@@ -28,7 +28,7 @@ public:
 	FVector3<Real> LineClosest, TriangleClosest, TriangleBaryCoords;
 
 
-	FDistLine3Triangle3(const FLine3<Real>& LineIn, const FTriangle3<Real>& TriangleIn) : Line(LineIn), Triangle(TriangleIn)
+	TDistLine3Triangle3(const TLine3<Real>& LineIn, const TTriangle3<Real>& TriangleIn) : Line(LineIn), Triangle(TriangleIn)
 	{
 	}
 
@@ -101,8 +101,8 @@ public:
 		Real sqrDist = TMathUtil<Real>::MaxReal;
 		for (int i0 = 2, i1 = 0; i1 < 3; i0 = i1++)
 		{
-			FSegment3<Real> segment(Triangle.V[i0], Triangle.V[i1]);
-			FDistLine3Segment3<Real> queryLS(Line, segment);
+			TSegment3<Real> segment(Triangle.V[i0], Triangle.V[i1]);
+			TDistLine3Segment3<Real> queryLS(Line, segment);
 			Real sqrDistTmp = queryLS.GetSquared();
 			if (sqrDistTmp < sqrDist)
 			{
@@ -122,6 +122,6 @@ public:
 	}
 };
 
-typedef FDistLine3Triangle3<float> FDistLine3Triangle3f;
-typedef FDistLine3Triangle3<double> FDistLine3Triangle3d;
+typedef TDistLine3Triangle3<float> FDistLine3Triangle3f;
+typedef TDistLine3Triangle3<double> FDistLine3Triangle3d;
 
diff --git a/Engine/Plugins/Experimental/GeometryProcessing/Source/GeometricObjects/Public/Distance/DistPoint3Triangle3.h b/Engine/Plugins/Experimental/GeometryProcessing/Source/GeometricObjects/Public/Distance/DistPoint3Triangle3.h
--- a/Engine/Plugins/Experimental/GeometryProcessing/Source/GeometricObjects/Public/Distance/DistPoint3Triangle3.h
+++ b/Engine/Plugins/Experimental/GeometryProcessing/Source/GeometricObjects/Public/Distance/DistPoint3Triangle3.h
@@ -12,19 +12,19 @@
  * Compute unsigned distance between 3D point and 3D triangle
  */
 template <typename Real>
-class FDistPoint3Triangle3
+class TDistPoint3Triangle3
 {
 public:
 	// Input
 	FVector3<Real> Point;
-	FTriangle3<Real> Triangle;
+	TTriangle3<Real> Triangle;
 
 	// Results
 	FVector3<Real> TriangleBaryCoords;
 	FVector3<Real> ClosestTrianglePoint;  // do we need this just use Triangle.BarycentricPoint
 
 
-	FDistPoint3Triangle3(const FVector3<Real>& PointIn, const FTriangle3<Real>& TriangleIn)
+	TDistPoint3Triangle3(const FVector3<Real>& PointIn, const TTriangle3<Real>& TriangleIn)
 	{
 		Point = PointIn;
 		Triangle = TriangleIn;
@@ -270,6 +270,6 @@ public:
 
 };
 
-typedef FDistPoint3Triangle3<float> FDistPoint3Triangle3f;
-typedef FDistPoint3Triangle3<double> FDistPoint3Triangle3d;
+typedef TDistPoint3Triangle3<float> FDistPoint3Triangle3f;
+typedef TDistPoint3Triangle3<double> FDistPoint3Triangle3d;
 
diff --git a/Engine/Plugins/Experimental/GeometryProcessing/Source/GeometricObjects/Public/Distance/DistRay3Segment3.h b/Engine/Plugins/Experimental/GeometryProcessing/Source/GeometricObjects/Public/Distance/DistRay3Segment3.h
--- a/Engine/Plugins/Experimental/GeometryProcessing/Source/GeometricObjects/Public/Distance/DistRay3Segment3.h
+++ b/Engine/Plugins/Experimental/GeometryProcessing/Source/GeometricObjects/Public/Distance/DistRay3Segment3.h
@@ -16,8 +16,8 @@ class TDistRay3Segment3
 {
 public:
 	// Input
-	FRay3<Real> Ray;
-	FSegment3<Real> Segment;
+	TRay3<Real> Ray;
+	TSegment3<Real> Segment;
 
 	// Results
 	Real DistanceSquared = -1.0;
@@ -28,7 +28,7 @@ public:
 	Real SegmentParameter;
 
 
-	TDistRay3Segment3(const FRay3<Real>& RayIn, const FSegment3<Real>& SegmentIn)
+	TDistRay3Segment3(const TRay3<Real>& RayIn, const TSegment3<Real>& SegmentIn)
 	{
 		Ray = RayIn;
 		Segment = SegmentIn;
@@ -216,7 +216,7 @@ public:
 
 
 
-	static double SquaredDistance(const FRay3<Real>& Ray, const FSegment3<Real>& Segment,
+	static double SquaredDistance(const TRay3<Real>& Ray, const TSegment3<Real>& Segment,
 		Real& RayParam, Real& SegParam)
 	{
 		FVector3<Real> diff = Ray.Origin - Segment.Center;
diff --git a/Engine/Plugins/Experimental/GeometryProcessing/Source/GeometricObjects/Public/Distance/DistSegment3Triangle3.h b/Engine/Plugins/Experimental/GeometryProcessing/Source/GeometricObjects/Public/Distance/DistSegment3Triangle3.h
--- a/Engine/Plugins/Experimental/GeometryProcessing/Source/GeometricObjects/Public/Distance/DistSegment3Triangle3.h
+++ b/Engine/Plugins/Experimental/GeometryProcessing/Source/GeometricObjects/Public/Distance/DistSegment3Triangle3.h
@@ -15,19 +15,19 @@
 * Compute unsigned distance between 3D segment and 3D triangle
 */
 template <typename Real>
-class FDistSegment3Triangle3
+class TDistSegment3Triangle3
 {
 public:
 	// Input
-	FSegment3<Real> Segment;
-	FTriangle3<Real> Triangle;
+	TSegment3<Real> Segment;
+	TTriangle3<Real> Triangle;
 
 	// Output
 	Real DistanceSquared = -1.0;
 	Real SegmentParameter;
 	FVector3<Real> TriangleClosest, TriangleBaryCoords, SegmentClosest;
 	
-	FDistSegment3Triangle3(const FSegment3<Real>& SegmentIn, const FTriangle3<Real>& TriangleIn) : Segment(SegmentIn), Triangle(TriangleIn)
+	TDistSegment3Triangle3(const TSegment3<Real>& SegmentIn, const TTriangle3<Real>& TriangleIn) : Segment(SegmentIn), Triangle(TriangleIn)
 	{
 	}
 
@@ -47,8 +47,8 @@ public:
 			return DistanceSquared;
 		}
 
-		FLine3<Real> line(Segment.Center, Segment.Direction);
-		FDistLine3Triangle3<Real> queryLT(line, Triangle);
+		TLine3<Real> line(Segment.Center, Segment.Direction);
+		TDistLine3Triangle3<Real> queryLT(line, Triangle);
 		double sqrDist = queryLT.GetSquared();
 		SegmentParameter = queryLT.LineParam;
 
@@ -60,7 +60,7 @@ public:
 			}
 			else {
 				SegmentClosest = Segment.EndPoint();
-				FDistPoint3Triangle3<Real> queryPT(SegmentClosest, Triangle);
+				TDistPoint3Triangle3<Real> queryPT(SegmentClosest, Triangle);
 				sqrDist = queryPT.GetSquared();
 				TriangleClosest = queryPT.ClosestTrianglePoint;
 				SegmentParameter = Segment.Extent;
@@ -69,7 +69,7 @@ public:
 		}
 		else {
 			SegmentClosest = Segment.StartPoint();
-			FDistPoint3Triangle3<Real> queryPT(SegmentClosest, Triangle);
+			TDistPoint3Triangle3<Real> queryPT(SegmentClosest, Triangle);
 			sqrDist = queryPT.GetSquared();
 			TriangleClosest = queryPT.ClosestTrianglePoint;
 			SegmentParameter = -Segment.Extent;
@@ -81,6 +81,6 @@ public:
 	}
 };
 
-typedef FDistSegment3Triangle3<float> FDistSegment3Triangle3f;
-typedef FDistSegment3Triangle3<double> FDistSegment3Triangle3d;
+typedef TDistSegment3Triangle3<float> FDistSegment3Triangle3f;
+typedef TDistSegment3Triangle3<double> FDistSegment3Triangle3d;
 
diff --git a/Engine/Plugins/Experimental/GeometryProcessing/Source/GeometricObjects/Public/Distance/DistTriangle3Triangle3.h b/Engine/Plugins/Experimental/GeometryProcessing/Source/GeometricObjects/Public/Distance/DistTriangle3Triangle3.h
--- a/Engine/Plugins/Experimental/GeometryProcessing/Source/GeometricObjects/Public/Distance/DistTriangle3Triangle3.h
+++ b/Engine/Plugins/Experimental/GeometryProcessing/Source/GeometricObjects/Public/Distance/DistTriangle3Triangle3.h
@@ -16,11 +16,11 @@
 * Compute unsigned distance between 3D segment and 3D triangle
 */
 template <typename Real>
-class FDistTriangle3Triangle3
+class TDistTriangle3Triangle3
 {
 public:
 	// Input
-	FTriangle3<Real> Triangle[2];
+	TTriangle3<Real> Triangle[2];
 
 	// Output
 	Real DistanceSquared = -1.0;
@@ -28,10 +28,10 @@ public:
 	FVector3<Real> TriangleBaryCoords[2];
 
 
-	FDistTriangle3Triangle3()
+	TDistTriangle3Triangle3()
 	{
 	}
-	FDistTriangle3Triangle3(const FTriangle3<Real>& TriangleA, const FTriangle3<Real>& TriangleB)
+	TDistTriangle3Triangle3(const TTriangle3<Real>& TriangleA, const TTriangle3<Real>& TriangleB)
 	{
 		Triangle[0] = TriangleA;
 		Triangle[1] = TriangleB;
@@ -60,9 +60,9 @@ public:
 		int i0, i1;
 		for (i0 = 2, i1 = 0; i1 < 3; i0 = i1++)
 		{
-			FSegment3<Real> edge(Triangle[0].V[i0], Triangle[0].V[i1]);
+			TSegment3<Real> edge(Triangle[0].V[i0], Triangle[0].V[i1]);
 
-			FDistSegment3Triangle3<Real> queryST(edge, Triangle[1]);
+			TDistSegment3Triangle3<Real> queryST(edge, Triangle[1]);
 			sqrDistTmp = queryST.GetSquared();
 			if (sqrDistTmp < sqrDist)
 			{
@@ -87,9 +87,9 @@ public:
 		// Compare edges of Triangle[1] to the interior of Triangle[0].
 		for (i0 = 2, i1 = 0; i1 < 3; i0 = i1++)
 		{
-			FSegment3<Real> edge(Triangle[1].V[i0], Triangle[1].V[i1]);
+			TSegment3<Real> edge(Triangle[1].V[i0], Triangle[1].V[i1]);
 
-			FDistSegment3Triangle3<Real> queryST(edge, Triangle[0]);
+			TDistSegment3Triangle3<Real> queryST(edge, Triangle[0]);
 			sqrDistTmp = queryST.GetSquared();
 			if (sqrDistTmp < sqrDist)
 			{
@@ -117,6 +117,6 @@ public:
 	}
 };
 
-typedef FDistTriangle3Triangle3<float> FDistTriangle3Triangle3f;
-typedef FDistTriangle3Triangle3<double> FDistTriangle3Triangle3d;
+typedef TDistTriangle3Triangle3<float> FDistTriangle3Triangle3f;
+typedef TDistTriangle3Triangle3<double> FDistTriangle3Triangle3d;
 
diff --git a/Engine/Plugins/Experimental/GeometryProcessing/Source/GeometricObjects/Public/FrameTypes.h b/Engine/Plugins/Experimental/GeometryProcessing/Source/GeometricObjects/Public/FrameTypes.h
--- a/Engine/Plugins/Experimental/GeometryProcessing/Source/GeometricObjects/Public/FrameTypes.h
+++ b/Engine/Plugins/Experimental/GeometryProcessing/Source/GeometricObjects/Public/FrameTypes.h
@@ -9,14 +9,14 @@
 #include "Quaternion.h"
 
 /**
- * FFrame3 is an object that represents an oriented 3D coordinate frame, ie orthogonal X/Y/Z axes at a point in space.
+ * TFrame3 is an object that represents an oriented 3D coordinate frame, ie orthogonal X/Y/Z axes at a point in space.
  * One can think of this Frame as a local coordinate space measured along these axes.
  * Functions are provided to map geometric objects to/from the Frame coordinate space.
  * 
  * Internally the representation is the same as an FTransform, except a Frame has no Scale.
  */
 template<typename RealType>
-struct FFrame3
+struct TFrame3
 {
 	/**
 	 * Origin of the frame
@@ -26,30 +26,30 @@ struct FFrame3
 	/**
 	 * Rotation of the frame. Think of this as the rotation of the unit X/Y/Z axes to the 3D frame axes.
 	 */
-	FQuaternion<RealType> Rotation;
+	TQuaternion<RealType> Rotation;
 
 	/**
 	 * Construct a frame positioned at (0,0,0) aligned to the unit axes
 	 */
-	FFrame3()
+	TFrame3()
 	{
 		Origin = FVector3<RealType>::Zero();
-		Rotation = FQuaternion<RealType>::Identity();
+		Rotation = TQuaternion<RealType>::Identity();
 	}
 
 	/**
 	 * Construct a frame at the given Origin aligned to the unit axes
 	 */
-	FFrame3(const FVector3<RealType>& OriginIn)
+	TFrame3(const FVector3<RealType>& OriginIn)
 	{
 		Origin = OriginIn;
-		Rotation = FQuaternion<RealType>::Identity();
+		Rotation = TQuaternion<RealType>::Identity();
 	}
 
 	/**
 	 * Construct a Frame from the given Origin and Rotation
 	 */
-	FFrame3(const FVector3<RealType>& OriginIn, const FQuaternion<RealType> RotationIn)
+	TFrame3(const FVector3<RealType>& OriginIn, const TQuaternion<RealType> RotationIn)
 	{
 		Origin = OriginIn;
 		Rotation = RotationIn;
@@ -60,7 +60,7 @@ struct FFrame3
 	 * @param OriginIn origin of frame
 	 * @param SetZ target Z axis
 	 */
-	FFrame3(const FVector3<RealType>& OriginIn, const FVector3<RealType>& SetZ)
+	TFrame3(const FVector3<RealType>& OriginIn, const FVector3<RealType>& SetZ)
 	{
 		Origin = OriginIn;
 		Rotation.SetFromTo(FVector3<RealType>::UnitZ(), SetZ);
@@ -73,14 +73,14 @@ struct FFrame3
 	 * @param Y desired Y axis of frame
 	 * @param Z desired Z axis of frame
 	 */
-	FFrame3(const FVector3<RealType>& OriginIn, const FVector3<RealType>& X, const FVector3<RealType>& Y, const FVector3<RealType>& Z)
+	TFrame3(const FVector3<RealType>& OriginIn, const FVector3<RealType>& X, const FVector3<RealType>& Y, const FVector3<RealType>& Z)
 	{
 		Origin = OriginIn;
-		Rotation = FQuaternion<RealType>( FMatrix3<RealType>(X, Y, Z, false) );
+		Rotation = TQuaternion<RealType>( TMatrix3<RealType>(X, Y, Z, false) );
 	}
 
 	/** Construct a Frame from an FTransform */
-	FFrame3(const FTransform& Transform)
+	TFrame3(const FTransform& Transform)
 	{
 		Origin = Transform.GetTranslation();
 		Rotation = Transform.GetRotation();
@@ -169,38 +169,38 @@ struct FFrame3
 
 
 	/** @return input Quaternion transformed into local coordinate system of Frame */
-	FQuaternion<RealType> ToFrame(const FQuaternion<RealType>& Quat) const
+	TQuaternion<RealType> ToFrame(const TQuaternion<RealType>& Quat) const
 	{
 		return Rotation.Inverse() * Quat;
 	}
 	/** @return input Quaternion transformed from local coordinate system of Frame into "World" coordinate system */
-	FQuaternion<RealType> FromFrame(const FQuaternion<RealType>& Quat) const
+	TQuaternion<RealType> FromFrame(const TQuaternion<RealType>& Quat) const
 	{
 		return Rotation * Quat;
 	}
 
 
 	/** @return input Ray transformed into local coordinate system of Frame */
-	FRay3<RealType> ToFrame(const FRay3<RealType>& Ray) const
+	TRay3<RealType> ToFrame(const TRay3<RealType>& Ray) const
 	{
-		return FRay3<RealType>(ToFramePoint(Ray.Origin), ToFrameVector(Ray.Direction));
+		return TRay3<RealType>(ToFramePoint(Ray.Origin), ToFrameVector(Ray.Direction));
 	}
 	/** @return input Ray transformed from local coordinate system of Frame into "World" coordinate system */
-	FRay3<RealType> FromFrame(const FRay3<RealType>& Ray) const
+	TRay3<RealType> FromFrame(const TRay3<RealType>& Ray) const
 	{
-		return FRay3<RealType>(ToFramePoint(Ray.Origin), ToFrameVector(Ray.Direction));
+		return TRay3<RealType>(ToFramePoint(Ray.Origin), ToFrameVector(Ray.Direction));
 	}
 
 
 	/** @return input Frame transformed into local coordinate system of this Frame */
-	FFrame3<RealType> ToFrame(const FFrame3<RealType>& Frame) const
+	TFrame3<RealType> ToFrame(const TFrame3<RealType>& Frame) const
 	{
-		return FFrame3<RealType>(ToFramePoint(Frame.Origin), ToFrame(Frame.Rotation));
+		return TFrame3<RealType>(ToFramePoint(Frame.Origin), ToFrame(Frame.Rotation));
 	}
 	/** @return input Frame transformed from local coordinate system of this Frame into "World" coordinate system */
-	FFrame3<RealType> FromFrame(const FFrame3<RealType>& Frame) const
+	TFrame3<RealType> FromFrame(const TFrame3<RealType>& Frame) const
 	{
-		return FFrame3<RealType>(ToFramePoint(Frame.Origin), FromFrame(Frame.Rotation));
+		return TFrame3<RealType>(ToFramePoint(Frame.Origin), FromFrame(Frame.Rotation));
 	}
 
 
@@ -273,7 +273,7 @@ struct FFrame3
 	/**
 	 * Rotate this frame by given quaternion
 	 */
-	void Rotate(const FQuaternion<RealType>& Quat)
+	void Rotate(const TQuaternion<RealType>& Quat)
 	{
 		Rotation = Quat * Rotation;
 	}
@@ -285,7 +285,7 @@ struct FFrame3
 	void Transform(const FTransform& XForm)
 	{
 		Origin = (FVector3<RealType>)XForm.TransformPosition((FVector3f)Origin);
-		Rotation = FQuaternion<RealType>(XForm.GetRotation()) * Rotation;
+		Rotation = TQuaternion<RealType>(XForm.GetRotation()) * Rotation;
 	}
 
 
@@ -296,7 +296,7 @@ struct FFrame3
 	 */
 	void AlignAxis(int AxisIndex, const FVector3<RealType>& ToDirection)
 	{
-		FQuaternion<RealType> RelRotation(GetAxis(AxisIndex), ToDirection);
+		TQuaternion<RealType> RelRotation(GetAxis(AxisIndex), ToDirection);
 		Rotate(RelRotation);
 	}
 
@@ -312,7 +312,7 @@ struct FFrame3
 		//@todo PlaneAngleSigned does acos() and then SetAxisAngleD() does cos/sin...can we optimize this?
 		FVector3<RealType> AxisVec = GetAxis(AxisIndex);
 		RealType AngleDeg = VectorUtil::PlaneAngleSignedD(AxisVec, ToDirection, AroundVector);
-		FQuaternion<RealType> RelRotation;
+		TQuaternion<RealType> RelRotation;
 		RelRotation.SetAxisAngleD(AroundVector, AngleDeg);
 		Rotate(RelRotation);
 	}
@@ -372,6 +372,6 @@ struct FFrame3
 
 };
 
-typedef FFrame3<float> FFrame3f;
-typedef FFrame3<double> FFrame3d;
+typedef TFrame3<float> FFrame3f;
+typedef TFrame3<double> FFrame3d;
 
diff --git a/Engine/Plugins/Experimental/GeometryProcessing/Source/GeometricObjects/Public/Intersection/Intersector1.h b/Engine/Plugins/Experimental/GeometryProcessing/Source/GeometricObjects/Public/Intersection/Intersector1.h
--- a/Engine/Plugins/Experimental/GeometryProcessing/Source/GeometricObjects/Public/Intersection/Intersector1.h
+++ b/Engine/Plugins/Experimental/GeometryProcessing/Source/GeometricObjects/Public/Intersection/Intersector1.h
@@ -11,22 +11,22 @@
 #include "BoxTypes.h"
 
 /**
- * FIntersector1 computes the intersection of one-dimensional intervals [u0,u1] and [v0,v1].
+ * TIntersector1 computes the intersection of one-dimensional intervals [u0,u1] and [v0,v1].
  * The end points of the input intervals must be ordered u0 <= u1 and v0 <= v1.
  * Infinite and degenerate intervals are allowed.
  */
 template<typename RealType>
-class FIntersector1
+class TIntersector1
 {
 private:
 	/** intersection point/interval, access via GetIntersection */
-	FInterval1<RealType> Intersections;
+	TInterval1<RealType> Intersections;
 
 public:
 	/** First interval */
-	FInterval1<RealType> U;
+	TInterval1<RealType> U;
 	/** Second interval */
-	FInterval1<RealType> V;
+	TInterval1<RealType> V;
 
 	/**
 	 * Number of intersections found.
@@ -36,13 +36,13 @@ public:
 	 */
 	int NumIntersections = 0;
 
-	FIntersector1(RealType u0, RealType u1, RealType v0, RealType v1) : Intersections(0,0), U(u0, u1), V(v0, v1)
+	TIntersector1(RealType u0, RealType u1, RealType v0, RealType v1) : Intersections(0,0), U(u0, u1), V(v0, v1)
 	{
 		check(u0 <= u1);
 		check(v0 <= v1);
 	}
 
-	FIntersector1(const FInterval1<RealType>& u, const FInterval1<RealType>& v) : Intersections(0, 0), U(u), V(v)
+	TIntersector1(const TInterval1<RealType>& u, const TInterval1<RealType>& v) : Intersections(0, 0), U(u), V(v)
 	{
 	}
 
@@ -104,5 +104,5 @@ public:
 
 };
 
-typedef FIntersector1<double> FIntersector1d;
-typedef FIntersector1<float> FIntersector1f;
+typedef TIntersector1<double> FIntersector1d;
+typedef TIntersector1<float> FIntersector1f;
diff --git a/Engine/Plugins/Experimental/GeometryProcessing/Source/GeometricObjects/Public/Intersection/IntrLine2Line2.h b/Engine/Plugins/Experimental/GeometryProcessing/Source/GeometricObjects/Public/Intersection/IntrLine2Line2.h
--- a/Engine/Plugins/Experimental/GeometryProcessing/Source/GeometricObjects/Public/Intersection/IntrLine2Line2.h
+++ b/Engine/Plugins/Experimental/GeometryProcessing/Source/GeometricObjects/Public/Intersection/IntrLine2Line2.h
@@ -12,12 +12,12 @@
  * Compute intersection between two 2D lines
  */
 template<typename RealType>
-class FIntrLine2Line2
+class TIntrLine2Line2
 {
 protected:
 	// input data
-	FLine2<RealType> Line1;
-	FLine2<RealType> Line2;
+	TLine2<RealType> Line1;
+	TLine2<RealType> Line2;
 	RealType dotThresh = TMathUtil<RealType>::ZeroTolerance;
 	
 public:
@@ -31,29 +31,29 @@ public:
 
 
 
-	FIntrLine2Line2(const FLine2<RealType>& Line1In, const FLine2<RealType>& Line2In) 
+	TIntrLine2Line2(const TLine2<RealType>& Line1In, const TLine2<RealType>& Line2In)
 		: Line1(Line1In), Line2(Line2In)
 	{
 	}
 
 
-	const FLine2<RealType>& GetLine1() const
+	const TLine2<RealType>& GetLine1() const
 	{
 		return Line1;
 	}
 
-	void SetLine1(const FLine2<RealType>& Value)
+	void SetLine1(const TLine2<RealType>& Value)
 	{
 		Line1 = Value;
 		Result = EIntersectionResult::NotComputed;
 	}
 
-	const FLine2<RealType>& GetLine2() const
+	const TLine2<RealType>& GetLine2() const
 	{
 		return Line2;
 	}
 
-	void SetLine2(FLine2<RealType>& Value)
+	void SetLine2(TLine2<RealType>& Value)
 	{
 		Line2 = Value;
 		Result = EIntersectionResult::NotComputed;
@@ -77,7 +77,7 @@ public:
 	}
 
 
-	FIntrLine2Line2& Compute()
+	TIntrLine2Line2& Compute()
 	{
 		Find();
 		return *this;
@@ -171,5 +171,5 @@ public:
 
 };
 
-typedef FIntrLine2Line2<double> FIntrLine2FLine2d;
-typedef FIntrLine2Line2<float> FIntrLine2Line2f;
+typedef TIntrLine2Line2<double> FIntrLine2TLine2d;
+typedef TIntrLine2Line2<float> FIntrLine2Line2f;
diff --git a/Engine/Plugins/Experimental/GeometryProcessing/Source/GeometricObjects/Public/Intersection/IntrRay3AxisAlignedBox3.h b/Engine/Plugins/Experimental/GeometryProcessing/Source/GeometricObjects/Public/Intersection/IntrRay3AxisAlignedBox3.h
--- a/Engine/Plugins/Experimental/GeometryProcessing/Source/GeometricObjects/Public/Intersection/IntrRay3AxisAlignedBox3.h
+++ b/Engine/Plugins/Experimental/GeometryProcessing/Source/GeometricObjects/Public/Intersection/IntrRay3AxisAlignedBox3.h
@@ -13,7 +13,7 @@
  * Compute intersection between 3D ray and 3D axis-aligned box
  */
 template <typename RealType>
-class FIntrRay3AxisAlignedBox3
+class TIntrRay3AxisAlignedBox3
 {
 public:
 
@@ -28,7 +28,7 @@ public:
 	 * @param ExpandExtents box is expanded by this amount in each direction, useful for dealing with float imprecision
 	 * @return true if ray intersects box
 	 */
-	static bool TestIntersection(const FRay3<RealType>& Ray, const FAxisAlignedBox3<RealType>& Box, RealType ExpandExtents = 0)
+	static bool TestIntersection(const TRay3<RealType>& Ray, const TAxisAlignedBox3<RealType>& Box, RealType ExpandExtents = 0)
 	{
 		FVector3<RealType> WdU = FVector3<RealType>::Zero();
 		FVector3<RealType> AWdU = FVector3<RealType>::Zero();
@@ -102,7 +102,7 @@ public:
 	 * @param RayParamOut ray intersect T-value, or TNumericLimits::Max()
 	 * @return true if ray intersects box
 	 */
-	static bool FindIntersection(const FRay3<RealType>& Ray, const FAxisAlignedBox3<RealType>& Box, RealType& RayParamOut)
+	static bool FindIntersection(const TRay3<RealType>& Ray, const TAxisAlignedBox3<RealType>& Box, RealType& RayParamOut)
 	{
 		RealType RayParam0 = 0.0;
 		RealType RayParam1 = TNumericLimits<RealType>::Max();
@@ -134,7 +134,7 @@ protected:
 
 	static bool DoClipping(RealType t0, RealType t1,
 		const FVector3<RealType>& RayOrigin, const FVector3<RealType>& RayDirection,
-		const FAxisAlignedBox3<RealType>& Box, bool solid, 
+		const TAxisAlignedBox3<RealType>& Box, bool solid, 
 		int& quantity, FVector3<RealType>& Point0, FVector3<RealType>& Point1, EIntersectionType& intrType)
 	{
 		FVector3<RealType> BOrigin = RayOrigin - Box.Center();
@@ -216,5 +216,5 @@ protected:
 
 };
 
-typedef FIntrRay3AxisAlignedBox3<float> FIntrRay3AxisAlignedBox3f;
-typedef FIntrRay3AxisAlignedBox3<double> FIntrRay3AxisAlignedBox3d;
\ No newline at end of file
+typedef TIntrRay3AxisAlignedBox3<float> FIntrRay3AxisAlignedBox3f;
+typedef TIntrRay3AxisAlignedBox3<double> FIntrRay3AxisAlignedBox3d;
\ No newline at end of file
diff --git a/Engine/Plugins/Experimental/GeometryProcessing/Source/GeometricObjects/Public/Intersection/IntrRay3Triangle3.h b/Engine/Plugins/Experimental/GeometryProcessing/Source/GeometricObjects/Public/Intersection/IntrRay3Triangle3.h
--- a/Engine/Plugins/Experimental/GeometryProcessing/Source/GeometricObjects/Public/Intersection/IntrRay3Triangle3.h
+++ b/Engine/Plugins/Experimental/GeometryProcessing/Source/GeometricObjects/Public/Intersection/IntrRay3Triangle3.h
@@ -13,12 +13,12 @@
  * Compute intersection between 3D ray and 3D triangle
  */
 template <typename Real>
-class FIntrRay3Triangle3
+class TIntrRay3Triangle3
 {
 public:
 	// Input
-	FRay3<Real> Ray;
-	FTriangle3<Real> Triangle;
+	TRay3<Real> Ray;
+	TTriangle3<Real> Triangle;
 
 	// Output
 	Real RayParameter;
@@ -26,7 +26,7 @@ public:
 	EIntersectionType IntersectionType;
 
 
-	FIntrRay3Triangle3(const FRay3<Real>& RayIn, const FTriangle3<Real>& TriangleIn)
+	TIntrRay3Triangle3(const TRay3<Real>& RayIn, const TTriangle3<Real>& TriangleIn)
 	{
 		Ray = RayIn;
 		Triangle = TriangleIn;
@@ -168,5 +168,5 @@ public:
 
 };
 
-typedef FIntrRay3Triangle3<float> FIntrRay3Triangle3f;
-typedef FIntrRay3Triangle3<double> FIntrRay3Triangle3d;
\ No newline at end of file
+typedef TIntrRay3Triangle3<float> FIntrRay3Triangle3f;
+typedef TIntrRay3Triangle3<double> FIntrRay3Triangle3d;
\ No newline at end of file
diff --git a/Engine/Plugins/Experimental/GeometryProcessing/Source/GeometricObjects/Public/Intersection/IntrSegment2Segment2.h b/Engine/Plugins/Experimental/GeometryProcessing/Source/GeometricObjects/Public/Intersection/IntrSegment2Segment2.h
--- a/Engine/Plugins/Experimental/GeometryProcessing/Source/GeometricObjects/Public/Intersection/IntrSegment2Segment2.h
+++ b/Engine/Plugins/Experimental/GeometryProcessing/Source/GeometricObjects/Public/Intersection/IntrSegment2Segment2.h
@@ -46,12 +46,12 @@
  * Calculate intersection between two 2D line segments
  */
 template<typename RealType>
-class FIntrSegment2Segment2
+class TIntrSegment2Segment2
 {
 protected:
 	// inputs
-	FSegment2<RealType> Segment1;
-	FSegment2<RealType> Segment2;
+	TSegment2<RealType> Segment1;
+	TSegment2<RealType> Segment2;
 	RealType IntervalThreshold = 0;
 	RealType DotThreshold = TMathUtil<RealType>::ZeroTolerance;
 
@@ -69,29 +69,29 @@ public:
 
 
 
-	FIntrSegment2Segment2(const FSegment2<RealType>& Segment1In, const FSegment2<RealType>& Segment2In)
+	TIntrSegment2Segment2(const TSegment2<RealType>& Segment1In, const TSegment2<RealType>& Segment2In)
 		: Segment1(Segment1In), Segment2(Segment2In)
 	{
 	}
 
 
-	const FSegment2<RealType>& GetSegment1() const
+	const TSegment2<RealType>& GetSegment1() const
 	{
 		return Segment1;
 	}
 
-	void SetSegment1(const FSegment2<RealType>& Value)
+	void SetSegment1(const TSegment2<RealType>& Value)
 	{
 		Segment1 = Value;
 		Result = EIntersectionResult::NotComputed;
 	}
 
-	const FSegment2<RealType>& GetSegment2() const
+	const TSegment2<RealType>& GetSegment2() const
 	{
 		return Segment2;
 	}
 
-	void SetSegment2(const FSegment2<RealType>& Value)
+	void SetSegment2(const TSegment2<RealType>& Value)
 	{
 		Segment2 = Value;
 		Result = EIntersectionResult::NotComputed;
@@ -125,7 +125,7 @@ public:
 	}
 
 
-	FIntrSegment2Segment2& Compute()
+	TIntrSegment2Segment2& Compute()
 	{
 		Find();
 		return *this;
@@ -150,7 +150,7 @@ public:
 
 
 		FVector2<RealType> s = FVector2<RealType>::Zero();
-		Type = FIntrLine2Line2<RealType>::Classify(Segment1.Center, Segment1.Direction,
+		Type = TIntrLine2Line2<RealType>::Classify(Segment1.Center, Segment1.Direction,
 			Segment2.Center, Segment2.Direction,
 			DotThreshold, s);
 
@@ -177,7 +177,7 @@ public:
 			RealType t1 = Segment1.Direction.Dot(diff);
 			RealType tmin = t1 - Segment2.Extent;
 			RealType tmax = t1 + Segment2.Extent;
-			FIntersector1<RealType> calc(-Segment1.Extent, Segment1.Extent, tmin, tmax);
+			TIntersector1<RealType> calc(-Segment1.Extent, Segment1.Extent, tmin, tmax);
 			calc.Find();
 			Quantity = calc.NumIntersections;
 			if (Quantity == 2) 
@@ -242,5 +242,5 @@ protected:
 	}
 };
 
-typedef FIntrSegment2Segment2<double> FIntrSegment2Segment2d;
-typedef FIntrSegment2Segment2<float> FIntrSegment2Segment2f;
\ No newline at end of file
+typedef TIntrSegment2Segment2<double> FIntrSegment2Segment2d;
+typedef TIntrSegment2Segment2<float> FIntrSegment2Segment2f;
\ No newline at end of file
diff --git a/Engine/Plugins/Experimental/GeometryProcessing/Source/GeometricObjects/Public/LineTypes.h b/Engine/Plugins/Experimental/GeometryProcessing/Source/GeometricObjects/Public/LineTypes.h
--- a/Engine/Plugins/Experimental/GeometryProcessing/Source/GeometricObjects/Public/LineTypes.h
+++ b/Engine/Plugins/Experimental/GeometryProcessing/Source/GeometricObjects/Public/LineTypes.h
@@ -8,11 +8,11 @@
 #include "VectorTypes.h"
 
 /**
- * FLine2 is a two-dimensional infinite line.
+ * TLine2 is a two-dimensional infinite line.
  * The line is stored in (Center,Direction) form.
  */
 template<typename T>
-struct FLine2
+struct TLine2
 {
 	/** Origin / Center Point of Line */
 	FVector2<T> Origin;
@@ -23,7 +23,7 @@ struct FLine2
 	/**
 	 * Construct default line along X axis
 	 */
-	FLine2()
+	TLine2()
 	{
 		Origin = FVector2<T>::Zero();
 		Direction = FVector2<T>::UnitX();
@@ -32,7 +32,7 @@ struct FLine2
 	/**
 	 * Construct line with given Origin and Direction
 	 */
-	FLine2(const FVector2<T>& OriginIn, const FVector2<T>& DirectionIn)
+	TLine2(const FVector2<T>& OriginIn, const FVector2<T>& DirectionIn)
 		: Origin(OriginIn), Direction(DirectionIn)
 	{
 	}
@@ -41,9 +41,9 @@ struct FLine2
 	/**
 	 * @return line between two points
 	 */
-	static FLine2<T> FromPoints(const FVector2<T>& Point0, const FVector2<T>& Point1)
+	static TLine2<T> FromPoints(const FVector2<T>& Point0, const FVector2<T>& Point1)
 	{
-		return FLine2<T>(Point0, (Point1 - Point0).Normalized());
+		return TLine2<T>(Point0, (Point1 - Point0).Normalized());
 	}
 
 
@@ -105,12 +105,12 @@ struct FLine2
 	 * @param ParallelDotTolerance tolerance used to determine if lines are parallel (and hence cannot intersect)
 	 * @return true if lines intersect and IntersectionPointOut was computed
 	 */
-	bool IntersectionPoint(const FLine2<T>& OtherLine, FVector2<T>& IntersectionPointOut, T ParallelDotTolerance = TMathUtil<T>::ZeroTolerance) const
+	bool IntersectionPoint(const TLine2<T>& OtherLine, FVector2<T>& IntersectionPointOut, T ParallelDotTolerance = TMathUtil<T>::ZeroTolerance) const
 	{
-		// see IntrFLine2FLine2 for more detailed explanation
+		// see IntrTLine2TLine2 for more detailed explanation
 		FVector2<T> diff = OtherLine.Origin - Origin;
 		T D0DotPerpD1 = Direction.DotPerp(OtherLine.Direction);
-		if (TMathUtil<T>::Abs(D0DotPerpD1) > ParallelDotTolerance)                     // FLines intersect in a single point.
+		if (TMathUtil<T>::Abs(D0DotPerpD1) > ParallelDotTolerance)                     // TLines intersect in a single point.
 		{
 			T invD0DotPerpD1 = ((T)1) / D0DotPerpD1;
 			T diffDotPerpD1 = diff.DotPerp(OtherLine.Direction);
@@ -118,14 +118,14 @@ struct FLine2
 			IntersectionPointOut = Origin + s * Direction;
 			return true;
 		}
-		// FLines are parallel.
+		// TLines are parallel.
 		return false;
 	}
 };
 
 
-typedef FLine2<double> FLine2d;
-typedef FLine2<float> FLine2f;
+typedef TLine2<double> FLine2d;
+typedef TLine2<float> FLine2f;
 
 
 
@@ -133,11 +133,11 @@ typedef FLine2<float> FLine2f;
 
 
 /**
- * FLine3 is a three-dimensional infinite line.
+ * TLine3 is a three-dimensional infinite line.
  * The line is stored in (Center,Direction) form.
  */
 template<typename T>
-struct FLine3
+struct TLine3
 {
 	/** Origin / Center Point of Line */
 	FVector3<T> Origin;
@@ -148,7 +148,7 @@ struct FLine3
 	/**
 	 * Construct default line along X axis
 	 */
-	FLine3()
+	TLine3()
 	{
 		Origin = FVector3<T>::Zero();
 		Direction = FVector3<T>::UnitX();
@@ -157,7 +157,7 @@ struct FLine3
 	/**
 	 * Construct line with given Origin and Direction
 	 */
-	FLine3(const FVector3<T>& OriginIn, const FVector3<T>& DirectionIn)
+	TLine3(const FVector3<T>& OriginIn, const FVector3<T>& DirectionIn)
 		: Origin(OriginIn), Direction(DirectionIn)
 	{
 	}
@@ -166,9 +166,9 @@ struct FLine3
 	/**
 	 * @return line between two points
 	 */
-	static FLine3<T> FromPoints(const FVector3<T>& Point0, const FVector3<T>& Point1)
+	static TLine3<T> FromPoints(const FVector3<T>& Point0, const FVector3<T>& Point1)
 	{
-		return FLine3<T>(Point0, (Point1 - Point0).Normalized());
+		return TLine3<T>(Point0, (Point1 - Point0).Normalized());
 	}
 
 
@@ -211,5 +211,5 @@ struct FLine3
 
 };
 
-typedef FLine3<double> FLine3d;
-typedef FLine3<float> FLine3f;
\ No newline at end of file
+typedef TLine3<double> FLine3d;
+typedef TLine3<float> FLine3f;
\ No newline at end of file
diff --git a/Engine/Plugins/Experimental/GeometryProcessing/Source/GeometricObjects/Public/MatrixTypes.h b/Engine/Plugins/Experimental/GeometryProcessing/Source/GeometricObjects/Public/MatrixTypes.h
--- a/Engine/Plugins/Experimental/GeometryProcessing/Source/GeometricObjects/Public/MatrixTypes.h
+++ b/Engine/Plugins/Experimental/GeometryProcessing/Source/GeometricObjects/Public/MatrixTypes.h
@@ -6,24 +6,24 @@
 #include "VectorUtil.h"
 
 template <typename RealType>
-struct FMatrix3
+struct TMatrix3
 {
 	FVector3<RealType> Row0;
 	FVector3<RealType> Row1;
 	FVector3<RealType> Row2;
 
-	FMatrix3()
+	TMatrix3()
 	{
 	}
 
-	FMatrix3(RealType ConstantValue)
+	TMatrix3(RealType ConstantValue)
 	{
 		Row0 = FVector3<RealType>(ConstantValue, ConstantValue, ConstantValue);
 		Row1 = Row0;
 		Row2 = Row0;
 	}
 
-	FMatrix3(RealType Diag0, RealType Diag1, RealType Diag2)
+	TMatrix3(RealType Diag0, RealType Diag1, RealType Diag2)
 	{
 		Row0 = FVector3<RealType>(Diag0, 0, 0);
 		Row1 = FVector3<RealType>(0, Diag1, 0);
@@ -34,21 +34,21 @@ struct FMatrix3
 	 * Construct outer-product of U*transpose(V) of U and V
 	 * result is that Mij = u_i * v_j
 	 */
-	FMatrix3<RealType>(const FVector3<RealType>& U, const FVector3<RealType>& V)
+	TMatrix3<RealType>(const FVector3<RealType>& U, const FVector3<RealType>& V)
 		: Row0(U.X * V.X, U.X * V.Y, U.X * V.Z),
 		  Row1(U.Y * V.X, U.Y * V.Y, U.Y * V.Z),
 		  Row2(U.Z * V.X, U.Z * V.Y, U.Z * V.Z)
 	{
 	}
 
-	FMatrix3(RealType M00, RealType M01, RealType M02, RealType M10, RealType M11, RealType M12, RealType M20, RealType M21, RealType M22)
+	TMatrix3(RealType M00, RealType M01, RealType M02, RealType M10, RealType M11, RealType M12, RealType M20, RealType M21, RealType M22)
 		: Row0(M00, M01, M02),
 		  Row1(M10, M11, M12),
 		  Row2(M20, M21, M22)
 	{
 	}
 
-	FMatrix3(const FVector3<RealType>& V1, const FVector3<RealType>& V2, const FVector3<RealType>& V3, bool bRows)
+	TMatrix3(const FVector3<RealType>& V1, const FVector3<RealType>& V2, const FVector3<RealType>& V3, bool bRows)
 	{
 		if (bRows)
 		{
@@ -64,13 +64,13 @@ struct FMatrix3
 		}
 	}
 
-	static FMatrix3<RealType> Zero()
+	static TMatrix3<RealType> Zero()
 	{
-		return FMatrix3<RealType>(0);
+		return TMatrix3<RealType>(0);
 	}
-	static FMatrix3<RealType> Identity()
+	static TMatrix3<RealType> Identity()
 	{
-		return FMatrix3<RealType>(1, 1, 1);
+		return TMatrix3<RealType>(1, 1, 1);
 	}
 
 	RealType operator()(int Row, int Col) const
@@ -90,9 +90,9 @@ struct FMatrix3
 		}
 	}
 
-	FMatrix3<RealType> operator*(RealType Scale) const
+	TMatrix3<RealType> operator*(RealType Scale) const
 	{
-		return FMatrix3<RealType>(
+		return TMatrix3<RealType>(
 			Row0.X * Scale, Row0.Y * Scale, Row0.Z * Scale,
 			Row1.X * Scale, Row1.Y * Scale, Row1.Z * Scale,
 			Row2.X * Scale, Row2.Y * Scale, Row2.Z * Scale);
@@ -106,7 +106,7 @@ struct FMatrix3
 			Row2.X * V.X + Row2.Y * V.Y + Row2.Z * V.Z);
 	}
 
-	FMatrix3<RealType> operator*(const FMatrix3<RealType>& Mat2) const
+	TMatrix3<RealType> operator*(const TMatrix3<RealType>& Mat2) const
 	{
 		RealType M00 = Row0.X * Mat2.Row0.X + Row0.Y * Mat2.Row1.X + Row0.Z * Mat2.Row2.X;
 		RealType M01 = Row0.X * Mat2.Row0.Y + Row0.Y * Mat2.Row1.Y + Row0.Z * Mat2.Row2.Y;
@@ -120,20 +120,20 @@ struct FMatrix3
 		RealType M21 = Row2.X * Mat2.Row0.Y + Row2.Y * Mat2.Row1.Y + Row2.Z * Mat2.Row2.Y;
 		RealType M22 = Row2.X * Mat2.Row0.Z + Row2.Y * Mat2.Row1.Z + Row2.Z * Mat2.Row2.Z;
 
-		return FMatrix3<RealType>(M00, M01, M02, M10, M11, M12, M20, M21, M22);
+		return TMatrix3<RealType>(M00, M01, M02, M10, M11, M12, M20, M21, M22);
 	}
 
-	FMatrix3<RealType> operator+(const FMatrix3<RealType>& Mat2)
+	TMatrix3<RealType> operator+(const TMatrix3<RealType>& Mat2)
 	{
-		return FMatrix3<RealType>(Row0 + Mat2.Row0, Row1 + Mat2.Row1, Row2 + Mat2.Row2, true);
+		return TMatrix3<RealType>(Row0 + Mat2.Row0, Row1 + Mat2.Row1, Row2 + Mat2.Row2, true);
 	}
 
-	FMatrix3<RealType> operator-(const FMatrix3<RealType>& Mat2)
+	TMatrix3<RealType> operator-(const TMatrix3<RealType>& Mat2)
 	{
-		return FMatrix3<RealType>(Row0 - Mat2.Row0, Row1 - Mat2.Row1, Row2 - Mat2.Row2, true);
+		return TMatrix3<RealType>(Row0 - Mat2.Row0, Row1 - Mat2.Row1, Row2 - Mat2.Row2, true);
 	}
 
-	inline FMatrix3<RealType>& operator*=(const RealType& Scalar)
+	inline TMatrix3<RealType>& operator*=(const RealType& Scalar)
 	{
 		Row0 *= Scalar;
 		Row1 *= Scalar;
@@ -141,7 +141,7 @@ struct FMatrix3
 		return *this;
 	}
 
-	inline FMatrix3<RealType>& operator+=(const FMatrix3<RealType>& Mat2)
+	inline TMatrix3<RealType>& operator+=(const TMatrix3<RealType>& Mat2)
 	{
 		Row0 += Mat2.Row0;
 		Row1 += Mat2.Row1;
@@ -149,7 +149,7 @@ struct FMatrix3
 		return *this;
 	}
 
-	RealType InnerProduct(const FMatrix3<RealType>& Mat2) const
+	RealType InnerProduct(const TMatrix3<RealType>& Mat2) const
 	{
 		return Row0.Dot(Mat2.Row0) + Row1.Dot(Mat2.Row1) + Row2.Dot(Mat2.Row2);
 	}
@@ -163,7 +163,7 @@ struct FMatrix3
 		return a11 * i00 + a21 * i01 + a31 * i02;
 	}
 
-	FMatrix3<RealType> Inverse() const
+	TMatrix3<RealType> Inverse() const
 	{
 		RealType a11 = Row0.X, a12 = Row0.Y, a13 = Row0.Z, a21 = Row1.X, a22 = Row1.Y, a23 = Row1.Z, a31 = Row2.X, a32 = Row2.Y, a33 = Row2.Z;
 		RealType i00 = a33 * a22 - a32 * a23;
@@ -181,25 +181,25 @@ struct FMatrix3
 		RealType det = a11 * i00 + a21 * i01 + a31 * i02;
 		ensure(TMathUtil<RealType>::Abs(det) >= TMathUtil<RealType>::Epsilon);
 		det = 1.0 / det;
-		return FMatrix3<RealType>(i00 * det, i01 * det, i02 * det, i10 * det, i11 * det, i12 * det, i20 * det, i21 * det, i22 * det);
+		return TMatrix3<RealType>(i00 * det, i01 * det, i02 * det, i10 * det, i11 * det, i12 * det, i20 * det, i21 * det, i22 * det);
 	}
 
-	FMatrix3<RealType> Transpose() const
+	TMatrix3<RealType> Transpose() const
 	{
-		return FMatrix3<RealType>(
+		return TMatrix3<RealType>(
 			Row0.X, Row1.X, Row2.X,
 			Row0.Y, Row1.Y, Row2.Y,
 			Row0.Z, Row1.Z, Row2.Z);
 	}
 
-	bool EpsilonEqual(const FMatrix3<RealType>& Mat2, RealType Epsilon) const
+	bool EpsilonEqual(const TMatrix3<RealType>& Mat2, RealType Epsilon) const
 	{
 		return VectorUtil::EpsilonEqual(Row0, Mat2.Row0, Epsilon) &&
 			   VectorUtil::EpsilonEqual(Row1, Mat2.Row1, Epsilon) &&
 			   VectorUtil::EpsilonEqual(Row2, Mat2.Row2, Epsilon);
 	}
 
-	static FMatrix3<RealType> AxisAngleR(const FVector3<RealType>& Axis, RealType AngleRad)
+	static TMatrix3<RealType> AxisAngleR(const FVector3<RealType>& Axis, RealType AngleRad)
 	{
 		RealType cs = TMathUtil<RealType>::Cos(AngleRad);
 		RealType sn = TMathUtil<RealType>::Sin(AngleRad);
@@ -213,13 +213,13 @@ struct FMatrix3
 		RealType xSin = Axis[0] * sn;
 		RealType ySin = Axis[1] * sn;
 		RealType zSin = Axis[2] * sn;
-		return FMatrix3<RealType>(
+		return TMatrix3<RealType>(
 			x2 * oneMinusCos + cs, xym - zSin, xzm + ySin,
 			xym + zSin, y2 * oneMinusCos + cs, yzm - xSin,
 			xzm - ySin, yzm + xSin, z2 * oneMinusCos + cs);
 	}
 
-	static FMatrix3<RealType> AxisAngleD(const FVector3<RealType>& Axis, RealType AngleDeg)
+	static TMatrix3<RealType> AxisAngleD(const FVector3<RealType>& Axis, RealType AngleDeg)
 	{
 		return AxisAngleR(Axis, TMathUtil<RealType>::DegreesToRadians(AngleDeg));
 	}
@@ -227,21 +227,21 @@ struct FMatrix3
 
 
 template <typename RealType>
-struct FMatrix2
+struct TMatrix2
 {
 	FVector2<RealType> Row0;
 	FVector2<RealType> Row1;
 
-	FMatrix2()
+	TMatrix2()
 	{
 	}
 
-	FMatrix2(RealType ConstantValue)
+	TMatrix2(RealType ConstantValue)
 	{
 		Row0 = Row1 = FVector2<RealType>(ConstantValue, ConstantValue);
 	}
 
-	FMatrix2(RealType Diag0, RealType Diag1)
+	TMatrix2(RealType Diag0, RealType Diag1)
 	{
 		Row0 = FVector2<RealType>(Diag0, 0);
 		Row1 = FVector2<RealType>(0, Diag1);
@@ -251,19 +251,19 @@ struct FMatrix2
 	 * Construct outer-product of U*transpose(V) of U and V
 	 * result is that Mij = u_i * v_j
 	 */
-	FMatrix2<RealType>(const FVector2<RealType>& U, const FVector2<RealType>& V)
+	TMatrix2<RealType>(const FVector2<RealType>& U, const FVector2<RealType>& V)
 		: Row0(U.X * V.X, U.X * V.Y),
 		  Row1(U.Y * V.X, U.Y * V.Y)
 	{
 	}
 
-	FMatrix2(RealType M00, RealType M01, RealType M10, RealType M11)
+	TMatrix2(RealType M00, RealType M01, RealType M10, RealType M11)
 		: Row0(M00, M01),
 		  Row1(M10, M11)
 	{
 	}
 
-	FMatrix2(const FVector2<RealType>& V1, const FVector2<RealType>& V2, bool bRows)
+	TMatrix2(const FVector2<RealType>& V1, const FVector2<RealType>& V2, bool bRows)
 	{
 		if (bRows)
 		{
@@ -277,13 +277,13 @@ struct FMatrix2
 		}
 	}
 
-	static FMatrix2<RealType> Zero()
+	static TMatrix2<RealType> Zero()
 	{
-		return FMatrix2<RealType>(0);
+		return TMatrix2<RealType>(0);
 	}
-	static FMatrix2<RealType> Identity()
+	static TMatrix2<RealType> Identity()
 	{
-		return FMatrix2<RealType>(1, 1);
+		return TMatrix2<RealType>(1, 1);
 	}
 
 	RealType operator()(int Row, int Col) const
@@ -299,9 +299,9 @@ struct FMatrix2
 		}
 	}
 
-	FMatrix2<RealType> operator*(RealType Scale) const
+	TMatrix2<RealType> operator*(RealType Scale) const
 	{
-		return FMatrix2<RealType>(
+		return TMatrix2<RealType>(
 			Row0.X * Scale, Row0.Y * Scale,
 			Row1.X * Scale, Row1.Y * Scale);
 	}
@@ -313,7 +313,7 @@ struct FMatrix2
 			Row1.X * V.X + Row1.Y * V.Y);
 	}
 
-	FMatrix2<RealType> operator*(const FMatrix2<RealType>& Mat2) const
+	TMatrix2<RealType> operator*(const TMatrix2<RealType>& Mat2) const
 	{
 		RealType M00 = Row0.X * Mat2.Row0.X + Row0.Y * Mat2.Row1.X;
 		RealType M01 = Row0.X * Mat2.Row0.Y + Row0.Y * Mat2.Row1.Y;
@@ -321,34 +321,34 @@ struct FMatrix2
 		RealType M10 = Row1.X * Mat2.Row0.X + Row1.Y * Mat2.Row1.X;
 		RealType M11 = Row1.X * Mat2.Row0.Y + Row1.Y * Mat2.Row1.Y;
 
-		return FMatrix2<RealType>(M00, M01, M10, M11);
+		return TMatrix2<RealType>(M00, M01, M10, M11);
 	}
 
-	FMatrix2<RealType> operator+(const FMatrix2<RealType>& Mat2)
+	TMatrix2<RealType> operator+(const TMatrix2<RealType>& Mat2)
 	{
-		return FMatrix2<RealType>(Row0 + Mat2.Row0, Row1 + Mat2.Row1, true);
+		return TMatrix2<RealType>(Row0 + Mat2.Row0, Row1 + Mat2.Row1, true);
 	}
 
-	FMatrix2<RealType> operator-(const FMatrix2<RealType>& Mat2)
+	TMatrix2<RealType> operator-(const TMatrix2<RealType>& Mat2)
 	{
-		return FMatrix2<RealType>(Row0 - Mat2.Row0, Row1 - Mat2.Row1, true);
+		return TMatrix2<RealType>(Row0 - Mat2.Row0, Row1 - Mat2.Row1, true);
 	}
 
-	inline FMatrix2<RealType>& operator*=(const RealType& Scalar)
+	inline TMatrix2<RealType>& operator*=(const RealType& Scalar)
 	{
 		Row0 *= Scalar;
 		Row1 *= Scalar;
 		return *this;
 	}
 
-	inline FMatrix2<RealType>& operator+=(const FMatrix2<RealType>& Mat2)
+	inline TMatrix2<RealType>& operator+=(const TMatrix2<RealType>& Mat2)
 	{
 		Row0 += Mat2.Row0;
 		Row1 += Mat2.Row1;
 		return *this;
 	}
 
-	RealType InnerProduct(const FMatrix2<RealType>& Mat2) const
+	RealType InnerProduct(const TMatrix2<RealType>& Mat2) const
 	{
 		return Row0.Dot(Mat2.Row0) + Row1.Dot(Mat2.Row1);
 	}
@@ -358,32 +358,32 @@ struct FMatrix2
 		return Row0.X * Row1.Y - Row0.Y * Row1.X;
 	}
 
-	FMatrix2<RealType> Inverse() const
+	TMatrix2<RealType> Inverse() const
 	{
 		RealType Det = Determinant();
 		ensure(TMathUtil<RealType>::Abs(Det) >= TMathUtil<RealType>::Epsilon);
 		RealType DetInv = 1.0 / Det;
-		return FMatrix2<RealType>(Row1.Y * DetInv, -Row0.Y * DetInv, -Row1.X * DetInv, Row0.X * DetInv);
+		return TMatrix2<RealType>(Row1.Y * DetInv, -Row0.Y * DetInv, -Row1.X * DetInv, Row0.X * DetInv);
 	}
 
-	FMatrix2<RealType> Transpose() const
+	TMatrix2<RealType> Transpose() const
 	{
-		return FMatrix2<RealType>(
+		return TMatrix2<RealType>(
 			Row0.X, Row1.X,
 			Row0.Y, Row1.Y);
 	}
 
-	bool EpsilonEqual(const FMatrix2<RealType>& Mat2, RealType Epsilon) const
+	bool EpsilonEqual(const TMatrix2<RealType>& Mat2, RealType Epsilon) const
 	{
 		return VectorUtil::EpsilonEqual(Row0, Mat2.Row0, Epsilon) &&
 			   VectorUtil::EpsilonEqual(Row1, Mat2.Row1, Epsilon);
 	}
 
-	static FMatrix2<RealType> RotationRad(RealType AngleRad)
+	static TMatrix2<RealType> RotationRad(RealType AngleRad)
 	{
 		RealType cs = TMathUtil<RealType>::Cos(AngleRad);
 		RealType sn = TMathUtil<RealType>::Sin(AngleRad);
-		return FMatrix2<RealType>(cs, -sn, sn, cs);
+		return TMatrix2<RealType>(cs, -sn, sn, cs);
 	}
 
 	/**
@@ -397,23 +397,23 @@ struct FMatrix2
 };
 
 template <typename RealType>
-inline FMatrix3<RealType> operator*(RealType Scale, const FMatrix3<RealType>& Mat)
+inline TMatrix3<RealType> operator*(RealType Scale, const TMatrix3<RealType>& Mat)
 {
-	return FMatrix3<RealType>(
+	return TMatrix3<RealType>(
 		Mat.Row0.X * Scale, Mat.Row0.Y * Scale, Mat.Row0.Z * Scale,
 		Mat.Row1.X * Scale, Mat.Row1.Y * Scale, Mat.Row1.Z * Scale,
 		Mat.Row2.X * Scale, Mat.Row2.Y * Scale, Mat.Row2.Z * Scale);
 }
 
 template <typename RealType>
-inline FMatrix2<RealType> operator*(RealType Scale, const FMatrix2<RealType>& Mat)
+inline TMatrix2<RealType> operator*(RealType Scale, const TMatrix2<RealType>& Mat)
 {
-	return FMatrix2<RealType>(
+	return TMatrix2<RealType>(
 		Mat.Row0.X * Scale, Mat.Row0.Y * Scale,
 		Mat.Row1.X * Scale, Mat.Row1.Y * Scale);
 }
 
-typedef FMatrix3<float> FMatrix3f;
-typedef FMatrix3<double> FMatrix3d;
-typedef FMatrix2<float> FMatrix2f;
-typedef FMatrix2<double> FMatrix2d;
\ No newline at end of file
+typedef TMatrix3<float> FMatrix3f;
+typedef TMatrix3<double> FMatrix3d;
+typedef TMatrix2<float> FMatrix2f;
+typedef TMatrix2<double> FMatrix2d;
\ No newline at end of file
diff --git a/Engine/Plugins/Experimental/GeometryProcessing/Source/GeometricObjects/Public/OrientedBoxTypes.h b/Engine/Plugins/Experimental/GeometryProcessing/Source/GeometricObjects/Public/OrientedBoxTypes.h
--- a/Engine/Plugins/Experimental/GeometryProcessing/Source/GeometricObjects/Public/OrientedBoxTypes.h
+++ b/Engine/Plugins/Experimental/GeometryProcessing/Source/GeometricObjects/Public/OrientedBoxTypes.h
@@ -9,26 +9,26 @@
 #include "FrameTypes.h"
 
 /**
- * FOrientedBox3 is a non-axis-aligned 3D box defined by a 3D frame and extents along the axes of that frame
+ * TOrientedBox3 is a non-axis-aligned 3D box defined by a 3D frame and extents along the axes of that frame
  * The frame is at the center of the box.
  */
 template<typename RealType>
-struct FOrientedBox3
+struct TOrientedBox3
 {
 	// available for porting: ContainPoint, MergeBoxes()
 
 
 	/** 3D position (center) and orientation (axes) of the box */
-	FFrame3<RealType> Frame;
+	TFrame3<RealType> Frame;
 	/** Half-dimensions of box measured along the three axes */
 	FVector3<RealType> Extents;
 
-	FOrientedBox3() : Extents(1,1,1) {}
+	TOrientedBox3() : Extents(1,1,1) {}
 
 	/**
 	 * Create axis-aligned box with given Origin and Extents
 	 */
-	FOrientedBox3(const FVector3<RealType>& Origin, const FVector3<RealType> & ExtentsIn)
+	TOrientedBox3(const FVector3<RealType>& Origin, const FVector3<RealType> & ExtentsIn)
 		: Frame(Origin), Extents(ExtentsIn)
 	{
 	}
@@ -36,7 +36,7 @@ struct FOrientedBox3
 	/**
 	 * Create oriented box with given Frame and Extents
 	 */
-	FOrientedBox3(const FFrame3<RealType>& FrameIn, const FVector3<RealType> & ExtentsIn)
+	TOrientedBox3(const TFrame3<RealType>& FrameIn, const FVector3<RealType> & ExtentsIn)
 		: Frame(FrameIn), Extents(ExtentsIn)
 	{
 	}
@@ -44,17 +44,17 @@ struct FOrientedBox3
 	/**
 	 * Create oriented box from axis-aligned box
 	 */
-	FOrientedBox3(const FAxisAlignedBox3<RealType>& AxisBox)
+	TOrientedBox3(const TAxisAlignedBox3<RealType>& AxisBox)
 		: Frame(AxisBox.Center()), Extents((RealType)0.5 * AxisBox.Diagonal())
 	{
 	}
 
 
 	/** @return box with unit dimensions centered at origin */
-	static FOrientedBox3<RealType> UnitZeroCentered() { return FOrientedBox3<RealType>(FVector3<RealType>::Zero(), (RealType)0.5*FVector3<RealType>::One()); }
+	static TOrientedBox3<RealType> UnitZeroCentered() { return TOrientedBox3<RealType>(FVector3<RealType>::Zero(), (RealType)0.5*FVector3<RealType>::One()); }
 
 	/** @return box with unit dimensions where minimum corner is at origin */
-	static FOrientedBox3<RealType> UnitPositive() { return FOrientedBox3<RealType>((RealType)0.5*FVector3<RealType>::One(), (RealType)0.5*FVector3<RealType>::One()); }
+	static TOrientedBox3<RealType> UnitPositive() { return TOrientedBox3<RealType>((RealType)0.5*FVector3<RealType>::One(), (RealType)0.5*FVector3<RealType>::One()); }
 
 
 	/** @return center of the box */
@@ -218,5 +218,5 @@ struct FOrientedBox3
 
 
 
-typedef FOrientedBox3<float> FOrientedBox3f;
-typedef FOrientedBox3<double> FOrientedBox3d;
+typedef TOrientedBox3<float> FOrientedBox3f;
+typedef TOrientedBox3<double> FOrientedBox3d;
diff --git a/Engine/Plugins/Experimental/GeometryProcessing/Source/GeometricObjects/Public/Polygon2.h b/Engine/Plugins/Experimental/GeometryProcessing/Source/GeometricObjects/Public/Polygon2.h
--- a/Engine/Plugins/Experimental/GeometryProcessing/Source/GeometricObjects/Public/Polygon2.h
+++ b/Engine/Plugins/Experimental/GeometryProcessing/Source/GeometricObjects/Public/Polygon2.h
@@ -15,12 +15,12 @@
 #include "Util/DynamicVector.h"
 
 /**
- * FPolygon2 is a 2D polygon represented as a list of Vertices.
+ * TPolygon2 is a 2D polygon represented as a list of Vertices.
  * 
  * @todo move operators
  */
 template<typename T>
-class FPolygon2
+class TPolygon2
 {
 protected:
 	/** The list of vertices/corners of the polygon */
@@ -31,21 +31,21 @@ protected:
 
 public:
 
-	FPolygon2() : Timestamp(0)
+	TPolygon2() : Timestamp(0)
 	{
 	}
 
 	/**
 	 * Construct polygon that is a copy of another polygon
 	 */
-	FPolygon2(const FPolygon2& Copy) : Vertices(Copy.Vertices), Timestamp(Copy.Timestamp)
+	TPolygon2(const TPolygon2& Copy) : Vertices(Copy.Vertices), Timestamp(Copy.Timestamp)
 	{
 	}
 
 	/**
 	 * Construct polygon with given list of vertices
 	 */
-	FPolygon2(const TArray<FVector2<T>>& VertexList) : Vertices(VertexList), Timestamp(0)
+	TPolygon2(const TArray<FVector2<T>>& VertexList) : Vertices(VertexList), Timestamp(0)
 	{
 	}
 
@@ -207,16 +207,16 @@ public:
 	/**
 	 * @return the bounding box of the Polygon Vertices
 	 */
-	FAxisAlignedBox2<T> Bounds() const
+	TAxisAlignedBox2<T> Bounds() const
 	{
-		FAxisAlignedBox2<T> box = FAxisAlignedBox2<T>::Empty();
+		TAxisAlignedBox2<T> box = TAxisAlignedBox2<T>::Empty();
 		box.Contain(Vertices);
 		return box;
 	}
 
 	
 	/**
-	 * SegmentIterator is used to iterate over the FSegment2<T> segments of the polygon
+	 * SegmentIterator is used to iterate over the TSegment2<T> segments of the polygon
 	 */
 	class SegmentIterator 
 	{
@@ -225,12 +225,12 @@ public:
 		{
 			return i < polygon->VertexCount();
 		}
-		inline FSegment2<T> operator*() const
+		inline TSegment2<T> operator*() const
 		{
 			check(polygon != nullptr && i < polygon->VertexCount());
-			return FSegment2<T>(polygon->Vertices[i], polygon->Vertices[i+1 % polygon->VertexCount()]);
+			return TSegment2<T>(polygon->Vertices[i], polygon->Vertices[i+1 % polygon->VertexCount()]);
 		}
-		//inline FSegment2<T> & operator*();
+		//inline TSegment2<T> & operator*();
 		inline SegmentIterator & operator++() 		// prefix
 		{
 			i++;
@@ -245,10 +245,10 @@ public:
 		inline bool operator==(const SegmentIterator & i2) { return i2.polygon == polygon && i2.i == i; }
 		inline bool operator!=(const SegmentIterator & i2) { return i2.polygon != polygon || i2.i != i; }
 	protected:
-		const FPolygon2 * polygon;
+		const TPolygon2 * polygon;
 		int i;
-		inline SegmentIterator(const FPolygon2 * p, int iCur) : polygon(p), i(iCur) {}
-		friend class FPolygon2;
+		inline SegmentIterator(const TPolygon2 * p, int iCur) : polygon(p), i(iCur) {}
+		friend class TPolygon2;
 	};
 	friend class SegmentIterator;
 
@@ -263,9 +263,9 @@ public:
 	class SegmentEnumerable
 	{
 	public:
-		const FPolygon2<T>* polygon;
+		const TPolygon2<T>* polygon;
 		SegmentEnumerable() : polygon(nullptr) {}
-		SegmentEnumerable(const FPolygon2<T> * p) : polygon(p) {}
+		SegmentEnumerable(const TPolygon2<T> * p) : polygon(p) {}
 		SegmentIterator begin() { return polygon->SegmentItr(); }
 		SegmentIterator end() { return SegmentIterator(polygon, polygon->VertexCount()); }
 	};
@@ -426,7 +426,7 @@ public:
 	/**
 	 * @return true if the Polygon fully contains the OtherPolygon
 	 */
-	bool Contains(const FPolygon2<T>& OtherPoly) const
+	bool Contains(const TPolygon2<T>& OtherPoly) const
 	{
 		// @todo fast bbox check?
 
@@ -451,7 +451,7 @@ public:
 	/**
 	 * @return true if the Segment is fully contained inside the Polygon
 	 */
-	bool Contains(const FSegment2<T>& Segment) const
+	bool Contains(const TSegment2<T>& Segment) const
 	{
 		// [TODO] Add bbox check
 		if (Contains(Segment.StartPoint()) == false || Contains(Segment.EndPoint()) == false)
@@ -459,7 +459,7 @@ public:
 			return false;
 		}
 
-		for (FSegment2<T> seg : Segments())
+		for (TSegment2<T> seg : Segments())
 		{
 			if (seg.Intersects(Segment))
 			{
@@ -473,16 +473,16 @@ public:
 	/**
 	 * @return true if at least one edge of the OtherPolygon intersects the Polygon
 	 */
-	bool Intersects(const FPolygon2<T>& OtherPoly) const
+	bool Intersects(const TPolygon2<T>& OtherPoly) const
 	{
 		if (!Bounds().Intersects(OtherPoly.Bounds()))
 		{
 			return false;
 		}
 
-		for (FSegment2<T> seg : Segments()) 
+		for (TSegment2<T> seg : Segments()) 
 		{
-			for (FSegment2<T> oseg : OtherPoly.Segments()) 
+			for (TSegment2<T> oseg : OtherPoly.Segments()) 
 			{
 				if (seg.Intersects(oseg))
 				{
@@ -497,7 +497,7 @@ public:
 	/**
 	 * @return true if the Segment intersects an edge of the Polygon
 	 */
-	bool Intersects(const FSegment2<T>& Segment) const
+	bool Intersects(const TSegment2<T>& Segment) const
 	{
 		// [TODO] Add bbox check
 		if (Contains(Segment.StartPoint()) == true || Contains(Segment.EndPoint()) == true)
@@ -506,7 +506,7 @@ public:
 		}
 
 		// [TODO] Add bbox check
-		for (FSegment2<T> seg : Segments())
+		for (TSegment2<T> seg : Segments())
 		{
 			if (seg.Intersects(Segment))
 			{
@@ -523,7 +523,7 @@ public:
 	 * @param OutArray intersection points are stored here
 	 * @return true if any intersections were found
 	 */
-	bool FindIntersections(const FPolygon2<T>& OtherPoly, TArray<FVector2<T>>& OutArray) const
+	bool FindIntersections(const TPolygon2<T>& OtherPoly, TArray<FVector2<T>>& OutArray) const
 	{
 		if (!Bounds().Intersects(OtherPoly.Bounds()))
 		{
@@ -531,16 +531,16 @@ public:
 		}
 
 		bool bFoundIntersections = false;
-		for (FSegment2<T> seg : Segments()) 
+		for (TSegment2<T> seg : Segments()) 
 		{
-			for (FSegment2<T> oseg : OtherPoly.Segments())
+			for (TSegment2<T> oseg : OtherPoly.Segments())
 			{
 				// this computes test twice for intersections, but seg.intersects doesn't
 				// create any new objects so it should be much faster for majority of segments (should profile!)
 				if (seg.Intersects(oseg)) 
 				{
 					//@todo can we replace with something like seg.intersects?
-					FIntrSegment2Segment2<T> intr(seg, oseg);
+					TIntrSegment2Segment2<T> intr(seg, oseg);
 					if (intr.Find()) 
 					{
 						bFoundIntersections = true;
@@ -561,9 +561,9 @@ public:
 	/**
 	 * @return edge of the polygon starting at vertex SegmentIndex
 	 */
-	FSegment2<T> Segment(int SegmentIndex) const
+	TSegment2<T> Segment(int SegmentIndex) const
 	{
-		return FSegment2<T>(Vertices[SegmentIndex], Vertices[(SegmentIndex + 1) % Vertices.Num()]);
+		return TSegment2<T>(Vertices[SegmentIndex], Vertices[(SegmentIndex + 1) % Vertices.Num()]);
 	}
 
 	/**
@@ -573,7 +573,7 @@ public:
 	 */
 	FVector2<T> PointAt(int SegmentIndex, T SegmentParam) const
 	{
-		FSegment2<T> seg(Vertices[SegmentIndex], Vertices[(SegmentIndex + 1) % Vertices.Num()]);
+		TSegment2<T> seg(Vertices[SegmentIndex], Vertices[(SegmentIndex + 1) % Vertices.Num()]);
 		return seg.PointAt(SegmentParam);
 	}
 
@@ -585,7 +585,7 @@ public:
 	 */
 	FVector2<T> GetNormal(int iSeg, T SegmentParam) const
 	{
-		FSegment2<T> seg(Vertices[iSeg], Vertices[(iSeg + 1) % Vertices.Num()]);
+		TSegment2<T> seg(Vertices[iSeg], Vertices[(iSeg + 1) % Vertices.Num()]);
 		T t = ((SegmentParam / seg.Extent) + 1.0) / 2.0;
 
 		FVector2<T> n0 = GetNormal(iSeg);
@@ -611,7 +611,7 @@ public:
 		for (int vi = 0; vi < N; ++vi) 
 		{
 			// @todo can't we just use segment function here now?
-			FSegment2<T> seg = FSegment2<T>(Vertices[vi], Vertices[(vi + 1) % N]);
+			TSegment2<T> seg = TSegment2<T>(Vertices[vi], Vertices[(vi + 1) % N]);
 			T t = (QueryPoint - seg.Center).Dot(seg.Direction);
 			T d = TNumericLimits<T>::Max();
 			if (t >= seg.Extent)
@@ -667,7 +667,7 @@ public:
 	 * Translate the polygon
 	 * @returns the Polygon, so that you can chain calls like Translate().Scale()
 	 */
-	FPolygon2<T>& Translate(const FVector2<T>& Translate) 
+	TPolygon2<T>& Translate(const FVector2<T>& Translate) 
 	{
 		int N = Vertices.Num();
 		for (int i = 0; i < N; ++i)
@@ -682,7 +682,7 @@ public:
 	 * Scale the Polygon relative to a given point
 	 * @returns the Polygon, so that you can chain calls like Translate().Scale()
 	 */
-	FPolygon2<T>& Scale(const FVector2<T>& Scale, const FVector2<T>& Origin)
+	TPolygon2<T>& Scale(const FVector2<T>& Scale, const FVector2<T>& Origin)
 	{
 		int N = Vertices.Num();
 		for (int i = 0; i < N; ++i)
@@ -698,7 +698,7 @@ public:
 	 * Apply an arbitrary transformation to the Polygon
 	 * @returns the Polygon, so that you can chain calls like Translate().Scale()
 	 */
-	FPolygon2<T>& Transform(const TFunction<FVector2<T> (const FVector2<T>&)>& TransformFunc)
+	TPolygon2<T>& Transform(const TFunction<FVector2<T> (const FVector2<T>&)>& TransformFunc)
 	{
 		int N = Vertices.Num();
 		for (int i = 0; i < N; ++i) 
@@ -761,8 +761,8 @@ public:
 			FVector2<T> prev = Vertices[k == 0 ? Vertices.Num() - 1 : k - 1];
 			FVector2<T> dn = (next - v).Normalized();
 			FVector2<T> dp = (prev - v).Normalized();
-			FLine2<T> ln(v + OffsetDistance * dn.Perp(), dn);
-			FLine2<T> lp(v - OffsetDistance * dp.Perp(), dp);
+			TLine2<T> ln(v + OffsetDistance * dn.Perp(), dn);
+			TLine2<T> lp(v - OffsetDistance * dp.Perp(), dp);
 
 			bool bIntersectsAtPoint = ln.IntersectionPoint(lp, NewVertices[k]);
 			if (!bIntersectsAtPoint)  // lines were parallel
@@ -800,7 +800,7 @@ private:
 		int maxi = j;          // index of vertex farthest from S
 		T maxd2 = 0;         // distance squared of farthest vertex
 		T tol2 = Tolerance * Tolerance;  // tolerance squared
-		FSegment2<T> S = FSegment2<T>(Vertices[j], Vertices[k]);    // segment from v[j] to v[k]
+		TSegment2<T> S = TSegment2<T>(Vertices[j], Vertices[k]);    // segment from v[j] to v[k]
 
 		// test each vertex v[i] for max distance from S
 		// Note: this works in any dimension (2D, 3D, ...)
@@ -1010,9 +1010,9 @@ public:
 	/**
 	 * Construct a four-vertex rectangle Polygon
 	 */
-	static FPolygon2<T> MakeRectangle(const FVector2<T>& Center, T Width, T Height)
+	static TPolygon2<T> MakeRectangle(const FVector2<T>& Center, T Width, T Height)
 	{
-		FPolygon2<T> Rectangle;
+		TPolygon2<T> Rectangle;
 		Rectangle.Vertices.SetNumUninitialized(4);
 		Rectangle.Set(0, FVector2<T>(Center.X - Width / 2, Center.Y - Height / 2));
 		Rectangle.Set(1, FVector2<T>(Center.X + Width / 2, Center.Y - Height / 2));
@@ -1025,9 +1025,9 @@ public:
 	/**
 	 * Construct a circular Polygon
 	 */
-	static FPolygon2<T> MakeCircle(T Radius, int Steps, T AngleShiftRadians = 0)
+	static TPolygon2<T> MakeCircle(T Radius, int Steps, T AngleShiftRadians = 0)
 	{
-		FPolygon2<T> Circle;
+		TPolygon2<T> Circle;
 		Circle.Vertices.SetNumUninitialized(Steps);
 
 		for (int i = 0; i < Steps; ++i)
@@ -1041,5 +1041,5 @@ public:
 	}
 };
 
-typedef FPolygon2<double> FPolygon2d;
-typedef FPolygon2<float> FPolygon2f;
+typedef TPolygon2<double> FPolygon2d;
+typedef TPolygon2<float> FPolygon2f;
diff --git a/Engine/Plugins/Experimental/GeometryProcessing/Source/GeometricObjects/Public/Polyline3.h b/Engine/Plugins/Experimental/GeometryProcessing/Source/GeometricObjects/Public/Polyline3.h
--- a/Engine/Plugins/Experimental/GeometryProcessing/Source/GeometricObjects/Public/Polyline3.h
+++ b/Engine/Plugins/Experimental/GeometryProcessing/Source/GeometricObjects/Public/Polyline3.h
@@ -210,9 +210,9 @@ public:
 	/**
 	 * @return edge of the poyline starting at vertex SegmentIndex
 	 */
-	FSegment3<T> GetSegment(int SegmentIndex) const
+	TSegment3<T> GetSegment(int SegmentIndex) const
 	{
-		return FSegment3<T>(Vertices[SegmentIndex], Vertices[SegmentIndex+1]);
+		return TSegment3<T>(Vertices[SegmentIndex], Vertices[SegmentIndex+1]);
 	}
 
 
@@ -223,7 +223,7 @@ public:
 	 */
 	FVector3<T> PointAt(int SegmentIndex, T SegmentParam) const
 	{
-		FSegment3<T> seg(Vertices[SegmentIndex], Vertices[SegmentIndex + 1]);
+		TSegment3<T> seg(Vertices[SegmentIndex], Vertices[SegmentIndex + 1]);
 		return seg.PointAt(SegmentParam);
 	}
 
@@ -233,9 +233,9 @@ public:
 	/**
 	 * @return the bounding box of the polyline Vertices
 	 */
-	FAxisAlignedBox3<T> GetBounds() const
+	TAxisAlignedBox3<T> GetBounds() const
 	{
-		FAxisAlignedBox3<T> box = FAxisAlignedBox3<T>::Empty();
+		TAxisAlignedBox3<T> box = TAxisAlignedBox3<T>::Empty();
 		int NumVertices = Vertices.Num();
 		for (int k = 0; k < NumVertices; ++k)
 		{
@@ -263,7 +263,7 @@ public:
 
 
 	/**
-	 * SegmentIterator is used to iterate over the FSegment3<T> segments of the polyline
+	 * SegmentIterator is used to iterate over the TSegment3<T> segments of the polyline
 	 */
 	class SegmentIterator
 	{
@@ -272,10 +272,10 @@ public:
 		{
 			return i < Polyline->SegmentCount();
 		}
-		inline FSegment3<T> operator*() const
+		inline TSegment3<T> operator*() const
 		{
 			check(Polyline != nullptr && i < Polyline->SegmentCount());
-			return FSegment3<T>(Polyline->Vertices[i], Polyline->Vertices[i+1]);
+			return TSegment3<T>(Polyline->Vertices[i], Polyline->Vertices[i+1]);
 		}
 		inline SegmentIterator & operator++() 		// prefix
 		{
@@ -345,7 +345,7 @@ public:
 		for (int vi = 0; vi < N; ++vi)
 		{
 			// @todo can't we just use segment function here now?
-			FSegment3<T> seg = FSegment3<T>(Vertices[vi], Vertices[vi+1]);
+			TSegment3<T> seg = TSegment3<T>(Vertices[vi], Vertices[vi+1]);
 			T t = (QueryPoint - seg.Center).Dot(seg.Direction);
 			T d = TNumericLimits<T>::Max();
 			if (t >= seg.Extent)
diff --git a/Engine/Plugins/Experimental/GeometryProcessing/Source/GeometricObjects/Public/QuadricError.h b/Engine/Plugins/Experimental/GeometryProcessing/Source/GeometricObjects/Public/QuadricError.h
--- a/Engine/Plugins/Experimental/GeometryProcessing/Source/GeometricObjects/Public/QuadricError.h
+++ b/Engine/Plugins/Experimental/GeometryProcessing/Source/GeometricObjects/Public/QuadricError.h
@@ -413,7 +413,7 @@ public:
 		*
 		* For this quadric, each component of the normal is treated as separate attribute.
 		*/
-		FMatrix3<RealType> Matrix(P0, P1, P2, true); // row-based matrix constructor.
+		TMatrix3<RealType> Matrix(P0, P1, P2, true); // row-based matrix constructor.
 
 		bool bSolved = false;
 		double det = Matrix.Determinant();
diff --git a/Engine/Plugins/Experimental/GeometryProcessing/Source/GeometricObjects/Public/Quaternion.h b/Engine/Plugins/Experimental/GeometryProcessing/Source/GeometricObjects/Public/Quaternion.h
--- a/Engine/Plugins/Experimental/GeometryProcessing/Source/GeometricObjects/Public/Quaternion.h
+++ b/Engine/Plugins/Experimental/GeometryProcessing/Source/GeometricObjects/Public/Quaternion.h
@@ -9,7 +9,7 @@
 // ported from geometry3Sharp Quaternion
 
 template<typename RealType>
-struct FQuaternion
+struct TQuaternion
 {
 	// note: in Wm5 version, this is a 4-element arraY stored in order (w,x,y,z).
 	RealType X;
@@ -17,28 +17,28 @@ struct FQuaternion
 	RealType Z;
 	RealType W;
 
-	FQuaternion();
-	FQuaternion(RealType X, RealType Y, RealType Z, RealType W);
-	FQuaternion(const RealType* Values);
-	FQuaternion(const FQuaternion& Copy);
-	FQuaternion(const FVector3<RealType>& Axis, RealType Angle, bool bAngleIsDegrees);
-	FQuaternion(const FVector3<RealType>& From, const FVector3<RealType>& To);
-	FQuaternion(const FQuaternion<RealType>& From, const FQuaternion<RealType>& To, RealType InterpT);
-	FQuaternion(const FMatrix3<RealType>& RotationMatrix);
+	TQuaternion();
+	TQuaternion(RealType X, RealType Y, RealType Z, RealType W);
+	TQuaternion(const RealType* Values);
+	TQuaternion(const TQuaternion& Copy);
+	TQuaternion(const FVector3<RealType>& Axis, RealType Angle, bool bAngleIsDegrees);
+	TQuaternion(const FVector3<RealType>& From, const FVector3<RealType>& To);
+	TQuaternion(const TQuaternion<RealType>& From, const TQuaternion<RealType>& To, RealType InterpT);
+	TQuaternion(const TMatrix3<RealType>& RotationMatrix);
 
 	void SetAxisAngleD(const FVector3<RealType>& Axis, RealType AngleDeg);
 	void SetAxisAngleR(const FVector3<RealType>& Axis, RealType AngleRad);
 	void SetFromTo(const FVector3<RealType>& From, const FVector3<RealType>& To);
-	void SetToSlerp(const FQuaternion<RealType>& From, const FQuaternion<RealType>& To, RealType InterpT);
-	void SetFromRotationMatrix(const FMatrix3<RealType>& RotationMatrix);
+	void SetToSlerp(const TQuaternion<RealType>& From, const TQuaternion<RealType>& To, RealType InterpT);
+	void SetFromRotationMatrix(const TMatrix3<RealType>& RotationMatrix);
 
-	static FQuaternion<RealType> Zero() { return FQuaternion<RealType>((RealType)0, (RealType)0, (RealType)0, (RealType)0); }
-	static FQuaternion<RealType> Identity() { return FQuaternion<RealType>((RealType)0, (RealType)0, (RealType)0, (RealType)1); }
+	static TQuaternion<RealType> Zero() { return TQuaternion<RealType>((RealType)0, (RealType)0, (RealType)0, (RealType)0); }
+	static TQuaternion<RealType> Identity() { return TQuaternion<RealType>((RealType)0, (RealType)0, (RealType)0, (RealType)1); }
 
 	RealType & operator[](int i) { return (&X)[i]; }
 	const RealType & operator[](int i) const { return (&X)[i]; }
 
-	bool EpsilonEqual(const FQuaternion<RealType>& Other, RealType Epsilon) const;
+	bool EpsilonEqual(const TQuaternion<RealType>& Other, RealType Epsilon) const;
 
 	RealType Length() const { return (RealType)sqrt(X*X + Y*Y + Z*Z + W*W); }
 	RealType SquaredLength() const { return X*X + Y*Y + Z*Z + W*W; }
@@ -48,14 +48,14 @@ struct FQuaternion
 	FVector3<RealType> AxisZ() const;
 
 	RealType Normalize(const RealType epsilon = 0);
-	FQuaternion<RealType> Normalized(const RealType epsilon = 0) const;
+	TQuaternion<RealType> Normalized(const RealType epsilon = 0) const;
 
-	RealType Dot(const FQuaternion<RealType>& Other) const;
+	RealType Dot(const TQuaternion<RealType>& Other) const;
 
-	FQuaternion<RealType> Inverse() const;
+	TQuaternion<RealType> Inverse() const;
 	FVector3<RealType> InverseMultiply(const FVector3<RealType>& Other) const;
 
-	FMatrix3<RealType> ToRotationMatrix() const;
+	TMatrix3<RealType> ToRotationMatrix() const;
 
 
 	// available for porting:
@@ -71,7 +71,7 @@ struct FQuaternion
 		Quat.W = (float)W;
 		return Quat;
 	}
-	inline FQuaternion(const FQuat& Quat)
+	inline TQuaternion(const FQuat& Quat)
 	{
 		X = (RealType)Quat.X;
 		Y = (RealType)Quat.Y;
@@ -83,20 +83,20 @@ struct FQuaternion
 
 
 
-typedef FQuaternion<float> FQuaternionf;
-typedef FQuaternion<double> FQuaterniond;
+typedef TQuaternion<float> FQuaternionf;
+typedef TQuaternion<double> FQuaterniond;
 
 
 
 
 template<typename RealType>
-FQuaternion<RealType>::FQuaternion()
+TQuaternion<RealType>::TQuaternion()
 {
 	X = Y = Z = 0; W = 1;
 }
 
 template<typename RealType>
-FQuaternion<RealType>::FQuaternion(RealType x, RealType y, RealType z, RealType w)
+TQuaternion<RealType>::TQuaternion(RealType x, RealType y, RealType z, RealType w)
 {
 	X = x;
 	Y = y;
@@ -105,7 +105,7 @@ FQuaternion<RealType>::FQuaternion(RealType x, RealType y, RealType z, RealType
 }
 
 template<typename RealType>
-FQuaternion<RealType>::FQuaternion(const RealType* Values)
+TQuaternion<RealType>::TQuaternion(const RealType* Values)
 {
 	X = Values[0];
 	Y = Values[1];
@@ -114,7 +114,7 @@ FQuaternion<RealType>::FQuaternion(const RealType* Values)
 }
 
 template<typename RealType>
-FQuaternion<RealType>::FQuaternion(const FQuaternion& Copy)
+TQuaternion<RealType>::TQuaternion(const TQuaternion& Copy)
 {
 	X = Copy.X;
 	Y = Copy.Y;
@@ -123,28 +123,28 @@ FQuaternion<RealType>::FQuaternion(const FQuaternion& Copy)
 }
 
 template<typename RealType>
-FQuaternion<RealType>::FQuaternion(const FVector3<RealType>& Axis, RealType Angle, bool bAngleIsDegrees)
+TQuaternion<RealType>::TQuaternion(const FVector3<RealType>& Axis, RealType Angle, bool bAngleIsDegrees)
 {
 	X = Y = Z = 0; W = 1;
 	SetAxisAngleR(Axis, Angle * (bAngleIsDegrees ? TMathUtil<RealType>::DegToRad : (RealType)1));
 }
 
 template<typename RealType>
-FQuaternion<RealType>::FQuaternion(const FVector3<RealType>& From, const FVector3<RealType>& To)
+TQuaternion<RealType>::TQuaternion(const FVector3<RealType>& From, const FVector3<RealType>& To)
 {
 	X = Y = Z = 0; W = 1;
 	SetFromTo(From, To);
 }
 
 template<typename RealType>
-FQuaternion<RealType>::FQuaternion(const FQuaternion<RealType>& From, const FQuaternion<RealType>& To, RealType InterpT)
+TQuaternion<RealType>::TQuaternion(const TQuaternion<RealType>& From, const TQuaternion<RealType>& To, RealType InterpT)
 {
 	X = Y = Z = 0; W = 1;
 	SetToSlerp(From, To, InterpT);
 }
 
 template<typename RealType>
-FQuaternion<RealType>::FQuaternion(const FMatrix3<RealType>& RotationMatrix)
+TQuaternion<RealType>::TQuaternion(const TMatrix3<RealType>& RotationMatrix)
 {
 	X = Y = Z = 0; W = 1;
 	SetFromRotationMatrix(RotationMatrix);
@@ -155,7 +155,7 @@ FQuaternion<RealType>::FQuaternion(const FMatrix3<RealType>& RotationMatrix)
 
 
 template<typename RealType>
-RealType FQuaternion<RealType>::Normalize(const RealType Epsilon)
+RealType TQuaternion<RealType>::Normalize(const RealType Epsilon)
 {
 	RealType length = Length();
 	if (length > Epsilon)
@@ -172,42 +172,42 @@ RealType FQuaternion<RealType>::Normalize(const RealType Epsilon)
 }
 
 template<typename RealType>
-FQuaternion<RealType> FQuaternion<RealType>::Normalized(const RealType Epsilon) const
+TQuaternion<RealType> TQuaternion<RealType>::Normalized(const RealType Epsilon) const
 {
 	RealType length = Length();
 	if (length > Epsilon)
 	{
 		RealType invLength = ((RealType)1) / length;
-		return FQuaternion<RealType>(X*invLength, Y*invLength, Z*invLength, W*invLength);
+		return TQuaternion<RealType>(X*invLength, Y*invLength, Z*invLength, W*invLength);
 	}
 	return Zero();
 }
 
 template<typename RealType>
-RealType FQuaternion<RealType>::Dot(const FQuaternion<RealType>& Other) const
+RealType TQuaternion<RealType>::Dot(const TQuaternion<RealType>& Other) const
 {
 	return X * Other.X + Y * Other.Y + Z * Other.Z + W * Other.W;
 }
 
 
 template<typename RealType>
-FQuaternion<RealType> operator*(const FQuaternion<RealType>& A, const FQuaternion<RealType>& B) 
+TQuaternion<RealType> operator*(const TQuaternion<RealType>& A, const TQuaternion<RealType>& B) 
 {
 	RealType W = A.W * B.W - A.X * B.X - A.Y * B.Y - A.Z * B.Z;
 	RealType X = A.W * B.X + A.X * B.W + A.Y * B.Z - A.Z * B.Y;
 	RealType Y = A.W * B.Y + A.Y * B.W + A.Z * B.X - A.X * B.Z;
 	RealType Z = A.W * B.Z + A.Z * B.W + A.X * B.Y - A.Y * B.X;
-	return FQuaternion<RealType>(X, Y, Z, W);
+	return TQuaternion<RealType>(X, Y, Z, W);
 }
 
 template<typename RealType>
-FQuaternion<RealType> operator -(const FQuaternion<RealType>& A, const FQuaternion<RealType>& B) 
+TQuaternion<RealType> operator -(const TQuaternion<RealType>& A, const TQuaternion<RealType>& B) 
 {
-	return FQuaternion<RealType>(A.X - B.X, A.Y - B.Y, A.Z - B.Z, A.W - B.W);
+	return TQuaternion<RealType>(A.X - B.X, A.Y - B.Y, A.Z - B.Z, A.W - B.W);
 }
 
 template<typename RealType>
-FVector3<RealType> operator*(const FQuaternion<RealType>& Q, const FVector3<RealType>& V) 
+FVector3<RealType> operator*(const TQuaternion<RealType>& Q, const FVector3<RealType>& V) 
 {
 	//return q.ToRotationMatrix() * v;
 	// inline-expansion of above:
@@ -223,7 +223,7 @@ FVector3<RealType> operator*(const FQuaternion<RealType>& Q, const FVector3<Real
 
 
 template<typename RealType>
-FVector3<RealType> FQuaternion<RealType>::InverseMultiply(const FVector3<RealType>& V) const
+FVector3<RealType> TQuaternion<RealType>::InverseMultiply(const FVector3<RealType>& V) const
 {
 	RealType norm = SquaredLength();
 	if (norm > 0) 
@@ -244,7 +244,7 @@ FVector3<RealType> FQuaternion<RealType>::InverseMultiply(const FVector3<RealTyp
 
 
 template<typename RealType>
-FVector3<RealType> FQuaternion<RealType>::AxisX() const
+FVector3<RealType> TQuaternion<RealType>::AxisX() const
 {
 	RealType twoY = (RealType)2 * Y; RealType twoZ = (RealType)2 * Z;
 	RealType twoWY = twoY * W; RealType twoWZ = twoZ * W;
@@ -254,7 +254,7 @@ FVector3<RealType> FQuaternion<RealType>::AxisX() const
 }
 
 template<typename RealType>
-FVector3<RealType> FQuaternion<RealType>::AxisY() const
+FVector3<RealType> TQuaternion<RealType>::AxisY() const
 {
 	RealType twoX = (RealType)2 * X; RealType twoY = (RealType)2 * Y; RealType twoZ = (RealType)2 * Z;
 	RealType twoWX = twoX * W; RealType twoWZ = twoZ * W; RealType twoXX = twoX * X;
@@ -263,7 +263,7 @@ FVector3<RealType> FQuaternion<RealType>::AxisY() const
 }
 
 template<typename RealType>
-FVector3<RealType> FQuaternion<RealType>::AxisZ() const
+FVector3<RealType> TQuaternion<RealType>::AxisZ() const
 {
 	RealType twoX = (RealType)2 * X; RealType twoY = (RealType)2 * Y; RealType twoZ = (RealType)2 * Z;
 	RealType twoWX = twoX * W; RealType twoWY = twoY * W; RealType twoXX = twoX * X;
@@ -272,27 +272,27 @@ FVector3<RealType> FQuaternion<RealType>::AxisZ() const
 }
 
 template<typename RealType>
-FQuaternion<RealType> FQuaternion<RealType>::Inverse() const
+TQuaternion<RealType> TQuaternion<RealType>::Inverse() const
 {
 	RealType norm = SquaredLength();
 	if (norm > (RealType)0) 
 	{
 		RealType invNorm = (RealType)1 / norm;
-		return FQuaternion<RealType>(
+		return TQuaternion<RealType>(
 			-X * invNorm, -Y * invNorm, -Z * invNorm, W * invNorm);
 	}
-	return FQuaternion<RealType>::Zero();
+	return TQuaternion<RealType>::Zero();
 }
 
 
 template<typename RealType>
-void FQuaternion<RealType>::SetAxisAngleD(const FVector3<RealType>& Axis, RealType AngleDeg) 
+void TQuaternion<RealType>::SetAxisAngleD(const FVector3<RealType>& Axis, RealType AngleDeg) 
 {
 	SetAxisAngleR(Axis, TMathUtil<RealType>::DegToRad * AngleDeg);
 }
 
 template<typename RealType>
-void FQuaternion<RealType>::SetAxisAngleR(const FVector3<RealType>& Axis, RealType AngleRad)
+void TQuaternion<RealType>::SetAxisAngleR(const FVector3<RealType>& Axis, RealType AngleRad)
 {
 	RealType halfAngle = (RealType)0.5 * AngleRad;
 	RealType sn = (RealType)sin(halfAngle);
@@ -306,13 +306,13 @@ void FQuaternion<RealType>::SetAxisAngleR(const FVector3<RealType>& Axis, RealTy
 
 // this function can take non-normalized vectors vFrom and vTo (normalizes internally)
 template<typename RealType>
-void FQuaternion<RealType>::SetFromTo(const FVector3<RealType>& From, const FVector3<RealType>& To)
+void TQuaternion<RealType>::SetFromTo(const FVector3<RealType>& From, const FVector3<RealType>& To)
 {
 	// [TODO] this page seems to have optimized version:
 	//    http://lolengine.net/blog/2013/09/18/beautiful-maths-quaternion-from-vectors
 
 	// [RMS] not ideal to explicitly normalize here, but if we don't,
-	//   output FQuaternion is not normalized and this causes problems,
+	//   output TQuaternion is not normalized and this causes problems,
 	//   eg like drift if we do repeated SetFromTo()
 	FVector3<RealType> from = From.Normalized(), to = To.Normalized();
 	FVector3<RealType> bisector = (from + to).Normalized();
@@ -350,7 +350,7 @@ void FQuaternion<RealType>::SetFromTo(const FVector3<RealType>& From, const FVec
 
 
 template<typename RealType>
-void FQuaternion<RealType>::SetToSlerp( const FQuaternion<RealType>& From, const FQuaternion<RealType>& To, RealType InterpT)
+void TQuaternion<RealType>::SetToSlerp( const TQuaternion<RealType>& From, const TQuaternion<RealType>& To, RealType InterpT)
 {
 	RealType cs = From.Dot(To);
 	RealType angle = (RealType)acos(cs);
@@ -378,10 +378,10 @@ void FQuaternion<RealType>::SetToSlerp( const FQuaternion<RealType>& From, const
 
 
 template<typename RealType>
-void FQuaternion<RealType>::SetFromRotationMatrix(const FMatrix3<RealType>& RotationMatrix)
+void TQuaternion<RealType>::SetFromRotationMatrix(const TMatrix3<RealType>& RotationMatrix)
 {
 	// Algorithm in Ken Shoemake's article in 1987 SIGGRAPH course notes
-	// article "FQuaternion Calculus and Fast Animation".
+	// article "TQuaternion Calculus and Fast Animation".
 	FIndex3i next(1, 2, 0);
 
 	RealType trace = RotationMatrix(0,0) + RotationMatrix(1,1) + RotationMatrix(2, 2);
@@ -423,20 +423,20 @@ void FQuaternion<RealType>::SetFromRotationMatrix(const FMatrix3<RealType>& Rota
 		X = quat.X; Y = quat.Y; Z = quat.Z;
 	}
 
-	Normalize();   // we prefer normalized FQuaternions...
+	Normalize();   // we prefer normalized TQuaternions...
 }
 
 
 
 
 template<typename RealType>
-FMatrix3<RealType> FQuaternion<RealType>::ToRotationMatrix() const
+TMatrix3<RealType> TQuaternion<RealType>::ToRotationMatrix() const
 {
 	RealType twoX = 2 * X; RealType twoY = 2 * Y; RealType twoZ = 2 * Z;
 	RealType twoWX = twoX * W; RealType twoWY = twoY * W; RealType twoWZ = twoZ * W;
 	RealType twoXX = twoX * X; RealType twoXY = twoY * X; RealType twoXZ = twoZ * X;
 	RealType twoYY = twoY * Y; RealType twoYZ = twoZ * Y; RealType twoZZ = twoZ * Z;
-	FMatrix3<RealType> m = FMatrix3<RealType>::Zero();
+	TMatrix3<RealType> m = TMatrix3<RealType>::Zero();
 	m.Row0 = FVector3<RealType>(1 - (twoYY + twoZZ), twoXY - twoWZ, twoXZ + twoWY);
 	m.Row1 = FVector3<RealType>(twoXY + twoWZ, 1 - (twoXX + twoZZ), twoYZ - twoWX);
 	m.Row2 = FVector3<RealType>(twoXZ - twoWY, twoYZ + twoWX, 1 - (twoXX + twoYY));
@@ -446,7 +446,7 @@ FMatrix3<RealType> FQuaternion<RealType>::ToRotationMatrix() const
 
 
 template<typename RealType>
-bool FQuaternion<RealType>::EpsilonEqual(const FQuaternion<RealType>& Other, RealType Epsilon) const
+bool TQuaternion<RealType>::EpsilonEqual(const TQuaternion<RealType>& Other, RealType Epsilon) const
 {
 	return (RealType)fabs(X - Other.X) <= Epsilon &&
 		(RealType)fabs(Y - Other.Y) <= Epsilon &&
diff --git a/Engine/Plugins/Experimental/GeometryProcessing/Source/GeometricObjects/Public/RayTypes.h b/Engine/Plugins/Experimental/GeometryProcessing/Source/GeometricObjects/Public/RayTypes.h
--- a/Engine/Plugins/Experimental/GeometryProcessing/Source/GeometricObjects/Public/RayTypes.h
+++ b/Engine/Plugins/Experimental/GeometryProcessing/Source/GeometricObjects/Public/RayTypes.h
@@ -12,7 +12,7 @@
  * 3D Ray stored as Origin point and normalized Direction vector
  */
 template<typename RealType>
-class FRay3
+class TRay3
 {
 public:
 	/** Origin point */
@@ -21,7 +21,7 @@ public:
 	FVector3<RealType> Direction;
 
 	/** Construct ray at origin pointed down Z axis */
-	FRay3()
+	TRay3()
 	{
 		Origin = FVector3<RealType>::Zero();
 		Direction = FVector3<RealType>::UnitZ();
@@ -33,7 +33,7 @@ public:
 	 * @param Direction direction vector. Must be normalized if bIsNormalized=true
 	 * @param bIsNormalized if true, Direction will not be re-normalized
 	 */
-	FRay3(const FVector3<RealType>& Origin, const FVector3<RealType>& Direction, bool bIsNormalized = false)
+	TRay3(const FVector3<RealType>& Origin, const FVector3<RealType>& Direction, bool bIsNormalized = false)
 	{
 		this->Origin = Origin;
 		this->Direction = Direction;
@@ -102,7 +102,7 @@ public:
 	{
 		return FRay((FVector)Origin, (FVector)Direction);
 	}
-	inline FRay3(const FRay & RayIn)
+	inline TRay3(const FRay & RayIn)
 	{
 		Origin = (FVector3<RealType>)RayIn.Origin;
 		Direction = (FVector3<RealType>)RayIn.Direction;
@@ -110,6 +110,6 @@ public:
 
 
 };
-typedef FRay3<float> FRay3f;
-typedef FRay3<double> FRay3d;
+typedef TRay3<float> FRay3f;
+typedef TRay3<double> FRay3d;
 
diff --git a/Engine/Plugins/Experimental/GeometryProcessing/Source/GeometricObjects/Public/SegmentTypes.h b/Engine/Plugins/Experimental/GeometryProcessing/Source/GeometricObjects/Public/SegmentTypes.h
--- a/Engine/Plugins/Experimental/GeometryProcessing/Source/GeometricObjects/Public/SegmentTypes.h
+++ b/Engine/Plugins/Experimental/GeometryProcessing/Source/GeometricObjects/Public/SegmentTypes.h
@@ -12,7 +12,7 @@
  * 2D Line Segmented stored as Center point, normalized Direction vector, and scalar Extent
  */
 template<typename T>
-struct FSegment2
+struct TSegment2
 {
 public:
 	/** Center point of segment */
@@ -25,7 +25,7 @@ public:
 	/**
 	 * Construct a Segment from two Points
 	 */
-	FSegment2(const FVector2<T>& Point0, const FVector2<T>& Point1)
+	TSegment2(const FVector2<T>& Point0, const FVector2<T>& Point1)
 	{
 		// set from endpoints 
 		Center = T(.5) * (Point0 + Point1);
@@ -36,7 +36,7 @@ public:
 	/**
 	 * Construct a segment from a Center Point, normalized Direction, and scalar Extent
 	 */
-	FSegment2(const FVector2<T>& CenterIn, const FVector2<T>& DirectionIn, T ExtentIn)
+	TSegment2(const FVector2<T>& CenterIn, const FVector2<T>& DirectionIn, T ExtentIn)
 	{
 		Center = CenterIn;
 		Direction = DirectionIn; 
@@ -204,7 +204,7 @@ public:
 	 * @param IntervalThresh distance tolerance used to allow slighly-not-touching segments to be considered overlapping
 	 * @return true if segments intersect
 	 */
-	bool Intersects(const FSegment2<T>& OtherSegment, T DotThresh = TMathUtil<T>::Epsilon, T IntervalThresh = 0) const
+	bool Intersects(const TSegment2<T>& OtherSegment, T DotThresh = TMathUtil<T>::Epsilon, T IntervalThresh = 0) const
 	{
 		// see IntrLine2Line2 and IntrSegment2Segment2 for details on this code
 
@@ -231,8 +231,8 @@ public:
 			T t1 = Direction.Dot(diff);
 			T tmin = t1 - OtherSegment.Extent;
 			T tmax = t1 + OtherSegment.Extent;
-			FInterval1<T> extents(-Extent, Extent);
-			if (extents.Overlaps(FInterval1<T>(tmin, tmax)))
+			TInterval1<T> extents(-Extent, Extent);
+			if (extents.Overlaps(TInterval1<T>(tmin, tmax)))
 			{
 				return true;
 			}
@@ -309,8 +309,8 @@ protected:
 	}
 
 };
-typedef FSegment2<float> FSegment2f;
-typedef FSegment2<double> FSegment2d;
+typedef TSegment2<float> FSegment2f;
+typedef TSegment2<double> FSegment2d;
 
 
 
@@ -323,7 +323,7 @@ typedef FSegment2<double> FSegment2d;
  * 3D Line Segmented stored as Center point, normalized Direction vector, and scalar Extent
  */
 template<typename T>
-struct FSegment3
+struct TSegment3
 {
 public:
 	/** Center point of segment */
@@ -336,7 +336,7 @@ public:
 	/**
 	 * Construct a Segment from two Points
 	 */
-	FSegment3(const FVector3<T>& Point0, const FVector3<T>& Point1)
+	TSegment3(const FVector3<T>& Point0, const FVector3<T>& Point1)
 	{
 		// set from endpoints 
 		Center = T(.5) * (Point0 + Point1);
@@ -347,7 +347,7 @@ public:
 	/**
 	 * Construct a segment from a Center Point, normalized Direction, and scalar Extent
 	 */
-	FSegment3(const FVector3<T>& CenterIn, const FVector3<T>& DirectionIn, T ExtentIn)
+	TSegment3(const FVector3<T>& CenterIn, const FVector3<T>& DirectionIn, T ExtentIn)
 	{
 		Center = CenterIn;
 		Direction = DirectionIn;
@@ -502,5 +502,5 @@ protected:
 	}
 
 };
-typedef FSegment3<float> FSegment3f;
-typedef FSegment3<double> FSegment3d;
+typedef TSegment3<float> FSegment3f;
+typedef TSegment3<double> FSegment3d;
diff --git a/Engine/Plugins/Experimental/GeometryProcessing/Source/GeometricObjects/Public/Spatial/MeshAABBTree3.h b/Engine/Plugins/Experimental/GeometryProcessing/Source/GeometricObjects/Public/Spatial/MeshAABBTree3.h
--- a/Engine/Plugins/Experimental/GeometryProcessing/Source/GeometricObjects/Public/Spatial/MeshAABBTree3.h
+++ b/Engine/Plugins/Experimental/GeometryProcessing/Source/GeometricObjects/Public/Spatial/MeshAABBTree3.h
@@ -414,9 +414,9 @@ protected:
 	// storage for Box Nodes.
 	//   - BoxToIndex is a pointer into IndexList
 	//   - BoxCenters and BoxExtents are the Centers/extents of the bounding boxes
-	FDynamicVector<int> BoxToIndex;
-	FDynamicVector<FVector3d> BoxCenters;
-	FDynamicVector<FVector3d> BoxExtents;
+	TDynamicVector<int> BoxToIndex;
+	TDynamicVector<FVector3d> BoxCenters;
+	TDynamicVector<FVector3d> BoxExtents;
 
 	// list of indices for a given Box. There is *no* marker/sentinel between
 	// boxes, you have to get the starting index from BoxToIndex[]
@@ -428,7 +428,7 @@ protected:
 	//       internal Box, with index (-IndexList[i])-1     (shift-by-one in case actual value is 0!)
 	//   - if i > TrianglesEnd and IndexList[i] > 0, this is a two-child
 	//       internal Box, with indices IndexList[i]-1 and IndexList[i+1]-1
-	FDynamicVector<int> IndexList;
+	TDynamicVector<int> IndexList;
 
 	// IndexList[i] for i < TrianglesEnd is a triangle-index list, otherwise Box-index pair/single
 	int TrianglesEnd = -1;
@@ -438,10 +438,10 @@ protected:
 
 	struct FBoxesSet
 	{
-		FDynamicVector<int> BoxToIndex;
-		FDynamicVector<FVector3d> BoxCenters;
-		FDynamicVector<FVector3d> BoxExtents;
-		FDynamicVector<int> IndexList;
+		TDynamicVector<int> BoxToIndex;
+		TDynamicVector<FVector3d> BoxCenters;
+		TDynamicVector<FVector3d> BoxExtents;
+		TDynamicVector<int> IndexList;
 		int IBoxCur;
 		int IIndicesCur;
 		FBoxesSet()
diff --git a/Engine/Plugins/Experimental/GeometryProcessing/Source/GeometricObjects/Public/TriangleTypes.h b/Engine/Plugins/Experimental/GeometryProcessing/Source/GeometricObjects/Public/TriangleTypes.h
--- a/Engine/Plugins/Experimental/GeometryProcessing/Source/GeometricObjects/Public/TriangleTypes.h
+++ b/Engine/Plugins/Experimental/GeometryProcessing/Source/GeometricObjects/Public/TriangleTypes.h
@@ -1,6 +1,6 @@
 // Copyright 1998-2019 Epic Games, Inc. All Rights Reserved.
 
-// Port of WildMagic DistPoint3FTriangle3
+// Port of WildMagic DistPoint3TTriangle3
 
 #pragma once
 
@@ -27,20 +27,20 @@ namespace TriangleUtil
 
 
 template<typename RealType>
-struct FTriangle2
+struct TTriangle2
 {
 	FVector2<RealType> V[3];
 
-	FTriangle2() {}
+	TTriangle2() {}
 
-	FTriangle2(const FVector2<RealType>& V0, const FVector2<RealType>& V1, const FVector2<RealType>& V2)
+	TTriangle2(const FVector2<RealType>& V0, const FVector2<RealType>& V1, const FVector2<RealType>& V2)
 	{
 		V[0] = V0;
 		V[1] = V1;
 		V[2] = V2;
 	}
 
-	FTriangle2(const FVector2<RealType> VIn[3])
+	TTriangle2(const FVector2<RealType> VIn[3])
 	{
 		V[0] = VIn[0];
 		V[1] = VIn[1];
@@ -100,9 +100,9 @@ struct FTriangle2
 
 };
 
-typedef FTriangle2<float> FTriangle2f;
-typedef FTriangle2<double> FTriangle2d;
-typedef FTriangle2<int> FTriangle2i;
+typedef TTriangle2<float> FTriangle2f;
+typedef TTriangle2<double> FTriangle2d;
+typedef TTriangle2<int> FTriangle2i;
 
 
 
@@ -110,20 +110,20 @@ typedef FTriangle2<int> FTriangle2i;
 
 
 template<typename RealType>
-struct FTriangle3
+struct TTriangle3
 {
 	FVector3<RealType> V[3];
 
-	FTriangle3() {}
+	TTriangle3() {}
 
-	FTriangle3(const FVector3<RealType>& V0, const FVector3<RealType>& V1, const FVector3<RealType>& V2)
+	TTriangle3(const FVector3<RealType>& V0, const FVector3<RealType>& V1, const FVector3<RealType>& V2)
 	{
 		V[0] = V0;
 		V[1] = V1;
 		V[2] = V2;
 	}
 
-	FTriangle3(const FVector3<RealType> VIn[3])
+	TTriangle3(const FVector3<RealType> VIn[3])
 	{
 		V[0] = VIn[0];
 		V[1] = VIn[1];
@@ -168,7 +168,7 @@ struct FTriangle3
 	}
 };
 
-typedef FTriangle3<float> FTriangle3f;
-typedef FTriangle3<double> FTriangle3d;
-typedef FTriangle3<int> FTriangle3i;
+typedef TTriangle3<float> FTriangle3f;
+typedef TTriangle3<double> FTriangle3d;
+typedef TTriangle3<int> FTriangle3i;
 
diff --git a/Engine/Plugins/Experimental/GeometryProcessing/Source/GeometricObjects/Public/Util/DynamicVector.h b/Engine/Plugins/Experimental/GeometryProcessing/Source/GeometricObjects/Public/Util/DynamicVector.h
--- a/Engine/Plugins/Experimental/GeometryProcessing/Source/GeometricObjects/Public/Util/DynamicVector.h
+++ b/Engine/Plugins/Experimental/GeometryProcessing/Source/GeometricObjects/Public/Util/DynamicVector.h
@@ -17,30 +17,30 @@
  * Iterator functions suitable for use with range-based for are provided
  */
 template <class Type>
-class FDynamicVector
+class TDynamicVector
 {
 public:
-	FDynamicVector()
+	TDynamicVector()
 	{
 		CurBlock = 0;
 		CurBlockUsed = 0;
 		Blocks.Emplace();
 	}
 
-	FDynamicVector(const FDynamicVector& Copy)
+	TDynamicVector(const TDynamicVector& Copy)
 	{
 		*this = Copy;
 	}
 
-	FDynamicVector(FDynamicVector&& Moved)
+	TDynamicVector(TDynamicVector&& Moved)
 	{
 		*this = MoveTemp(Moved);
 	}
 
-	~FDynamicVector() {}
+	~TDynamicVector() {}
 
-	const FDynamicVector& operator=(const FDynamicVector& Copy);
-	const FDynamicVector& operator=(FDynamicVector&& Moved);
+	const TDynamicVector& operator=(const TDynamicVector& Copy);
+	const TDynamicVector& operator=(TDynamicVector&& Moved);
 
 	inline void Clear();
 	inline void Fill(const Type& Value);
@@ -55,7 +55,7 @@ public:
 	inline size_t GetByteCount() const { return Blocks.Num() * BlockSize * sizeof(Type); }
 
 	inline void Add(const Type& Data);
-	inline void Add(const FDynamicVector<Type>& Data);
+	inline void Add(const TDynamicVector<Type>& Data);
 	inline void PopBack();
 
 	inline void InsertAt(const Type& Data, unsigned int Index);
@@ -128,14 +128,14 @@ public:
 		}
 
 	protected:
-		FDynamicVector<Type>* DVector;
+		TDynamicVector<Type>* DVector;
 		int Idx;
-		inline FIterator(FDynamicVector<Type>* Parent, int ICur)
+		inline FIterator(TDynamicVector<Type>* Parent, int ICur)
 		{
 			DVector = Parent;
 			Idx = ICur;
 		}
-		friend class FDynamicVector<Type>;
+		friend class TDynamicVector<Type>;
 	};
 
 	/** @return iterator at beginning of vector */
@@ -156,30 +156,30 @@ public:
 
 
 template <class Type, int N>
-class FDynamicVectorN
+class TDynamicVectorN
 {
 public:
-	FDynamicVectorN()
+	TDynamicVectorN()
 	{
 	}
-	FDynamicVectorN(const FDynamicVectorN& Copy)
+	TDynamicVectorN(const TDynamicVectorN& Copy)
 		: Data(Copy.Data)
 	{
 	}
-	FDynamicVectorN(FDynamicVectorN&& Moved)
+	TDynamicVectorN(TDynamicVectorN&& Moved)
 		: Data(MoveTemp(Moved.Data))
 	{
 	}
-	~FDynamicVectorN()
+	~TDynamicVectorN()
 	{
 	}
 
-	const FDynamicVectorN& operator=(const FDynamicVectorN& Copy)
+	const TDynamicVectorN& operator=(const TDynamicVectorN& Copy)
 	{
 		Data = Copy.Data;
 		return *this;
 	}
-	const FDynamicVectorN& operator=(FDynamicVectorN&& Moved)
+	const TDynamicVectorN& operator=(TDynamicVectorN&& Moved)
 	{
 		Data = MoveTemp(Moved.Data);
 		return *this;
@@ -308,26 +308,26 @@ public:
 	}
 
 protected:
-	FDynamicVector<Type> Data;
+	TDynamicVector<Type> Data;
 
 	friend class FIterator;
 };
 
-template class FDynamicVectorN<double, 2>;
+template class TDynamicVectorN<double, 2>;
 
-typedef FDynamicVectorN<float, 3> FDynamicVector3f;
-typedef FDynamicVectorN<float, 2> FDynamicVector2f;
-typedef FDynamicVectorN<double, 3> FDynamicVector3d;
-typedef FDynamicVectorN<double, 2> FDynamicVector2d;
-typedef FDynamicVectorN<int, 3> FDynamicVector3i;
-typedef FDynamicVectorN<int, 2> FDynamicVector2i;
+typedef TDynamicVectorN<float, 3> TDynamicVector3f;
+typedef TDynamicVectorN<float, 2> TDynamicVector2f;
+typedef TDynamicVectorN<double, 3> TDynamicVector3d;
+typedef TDynamicVectorN<double, 2> TDynamicVector2d;
+typedef TDynamicVectorN<int, 3> TDynamicVector3i;
+typedef TDynamicVectorN<int, 2> TDynamicVector2i;
 
 
 
 
 
 template <class Type>
-const FDynamicVector<Type>& FDynamicVector<Type>::operator=(const FDynamicVector& Copy)
+const TDynamicVector<Type>& TDynamicVector<Type>::operator=(const TDynamicVector& Copy)
 {
 	Blocks = Copy.Blocks;
 	CurBlock = Copy.CurBlock;
@@ -336,7 +336,7 @@ const FDynamicVector<Type>& FDynamicVector<Type>::operator=(const FDynamicVector
 }
 
 template <class Type>
-const FDynamicVector<Type>& FDynamicVector<Type>::operator=(FDynamicVector&& Moved)
+const TDynamicVector<Type>& TDynamicVector<Type>::operator=(TDynamicVector&& Moved)
 {
 	Blocks = MoveTemp(Moved.Blocks);
 	CurBlock = Moved.CurBlock;
@@ -345,7 +345,7 @@ const FDynamicVector<Type>& FDynamicVector<Type>::operator=(FDynamicVector&& Mov
 }
 
 template <class Type>
-void FDynamicVector<Type>::Clear()
+void TDynamicVector<Type>::Clear()
 {
 	Blocks.Empty();
 	CurBlock = 0;
@@ -354,7 +354,7 @@ void FDynamicVector<Type>::Clear()
 }
 
 template <class Type>
-void FDynamicVector<Type>::Fill(const Type& Value)
+void TDynamicVector<Type>::Fill(const Type& Value)
 {
 	size_t Count = Blocks.Num();
 	for (unsigned int i = 0; i < Count; ++i)
@@ -364,7 +364,7 @@ void FDynamicVector<Type>::Fill(const Type& Value)
 }
 
 template <class Type>
-void FDynamicVector<Type>::Resize(size_t Count)
+void TDynamicVector<Type>::Resize(size_t Count)
 {
 	if (GetLength() == Count)
 	{
@@ -403,7 +403,7 @@ void FDynamicVector<Type>::Resize(size_t Count)
 }
 
 template <class Type>
-void FDynamicVector<Type>::Resize(size_t Count, const Type& InitValue)
+void TDynamicVector<Type>::Resize(size_t Count, const Type& InitValue)
 {
 	size_t nCurSize = GetLength();
 	Resize(Count);
@@ -414,7 +414,7 @@ void FDynamicVector<Type>::Resize(size_t Count, const Type& InitValue)
 }
 
 template <class Type>
-void FDynamicVector<Type>::Add(const Type& Value)
+void TDynamicVector<Type>::Add(const Type& Value)
 {
 	if (CurBlockUsed == BlockSize)
 	{
@@ -431,7 +431,7 @@ void FDynamicVector<Type>::Add(const Type& Value)
 
 
 template <class Type>
-void FDynamicVector<Type>::Add(const FDynamicVector<Type>& AddData)
+void TDynamicVector<Type>::Add(const TDynamicVector<Type>& AddData)
 {
 	// @todo it could be more efficient to use memcopies here...
 	size_t nSize = AddData.Num();
@@ -442,7 +442,7 @@ void FDynamicVector<Type>::Add(const FDynamicVector<Type>& AddData)
 }
 
 template <class Type>
-void FDynamicVector<Type>::PopBack()
+void TDynamicVector<Type>::PopBack()
 {
 	if (CurBlockUsed > 0)
 	{
@@ -457,7 +457,7 @@ void FDynamicVector<Type>::PopBack()
 }
 
 template <class Type>
-void FDynamicVector<Type>::InsertAt(const Type& AddData, unsigned int Index)
+void TDynamicVector<Type>::InsertAt(const Type& AddData, unsigned int Index)
 {
 	size_t s = GetLength();
 	if (Index == s)
@@ -477,7 +477,7 @@ void FDynamicVector<Type>::InsertAt(const Type& AddData, unsigned int Index)
 
 template <typename Type>
 template <typename Func>
-void FDynamicVector<Type>::Apply(const Func& f)
+void TDynamicVector<Type>::Apply(const Func& applyF)
 {
 	for (int bi = 0; bi < CurBlock; ++bi)
 	{
diff --git a/Engine/Plugins/Experimental/GeometryProcessing/Source/GeometricObjects/Public/Util/IndexPriorityQueue.h b/Engine/Plugins/Experimental/GeometryProcessing/Source/GeometricObjects/Public/Util/IndexPriorityQueue.h
--- a/Engine/Plugins/Experimental/GeometryProcessing/Source/GeometricObjects/Public/Util/IndexPriorityQueue.h
+++ b/Engine/Plugins/Experimental/GeometryProcessing/Source/GeometricObjects/Public/Util/IndexPriorityQueue.h
@@ -11,7 +11,7 @@
  * Integer IDs must be provided by the user to identify unique nodes.
  * Internally an array is used to keep track of the mapping from ids to internal indices, so the max ID must also be provided.
  *
- * @todo based on C# code where FDynamicVector could not return a reference to internal element. In C++ we can and code could be updated to be more efficient in many places.
+ * @todo based on C# code where TDynamicVector could not return a reference to internal element. In C++ we can and code could be updated to be more efficient in many places.
  * @todo id_to_index could be sparse in many situations...
  */
 class FIndexPriorityQueue
@@ -29,7 +29,7 @@ public:
 	};
 
 	/** tree of allocated nodes, stored linearly. active up to num_nodes (allocated may be larger) */
-	FDynamicVector<FQueueNode> nodes;
+	TDynamicVector<FQueueNode> nodes;
 	/** count of active nodes */
 	int num_nodes;
 	/** mapping from external ids to internal node indices */
diff --git a/Engine/Plugins/Experimental/GeometryProcessing/Source/GeometricObjects/Public/Util/IndexUtil.h b/Engine/Plugins/Experimental/GeometryProcessing/Source/GeometricObjects/Public/Util/IndexUtil.h
--- a/Engine/Plugins/Experimental/GeometryProcessing/Source/GeometricObjects/Public/Util/IndexUtil.h
+++ b/Engine/Plugins/Experimental/GeometryProcessing/Source/GeometricObjects/Public/Util/IndexUtil.h
@@ -114,7 +114,7 @@ namespace IndexUtil
 	 * @return vertex id of other vertex, or InvalidID if not found
 	 */
 	template<typename T>
-	inline int FindTriOtherVtx(T VertexID1, T VertexID2, const FDynamicVector<T> & TriIndexArray, T TriangleIndex)
+	inline int FindTriOtherVtx(T VertexID1, T VertexID2, const TDynamicVector<T> & TriIndexArray, T TriangleIndex)
 	{
 		int i = 3 * TriangleIndex;
 		for (int j = 0; j < 3; ++j) 
diff --git a/Engine/Plugins/Experimental/GeometryProcessing/Source/GeometricObjects/Public/Util/MeshCaches.h b/Engine/Plugins/Experimental/GeometryProcessing/Source/GeometricObjects/Public/Util/MeshCaches.h
--- a/Engine/Plugins/Experimental/GeometryProcessing/Source/GeometricObjects/Public/Util/MeshCaches.h
+++ b/Engine/Plugins/Experimental/GeometryProcessing/Source/GeometricObjects/Public/Util/MeshCaches.h
@@ -9,9 +9,9 @@
  */
 struct FMeshTriInfoCache
 {
-	FDynamicVector<FVector3d> Centroids;
-	FDynamicVector<FVector3d> Normals;
-	FDynamicVector<double> Areas;
+	TDynamicVector<FVector3d> Centroids;
+	TDynamicVector<FVector3d> Normals;
+	TDynamicVector<double> Areas;
 
 	void GetTriInfo(int TriangleID, FVector3d& NormalOut, double& AreaOut, FVector3d& CentroidOut) const
 	{
diff --git a/Engine/Plugins/Experimental/GeometryProcessing/Source/GeometricObjects/Public/Util/RefCountVector.h b/Engine/Plugins/Experimental/GeometryProcessing/Source/GeometricObjects/Public/Util/RefCountVector.h
--- a/Engine/Plugins/Experimental/GeometryProcessing/Source/GeometricObjects/Public/Util/RefCountVector.h
+++ b/Engine/Plugins/Experimental/GeometryProcessing/Source/GeometricObjects/Public/Util/RefCountVector.h
@@ -22,21 +22,21 @@ class FRefCountVector
 public:
 	static constexpr short INVALID_REF_COUNT = -1;
 
-	FDynamicVector<short> RefCounts;
-	FDynamicVector<int> FreeIndices;
+	TDynamicVector<short> RefCounts;
+	TDynamicVector<int> FreeIndices;
 	int UsedCount;
 
 	FRefCountVector()
 	{
-		RefCounts = FDynamicVector<short>();
-		FreeIndices = FDynamicVector<int>();
+		RefCounts = TDynamicVector<short>();
+		FreeIndices = TDynamicVector<int>();
 		UsedCount = 0;
 	}
 
 	FRefCountVector(const FRefCountVector& Copy)
 	{
-		RefCounts = FDynamicVector<short>(Copy.RefCounts);
-		FreeIndices = FDynamicVector<int>(Copy.FreeIndices);
+		RefCounts = TDynamicVector<short>(Copy.RefCounts);
+		FreeIndices = TDynamicVector<int>(Copy.FreeIndices);
 		UsedCount = Copy.UsedCount;
 	}
 
@@ -222,7 +222,7 @@ public:
 
 
 
-	const FDynamicVector<short>& GetRawRefCounts() const
+	const TDynamicVector<short>& GetRawRefCounts() const
 	{
 		return RefCounts;
 	}
@@ -230,7 +230,7 @@ public:
 	/**
 	 * @warning you should not use this!
 	 */
-	FDynamicVector<short>& GetRawRefCountsUnsafe()
+	TDynamicVector<short>& GetRawRefCountsUnsafe()
 	{
 		return RefCounts;
 	}
@@ -250,7 +250,7 @@ public:
 
 	void RebuildFreeList()
 	{
-		FreeIndices = FDynamicVector<int>();
+		FreeIndices = TDynamicVector<int>();
 		UsedCount = 0;
 
 		int N = (int)RefCounts.GetLength();
@@ -271,7 +271,7 @@ public:
 
 	void Trim(int maxIndex)
 	{
-		FreeIndices = FDynamicVector<int>();
+		FreeIndices = TDynamicVector<int>();
 		RefCounts.Resize(maxIndex);
 		UsedCount = maxIndex;
 	}
diff --git a/Engine/Plugins/Experimental/GeometryProcessing/Source/GeometricObjects/Public/Util/SmallListSet.h b/Engine/Plugins/Experimental/GeometryProcessing/Source/GeometricObjects/Public/Util/SmallListSet.h
--- a/Engine/Plugins/Experimental/GeometryProcessing/Source/GeometricObjects/Public/Util/SmallListSet.h
+++ b/Engine/Plugins/Experimental/GeometryProcessing/Source/GeometricObjects/Public/Util/SmallListSet.h
@@ -8,11 +8,11 @@
 
 /**
  * FSmallListSet stores a set of short integer-valued variable-size lists.
- * The lists are encoded into a few large FDynamicVector buffers, with internal pooling,
+ * The lists are encoded into a few large TDynamicVector buffers, with internal pooling,
  * so adding/removing lists usually does not involve any or delete ops.
  * 
  * The lists are stored in two parts. The first N elements are stored in a linear
- * subset of a FDynamicVector. If the list spills past these N elements, the extra elements
+ * subset of a TDynamicVector. If the list spills past these N elements, the extra elements
  * are stored in a linked list (which is also stored in a flat array).
  * 
  * Each list stores its count, so list-size operations are constant time.
@@ -32,16 +32,16 @@ protected:
 	static constexpr int32 BLOCK_LIST_OFFSET = BLOCKSIZE + 1;
 
 	/** mapping from list index to offset into block_store that contains list data */
-	FDynamicVector<int32> ListHeads;
+	TDynamicVector<int32> ListHeads;
 
 	/** 
 	 * flat buffer used to store per-list linear-memory blocks. 
 	 * blocks are BLOCKSIZE+2 long, elements are [CurrentCount, item0...itemN, LinkedListPtr]
 	 */
-	FDynamicVector<int32> ListBlocks; 
+	TDynamicVector<int32> ListBlocks; 
 
 	/** list of free blocks as indices/offsets into block_store */
-	FDynamicVector<int32> FreeBlocks;
+	TDynamicVector<int32> FreeBlocks;
 
 	/** number of allocated lists */
 	int32 AllocatedCount = 0;
@@ -50,7 +50,7 @@ protected:
 	 * flat buffer used to store linked-list "spill" elements
 	 * each element is [value, next_ptr]
 	 */
-	FDynamicVector<int32> LinkedListElements;
+	TDynamicVector<int32> LinkedListElements;
 
 	/** index of first free element in linked_store */
 	int32 FreeHeadIndex;
