make geometric objects / dvec more unreal-y

diff --git a/Engine/Plugins/Experimental/GeometryProcessing/Source/GeometricObjects/Public/Util/DynamicVector.h b/Engine/Plugins/Experimental/GeometryProcessing/Source/GeometricObjects/Public/Util/DynamicVector.h
--- a/Engine/Plugins/Experimental/GeometryProcessing/Source/GeometricObjects/Public/Util/DynamicVector.h
+++ b/Engine/Plugins/Experimental/GeometryProcessing/Source/GeometricObjects/Public/Util/DynamicVector.h
@@ -8,8 +8,6 @@
 #include "VectorTypes.h"
 #include "IndexTypes.h"
 
-#include <array>
-
 
 /*
  * Blocked array with fixed, power-of-two sized blocks.
@@ -86,7 +84,7 @@ protected:
 	unsigned int CurBlock;
 	unsigned int CurBlockUsed;
 
-	using BlockType = std::array<Type, BlockSize>;
+	using BlockType = TStaticArray<Type, BlockSize>;
 	TArray<BlockType> Blocks;
 
 	friend class FIterator;
@@ -219,12 +217,18 @@ public:
 		return Data.GetByteCount();
 	}
 
-	inline void Add(const std::array<Type, N>& AddData)
+	// simple struct to help pass N-dimensional data without presuming a vector type (e.g. just via initializer list)
+	struct ElementVectorN
+	{
+		Type Data[N];
+	};
+
+	inline void Add(const ElementVectorN& AddData)
 	{
 		// todo specialize for N=2,3,4
 		for (int i = 0; i < N; i++)
 		{
-			Data.Add(AddData[i]);
+			Data.Add(AddData.Data[i]);
 		}
 	}
 
@@ -236,11 +240,11 @@ public:
 		}
 	} // TODO specialize
 
-	inline void InsertAt(const std::array<Type, N>& AddData, unsigned int Index)
+	inline void InsertAt(const ElementVectorN& AddData, unsigned int Index)
 	{
 		for (int i = 1; i <= N; i++)
 		{
-			Data.InsertAt(AddData[N - i], N * (Index + 1) - i);
+			Data.InsertAt(AddData.Data[N - i], N * (Index + 1) - i);
 		}
 	}
 
@@ -359,7 +363,10 @@ void TDynamicVector<Type>::Fill(const Type& Value)
 	size_t Count = Blocks.Num();
 	for (unsigned int i = 0; i < Count; ++i)
 	{
-		Blocks[i].fill(Value);
+		for (uint32 ElementIndex = 0; ElementIndex < BlockSize; ++ElementIndex)
+		{
+			Blocks[i][ElementIndex] = Value;
+		}
 	}
 }
 
@@ -377,11 +384,6 @@ void TDynamicVector<Type>::Resize(size_t Count)
 	// figure out how many are currently allocated...
 	size_t nCurCount = Blocks.Num();
 
-	// erase extra segments memory
-	// [RMS] not necessary right? std::array will deallocate?
-	//for (int i = nNumSegs; i < nCurCount; ++i)
-	//	Blocks[i] = null;
-
 	// resize to right number of segments
 	if (nNumSegs >= Blocks.Num())
 	{
