add missing header
#rnx

diff --git a/Engine/Plugins/Experimental/GeometryProcessing/Source/GeometricObjects/Public/Util/DynamicVector.h b/Engine/Plugins/Experimental/GeometryProcessing/Source/GeometricObjects/Public/Util/DynamicVector.h
--- a/Engine/Plugins/Experimental/GeometryProcessing/Source/GeometricObjects/Public/Util/DynamicVector.h
+++ b/Engine/Plugins/Experimental/GeometryProcessing/Source/GeometricObjects/Public/Util/DynamicVector.h
@@ -5,9 +5,12 @@
 #pragma once
 
 #include <CoreMinimal.h>
+#include "Containers/StaticArray.h"
 #include "VectorTypes.h"
 #include "IndexTypes.h"
 
+#include <array>
+
 
 /*
  * Blocked array with fixed, power-of-two sized blocks.
@@ -84,7 +87,7 @@ protected:
 	unsigned int CurBlock;
 	unsigned int CurBlockUsed;
 
-	using BlockType = TStaticArray<Type, BlockSize>;
+	using BlockType = std::array<Type, BlockSize>;
 	TArray<BlockType> Blocks;
 
 	friend class FIterator;
@@ -217,18 +220,12 @@ public:
 		return Data.GetByteCount();
 	}
 
-	// simple struct to help pass N-dimensional data without presuming a vector type (e.g. just via initializer list)
-	struct ElementVectorN
-	{
-		Type Data[N];
-	};
-
-	inline void Add(const ElementVectorN& AddData)
+	inline void Add(const std::array<Type, N>& AddData)
 	{
 		// todo specialize for N=2,3,4
 		for (int i = 0; i < N; i++)
 		{
-			Data.Add(AddData.Data[i]);
+			Data.Add(AddData[i]);
 		}
 	}
 
@@ -240,11 +237,11 @@ public:
 		}
 	} // TODO specialize
 
-	inline void InsertAt(const ElementVectorN& AddData, unsigned int Index)
+	inline void InsertAt(const std::array<Type, N>& AddData, unsigned int Index)
 	{
 		for (int i = 1; i <= N; i++)
 		{
-			Data.InsertAt(AddData.Data[N - i], N * (Index + 1) - i);
+			Data.InsertAt(AddData[N - i], N * (Index + 1) - i);
 		}
 	}
 
@@ -363,10 +360,7 @@ void TDynamicVector<Type>::Fill(const Type& Value)
 	size_t Count = Blocks.Num();
 	for (unsigned int i = 0; i < Count; ++i)
 	{
-		for (uint32 ElementIndex = 0; ElementIndex < BlockSize; ++ElementIndex)
-		{
-			Blocks[i][ElementIndex] = Value;
-		}
+		Blocks[i].fill(Value);
 	}
 }
 
@@ -384,6 +378,11 @@ void TDynamicVector<Type>::Resize(size_t Count)
 	// figure out how many are currently allocated...
 	size_t nCurCount = Blocks.Num();
 
+	// erase extra segments memory
+	// [RMS] not necessary right? std::array will deallocate?
+	//for (int i = nNumSegs; i < nCurCount; ++i)
+	//	Blocks[i] = null;
+
 	// resize to right number of segments
 	if (nNumSegs >= Blocks.Num())
 	{
