GeometryProcessing Updates:
- make FDynamicMesh3 constructor explicit because it allowed a temporary empty mesh to be constructed if FDynamicMesh* was passed to function that was expecting const FDynamicMesh&
- add FDynamicMeshEditor::DisconnectTriangles
- replace Polygon2/Polyline3 PointAt() with GetSegmentPoint() and GetSegmentPointUnitParam()
- add FPolyline3::SmoothSubdivide()
- add MeshIndexUtil, and function for converting Tri to Point indices
#rnx

diff --git a/Engine/Plugins/Experimental/GeometryProcessing/Source/DynamicMesh/Private/DynamicMeshEditor.cpp b/Engine/Plugins/Experimental/GeometryProcessing/Source/DynamicMesh/Private/DynamicMeshEditor.cpp
--- a/Engine/Plugins/Experimental/GeometryProcessing/Source/DynamicMesh/Private/DynamicMeshEditor.cpp
+++ b/Engine/Plugins/Experimental/GeometryProcessing/Source/DynamicMesh/Private/DynamicMeshEditor.cpp
@@ -4,6 +4,7 @@
 #include "DynamicMeshEditor.h"
 #include "DynamicMeshAttributeSet.h"
 #include "Util/BufferUtil.h"
+#include "MeshRegionBoundaryLoops.h"
 
 
 
@@ -148,13 +149,96 @@ void FDynamicMeshEditor::DuplicateTriangles(const TArray<int>& Triangles, FMeshI
 
 		CopyAttributes(TriangleID, NewTriangleID, IndexMaps, ResultOut);
 
-		Mesh->CheckValidity(true);
+		//Mesh->CheckValidity(true);
 	}
 
 }
 
 
 
+
+bool FDynamicMeshEditor::DisconnectTriangles(const TArray<int>& Triangles, TArray<FLoopPairSet>& LoopSetOut)
+{
+	check(Mesh->HasAttributes() == false);  // not supported yet
+
+	// find the region boundary loops
+	FMeshRegionBoundaryLoops RegionLoops(Mesh, Triangles, false);
+	bool bOK = RegionLoops.Compute();
+	check(bOK);
+	if (!bOK)
+	{
+		return false;
+	}
+	TArray<FEdgeLoop>& Loops = RegionLoops.Loops;
+
+	// need to test Contains() many times
+	TSet<int> TriangleSet;
+	TriangleSet.Reserve(Triangles.Num() * 3);
+	for (int TriID : Triangles)
+	{
+		TriangleSet.Add(TriID);
+	}
+
+	// process each loop island
+	int NumLoops = Loops.Num();
+	LoopSetOut.SetNum(NumLoops);
+	for ( int li = 0; li < NumLoops; ++li)
+	{
+		FEdgeLoop& Loop = Loops[li];
+		FLoopPairSet& LoopPair = LoopSetOut[li];
+		LoopPair.LoopA = Loop;
+
+		// duplicate the vertices
+		int NumVertices = Loop.Vertices.Num();
+		TMap<int, int> LoopVertexMap; LoopVertexMap.Reserve(NumVertices);
+		TArray<int> NewVertexLoop; NewVertexLoop.SetNum(NumVertices);
+		for (int vi = 0; vi < NumVertices; ++vi)
+		{
+			int VertID = Loop.Vertices[vi];
+			int NewVertID = Mesh->AppendVertex(*Mesh, VertID);
+			LoopVertexMap.Add(Loop.Vertices[vi], NewVertID);
+			NewVertexLoop[vi] = NewVertID;
+		}
+
+		// for each border triangle, rewrite vertices
+		int NumEdges = Loop.Edges.Num();
+		for (int ei = 0; ei < NumEdges; ++ei)
+		{
+			int EdgeID = Loop.Edges[ei];
+			FIndex2i EdgeTris = Mesh->GetEdgeT(EdgeID);
+			int EditTID = TriangleSet.Contains(EdgeTris.A) ? EdgeTris.A : EdgeTris.B;
+			if (EditTID == FDynamicMesh3::InvalidID)
+			{
+				continue;		// happens on final edge, and on input boundary edges
+			}
+
+			FIndex3i OldTri = Mesh->GetTriangle(EditTID);
+			FIndex3i NewTri = OldTri;
+			int Modified = 0;
+			for (int j = 0; j < 3; ++j)
+			{
+				const int* NewVertID = LoopVertexMap.Find(OldTri[j]);
+				if (NewVertID != nullptr)
+				{
+					NewTri[j] = *NewVertID;
+					++Modified;
+				}
+			}
+			if (Modified > 0)
+			{
+				Mesh->SetTriangle(EditTID, NewTri, false);
+			}
+		}
+
+		LoopPair.LoopB.InitializeFromVertices(Mesh, NewVertexLoop, false);
+	}
+
+	return true;
+}
+
+
+
+
 FVector3f FDynamicMeshEditor::ComputeAndSetQuadNormal(const FIndex2i& QuadTris, bool bIsPlanar)
 {
 	FVector3f Normal(0, 0, 1);
@@ -461,7 +545,7 @@ int FDynamicMeshEditor::FindOrCreateDuplicateVertex(int VertexID, FMeshIndexMapp
 	int NewVertexID = IndexMaps.GetNewVertex(VertexID);
 	if (NewVertexID == IndexMaps.InvalidID())
 	{
-		NewVertexID = Mesh->AppendVertex(Mesh, VertexID);
+		NewVertexID = Mesh->AppendVertex(*Mesh, VertexID);
 		IndexMaps.SetVertex(VertexID, NewVertexID);
 		ResultOut.NewVertices.Add(NewVertexID);
 	}
@@ -481,4 +565,50 @@ int FDynamicMeshEditor::FindOrCreateDuplicateGroup(int TriangleID, FMeshIndexMap
 		ResultOut.NewGroups.Add(NewGroupID);
 	}
 	return NewGroupID;
+}
+
+
+
+
+void FDynamicMeshEditor::AppendMesh(const FDynamicMesh3* AppendMesh,
+	FMeshIndexMappings& IndexMapsOut, FDynamicMeshEditResult& ResultOut,
+	TFunction<FVector3d(int, const FVector3d&)> PositionTransform,
+	TFunction<FVector3f(int, const FVector3f&)> NormalTransform)
+{
+	IndexMapsOut.Reset();
+
+	FIndexMapi& VertexMap = IndexMapsOut.GetVertexMap();
+	VertexMap.Reserve(AppendMesh->VertexCount());
+	for (int VertID : AppendMesh->VertexIndicesItr())
+	{
+		FVector3d Position = AppendMesh->GetVertex(VertID);
+		if (PositionTransform != nullptr)
+		{
+			Position = PositionTransform(VertID, Position);
+		}
+		int NewVertID = Mesh->AppendVertex(Position);
+		VertexMap.Add(VertID, NewVertID);
+
+		if (AppendMesh->HasVertexNormals() && Mesh->HasVertexNormals())
+		{
+			FVector3f Normal = AppendMesh->GetVertexNormal(VertID);
+			if (NormalTransform != nullptr)
+			{
+				Normal = NormalTransform(VertID, Normal);
+			}
+			Mesh->SetVertexNormal(NewVertID, Normal);
+		}
+
+		if (AppendMesh->HasVertexColors() && Mesh->HasVertexColors()) 
+		{
+			FVector3f Color = AppendMesh->GetVertexColor(VertID);
+			Mesh->SetVertexColor(NewVertID, Color);
+		}
+	}
+
+	for (int TriID : AppendMesh->TriangleIndicesItr())
+	{
+		FIndex3i Tri = AppendMesh->GetTriangle(TriID);
+		int NewTriID = Mesh->AppendTriangle(VertexMap.GetTo(Tri.A), VertexMap.GetTo(Tri.B), VertexMap.GetTo(Tri.C));
+	}
 }
\ No newline at end of file
diff --git a/Engine/Plugins/Experimental/GeometryProcessing/Source/DynamicMesh/Private/MeshIndexUtil.cpp b/Engine/Plugins/Experimental/GeometryProcessing/Source/DynamicMesh/Private/MeshIndexUtil.cpp
new file mode 100644
--- /dev/null
+++ b/Engine/Plugins/Experimental/GeometryProcessing/Source/DynamicMesh/Private/MeshIndexUtil.cpp
@@ -0,0 +1,48 @@
+// Copyright 1998-2019 Epic Games, Inc. All Rights Reserved.
+
+
+#include "MeshIndexUtil.h"
+
+
+void MeshIndexUtil::ConverTriangleToVertexIDs(const FDynamicMesh3* Mesh, const TArray<int>& TriangleIDs, TArray<int>& VertexIDsOut)
+{
+	// if we are getting close to full mesh it is probably more efficient to use a bitmap...
+
+	int NumTris = TriangleIDs.Num();
+
+	// @todo profile this constant
+	if (NumTris < 25)
+	{
+		for (int k = 0; k < NumTris; ++k)
+		{
+			if (Mesh->IsTriangle(TriangleIDs[k]))
+			{
+				FIndex3i Tri = Mesh->GetTriangle(TriangleIDs[k]);
+				VertexIDsOut.AddUnique(Tri[0]);
+				VertexIDsOut.AddUnique(Tri[1]);
+				VertexIDsOut.AddUnique(Tri[2]);
+			}
+		}
+	}
+	else
+	{
+		TSet<int> VertexSet;
+		VertexSet.Reserve(TriangleIDs.Num()*3);
+		for (int k = 0; k < NumTris; ++k)
+		{
+			if (Mesh->IsTriangle(TriangleIDs[k]))
+			{
+				FIndex3i Tri = Mesh->GetTriangle(TriangleIDs[k]);
+				VertexSet.Add(Tri[0]);
+				VertexSet.Add(Tri[1]);
+				VertexSet.Add(Tri[2]);
+			}
+		}
+
+		VertexIDsOut.Reserve(VertexSet.Num());
+		for (int VertexID : VertexSet)
+		{
+			VertexIDsOut.Add(VertexID);
+		}
+	}
+}
\ No newline at end of file
diff --git a/Engine/Plugins/Experimental/GeometryProcessing/Source/DynamicMesh/Private/Operations/ExtrudeMesh.cpp b/Engine/Plugins/Experimental/GeometryProcessing/Source/DynamicMesh/Private/Operations/ExtrudeMesh.cpp
--- a/Engine/Plugins/Experimental/GeometryProcessing/Source/DynamicMesh/Private/Operations/ExtrudeMesh.cpp
+++ b/Engine/Plugins/Experimental/GeometryProcessing/Source/DynamicMesh/Private/Operations/ExtrudeMesh.cpp
@@ -99,23 +99,26 @@ bool FExtrudeMesh::Apply()
 		StitchPolygonIDs[LoopIndex] = StitchResult.NewGroups;
 
 		// for each polygon we created in stitch, set UVs and normals
-		int NumNewQuads = StitchResult.NewQuads.Num();
-		for (int k = 0; k < NumNewQuads; k++)
+		if (Mesh->HasAttributes())
 		{
-			FVector3f Normal = Editor.ComputeAndSetQuadNormal(StitchResult.NewQuads[k], true);
-
-			// @todo is there a simpler way to construct rotation from 3 known axes (third id Normal.Cross(UnitY))
-			//  (converting from matrix might end up being more efficient due to trig in ConstrainedAlignAxis?)
-			FFrame3f ProjectFrame(FVector3f::Zero(), Normal);
-			if (FMathd::Abs(ProjectFrame.Y().Dot(FVector3f::UnitY())) < 0.01)
-			{
-				ProjectFrame.ConstrainedAlignAxis(0, FVector3f::UnitX(), ProjectFrame.Z());
-			}
-			else
+			int NumNewQuads = StitchResult.NewQuads.Num();
+			for (int k = 0; k < NumNewQuads; k++)
 			{
-				ProjectFrame.ConstrainedAlignAxis(1, FVector3f::UnitY(), ProjectFrame.Z());
+				FVector3f Normal = Editor.ComputeAndSetQuadNormal(StitchResult.NewQuads[k], true);
+
+				// @todo is there a simpler way to construct rotation from 3 known axes (third id Normal.Cross(UnitY))
+				//  (converting from matrix might end up being more efficient due to trig in ConstrainedAlignAxis?)
+				FFrame3f ProjectFrame(FVector3f::Zero(), Normal);
+				if (FMathd::Abs(ProjectFrame.Y().Dot(FVector3f::UnitY())) < 0.01)
+				{
+					ProjectFrame.ConstrainedAlignAxis(0, FVector3f::UnitX(), ProjectFrame.Z());
+				}
+				else
+				{
+					ProjectFrame.ConstrainedAlignAxis(1, FVector3f::UnitY(), ProjectFrame.Z());
+				}
+				Editor.SetQuadUVsFromProjection(StitchResult.NewQuads[k], ProjectFrame, UVScaleFactor);
 			}
-			Editor.SetQuadUVsFromProjection(StitchResult.NewQuads[k], ProjectFrame, UVScaleFactor);
 		}
 
 		NewLoops[LoopIndex].InitializeFromVertices(Mesh, OffsetLoop);
diff --git a/Engine/Plugins/Experimental/GeometryProcessing/Source/DynamicMesh/Public/DynamicMesh3.h b/Engine/Plugins/Experimental/GeometryProcessing/Source/DynamicMesh/Public/DynamicMesh3.h
--- a/Engine/Plugins/Experimental/GeometryProcessing/Source/DynamicMesh/Public/DynamicMesh3.h
+++ b/Engine/Plugins/Experimental/GeometryProcessing/Source/DynamicMesh/Public/DynamicMesh3.h
@@ -209,7 +209,7 @@ protected:
 public:
 	virtual ~FDynamicMesh3();
 
-	FDynamicMesh3(bool bWantNormals = true, bool bWantColors = false, bool bWantUVs = false, bool bWantTriGroups = false);
+	explicit FDynamicMesh3(bool bWantNormals = true, bool bWantColors = false, bool bWantUVs = false, bool bWantTriGroups = false);
 
 	FDynamicMesh3(EMeshComponents flags)
 		: FDynamicMesh3(((int)flags & (int)EMeshComponents::VertexNormals) != 0, ((int)flags & (int)EMeshComponents::VertexColors) != 0,
diff --git a/Engine/Plugins/Experimental/GeometryProcessing/Source/DynamicMesh/Public/DynamicMeshEditor.h b/Engine/Plugins/Experimental/GeometryProcessing/Source/DynamicMesh/Public/DynamicMeshEditor.h
--- a/Engine/Plugins/Experimental/GeometryProcessing/Source/DynamicMesh/Public/DynamicMeshEditor.h
+++ b/Engine/Plugins/Experimental/GeometryProcessing/Source/DynamicMesh/Public/DynamicMeshEditor.h
@@ -5,6 +5,7 @@
 #pragma once
 
 #include "DynamicMesh3.h"
+#include "EdgeLoop.h"
 
 
 
@@ -31,6 +32,21 @@ public:
 	/** @return the value used to indicate "invalid" in the mapping */
 	int InvalidID() { return VertexMap.GetInvalidID(); }
 
+	void Reset()
+	{
+		VertexMap.Reset();
+		TriangleMap.Reset();
+		GroupMap.Reset();
+		for (FIndexMapi& UVMap : UVMaps)
+		{
+			UVMap.Reset();
+		}
+		for (FIndexMapi& NormalMap : NormalMaps)
+		{
+			NormalMap.Reset();
+		}
+	}
+
 	FIndexMapi& GetVertexMap() { return VertexMap; }
 	inline void SetVertex(int FromID, int ToID) { VertexMap.Add(FromID, ToID); }
 	inline int GetNewVertex(int FromID) const { return VertexMap.GetTo(FromID); }
@@ -136,6 +152,25 @@ public:
 	void DuplicateTriangles(const TArray<int>& Triangles, FMeshIndexMappings& IndexMaps, FDynamicMeshEditResult& ResultOut);
 
 
+	/**
+	 * Pair of associated edge loops.
+	 */
+	struct FLoopPairSet
+	{
+		FEdgeLoop LoopA;
+		FEdgeLoop LoopB;
+	};
+
+	/**
+	 * Finds boundary loops of connected components of a set of triangles, and duplicates the vertices
+	 * along the boundary, such that the triangles become disconnected.
+	 * @param Triangles set of triangles
+	 * @param LoopSetOut set of boundary loops. LoopA is original loop which remains with "outer" triangles, and LoopB is new boundary loop of triangle set
+	 * @return true on success
+	 */
+	bool DisconnectTriangles(const TArray<int>& Triangles, TArray<FLoopPairSet>& LoopSetOut);
+
+
 	/**
 	 * Remove a list of triangles from the mesh, and optionally any vertices that are now orphaned
 	 * @param Triangles the triangles to remove
@@ -259,6 +294,11 @@ public:
 
 
 
+
+	void AppendMesh(const FDynamicMesh3* AppendMesh, FMeshIndexMappings& IndexMapsOut, FDynamicMeshEditResult& ResultOut,
+		TFunction<FVector3d(int, const FVector3d&)> PositionTransform = nullptr,
+		TFunction<FVector3f(int, const FVector3f&)> NormalTransform = nullptr);
+
 };
 
 
diff --git a/Engine/Plugins/Experimental/GeometryProcessing/Source/DynamicMesh/Public/MeshIndexUtil.h b/Engine/Plugins/Experimental/GeometryProcessing/Source/DynamicMesh/Public/MeshIndexUtil.h
new file mode 100644
--- /dev/null
+++ b/Engine/Plugins/Experimental/GeometryProcessing/Source/DynamicMesh/Public/MeshIndexUtil.h
@@ -0,0 +1,20 @@
+// Copyright 1998-2019 Epic Games, Inc. All Rights Reserved.
+
+#pragma once
+
+#include "DynamicMesh3.h"
+
+/**
+ * Utility functions for dealing with mesh indices
+ */
+namespace MeshIndexUtil
+{
+	/**
+	 * @param Mesh input mesh
+	 * @param TriangleIDs list of triangle IDs of Mesh
+	 * @param VertexIDsOut list of vertices contained by triangles
+	 */
+	DYNAMICMESH_API void ConverTriangleToVertexIDs(const FDynamicMesh3* Mesh, const TArray<int>& TriangleIDs, TArray<int>& VertexIDsOut);
+
+
+}
diff --git a/Engine/Plugins/Experimental/GeometryProcessing/Source/GeometricObjects/Private/Spatial/GeometrySet3.cpp b/Engine/Plugins/Experimental/GeometryProcessing/Source/GeometricObjects/Private/Spatial/GeometrySet3.cpp
--- a/Engine/Plugins/Experimental/GeometryProcessing/Source/GeometricObjects/Private/Spatial/GeometrySet3.cpp
+++ b/Engine/Plugins/Experimental/GeometryProcessing/Source/GeometricObjects/Private/Spatial/GeometrySet3.cpp
@@ -128,7 +128,7 @@ bool FGeometrySet3::FindNearestCurveToRay(const FRay3d& Ray, FNearest& ResultOut
 			if (SegRayParam < CurveMinRayParamT)
 			{
 				FVector3d RayPosition = Ray.PointAt(SegRayParam);
-				FVector3d CurvePosition = Polyline.PointAt(si, SegSegParam);
+				FVector3d CurvePosition = Polyline.GetSegmentPoint(si, SegSegParam);
 				if (PointWithinToleranceTest(RayPosition, CurvePosition))
 				{
 					CurveMinRayParamT = SegRayParam;
@@ -159,7 +159,7 @@ bool FGeometrySet3::FindNearestCurveToRay(const FRay3d& Ray, FNearest& ResultOut
 		ResultOut.ID = NearestID;
 		ResultOut.bIsPoint = false;
 		ResultOut.NearestRayPoint = Ray.PointAt(MinRayParamT);
-		ResultOut.NearestGeoPoint = Curves[NearestIndex].Geometry.PointAt(NearestSegmentIdx, NearestSegmentParam);
+		ResultOut.NearestGeoPoint = Curves[NearestIndex].Geometry.GetSegmentPoint(NearestSegmentIdx, NearestSegmentParam);
 		ResultOut.RayParam = MinRayParamT;
 		ResultOut.PolySegmentIdx = NearestSegmentIdx;
 		ResultOut.PolySegmentParam = NearestSegmentParam;
diff --git a/Engine/Plugins/Experimental/GeometryProcessing/Source/GeometricObjects/Public/GeometryTypes.h b/Engine/Plugins/Experimental/GeometryProcessing/Source/GeometricObjects/Public/GeometryTypes.h
--- a/Engine/Plugins/Experimental/GeometryProcessing/Source/GeometricObjects/Public/GeometryTypes.h
+++ b/Engine/Plugins/Experimental/GeometryProcessing/Source/GeometricObjects/Public/GeometryTypes.h
@@ -99,6 +99,12 @@ public:
 		InvalidID = (IntType)-9999999;
 	}
 
+	void Reset()
+	{
+		ForwardMap.Reset();
+		ReverseMap.Reset();
+	}
+
 	/** @return the value used to indicate "invalid" in the mapping */
 	inline IntType GetInvalidID() const { return InvalidID; }
 
@@ -145,6 +151,18 @@ public:
 		const IntType* FoundVal = ReverseMap.Find(ToID);
 		return (FoundVal == nullptr) ? InvalidID : *FoundVal;
 	}
+
+	void Reserve(int NumElements)
+	{
+		if (bWantForward)
+		{
+			ForwardMap.Reserve(NumElements);
+		}
+		if (bWantReverse)
+		{
+			ReverseMap.Reserve(NumElements);
+		}
+	}
 };
 typedef TIndexMap<int> FIndexMapi;
 
diff --git a/Engine/Plugins/Experimental/GeometryProcessing/Source/GeometricObjects/Public/Polygon2.h b/Engine/Plugins/Experimental/GeometryProcessing/Source/GeometricObjects/Public/Polygon2.h
--- a/Engine/Plugins/Experimental/GeometryProcessing/Source/GeometricObjects/Public/Polygon2.h
+++ b/Engine/Plugins/Experimental/GeometryProcessing/Source/GeometricObjects/Public/Polygon2.h
@@ -568,15 +568,26 @@ public:
 
 	/**
 	 * @param SegmentIndex index of first vertex of the edge
-	 * @param SegmentParam parameter in range [0,1] along segment
+	 * @param SegmentParam parameter in range [-Extent,Extent] along segment
 	 * @return point on the segment at the given parameter value
 	 */
-	FVector2<T> PointAt(int SegmentIndex, T SegmentParam) const
+	FVector2<T> GetSegmentPoint(int SegmentIndex, T SegmentParam) const
 	{
 		TSegment2<T> seg(Vertices[SegmentIndex], Vertices[(SegmentIndex + 1) % Vertices.Num()]);
 		return seg.PointAt(SegmentParam);
 	}
 
+	/**
+	 * @param SegmentIndex index of first vertex of the edge
+	 * @param SegmentParam parameter in range [0,1] along segment
+	 * @return point on the segment at the given parameter value
+	 */
+	FVector2<T> GetSegmentPointUnitParam(int SegmentIndex, T SegmentParam) const
+	{
+		TSegment2<T> seg(Vertices[SegmentIndex], Vertices[(SegmentIndex + 1) % Vertices.Num()]);
+		return seg.PointBetween(SegmentParam);
+	}
+
 
 	/**
 	 * @param SegmentIndex index of first vertex of the edge
diff --git a/Engine/Plugins/Experimental/GeometryProcessing/Source/GeometricObjects/Public/Polyline3.h b/Engine/Plugins/Experimental/GeometryProcessing/Source/GeometricObjects/Public/Polyline3.h
--- a/Engine/Plugins/Experimental/GeometryProcessing/Source/GeometricObjects/Public/Polyline3.h
+++ b/Engine/Plugins/Experimental/GeometryProcessing/Source/GeometricObjects/Public/Polyline3.h
@@ -6,6 +6,7 @@
 #include "SegmentTypes.h"
 #include "LineTypes.h"
 #include "RayTypes.h"
+#include "BoxTypes.h"
 
 /**
  * TPolyline3 represents a 3D polyline stored as a list of Vertices.
@@ -218,16 +219,27 @@ public:
 
 	/**
 	 * @param SegmentIndex index of first vertex of the edge
-	 * @param SegmentParam parameter in range [0,1] along segment
+	 * @param SegmentParam parameter in range [-Extent,Extent] along segment
 	 * @return point on the segment at the given parameter value
 	 */
-	FVector3<T> PointAt(int SegmentIndex, T SegmentParam) const
+	FVector3<T> GetSegmentPoint(int SegmentIndex, T SegmentParam) const
 	{
 		TSegment3<T> seg(Vertices[SegmentIndex], Vertices[SegmentIndex + 1]);
 		return seg.PointAt(SegmentParam);
 	}
 
 
+	/**
+	 * @param SegmentIndex index of first vertex of the edge
+	 * @param SegmentParam parameter in range [0,1] along segment
+	 * @return point on the segment at the given parameter value
+	 */
+	FVector3<T> GetSegmentPointUnitParam(int SegmentIndex, T SegmentParam) const
+	{
+		TSegment3<T> seg(Vertices[SegmentIndex], Vertices[SegmentIndex + 1]);
+		return seg.PointBetween(SegmentParam);
+	}
+
 
 
 	/**
@@ -399,6 +411,32 @@ public:
 	}
 
 
+
+	/**
+	 * Produce a new polyline that is smoother than this one
+	 */
+	void SmoothSubdivide(TPolyline3<T>& NewPolyline) const
+	{
+		const T Alpha = (T)1 / (T)3;
+		const T OneMinusAlpha = (T)2 / (T)3;
+
+		int N = Vertices.Num();
+		NewPolyline.Vertices.SetNum(2*(N-1));
+		NewPolyline.Vertices[0] = Vertices[0];
+		N = N - 1;
+		int k = 1;
+		for (int i = 1; i < N; ++i)
+		{
+			const FVector3<T>& Prev = Vertices[i-1];
+			const FVector3<T>& Cur = Vertices[i];
+			const FVector3<T>& Next = Vertices[i+1];
+			NewPolyline.Vertices[k++] = Alpha * Prev + OneMinusAlpha * Cur;
+			NewPolyline.Vertices[k++] = OneMinusAlpha * Cur + Alpha * Next;
+		}
+		NewPolyline.Vertices[k] = Vertices[N];
+	}
+
+
 };
 
 typedef TPolyline3<double> FPolyline3d;
