add copyright boilerplate

diff --git a/Engine/Plugins/Experimental/GeometryProcessing/Source/GeometricObjects/Private/Tests/VectorTypes.spec.cpp b/Engine/Plugins/Experimental/GeometryProcessing/Source/GeometricObjects/Private/Tests/VectorTypes.spec.cpp
--- a/Engine/Plugins/Experimental/GeometryProcessing/Source/GeometricObjects/Private/Tests/VectorTypes.spec.cpp
+++ b/Engine/Plugins/Experimental/GeometryProcessing/Source/GeometricObjects/Private/Tests/VectorTypes.spec.cpp
@@ -1,3 +1,5 @@
+// Copyright 1998-2019 Epic Games, Inc. All Rights Reserved.
+
 #include "Misc/AutomationTest.h"
 #include "VectorTypes.h"
 
