GeometryProcessing updates. Add MeshIndexUtil::VertexToTriangleOneRing, Support attributes in FDynamicMeshEditor::AppendMesh, add FDynamicMeshEditor::AppendNormals/AppendUVs, FDynamicMeshAttributeSet::GetNumLayers/UVs, TIndexMap find functions. Add DynamicMeshOverlay::InsertElement and unsafe-rebuild option
#rnx

diff --git a/Engine/Plugins/Experimental/GeometryProcessing/Source/DynamicMesh/Private/DynamicMeshEditor.cpp b/Engine/Plugins/Experimental/GeometryProcessing/Source/DynamicMesh/Private/DynamicMeshEditor.cpp
--- a/Engine/Plugins/Experimental/GeometryProcessing/Source/DynamicMesh/Private/DynamicMeshEditor.cpp
+++ b/Engine/Plugins/Experimental/GeometryProcessing/Source/DynamicMesh/Private/DynamicMeshEditor.cpp
@@ -571,11 +571,12 @@ int FDynamicMeshEditor::FindOrCreateDuplicateGroup(int TriangleID, FMeshIndexMap
 
 
 void FDynamicMeshEditor::AppendMesh(const FDynamicMesh3* AppendMesh,
-	FMeshIndexMappings& IndexMapsOut, FDynamicMeshEditResult& ResultOut,
+	FMeshIndexMappings& IndexMapsOut, 
 	TFunction<FVector3d(int, const FVector3d&)> PositionTransform,
-	TFunction<FVector3f(int, const FVector3f&)> NormalTransform)
+	TFunction<FVector3d(int, const FVector3d&)> NormalTransform)
 {
 	IndexMapsOut.Reset();
+	IndexMapsOut.Initialize(Mesh);
 
 	FIndexMapi& VertexMap = IndexMapsOut.GetVertexMap();
 	VertexMap.Reserve(AppendMesh->VertexCount());
@@ -594,21 +595,135 @@ void FDynamicMeshEditor::AppendMesh(const FDynamicMesh3* AppendMesh,
 			FVector3f Normal = AppendMesh->GetVertexNormal(VertID);
 			if (NormalTransform != nullptr)
 			{
-				Normal = NormalTransform(VertID, Normal);
+				Normal = (FVector3f)NormalTransform(VertID, (FVector3d)Normal);
 			}
 			Mesh->SetVertexNormal(NewVertID, Normal);
 		}
 
-		if (AppendMesh->HasVertexColors() && Mesh->HasVertexColors()) 
+		if (AppendMesh->HasVertexColors() && Mesh->HasVertexColors())
 		{
 			FVector3f Color = AppendMesh->GetVertexColor(VertID);
 			Mesh->SetVertexColor(NewVertID, Color);
 		}
 	}
 
+	FIndexMapi& TriangleMap = IndexMapsOut.GetTriangleMap();
+	bool bAppendGroups = AppendMesh->HasTriangleGroups() && Mesh->HasTriangleGroups();
+	FIndexMapi& GroupsMap = IndexMapsOut.GetGroupMap();
 	for (int TriID : AppendMesh->TriangleIndicesItr())
 	{
+		// append trigroup
+		int GroupID = FDynamicMesh3::InvalidID;
+		if (bAppendGroups)
+		{
+			GroupID = AppendMesh->GetTriangleGroup(TriID);
+			if (GroupID != FDynamicMesh3::InvalidID)
+			{
+				const int* FoundNewGroupID = GroupsMap.FindTo(GroupID);
+				if (FoundNewGroupID == nullptr)
+				{
+					int NewGroupID = Mesh->AllocateTriangleGroup();
+					GroupsMap.Add(GroupID, NewGroupID);
+					GroupID = NewGroupID;
+				}
+				else
+				{
+					GroupID = *FoundNewGroupID;
+				}
+			}
+		}
+
 		FIndex3i Tri = AppendMesh->GetTriangle(TriID);
-		int NewTriID = Mesh->AppendTriangle(VertexMap.GetTo(Tri.A), VertexMap.GetTo(Tri.B), VertexMap.GetTo(Tri.C));
+		int NewTriID = Mesh->AppendTriangle(VertexMap.GetTo(Tri.A), VertexMap.GetTo(Tri.B), VertexMap.GetTo(Tri.C), GroupID);
+		TriangleMap.Add(TriID, NewTriID);
 	}
-}
\ No newline at end of file
+
+
+	// @todo support multiple UV/normal layer copying
+	// @todo can we have a template fn that does this?
+
+	if (AppendMesh->HasAttributes() && Mesh->HasAttributes())
+	{
+		const FDynamicMeshNormalOverlay* FromNormals = AppendMesh->Attributes()->PrimaryNormals();
+		FDynamicMeshNormalOverlay* ToNormals = Mesh->Attributes()->PrimaryNormals();
+		if (FromNormals != nullptr && ToNormals != nullptr)
+		{
+			FIndexMapi& NormalMap = IndexMapsOut.GetNormalMap(0);
+			NormalMap.Reserve(FromNormals->ElementCount());
+			AppendNormals(AppendMesh, FromNormals, ToNormals,
+				VertexMap, TriangleMap, NormalTransform, NormalMap);
+		}
+
+		const FDynamicMeshUVOverlay* FromUVs = AppendMesh->Attributes()->PrimaryUV();
+		FDynamicMeshUVOverlay* ToUVs = Mesh->Attributes()->PrimaryUV();
+		if (FromUVs != nullptr && ToUVs != nullptr)
+		{
+			FIndexMapi& UVMap = IndexMapsOut.GetUVMap(0);
+			UVMap.Reserve(FromUVs->ElementCount());
+			AppendUVs(AppendMesh, FromUVs, ToUVs,
+				VertexMap, TriangleMap, UVMap);
+		}
+	}
+}
+
+
+
+void FDynamicMeshEditor::AppendNormals(const FDynamicMesh3* AppendMesh, 
+	const FDynamicMeshNormalOverlay* FromNormals, FDynamicMeshNormalOverlay* ToNormals,
+	const FIndexMapi& VertexMap, const FIndexMapi& TriangleMap,
+	TFunction<FVector3d(int, const FVector3d&)> NormalTransform,
+	FIndexMapi& NormalMapOut)
+{
+	// copy over normals
+	for (int ElemID : FromNormals->ElementIndicesItr())
+	{
+		int ParentVertID = FromNormals->GetParentVertex(ElemID);
+		FVector3f Normal = FromNormals->GetElement(ElemID);
+		if (NormalTransform != nullptr)
+		{
+			Normal = (FVector3f)NormalTransform(ParentVertID, (FVector3d)Normal);
+		}
+		int NewElemID = ToNormals->AppendElement(Normal, VertexMap.GetTo(ParentVertID));
+		NormalMapOut.Add(ElemID, NewElemID);
+	}
+
+	// now set new triangles
+	for (int TriID : AppendMesh->TriangleIndicesItr())
+	{
+		FIndex3i ElemTri = FromNormals->GetTriangle(TriID);
+		int NewTriID = TriangleMap.GetTo(TriID);
+		for (int j = 0; j < 3; ++j)
+		{
+			ElemTri[j] = FromNormals->IsElement(ElemTri[j]) ? NormalMapOut.GetTo(ElemTri[j]) : FDynamicMesh3::InvalidID;
+		}
+		ToNormals->SetTriangle(NewTriID, ElemTri);
+	}
+}
+
+
+void FDynamicMeshEditor::AppendUVs(const FDynamicMesh3* AppendMesh,
+	const FDynamicMeshUVOverlay* FromUVs, FDynamicMeshUVOverlay* ToUVs,
+	const FIndexMapi& VertexMap, const FIndexMapi& TriangleMap,
+	FIndexMapi& UVMapOut)
+{
+	// copy over uv elements
+	for (int ElemID : FromUVs->ElementIndicesItr())
+	{
+		int ParentVertID = FromUVs->GetParentVertex(ElemID);
+		FVector2f UV = FromUVs->GetElement(ElemID);
+		int NewElemID = ToUVs->AppendElement(UV, VertexMap.GetTo(ParentVertID));
+		UVMapOut.Add(ElemID, NewElemID);
+	}
+
+	// now set new triangles
+	for (int TriID : AppendMesh->TriangleIndicesItr())
+	{
+		FIndex3i ElemTri = FromUVs->GetTriangle(TriID);
+		int NewTriID = TriangleMap.GetTo(TriID);
+		for (int j = 0; j < 3; ++j)
+		{
+			ElemTri[j] = FromUVs->IsElement(ElemTri[j]) ? UVMapOut.GetTo(ElemTri[j]) : FDynamicMesh3::InvalidID;
+		}
+		ToUVs->SetTriangle(NewTriID, ElemTri);
+	}
+}
diff --git a/Engine/Plugins/Experimental/GeometryProcessing/Source/DynamicMesh/Private/DynamicMeshOverlay.cpp b/Engine/Plugins/Experimental/GeometryProcessing/Source/DynamicMesh/Private/DynamicMeshOverlay.cpp
--- a/Engine/Plugins/Experimental/GeometryProcessing/Source/DynamicMesh/Private/DynamicMeshOverlay.cpp
+++ b/Engine/Plugins/Experimental/GeometryProcessing/Source/DynamicMesh/Private/DynamicMeshOverlay.cpp
@@ -18,7 +18,7 @@ void TDynamicMeshOverlay<RealType, ElementSize>::ClearElements()
 
 
 template<typename RealType, int ElementSize>
-int TDynamicMeshOverlay<RealType, ElementSize>::AppendElement(RealType ConstantValue, int SourceVertex)
+int TDynamicMeshOverlay<RealType, ElementSize>::AppendElement(RealType ConstantValue, int ParentVertex)
 {
 	int vid = ElementsRefCounts.Allocate();
 	int i = ElementSize * vid;
@@ -26,7 +26,7 @@ int TDynamicMeshOverlay<RealType, ElementSize>::AppendElement(RealType ConstantV
 	{
 		Elements.InsertAt(ConstantValue, i + k);
 	}
-	ParentVertices.InsertAt(SourceVertex, vid);
+	ParentVertices.InsertAt(ParentVertex, vid);
 
 	//updateTimeStamp(true);
 	return vid;
@@ -34,7 +34,7 @@ int TDynamicMeshOverlay<RealType, ElementSize>::AppendElement(RealType ConstantV
 
 
 template<typename RealType, int ElementSize>
-int TDynamicMeshOverlay<RealType, ElementSize>::AppendElement(const RealType* Value, int SourceVertex)
+int TDynamicMeshOverlay<RealType, ElementSize>::AppendElement(const RealType* Value, int ParentVertex)
 {
 	int vid = ElementsRefCounts.Allocate();
 	int i = ElementSize * vid;
@@ -45,13 +45,47 @@ int TDynamicMeshOverlay<RealType, ElementSize>::AppendElement(const RealType* Va
 		Elements.InsertAt(Value[k], i + k);
 	}
 
-	ParentVertices.InsertAt(SourceVertex, vid);
+	ParentVertices.InsertAt(ParentVertex, vid);
 
 	//updateTimeStamp(true);
 	return vid;
 }
 
 
+template<typename RealType, int ElementSize>
+EMeshResult TDynamicMeshOverlay<RealType, ElementSize>::InsertElement(int ElementID, const RealType* Value, int ParentVertex, bool bUnsafe)
+{
+	if (ElementsRefCounts.IsValid(ElementID))
+	{
+		return EMeshResult::Failed_VertexAlreadyExists;
+	}
+
+	bool bOK = (bUnsafe) ? ElementsRefCounts.AllocateAtUnsafe(ElementID) :
+		ElementsRefCounts.AllocateAt(ElementID);
+	if (bOK == false)
+	{
+		return EMeshResult::Failed_CannotAllocateVertex;
+	}
+
+	int i = ElementSize * ElementID;
+	// insert in reverse order so that Resize() is only called once
+	for (int k = ElementSize - 1; k >= 0; --k)
+	{
+		Elements.InsertAt(Value[k], i + k);
+	}
+
+	ParentVertices.InsertAt(ParentVertex, ElementID);
+
+	//UpdateTimeStamp(true, true);
+	return EMeshResult::Ok;
+}
+
+
+
+
+
+
+
 
 
 template<typename RealType, int ElementSize>
diff --git a/Engine/Plugins/Experimental/GeometryProcessing/Source/DynamicMesh/Private/MeshIndexUtil.cpp b/Engine/Plugins/Experimental/GeometryProcessing/Source/DynamicMesh/Private/MeshIndexUtil.cpp
--- a/Engine/Plugins/Experimental/GeometryProcessing/Source/DynamicMesh/Private/MeshIndexUtil.cpp
+++ b/Engine/Plugins/Experimental/GeometryProcessing/Source/DynamicMesh/Private/MeshIndexUtil.cpp
@@ -4,7 +4,7 @@
 #include "MeshIndexUtil.h"
 
 
-void MeshIndexUtil::ConverTriangleToVertexIDs(const FDynamicMesh3* Mesh, const TArray<int>& TriangleIDs, TArray<int>& VertexIDsOut)
+void MeshIndexUtil::TriangleToVertexIDs(const FDynamicMesh3* Mesh, const TArray<int>& TriangleIDs, TArray<int>& VertexIDsOut)
 {
 	// if we are getting close to full mesh it is probably more efficient to use a bitmap...
 
@@ -45,4 +45,24 @@ void MeshIndexUtil::ConverTriangleToVertexIDs(const FDynamicMesh3* Mesh, const T
 			VertexIDsOut.Add(VertexID);
 		}
 	}
+}
+
+
+
+
+
+void MeshIndexUtil::VertexToTriangleOneRing(const FDynamicMesh3* Mesh, const TArray<int>& VertexIDs, TSet<int>& TriangleIDsOut)
+{
+	int NumVerts = VertexIDs.Num();
+	TriangleIDsOut.Reserve( (NumVerts < 5) ? NumVerts*6 : NumVerts*4);
+	for (int k = 0; k < NumVerts; ++k)
+	{
+		if (Mesh->IsVertex(VertexIDs[k]))
+		{
+			for (int TriID : Mesh->VtxTrianglesItr(VertexIDs[k]))
+			{
+				TriangleIDsOut.Add(TriID);
+			}
+		}
+	}
 }
\ No newline at end of file
diff --git a/Engine/Plugins/Experimental/GeometryProcessing/Source/DynamicMesh/Public/DynamicMeshAttributeSet.h b/Engine/Plugins/Experimental/GeometryProcessing/Source/DynamicMesh/Public/DynamicMeshAttributeSet.h
--- a/Engine/Plugins/Experimental/GeometryProcessing/Source/DynamicMesh/Public/DynamicMeshAttributeSet.h
+++ b/Engine/Plugins/Experimental/GeometryProcessing/Source/DynamicMesh/Public/DynamicMeshAttributeSet.h
@@ -67,6 +67,11 @@ public:
 	// UV Layers 
 	//
 
+	/** @return The number of UV layers in this attribute set */
+	int GetNumUVLayers() const
+	{
+		return UVLayers.Num();
+	}
 
 	/** @return the UV layer at the given Index */
 	FDynamicMeshUVOverlay* GetUVLayer(int Index)
@@ -102,6 +107,12 @@ public:
 	// Normal Layers 
 	//
 
+	/** @return The number of normal layers in this attribute set */
+	int GetNumNormalLayers() const
+	{
+		return NormalLayers.Num();
+	}
+
 	/** @return the Normal layer at the given Index */
 	FDynamicMeshNormalOverlay* GetNormalLayer(int Index)
 	{
diff --git a/Engine/Plugins/Experimental/GeometryProcessing/Source/DynamicMesh/Public/DynamicMeshEditor.h b/Engine/Plugins/Experimental/GeometryProcessing/Source/DynamicMesh/Public/DynamicMeshEditor.h
--- a/Engine/Plugins/Experimental/GeometryProcessing/Source/DynamicMesh/Public/DynamicMeshEditor.h
+++ b/Engine/Plugins/Experimental/GeometryProcessing/Source/DynamicMesh/Public/DynamicMeshEditor.h
@@ -5,11 +5,9 @@
 #pragma once
 
 #include "DynamicMesh3.h"
+#include "DynamicMeshAttributeSet.h"
 #include "EdgeLoop.h"
 
-
-
-
 /**
  * FMeshIndexMappings stores a set of integer IndexMaps for a mesh
  * This is a convenient object to have, to avoid passing around large numbers of separate maps.
@@ -295,10 +293,52 @@ public:
 
 
 
-	void AppendMesh(const FDynamicMesh3* AppendMesh, FMeshIndexMappings& IndexMapsOut, FDynamicMeshEditResult& ResultOut,
+	/**
+	 * Append input mesh to our internal mesh
+	 * @param AppendMesh mesh to append
+	 * @param IndexMapsOut mesh element index mappings generated in this append operation
+	 * @param PositionTransform optional transformation function applied to mesh vertex positions
+	 * @param NormalTransform optional transformation function applied to mesh normals
+	 */
+	void AppendMesh(const FDynamicMesh3* AppendMesh, FMeshIndexMappings& IndexMapsOut, 
 		TFunction<FVector3d(int, const FVector3d&)> PositionTransform = nullptr,
-		TFunction<FVector3f(int, const FVector3f&)> NormalTransform = nullptr);
+		TFunction<FVector3d(int, const FVector3d&)> NormalTransform = nullptr);
+
 
+	/**
+	 * Append normals from one attribute overlay to another.
+	 * Assumes that AppendMesh has already been appended to Mesh.
+	 * Note that this function has no dependency on .Mesh, it could be static
+	 * @param AppendMesh mesh that owns FromNormals attribute overlay
+	 * @param FromNormals Normals overlay we want to append from (owned by AppendMesh)
+	 * @param ToNormals Normals overlay we want to append to (owned by Mesh)
+	 * @param VertexMap map from AppendMesh vertex IDs to vertex IDs applicable to ToNormals (ie of .Mesh)
+	 * @param TriangleMap map from AppendMesh triangle IDs to triangle IDs applicable to ToNormals (ie of .Mesh)
+	 * @param NormalTransform optional transformation function applied to mesh normals
+	 * @param NormalMapOut Mapping from element IDs of FromNormals to new element IDs in ToNormals
+	 */
+	void AppendNormals(const FDynamicMesh3* AppendMesh,
+		const FDynamicMeshNormalOverlay* FromNormals, FDynamicMeshNormalOverlay* ToNormals,
+		const FIndexMapi& VertexMap, const FIndexMapi& TriangleMap,
+		TFunction<FVector3d(int, const FVector3d&)> NormalTransform,
+		FIndexMapi& NormalMapOut);
+
+
+	/**
+	 * Append UVs from one attribute overlay to another.
+	 * Assumes that AppendMesh has already been appended to Mesh.
+	 * Note that this function has no dependency on .Mesh, it could be static
+	 * @param AppendMesh mesh that owns FromUVs attribute overlay
+	 * @param FromUVs UV overlay we want to append from (owned by AppendMesh)
+	 * @param ToUVs UV overlay we want to append to (owned by Mesh)
+	 * @param VertexMap map from AppendMesh vertex IDs to vertex IDs applicable to ToUVs (ie of .Mesh)
+	 * @param TriangleMap map from AppendMesh triangle IDs to triangle IDs applicable to ToUVs (ie of .Mesh)
+	 * @param UVMapOut Mapping from element IDs of FromUVs to new element IDs in ToUVs
+	 */
+	void AppendUVs(const FDynamicMesh3* AppendMesh,
+		const FDynamicMeshUVOverlay* FromUVs, FDynamicMeshUVOverlay* ToUVs,
+		const FIndexMapi& VertexMap, const FIndexMapi& TriangleMap,
+		FIndexMapi& UVMapOut);
 };
 
 
diff --git a/Engine/Plugins/Experimental/GeometryProcessing/Source/DynamicMesh/Public/DynamicMeshOverlay.h b/Engine/Plugins/Experimental/GeometryProcessing/Source/DynamicMesh/Public/DynamicMeshOverlay.h
--- a/Engine/Plugins/Experimental/GeometryProcessing/Source/DynamicMesh/Public/DynamicMeshOverlay.h
+++ b/Engine/Plugins/Experimental/GeometryProcessing/Source/DynamicMesh/Public/DynamicMeshOverlay.h
@@ -102,6 +102,41 @@ public:
 	/** Set the triangle to the given Element index tuple, and increment element reference counts */
 	EMeshResult SetTriangle(int TriangleID, const FIndex3i& TriElements);
 
+
+
+	//
+	// Support for inserting element at specific ID. This is a bit tricky
+	// because we likely will need to update the free list in the RefCountVector, which
+	// can be expensive. If you are going to do many inserts (eg inside a loop), wrap in
+	// BeginUnsafe / EndUnsafe calls, and pass bUnsafe = true to the InsertElement() calls,
+	// to the defer free list rebuild until you are done.
+	//
+
+	/** Call this before a set of unsafe InsertVertex() calls */
+	void BeginUnsafeElementsInsert()
+	{
+		// do nothing...
+	}
+
+	/** Call after a set of unsafe InsertVertex() calls to rebuild free list */
+	void EndUnsafeElementsInsert()
+	{
+		ElementsRefCounts.RebuildFreeList();
+	}
+
+	/**
+	 * Insert element at given index, assuming it is unused.
+	 * If bUnsafe, we use fast id allocation that does not update free list.
+	 * You should only be using this between BeginUnsafeElementsInsert() / EndUnsafeElementsInsert() calls
+	 */
+	EMeshResult InsertElement(int ElementID, const RealType* Value, int ParentVertex, bool bUnsafe = false);
+
+
+	//
+	// Accessors/Queries
+	//  
+
+
 	/** Get the element at a given index */
 	inline void GetElement(int ElementID, RealType* Data) const
 	{
diff --git a/Engine/Plugins/Experimental/GeometryProcessing/Source/DynamicMesh/Public/MeshIndexUtil.h b/Engine/Plugins/Experimental/GeometryProcessing/Source/DynamicMesh/Public/MeshIndexUtil.h
--- a/Engine/Plugins/Experimental/GeometryProcessing/Source/DynamicMesh/Public/MeshIndexUtil.h
+++ b/Engine/Plugins/Experimental/GeometryProcessing/Source/DynamicMesh/Public/MeshIndexUtil.h
@@ -10,11 +10,20 @@
 namespace MeshIndexUtil
 {
 	/**
+	 * Find list of unique vertices that are contained in one or more triangles
 	 * @param Mesh input mesh
 	 * @param TriangleIDs list of triangle IDs of Mesh
 	 * @param VertexIDsOut list of vertices contained by triangles
 	 */
-	DYNAMICMESH_API void ConverTriangleToVertexIDs(const FDynamicMesh3* Mesh, const TArray<int>& TriangleIDs, TArray<int>& VertexIDsOut);
+	DYNAMICMESH_API void TriangleToVertexIDs(const FDynamicMesh3* Mesh, const TArray<int>& TriangleIDs, TArray<int>& VertexIDsOut);
 
 
+	/**
+	 * Find all the triangles in all the one rings of a set of vertices
+	 * @param Mesh input mesh
+	 * @param VertexIDs list of Vertex IDs of Mesh
+	 * @param TriangleIDsOut list of triangle IDs where any vertex of triangle touches an element of VertexIDs
+	 */
+	DYNAMICMESH_API void VertexToTriangleOneRing(const FDynamicMesh3* Mesh, const TArray<int>& VertexIDs, TSet<int>& TriangleIDsOut);
+
 }
diff --git a/Engine/Plugins/Experimental/GeometryProcessing/Source/GeometricObjects/Public/GeometryTypes.h b/Engine/Plugins/Experimental/GeometryProcessing/Source/GeometricObjects/Public/GeometryTypes.h
--- a/Engine/Plugins/Experimental/GeometryProcessing/Source/GeometricObjects/Public/GeometryTypes.h
+++ b/Engine/Plugins/Experimental/GeometryProcessing/Source/GeometricObjects/Public/GeometryTypes.h
@@ -152,6 +152,21 @@ public:
 		return (FoundVal == nullptr) ? InvalidID : *FoundVal;
 	}
 
+	/** @return forward-map of input value or null if not found */
+	inline const IntType* FindTo(IntType FromID) const
+	{
+		check(bWantForward);
+		return ForwardMap.Find(FromID);
+	}
+
+	/** @return reverse-map of input value or null if not found */
+	inline const IntType* FindFrom(IntType ToID) const
+	{
+		check(bWantReverse);
+		return ReverseMap.Find(ToID);
+	}
+
+
 	void Reserve(int NumElements)
 	{
 		if (bWantForward)
diff --git a/Engine/Plugins/Experimental/GeometryProcessing/Source/GeometricObjects/Public/Polyline3.h b/Engine/Plugins/Experimental/GeometryProcessing/Source/GeometricObjects/Public/Polyline3.h
--- a/Engine/Plugins/Experimental/GeometryProcessing/Source/GeometricObjects/Public/Polyline3.h
+++ b/Engine/Plugins/Experimental/GeometryProcessing/Source/GeometricObjects/Public/Polyline3.h
@@ -420,10 +420,9 @@ public:
 		const T Alpha = (T)1 / (T)3;
 		const T OneMinusAlpha = (T)2 / (T)3;
 
-		int N = Vertices.Num();
-		NewPolyline.Vertices.SetNum(2*(N-1));
+		int N = Vertices.Num() - 1;
+		NewPolyline.Vertices.SetNum(2*N);
 		NewPolyline.Vertices[0] = Vertices[0];
-		N = N - 1;
 		int k = 1;
 		for (int i = 1; i < N; ++i)
 		{
diff --git a/Engine/Plugins/Experimental/GeometryProcessing/Source/GeometricObjects/Public/SegmentTypes.h b/Engine/Plugins/Experimental/GeometryProcessing/Source/GeometricObjects/Public/SegmentTypes.h
--- a/Engine/Plugins/Experimental/GeometryProcessing/Source/GeometricObjects/Public/SegmentTypes.h
+++ b/Engine/Plugins/Experimental/GeometryProcessing/Source/GeometricObjects/Public/SegmentTypes.h
@@ -6,6 +6,7 @@
 
 #include "Math/UnrealMath.h"
 #include "VectorTypes.h"
+#include "BoxTypes.h"
 
 
 /*
