add FDynamicMeshChangeTracker and various implementation classes, which supports tracking and undo/redo of mesh topology changes in a FDynamicMesh3
#rnx

diff --git a/Engine/Plugins/Experimental/GeometryProcessing/Source/DynamicMesh/Private/DynamicMeshChangeTracker.cpp b/Engine/Plugins/Experimental/GeometryProcessing/Source/DynamicMesh/Private/DynamicMeshChangeTracker.cpp
new file mode 100644
--- /dev/null
+++ b/Engine/Plugins/Experimental/GeometryProcessing/Source/DynamicMesh/Private/DynamicMeshChangeTracker.cpp
@@ -0,0 +1,670 @@
+// Copyright 1998-2019 Epic Games, Inc. All Rights Reserved.
+
+#include "DynamicMeshChangeTracker.h"
+
+
+
+//
+//
+// TDynamicMeshAttributeChange - stores change in attribute overlay
+// 
+//
+
+
+
+template<typename RealType,int ElementSize>
+void TDynamicMeshAttributeChange<RealType,ElementSize>::SaveInitialElement(const TDynamicMeshOverlay<RealType,ElementSize>* Overlay, int ElementID)
+{
+	FChangeElement Element;
+	Element.ElementID = ElementID;
+	Element.ParentVertexID = Overlay->GetParentVertex(ElementID);
+	Element.DataIndex = OldElementData.Num();
+	TStaticArray<RealType,ElementSize> ElemData;
+	Overlay->GetElement(ElementID, &ElemData[0]);
+	for (int j = 0; j < ElementSize; ++j)
+	{
+		OldElementData.Add(ElemData[j]);
+	}
+	OldElements.Add(Element);		// @todo emplace
+}
+
+template<typename RealType,int ElementSize>
+void TDynamicMeshAttributeChange<RealType,ElementSize>::SaveInitialTriangle(const TDynamicMeshOverlay<RealType,ElementSize>* Overlay, int TriangleID)
+{
+	FChangeTriangle Triangle;
+	Triangle.TriangleID = TriangleID;
+	Triangle.Elements = Overlay->GetTriangle(TriangleID);
+	OldTriangles.Add(Triangle);
+}
+
+template<typename RealType,int ElementSize>
+void TDynamicMeshAttributeChange<RealType,ElementSize>::StoreFinalElement(const TDynamicMeshOverlay<RealType,ElementSize>* Overlay, int ElementID)
+{
+	FChangeElement Element;
+	Element.ElementID = ElementID;
+	Element.ParentVertexID = Overlay->GetParentVertex(ElementID);
+	Element.DataIndex = NewElementData.Num();
+	TStaticArray<RealType, ElementSize> ElemData;
+	Overlay->GetElement(ElementID, &ElemData[0]);
+	for (int j = 0; j < ElementSize; ++j)
+	{
+		NewElementData.Add(ElemData[j]);
+	}
+	NewElements.Add(Element);		// @todo emplace
+}
+
+template<typename RealType,int ElementSize>
+void TDynamicMeshAttributeChange<RealType,ElementSize>::StoreFinalTriangle(const TDynamicMeshOverlay<RealType,ElementSize>* Overlay, int TriangleID)
+{
+	FChangeTriangle Triangle;
+	Triangle.TriangleID = TriangleID;
+	Triangle.Elements = Overlay->GetTriangle(TriangleID);
+	NewTriangles.Add(Triangle);
+}
+
+template<typename RealType,int ElementSize>
+bool TDynamicMeshAttributeChange<RealType,ElementSize>::Apply(TDynamicMeshOverlay<RealType,ElementSize>* Overlay, bool bRevert) const
+{
+	if (bRevert)
+	{
+		ApplyReplaceChange(Overlay, NewTriangles, OldElements, OldElementData, OldTriangles);
+	}
+	else
+	{
+		ApplyReplaceChange(Overlay, OldTriangles, NewElements, NewElementData, NewTriangles);
+	}
+
+	return true;
+}
+
+
+template<typename RealType,int ElementSize>
+void TDynamicMeshAttributeChange<RealType,ElementSize>::ApplyReplaceChange(TDynamicMeshOverlay<RealType, ElementSize>* Overlay,
+	const TArray<FChangeTriangle>& RemoveTris,
+	const TArray<FChangeElement>& InsertElements,
+	const TArray<RealType>& InsertElementData,
+	const TArray<FChangeTriangle>& InsertTris) const
+{
+	// following block is disabled because DynamicMesh::RemoveTriangle() will remove the attribute triangles,
+	// and discard elements. So we do not need to do this here. And *cannot* do this here. Which is not ideal...
+	// @todo optionally do this part via a flag, so we can use TDynamicMeshAttributeChange outside of a parent FDynamicMeshChange?
+
+
+	// clear triangles and discard any elements we don't need anymore
+	//for (const FChangeTriangle& TriInfo : RemoveTris)
+	//{
+	//	Overlay->OnRemoveTriangle(TriInfo.TriangleID, true);
+	//}
+
+	// if this is true, we use unsafe insert, which means we take shortcuts
+	// inside DMesh to speed up multiple inserts. This requires rebuilding
+	// the refcount free lists after the insertions are done, which is O(N)
+	bool bUseUnsafe = true;
+
+	// insert missing vertices
+	if (bUseUnsafe)
+	{
+		Overlay->BeginUnsafeElementsInsert();
+	}
+	for (const FChangeElement& ElemInfo : InsertElements)
+	{
+		// boundary elements may still exist
+		if (Overlay->IsElement(ElemInfo.ElementID) == false)
+		{
+			const RealType* Data = &InsertElementData[ElemInfo.DataIndex];
+			EMeshResult Result = Overlay->InsertElement(ElemInfo.ElementID, Data, ElemInfo.ParentVertexID, bUseUnsafe);
+			check(Result == EMeshResult::Ok);
+		}
+		else
+		{
+			const RealType* Data = &InsertElementData[ElemInfo.DataIndex];
+			Overlay->SetElement(ElemInfo.ElementID, Data);
+		}
+	}
+	if (bUseUnsafe)
+	{
+		Overlay->EndUnsafeElementsInsert();
+	}
+
+	// set new element triangles
+	for (const FChangeTriangle& TriInfo : InsertTris)
+	{
+		check(Overlay->GetParentMesh()->IsTriangle(TriInfo.TriangleID));
+		// existing overlay triangle should be empty...
+		EMeshResult Result = Overlay->SetTriangle(TriInfo.TriangleID, TriInfo.Elements);
+		check(Result == EMeshResult::Ok);
+	}
+}
+
+
+
+
+
+// because ::ApplyReplaceChange() is in the .cpp we need to declare possible attribute types here
+template class DYNAMICMESH_API TDynamicMeshAttributeChange<float, 1>;
+template class DYNAMICMESH_API TDynamicMeshAttributeChange<double, 1>;
+template class DYNAMICMESH_API TDynamicMeshAttributeChange<int, 1>;
+template class DYNAMICMESH_API TDynamicMeshAttributeChange<float, 2>;
+template class DYNAMICMESH_API TDynamicMeshAttributeChange<double, 2>;
+template class DYNAMICMESH_API TDynamicMeshAttributeChange<int, 2>;
+template class DYNAMICMESH_API TDynamicMeshAttributeChange<float, 3>;
+template class DYNAMICMESH_API TDynamicMeshAttributeChange<double, 3>;
+template class DYNAMICMESH_API TDynamicMeshAttributeChange<int, 3>;
+
+
+
+
+
+
+
+//
+//
+// FDynamicMeshChange - stores change in dynamic mesh
+// 
+//
+
+
+
+
+FDynamicMeshChange::~FDynamicMeshChange()
+{
+	check(true);
+}
+
+
+void FDynamicMeshChange::SaveInitialVertex(const FDynamicMesh3* Mesh, int VertexID)
+{
+	FChangeVertex Vertex;
+	Vertex.VertexID = VertexID;
+	Vertex.Info = Mesh->GetVertexInfo(VertexID);
+	OldVertices.Add(Vertex);		// @todo emplace
+}
+
+void FDynamicMeshChange::SaveInitialTriangle(const FDynamicMesh3* Mesh, int TriangleID)
+{
+	FChangeTriangle Triangle;
+	Triangle.TriangleID = TriangleID;
+	Triangle.Vertices = Mesh->GetTriangle(TriangleID);
+	Triangle.Edges = Mesh->GetTriEdges(TriangleID);
+	Triangle.GroupID = Mesh->GetTriangleGroup(TriangleID);
+	OldTriangles.Add(Triangle);
+}
+
+
+void FDynamicMeshChange::StoreFinalVertex(const FDynamicMesh3* Mesh, int VertexID)
+{
+	FChangeVertex Vertex;
+	Vertex.VertexID = VertexID;
+	Vertex.Info = Mesh->GetVertexInfo(VertexID);
+	NewVertices.Add(Vertex);		// @todo emplace
+}
+
+void FDynamicMeshChange::StoreFinalTriangle(const FDynamicMesh3* Mesh, int TriangleID)
+{
+	FChangeTriangle Triangle;
+	Triangle.TriangleID = TriangleID;
+	Triangle.Vertices = Mesh->GetTriangle(TriangleID);
+	Triangle.Edges = Mesh->GetTriEdges(TriangleID);
+	Triangle.GroupID = Mesh->GetTriangleGroup(TriangleID);
+	NewTriangles.Add(Triangle);
+}
+
+
+bool FDynamicMeshChange::Apply(FDynamicMesh3* Mesh, bool bRevert) const
+{
+	//Mesh->CheckValidity();
+
+	if (bRevert)
+	{
+		ApplyReplaceChange(Mesh, NewTriangles, OldVertices, OldTriangles);
+	}
+	else
+	{
+		ApplyReplaceChange(Mesh, OldTriangles, NewVertices, NewTriangles);
+	}
+
+	if (AttributeChanges.IsValid())
+	{
+		AttributeChanges->Apply(Mesh->Attributes(), bRevert);
+	}
+
+	//Mesh->CheckValidity();
+
+	return true;
+}
+
+
+void FDynamicMeshChange::ApplyReplaceChange(FDynamicMesh3* Mesh,
+	const TArray<FChangeTriangle>& RemoveTris,
+	const TArray<FChangeVertex>& InsertVerts,
+	const TArray<FChangeTriangle>& InsertTris) const
+{
+	// remove inserted triangles
+	for (const FChangeTriangle& TriInfo : RemoveTris)
+	{
+		EMeshResult Result = Mesh->RemoveTriangle(TriInfo.TriangleID);
+		check(Result == EMeshResult::Ok);
+	}
+
+	// if this is true, we use unsafe insert, which means we take shortcuts
+	// inside DMesh to speed up multiple inserts. This requires rebuilding
+	// the refcount free lists after the insertions are done, which is O(N)
+	bool bUseUnsafe = true;
+	
+
+	// insert missing vertices
+	if (bUseUnsafe)
+	{
+		Mesh->BeginUnsafeVerticesInsert();
+	}
+	for (const FChangeVertex& VertInfo : InsertVerts)
+	{
+		// boundary vertices may still exist. If interior vertices still exist we are in trouble but what could we possibly do to recover?
+		if (Mesh->IsVertex(VertInfo.VertexID) == false)
+		{
+			EMeshResult Result = Mesh->InsertVertex(VertInfo.VertexID, VertInfo.Info, bUseUnsafe);
+			check(Result == EMeshResult::Ok);
+		}
+		else
+		{
+			Mesh->SetVertex(VertInfo.VertexID, VertInfo.Info.Position);
+			if (VertInfo.Info.bHaveN)
+			{
+				Mesh->SetVertexNormal(VertInfo.VertexID, VertInfo.Info.Normal);
+			}
+			if (VertInfo.Info.bHaveC)
+			{
+				Mesh->SetVertexColor(VertInfo.VertexID, VertInfo.Info.Color);
+			}
+			if (VertInfo.Info.bHaveUV)
+			{
+				Mesh->SetVertexUV(VertInfo.VertexID, VertInfo.Info.UV);
+			}
+		}
+	}
+	if (bUseUnsafe)
+	{
+		Mesh->EndUnsafeVerticesInsert();
+	}
+
+	// insert new triangles
+	if (bUseUnsafe)
+	{
+		Mesh->BeginUnsafeTrianglesInsert();
+	}
+	for (const FChangeTriangle& TriInfo : InsertTris)
+	{
+		EMeshResult Result = Mesh->InsertTriangle(TriInfo.TriangleID, TriInfo.Vertices, TriInfo.GroupID, bUseUnsafe);
+		check(Result == EMeshResult::Ok);
+	}
+	if (bUseUnsafe)
+	{
+		Mesh->EndUnsafeTrianglesInsert();
+	}
+}
+
+
+bool FDynamicMeshChange::HasSavedVertex(int VertexID)
+{
+	for (const FChangeVertex& VertInfo : OldVertices)
+	{
+		if (VertInfo.VertexID == VertexID)
+			return true;
+	}
+	return false;
+}
+
+
+void FDynamicMeshChange::VerifySaveState() const
+{
+	TSet<int> SavedVertexIDs;
+	for (const FChangeVertex& VertInfo : OldVertices)
+	{
+		SavedVertexIDs.Add(VertInfo.VertexID);
+	}
+
+	for (const FChangeTriangle& TriInfo : OldTriangles)
+	{
+		check(SavedVertexIDs.Contains(TriInfo.Vertices.A));
+		check(SavedVertexIDs.Contains(TriInfo.Vertices.B));
+		check(SavedVertexIDs.Contains(TriInfo.Vertices.C));
+	}
+}
+
+
+
+
+
+//
+//
+// FDynamicMeshAttributeSetChangeTracker - tracks changes in all attribute overlays of a Mesh
+// 
+//
+
+
+
+
+FDynamicMeshAttributeSetChangeTracker::FDynamicMeshAttributeSetChangeTracker(const FDynamicMeshAttributeSet* AttribsIn)
+{
+	this->Attribs = AttribsIn;
+}
+
+void FDynamicMeshAttributeSetChangeTracker::BeginChange()
+{
+	// initialize new attribute set change
+	check(Change == nullptr);
+	Change = new FDynamicMeshAttributeChangeSet();
+	int NumUVLayers = Attribs->GetNumUVLayers();
+	Change->UVChanges.SetNum(NumUVLayers);
+	int NumNormalLayers = Attribs->GetNumNormalLayers();
+	Change->NormalChanges.SetNum(NumNormalLayers);
+
+	// initialize UV layer state tracking
+	if (UVStates.Num() < NumUVLayers)
+	{
+		UVStates.SetNum(NumUVLayers);
+	}
+	for (int k = 0; k < NumUVLayers; ++k)
+	{
+		FElementState& State = UVStates[k];
+		const FDynamicMeshUVOverlay* UVLayer = Attribs->GetUVLayer(k);
+		State.MaxElementID = UVLayer->MaxElementID();
+		State.StartElements.Init(false, State.MaxElementID);
+		State.ChangedElements.Init(false, State.MaxElementID);
+		for ( int ElemID : UVLayer->ElementIndicesItr())
+		{
+			State.StartElements[ElemID] = true;
+		}
+	}
+
+	// initialize Normal layer state tracking
+	if (NormalStates.Num() < NumNormalLayers)
+	{
+		NormalStates.SetNum(NumNormalLayers);
+	}
+	for (int k = 0; k < NumNormalLayers; ++k)
+	{
+		FElementState& State = NormalStates[k];
+		const FDynamicMeshNormalOverlay* NormalLayer = Attribs->GetNormalLayer(k);
+		State.MaxElementID = NormalLayer->MaxElementID();
+		State.StartElements.Init(false, State.MaxElementID);
+		State.ChangedElements.Init(false, State.MaxElementID);
+		for (int ElemID : NormalLayer->ElementIndicesItr())
+		{
+			State.StartElements[ElemID] = true;
+		}
+	}
+}
+
+
+TUniquePtr<FDynamicMeshAttributeChangeSet> FDynamicMeshAttributeSetChangeTracker::EndChange()
+{
+	TUniquePtr<FDynamicMeshAttributeChangeSet> Result(this->Change);
+	return MoveTemp(Result);
+}
+
+
+
+void FDynamicMeshAttributeSetChangeTracker::SaveInitialTriangle(int TriangleID)
+{
+	int NumUVLayers = Attribs->GetNumUVLayers();
+	for (int k = 0; k < NumUVLayers; ++k)
+	{
+		FDynamicMeshUVChange& UVChange = Change->UVChanges[k];
+		const FDynamicMeshUVOverlay* UVLayer = Attribs->GetUVLayer(k);
+		FIndex3i UVTriangle = UVLayer->GetTriangle(TriangleID);
+		SaveElement(UVTriangle.A, UVStates[k], UVLayer, UVChange);
+		SaveElement(UVTriangle.B, UVStates[k], UVLayer, UVChange);
+		SaveElement(UVTriangle.C, UVStates[k], UVLayer, UVChange);
+		UVChange.SaveInitialTriangle(UVLayer, TriangleID);
+	}
+
+	int NumNormalLayers = Attribs->GetNumNormalLayers();
+	for (int k = 0; k < NumNormalLayers; ++k)
+	{
+		FDynamicMeshNormalChange& NormalChange = Change->NormalChanges[k];
+		const FDynamicMeshNormalOverlay* NormalLayer = Attribs->GetNormalLayer(k);
+		FIndex3i NormTriangle = NormalLayer->GetTriangle(TriangleID);
+		SaveElement(NormTriangle.A, NormalStates[k], NormalLayer, NormalChange);
+		SaveElement(NormTriangle.B, NormalStates[k], NormalLayer, NormalChange);
+		SaveElement(NormTriangle.C, NormalStates[k], NormalLayer, NormalChange);
+		NormalChange.SaveInitialTriangle(NormalLayer, TriangleID);
+	}
+}
+
+
+void FDynamicMeshAttributeSetChangeTracker::StoreAllFinalTriangles(const TArray<int>& TriangleIDs)
+{
+	TSet<int> StoredVertices;	// re-used
+	
+	// store final UV elements for all modified triangles, and final triangles
+	int NumUVLayers = Attribs->GetNumUVLayers();
+	for (int k = 0; k < NumUVLayers; ++k)
+	{
+		FDynamicMeshUVChange& UVChange = Change->UVChanges[k];
+		const FDynamicMeshUVOverlay* UVLayer = Attribs->GetUVLayer(k);
+		StoredVertices.Reset();
+
+		for (int tid : TriangleIDs)
+		{
+			FIndex3i Tri = UVLayer->GetTriangle(tid);
+			for (int j = 0; j < 3; ++j)
+			{
+				if (StoredVertices.Contains(Tri[j]) == false)
+				{
+					UVChange.StoreFinalElement(UVLayer, Tri[j]);
+					StoredVertices.Add(Tri[j]);
+				}
+			}
+			UVChange.StoreFinalTriangle(UVLayer, tid);
+		}
+	}
+
+
+	// store final Normal elements for all modified triangles, and final triangles
+	int NumNormalLayers = Attribs->GetNumNormalLayers();
+	for (int k = 0; k < NumNormalLayers; ++k)
+	{
+		FDynamicMeshNormalChange& NormalChange = Change->NormalChanges[k];
+		const FDynamicMeshNormalOverlay* NormalLayer = Attribs->GetNormalLayer(k);
+		StoredVertices.Reset();
+
+		for (int tid : TriangleIDs)
+		{
+			FIndex3i Tri = NormalLayer->GetTriangle(tid);
+			for (int j = 0; j < 3; ++j)
+			{
+				if (StoredVertices.Contains(Tri[j]) == false)
+				{
+					NormalChange.StoreFinalElement(NormalLayer, Tri[j]);
+					StoredVertices.Add(Tri[j]);
+				}
+			}
+			NormalChange.StoreFinalTriangle(NormalLayer, tid);
+		}
+	}
+
+}
+
+
+bool FDynamicMeshAttributeChangeSet::Apply(FDynamicMeshAttributeSet* Attributes, bool bRevert) const
+{
+	int NumUVLayers = Attributes->GetNumUVLayers();
+	check(NumUVLayers == UVChanges.Num());
+	for (int k = 0; k < NumUVLayers; ++k)
+	{
+		FDynamicMeshUVOverlay* UVLayer = Attributes->GetUVLayer(k);
+		UVChanges[k].Apply(UVLayer, bRevert);
+	}
+
+	int NumNormalLayers = Attributes->GetNumNormalLayers();
+	check(NumNormalLayers == NormalChanges.Num());
+	for (int k = 0; k < NumNormalLayers; ++k)
+	{
+		FDynamicMeshNormalOverlay* NormalLayer = Attributes->GetNormalLayer(k);
+		NormalChanges[k].Apply(NormalLayer, bRevert);
+	}
+
+	return true;
+}
+
+
+
+
+
+
+
+
+//
+//
+// FDynamicMeshChangeTracker - tracks all changes in mesh and attribute overlay
+// 
+//
+
+
+
+
+FDynamicMeshChangeTracker::FDynamicMeshChangeTracker(const FDynamicMesh3* MeshIn)
+{
+	Mesh = MeshIn;
+	if (Mesh->HasAttributes())
+	{
+		AttribChangeTracker = new FDynamicMeshAttributeSetChangeTracker(Mesh->Attributes());
+	}
+}
+
+FDynamicMeshChangeTracker::~FDynamicMeshChangeTracker()
+{
+	if (AttribChangeTracker != nullptr)
+	{
+		delete AttribChangeTracker;
+		AttribChangeTracker = nullptr;
+	}
+}
+
+
+void FDynamicMeshChangeTracker::BeginChange()
+{
+	check(Change == nullptr);
+	Change = new FDynamicMeshChange();
+
+	// @todo should we do this on EndChange() so that delay is at end of stroke instead of beginning?
+	// @todo use TSet instead of a bitmap? Or implement HBitmap?
+
+	this->MaxVertexID = Mesh->MaxVertexID();
+	StartVertices.Init(false, MaxVertexID);
+	ChangedVertices.Init(false, MaxVertexID);
+	for (int vi = 0; vi < MaxVertexID; ++vi)
+	{
+		if (Mesh->IsVertex(vi))
+		{
+			StartVertices[vi] = true;
+		}
+	}
+
+	this->MaxTriangleID = Mesh->MaxTriangleID();
+	StartTriangles.Init(false, MaxTriangleID);    // does this re-use memory? would be preferable
+	ChangedTriangles.Init(false, MaxTriangleID);
+	for (int ti = 0; ti < MaxTriangleID; ++ti)
+	{
+		if (Mesh->IsTriangle(ti))
+		{
+			StartTriangles[ti] = true;
+		}
+	}
+
+	if (AttribChangeTracker != nullptr)
+	{
+		AttribChangeTracker->BeginChange();
+	}
+}
+
+
+
+void FDynamicMeshChangeTracker::SaveVertex(int VertexID)
+{
+	// only save initial vertex state if it existed at BeginChange() and hasn't been modified yet
+	if (VertexID >= MaxVertexID || ChangedVertices[VertexID] == true || StartVertices[VertexID] == false)
+	{
+		return;
+	}
+
+	Change->SaveInitialVertex(Mesh, VertexID);
+
+	ChangedVertices[VertexID] = true;
+}
+
+
+
+void FDynamicMeshChangeTracker::SaveTriangle(int TriangleID, bool bSaveVertices)
+{
+	// only save initial triangle state if it existed at BeginChange() and hasn't been modified yet
+	if (TriangleID >= MaxTriangleID || ChangedTriangles[TriangleID] == true || StartTriangles[TriangleID] == false)
+	{
+		return;
+	}
+
+	if (bSaveVertices)
+	{
+		FIndex3i Tri = Mesh->GetTriangle(TriangleID);
+		SaveVertex(Tri.A);
+		SaveVertex(Tri.B);
+		SaveVertex(Tri.C);
+	}
+
+	Change->SaveInitialTriangle(Mesh, TriangleID);
+
+	ChangedTriangles[TriangleID] = true;
+
+	if (AttribChangeTracker != nullptr)
+	{
+		AttribChangeTracker->SaveInitialTriangle(TriangleID);
+	}
+}
+
+
+
+TUniquePtr<FDynamicMeshChange> FDynamicMeshChangeTracker::EndChange()
+{
+	TSet<int> StoredVertices;   // @todo maybe should be TBitArray? 
+	TArray<int> StoredTriangles;
+
+	for (int tid : Mesh->TriangleIndicesItr())
+	{
+		// determine if this is a new or modified triangle, and if so, store it
+		if (tid >= MaxTriangleID || StartTriangles[tid] == false || ChangedTriangles[tid] == true)
+		{
+			// only store new vertices 
+			FIndex3i Tri = Mesh->GetTriangle(tid);
+			for (int j = 0; j < 3; ++j)
+			{
+				if (StoredVertices.Contains(Tri[j]) == false)
+				{
+					Change->StoreFinalVertex(Mesh, Tri[j]);
+					StoredVertices.Add(Tri[j]);
+				}
+			}
+
+			Change->StoreFinalTriangle(Mesh, tid);
+			StoredTriangles.Add(tid);
+		}
+	}
+	TUniquePtr<FDynamicMeshChange> Result(Change);
+
+	// handle attribute overlays
+	if (AttribChangeTracker != nullptr)
+	{
+		AttribChangeTracker->StoreAllFinalTriangles(StoredTriangles);
+		Result->AttachAttributeChanges(AttribChangeTracker->EndChange());
+	}
+
+	return MoveTemp(Result);
+}
+
+
+
+
+void FDynamicMeshChangeTracker::VerifySaveState()
+{
+	Change->VerifySaveState();
+}
\ No newline at end of file
diff --git a/Engine/Plugins/Experimental/GeometryProcessing/Source/DynamicMesh/Public/DynamicMeshChangeTracker.h b/Engine/Plugins/Experimental/GeometryProcessing/Source/DynamicMesh/Public/DynamicMeshChangeTracker.h
new file mode 100644
--- /dev/null
+++ b/Engine/Plugins/Experimental/GeometryProcessing/Source/DynamicMesh/Public/DynamicMeshChangeTracker.h
@@ -0,0 +1,280 @@
+// Copyright 1998-2019 Epic Games, Inc. All Rights Reserved.
+
+#pragma once
+
+#include "DynamicMesh3.h"
+#include "DynamicMeshOverlay.h"
+#include "DynamicMeshAttributeSet.h"
+
+
+//
+// Implementation of mesh change tracking for FDynamicMesh3.
+// 
+// The top-level class is FDynamicMeshChangeTracker, found at the bottom of this file.
+// You create an instance of this and then call BeginChange(), then call SaveVertex()/SaveTriangle()
+// before you make modifications to a vertex/triangle, and then call EndChange() at the end.
+// This function emits a FDynamicMeshChange instance, and you can call Apply() to apply/revert it.
+// So this is the object you would store in a higher-level undo/redo record, for example.
+// 
+// Attribute overlays make everything more complicated of course. Handling of these
+// follows a similar structure - the MeshChangeTracker creates a FDynamicMeshAttributeSetChangeTracker,
+// which in turn creates a TDynamicMeshAttributeChange for each UV and Normal overlay
+// (grouped together in a FDynamicMeshAttributeChangeSet). However you don't have to explicitly
+// do anything to get Attribute support, if the initial Mesh had attributes, then
+// this all happens automatically, and the attribute changes get attached to the FDynamicMeshChange.
+// 
+// @todo Currently the attribute support is hardcoded for UV (2-float) and Normal (3-float) overlays.
+// Perhaps it would be possible to make this more flexible to avoid hardcoding these?
+// Very little of the code depends on the element size or type.
+// 
+// 
+
+
+
+
+/**
+ * TDynamicMeshAttributeChange represents a change to an attribute overlay of a FDynamicMesh3.
+ * @warning This class is meant to be used via FDynamicMeshChange and is not fully functional
+ * on its own (see comments in ApplyReplaceChange)
+ */
+template<typename RealType, int ElementSize>
+class DYNAMICMESH_API TDynamicMeshAttributeChange
+{
+public:
+	void SaveInitialElement(const TDynamicMeshOverlay<RealType,ElementSize>* Overlay, int ElementID);
+	void SaveInitialTriangle(const TDynamicMeshOverlay<RealType,ElementSize>* Overlay, int TriangleID);
+
+	void StoreFinalElement(const TDynamicMeshOverlay<RealType, ElementSize>* Overlay, int ElementID);
+	void StoreFinalTriangle(const TDynamicMeshOverlay<RealType, ElementSize>* Overlay, int TriangleID);
+
+	bool Apply(TDynamicMeshOverlay<RealType, ElementSize>* Overlay, bool bRevert) const;
+
+protected:
+	struct FChangeElement
+	{
+		int ElementID;
+		int DataIndex;
+		int ParentVertexID;
+	};
+
+	struct FChangeTriangle
+	{
+		int TriangleID;
+		FIndex3i Elements;
+	};
+
+	TArray<FChangeElement> OldElements;
+	TArray<RealType> OldElementData;
+	TArray<FChangeTriangle> OldTriangles;
+
+	TArray<FChangeElement> NewElements;
+	TArray<RealType> NewElementData;
+	TArray<FChangeTriangle> NewTriangles;
+
+	void ApplyReplaceChange(TDynamicMeshOverlay<RealType,ElementSize>* Overlay,
+		const TArray<FChangeTriangle>& RemoveTris,
+		const TArray<FChangeElement>& InsertElements,
+		const TArray<RealType>& InsertElementData,
+		const TArray<FChangeTriangle>& InsertTris) const;
+};
+
+
+
+/** Standard UV overlay change type - 2-element float */
+typedef TDynamicMeshAttributeChange<float,2> FDynamicMeshUVChange;
+/** Standard Normal overlay change type - 3-element float */
+typedef TDynamicMeshAttributeChange<float,3> FDynamicMeshNormalChange;
+
+
+/**
+ * FDynamicMeshAttributeChangeSet stores a set of UV and Normal changes for a FDynamicMesh3
+ */
+class DYNAMICMESH_API FDynamicMeshAttributeChangeSet
+{
+public:
+	TArray<FDynamicMeshUVChange> UVChanges;
+	TArray<FDynamicMeshNormalChange> NormalChanges;
+
+	/** call ::Apply() on all the UV and Normal changes */
+	bool Apply(FDynamicMeshAttributeSet* Attributes, bool bRevert) const;
+};
+
+
+/**
+ * FDynamicMeshChange stores a "change" in a FDynamicMesh3, which in this context
+ * means the replacement of one set of triangles with a second set, that may have
+ * different vertices/attributes. The change can be applied and reverted via ::Apply()
+ * 
+ * Construction of a well-formed FDynamicMeshChange is quite complex and it is strongly
+ * suggested that you do so via FDynamicMeshChangeTracker
+ */
+class DYNAMICMESH_API FDynamicMeshChange
+{
+public:
+	~FDynamicMeshChange();
+
+	/** Store the initial state of a vertex */
+	void SaveInitialVertex(const FDynamicMesh3* Mesh, int VertexID);
+	/** Store the initial state of a triangle */
+	void SaveInitialTriangle(const FDynamicMesh3* Mesh, int TriangleID);
+
+	/** Store the final state of a vertex */
+	void StoreFinalVertex(const FDynamicMesh3* Mesh, int VertexID);
+	/** Store the final state of a triangle */
+	void StoreFinalTriangle(const FDynamicMesh3* Mesh, int TriangleID);
+
+	/** Attach an attribute change set to this mesh change, which will the be applied/reverted autoamtically */
+	void AttachAttributeChanges(TUniquePtr<FDynamicMeshAttributeChangeSet> AttribChanges)
+	{
+		this->AttributeChanges = MoveTemp(AttribChanges);
+	}
+
+	/** Apply or Revert this change using the given Mesh */
+	bool Apply(FDynamicMesh3* Mesh, bool bRevert) const;
+
+	/** Do (limited) sanity checks on this MeshChange to ensure it is well-formed */
+	void VerifySaveState() const;
+
+	/** @return true if this vertex was saved. Uses linear search. */
+	bool HasSavedVertex(int VertexID);
+
+protected:
+
+	struct FChangeVertex
+	{
+		int VertexID;
+		FVertexInfo Info;
+	};
+
+	struct FChangeTriangle
+	{
+		int TriangleID;
+		FIndex3i Vertices;
+		FIndex3i Edges;
+		int GroupID;
+	};
+
+	TArray<FChangeVertex> OldVertices;
+	TArray<FChangeTriangle> OldTriangles;
+
+	TArray<FChangeVertex> NewVertices;
+	TArray<FChangeTriangle> NewTriangles;
+
+	TUniquePtr<FDynamicMeshAttributeChangeSet> AttributeChanges;
+
+	void ApplyReplaceChange(FDynamicMesh3* Mesh,
+		const TArray<FChangeTriangle>& RemoveTris,
+		const TArray<FChangeVertex>& InsertVerts,
+		const TArray<FChangeTriangle>& InsertTris) const;
+
+};
+
+
+
+
+
+
+
+/**
+ * FDynamicMeshAttributeSetChangeTracker constructs a well-formed set of TDynamicMeshAttributeChange
+ * objects (stored in a FDynamicMeshAttributeChangeSet). You should not use this class
+ * directly, it is intended to be used via FDynamicMeshChangeTracker
+ */
+class DYNAMICMESH_API FDynamicMeshAttributeSetChangeTracker
+{
+public:
+	explicit FDynamicMeshAttributeSetChangeTracker(const FDynamicMeshAttributeSet* Attribs);
+
+	/** Start tracking a change */
+	void BeginChange();
+	/** End the change transaction and get the resulting change object */
+	TUniquePtr<FDynamicMeshAttributeChangeSet> EndChange();
+
+	/** Store the initial state of a triangle */
+	void SaveInitialTriangle(int TriangleID);
+
+	/** store the final state of a set of triangles */
+	void StoreAllFinalTriangles(const TArray<int>& TriangleIDs);
+
+
+protected:
+	const FDynamicMeshAttributeSet* Attribs = nullptr;
+
+	FDynamicMeshAttributeChangeSet* Change = nullptr;
+
+	struct FElementState
+	{
+		int MaxElementID;
+		TBitArray<> StartElements;
+		TBitArray<> ChangedElements;
+	};
+	TArray<FElementState> UVStates;
+	TArray<FElementState> NormalStates;
+
+	template<typename AttribOverlayType, typename AttribChangeType>
+	void SaveElement(int ElementID, FElementState& State, const AttribOverlayType* Overlay, AttribChangeType& ChangeIn)
+	{
+		if (ElementID < State.MaxElementID && State.ChangedElements[ElementID] == false && State.StartElements[ElementID] == true)
+		{
+			ChangeIn.SaveInitialElement(Overlay, ElementID);
+			State.ChangedElements[ElementID] = true;
+		}
+	}
+
+};
+
+
+
+
+
+
+
+
+
+/**
+ * FDynamicMeshChangeTracker tracks changes to a FDynamicMesh and returns a
+ * FDynamicMeshChange instance that represents this change and allows it to be reverted/reapplied.
+ * This is the top-level class you likely want to use to track mesh changes.
+ * 
+ * Call BeginChange() before making any changes to the mesh, then call SaveVertex()
+ * and SaveTriangle() before modifying the respective elements. Then call EndChange()
+ * to construct a FDynamicMeshChange that represents the mesh delta.
+ * 
+ */
+class DYNAMICMESH_API FDynamicMeshChangeTracker
+{
+public:
+	explicit FDynamicMeshChangeTracker(const FDynamicMesh3* Mesh);
+	~FDynamicMeshChangeTracker();
+
+	/** Initialize the change-tracking process */
+	void BeginChange();
+	/** Construct a change object that represents the delta between the Begin and End states */
+	TUniquePtr<FDynamicMeshChange> EndChange();
+
+	/** Save necessary information about a vertex before it is modified */
+	void SaveVertex(int VertexID);
+	/** Save necessary information about a triangle before it is modified */
+	void SaveTriangle(int TriangleID, bool bSaveVertices);
+
+	/** Do (limited) sanity checking to make sure that the change is well-formed*/
+	void VerifySaveState();
+
+protected:
+	const FDynamicMesh3* Mesh = nullptr;
+
+	/** Active attribute tracker, if Mesh has attribute overlays */
+	FDynamicMeshAttributeSetChangeTracker* AttribChangeTracker = nullptr;
+
+	/** Active change that is being constructed */
+	FDynamicMeshChange* Change = nullptr;
+
+	int MaxTriangleID;
+	TBitArray<> StartTriangles;		// bit is 1 if triangle ID was in initial mesh on BeginChange()
+	TBitArray<> ChangedTriangles;	// bit is set to 1 if triangle was in StartTriangles and then had SaveTriangle() called for it
+
+	int MaxVertexID;
+	TBitArray<> StartVertices;		// bit is 1 if vertex ID was in initial mesh on BeginChange()
+	TBitArray<> ChangedVertices;	// bit is set to 1 if vertex was in StartVertices and then had SaveVertex() called for it
+};
+
diff --git a/Engine/Plugins/Experimental/GeometryProcessing/Source/DynamicMesh/Public/SubRegionRemesher.h b/Engine/Plugins/Experimental/GeometryProcessing/Source/DynamicMesh/Public/SubRegionRemesher.h
--- a/Engine/Plugins/Experimental/GeometryProcessing/Source/DynamicMesh/Public/SubRegionRemesher.h
+++ b/Engine/Plugins/Experimental/GeometryProcessing/Source/DynamicMesh/Public/SubRegionRemesher.h
@@ -4,6 +4,7 @@
 
 #include "Remesher.h"
 #include "Util/BufferUtil.h"
+#include "DynamicMeshChangeTracker.h"
 
 /**
  * FSubRegionRemesher is an extension
@@ -48,10 +49,29 @@ public:
 		}
 		Edges.Reset();
 		BufferUtil::AppendElements(Edges, EdgeROI);
-		//Edges = std::vector<int>(EdgeROI.begin(), EdgeROI.end());
 	}
 
 
+	/**
+	 * Tell a MeshChangeTracker about the set of triangles that we might modify in the next remesh pass.
+	 * This could include one-rings of either side of an edge in the ROI, if we collapse.
+	 */
+	void SaveActiveROI(FDynamicMeshChangeTracker* Change)
+	{
+		const FDynamicMesh3* UseMesh = GetMesh();
+		for (int eid : EdgeROI)
+		{
+			FIndex2i EdgeVerts = UseMesh->GetEdgeV(eid);
+			for (int j = 0; j < 2; ++j)
+			{
+				for (int tid : UseMesh->VtxTrianglesItr(EdgeVerts[j]))
+				{
+					Change->SaveTriangle(tid, true);
+				}
+			}
+		}
+	}
+
 
 	//
 	// specialization of Remesher functionality
