vertex one-ring iterations were quite expensive because of returned lambda in cpp file. Made custom iterator in IteratorUtil, use in DynamicMesh3, significantly faster. Also added return-one-ring-via-tarray-out variant, and reduced number of times each tri is touched in SubRegionRemesher::SaveActiveROI  (this was hotspot in dynamic sculpting)
#rnx

diff --git a/Engine/Plugins/Experimental/GeometryProcessing/Source/DynamicMesh/Private/DynamicMesh3.cpp b/Engine/Plugins/Experimental/GeometryProcessing/Source/DynamicMesh/Private/DynamicMesh3.cpp
--- a/Engine/Plugins/Experimental/GeometryProcessing/Source/DynamicMesh/Private/DynamicMesh3.cpp
+++ b/Engine/Plugins/Experimental/GeometryProcessing/Source/DynamicMesh/Private/DynamicMesh3.cpp
@@ -549,48 +549,27 @@ FIndex3i FDynamicMesh3::GetTriNeighbourTris(int tID) const
 }
 
 
-
-/** Enumerate triangle one ring at a vertex */
-FDynamicMesh3::vtx_triangles_enumerable FDynamicMesh3::VtxTrianglesItr(int vID) const
+void FDynamicMesh3::GetVertexOneRingTriangles(int VertexID, TArray<int>& TrianglesOut) const
 {
-	check(VertexRefCounts.IsValid(vID));
-
-	TFunction<int(int, int&)> expand_f = [this,vID](int eid, int& k) 
+	check(VertexRefCounts.IsValid(VertexID));
+	for (int EdgeID : VertexEdgeLists.Values(VertexID))
 	{
-		int vOther = GetOtherEdgeVertex(eid, vID);
-		int i = 4 * eid;
-		if (k == -1) 
+		FIndex2i EdgePair = GetOrderedOneRingEdgeTris(VertexID, EdgeID);
+		if (EdgePair.A != InvalidID)
 		{
-			int et0 = Edges[i + 2];
-			if (TriHasSequentialVertices(et0, vID, vOther))
-			{
-				k = 0;
-				return et0;
-			}
+			TrianglesOut.Add(EdgePair.A);
 		}
-		if (k != 1) 
+		if (EdgePair.B != InvalidID)
 		{
-			int et1 = Edges[i + 3];
-			if (et1 != InvalidID && TriHasSequentialVertices(et1, vID, vOther))
-			{
-				k = 1;
-				return et1;
-			}
+			TrianglesOut.Add(EdgePair.B);
 		}
-		k = -1;
-		return -1;
-	};
-
-	return vtx_triangles_enumerable(VertexEdgeLists.Values(vID), expand_f);
+	}
 }
 
 
 
 
 
-
-
-
 FString FDynamicMesh3::MeshInfoString()
 {
 	FString VtxString = FString::Printf(TEXT("Vertices count %d max %d  %s  VtxEdges %s"),
diff --git a/Engine/Plugins/Experimental/GeometryProcessing/Source/DynamicMesh/Public/DynamicMesh3.h b/Engine/Plugins/Experimental/GeometryProcessing/Source/DynamicMesh/Public/DynamicMesh3.h
--- a/Engine/Plugins/Experimental/GeometryProcessing/Source/DynamicMesh/Public/DynamicMesh3.h
+++ b/Engine/Plugins/Experimental/GeometryProcessing/Source/DynamicMesh/Public/DynamicMesh3.h
@@ -363,7 +363,7 @@ public:
 	typedef typename FRefCountVector::IndexEnumerable edge_iterator;
 	template <typename T>
 	using value_iteration = FRefCountVector::MappedEnumerable<T>;
-	using vtx_triangles_enumerable = ExpandEnumerable<int, int, FSmallListSet::ValueIterator>;
+	using vtx_triangles_enumerable = TPairExpandEnumerable<FSmallListSet::ValueIterator>;
 
 	/** @return enumerable object for valid vertex indices suitable for use with range-based for, ie for ( int i : VertexIndicesItr() ) */
 	vertex_iterator VertexIndicesItr() const
@@ -435,9 +435,13 @@ public:
 	}
 
 	/** @return enumerable object for one-ring triangles of a vertex, suitable for use with range-based for, ie for ( int i : VtxTrianglesItr(VertexID) ) */
-	vtx_triangles_enumerable VtxTrianglesItr(int VertexID) const;
-
-
+	vtx_triangles_enumerable VtxTrianglesItr(int VertexID) const
+	{
+		check(VertexRefCounts.IsValid(VertexID));
+		return vtx_triangles_enumerable(VertexEdgeLists.Values(VertexID), [this, VertexID](int EdgeID) {
+			return GetOrderedOneRingEdgeTris(VertexID, EdgeID);
+		});
+	}
 
 	//
 	// Mesh Construction
@@ -548,6 +552,9 @@ public:
 	/** @return the max valence of all vertices in the mesh */
 	int GetMaxVtxEdgeCount() const;
 
+	/** adds triangles touching VertexID to the TrianglesOut list */
+	void GetVertexOneRingTriangles(int VertexID, TArray<int>& TrianglesOut) const;
+
 	/** Get triangle vertices */
 	inline FIndex3i GetTriangle(int TriangleID) const
 	{
@@ -1391,6 +1398,22 @@ protected:
 		}
 	}
 
+	// utility function that returns one or two triangles of edge, used to enumerate vertex one-ring triangles
+	// The logic is a bit tricky to follow without drawing it out on paper, but this will only return 
+	// each triangle once, for the 'outgoing' edge from the vertex, and each triangle only has one such edge
+	// at any vertex (including boundary triangles)
+	inline FIndex2i GetOrderedOneRingEdgeTris(int VertexID, int EdgeID) const
+	{
+		int vOther = GetOtherEdgeVertex(EdgeID, VertexID);
+		int i = 4 * EdgeID;
+		int et1 = Edges[i + 3];
+		et1 = (et1 != InvalidID && TriHasSequentialVertices(et1, VertexID, vOther)) ? et1 : InvalidID;
+		int et0 = Edges[i + 2];
+		return TriHasSequentialVertices(et0, VertexID, vOther) ? 
+			FIndex2i(et0, et1) : FIndex2i(et1, InvalidID);
+	}
+
+
 	void ReverseTriOrientationInternal(int TriangleID);
 
 	void UpdateTimeStamp(bool bShapeChange, bool bTopologyChange)
diff --git a/Engine/Plugins/Experimental/GeometryProcessing/Source/DynamicMesh/Public/SubRegionRemesher.h b/Engine/Plugins/Experimental/GeometryProcessing/Source/DynamicMesh/Public/SubRegionRemesher.h
--- a/Engine/Plugins/Experimental/GeometryProcessing/Source/DynamicMesh/Public/SubRegionRemesher.h
+++ b/Engine/Plugins/Experimental/GeometryProcessing/Source/DynamicMesh/Public/SubRegionRemesher.h
@@ -56,20 +56,34 @@ public:
 	 * Tell a MeshChangeTracker about the set of triangles that we might modify in the next remesh pass.
 	 * This could include one-rings of either side of an edge in the ROI, if we collapse.
 	 */
+	TSet<int> Vertices;
 	void SaveActiveROI(FDynamicMeshChangeTracker* Change)
 	{
+		// there is quite a bit of overhead here...perhaps remesher could just save triangles
+		// itself before it touches them?
+
+		TArray<int> OneRingTris; OneRingTris.Reserve(32);
+
+		// figuring out unique verts means we don't do each vertex N~=valence times, 
+		// which saves a lot of one-ring iterations that are somewhat expensive...
+		Vertices.Reset();
 		const FDynamicMesh3* UseMesh = GetMesh();
 		for (int eid : EdgeROI)
 		{
 			FIndex2i EdgeVerts = UseMesh->GetEdgeV(eid);
-			for (int j = 0; j < 2; ++j)
+			Vertices.Add(EdgeVerts.A);
+			Vertices.Add(EdgeVerts.B);
+		}
+		for ( int vid : Vertices)
+		{
+			OneRingTris.Reset();
+			UseMesh->GetVertexOneRingTriangles(vid, OneRingTris);
+			for (int tid : OneRingTris)
 			{
-				for (int tid : UseMesh->VtxTrianglesItr(EdgeVerts[j]))
-				{
-					Change->SaveTriangle(tid, true);
-				}
+				Change->SaveTriangle(tid, true);
 			}
 		}
+
 	}
 
 
diff --git a/Engine/Plugins/Experimental/GeometryProcessing/Source/GeometricObjects/Public/Util/IteratorUtil.h b/Engine/Plugins/Experimental/GeometryProcessing/Source/GeometricObjects/Public/Util/IteratorUtil.h
--- a/Engine/Plugins/Experimental/GeometryProcessing/Source/GeometricObjects/Public/Util/IteratorUtil.h
+++ b/Engine/Plugins/Experimental/GeometryProcessing/Source/GeometricObjects/Public/Util/IteratorUtil.h
@@ -3,6 +3,7 @@
 // Port of geometry3cpp iterator_util.h
 
 #pragma once
+#include "IndexTypes.h"
 
 
 /**
@@ -236,6 +237,157 @@ public:
 		return ExpandIteratorT(BeginItr, EndItr, ExpandFunc);
 	}
 
+	ExpandIteratorT end()
+	{
+		return ExpandIteratorT(EndItr, EndItr, ExpandFunc);
+	}
+};
+
+
+
+
+
+
+
+
+
+
+
+
+/**
+ * Wrapper around existing integer iterator that returns either 0, 1, or 2 integers
+ * for each value that the original iterator returns.
+ * 
+ * This is specifically used by FDynamicMesh3::VtxTrianglesItr, where for each edge
+ * around a vertex, between 0 and 2 triangles need to be returned. 
+ * 
+ * This is done via the PairExpandFunctionT TFunction, which returns a FIndex2i for
+ * a given integer. This pair must be either (a,invalid), (a, b), or (invalid, invalid),
+ * where invalid is integer < 0
+ */
+template<typename InputIteratorT>
+class TPairExpandIterator
+{
+	using PairExpandFunctionT = TFunction<FIndex2i(int)>;
+
+public:
+	inline TPairExpandIterator() { }
+
+	inline bool operator==(const TPairExpandIterator& Other) const
+	{
+		return Cur == Other.Cur;
+	}
+	inline bool operator!=(const TPairExpandIterator & Other) const
+	{
+		return Cur != Other.Cur;
+	}
+
+	inline int operator*() const
+	{
+		return CurValue;
+	}
+
+	inline const TPairExpandIterator& operator++() 		// prefix
+	{
+		goto_next();
+		return *this;
+	}
+
+	inline void goto_next()
+	{
+		while (Cur != End)
+		{
+			if (CurPairI == 0)
+			{
+				CurPair = PairFunc(*Cur);
+				if (CurPair.A >= 0)
+				{
+					CurValue = CurPair.A;
+					CurPairI = 1;	// want to take second branch
+					return;			// let caller see value
+				}
+				else
+				{
+					CurPairI = 0;
+					++Cur;  // done with this base value
+				}
+			}
+			else if (CurPairI == 1)
+			{
+				if (CurPair.B >= 0)
+				{
+					CurValue = CurPair.B;
+					CurPairI = 2;	// want to take third branch
+					return;		// let caller see value
+				}
+				else
+				{
+					CurPairI = 0;
+					++Cur;  // done with this base value
+				}
+			}
+			else
+			{
+				CurPairI = 0;
+				++Cur;  // done with this base value
+			}
+		}
+	}
+
+	inline TPairExpandIterator(const InputIteratorT& CurItr, const InputIteratorT& EndItr, const PairExpandFunctionT& PairFuncIn)
+	{
+		Cur = CurItr;
+		End = EndItr;
+		PairFunc = PairFuncIn;
+		CurPairI = 0;
+		goto_next();
+	}
+
+	InputIteratorT Cur;
+	InputIteratorT End;
+	FIndex2i CurPair;
+	int CurValue;
+	int CurPairI;
+	PairExpandFunctionT PairFunc;
+};
+
+
+
+
+/**
+ * Generic "enumerable" object that provides begin/end semantics for an TPairExpandIterator suitable for use with range-based for.
+ * You can either provide begin/end iterators, or another "enumerable" object that has begin()/end() functions.
+ */
+template<typename InputIteratorT>
+class TPairExpandEnumerable
+{
+	using ExpandFunctionT = TFunction<FIndex2i(int)>;
+	using ExpandIteratorT = TPairExpandIterator<InputIteratorT>;
+
+public:
+	ExpandFunctionT ExpandFunc;
+	InputIteratorT BeginItr, EndItr;
+
+	TPairExpandEnumerable(const InputIteratorT& BeginIn, const InputIteratorT& EndIn, ExpandFunctionT ExpandFuncIn)
+	{
+		this->BeginItr = BeginIn;
+		this->EndItr = EndIn;
+		this->ExpandFunc = ExpandFuncIn;
+	}
+
+	template<typename IteratorSource>
+	TPairExpandEnumerable(const IteratorSource& Source, ExpandFunctionT ExpandFuncIn)
+	{
+		this->BeginItr = Source.begin();
+		this->EndItr = Source.end();
+		this->ExpandFunc = ExpandFuncIn;
+	}
+
+	ExpandIteratorT begin()
+	{
+		return ExpandIteratorT(BeginItr, EndItr, ExpandFunc);
+	}
+
 	ExpandIteratorT end()
 	{
 		return ExpandIteratorT(EndItr, EndItr, ExpandFunc);
