Removed unused Engine include in MeshSolverUtilities

diff --git a/Engine/Plugins/Experimental/GeometryProcessing/Source/MeshSolverUtilities/MeshSolverUtilities.Build.cs b/Engine/Plugins/Experimental/GeometryProcessing/Source/MeshSolverUtilities/MeshSolverUtilities.Build.cs
--- a/Engine/Plugins/Experimental/GeometryProcessing/Source/MeshSolverUtilities/MeshSolverUtilities.Build.cs
+++ b/Engine/Plugins/Experimental/GeometryProcessing/Source/MeshSolverUtilities/MeshSolverUtilities.Build.cs
@@ -20,7 +20,6 @@ public class MeshSolverUtilities : ModuleRules
             new string[]
             {
                 "CoreUObject",
-                "Engine",
 				"GeometricObjects",
 				"DynamicMesh"
                 //"GeometricObjects",
