MeshPlaneCut & SimpleHoleFiller -- for now just following the c# implementations pretty directly

also fixes error handling for MeshEditor's StitchVertexLoopsMinimal (which I referenced for the hole filler's error handling, but looked broken)

diff --git a/Engine/Plugins/Experimental/GeometryProcessing/Source/DynamicMesh/Private/DynamicMeshEditor.cpp b/Engine/Plugins/Experimental/GeometryProcessing/Source/DynamicMesh/Private/DynamicMeshEditor.cpp
--- a/Engine/Plugins/Experimental/GeometryProcessing/Source/DynamicMesh/Private/DynamicMeshEditor.cpp
+++ b/Engine/Plugins/Experimental/GeometryProcessing/Source/DynamicMesh/Private/DynamicMeshEditor.cpp
@@ -85,10 +85,15 @@ bool FDynamicMeshEditor::StitchVertexLoopsMinimal(const TArray<int>& Loop1, cons
 
 operation_failed:
 	// remove what we added so far
-	if (i > 0) 
+	if (ResultOut.NewQuads.Num()) 
 	{
-		ResultOut.NewTriangles.SetNum(2 * i + 1);
-		if (RemoveTriangles(ResultOut.NewTriangles, false) == false)
+		TArray<int> Triangles; Triangles.Reserve(2*ResultOut.NewQuads.Num());
+		for (const FIndex2i& QuadTriIndices : ResultOut.NewQuads)
+		{
+			Triangles.Add(QuadTriIndices.A);
+			Triangles.Add(QuadTriIndices.B);
+		}
+		if (!RemoveTriangles(Triangles, false))
 		{
 			checkf(false, TEXT("FDynamicMeshEditor::StitchLoop: failed to add all triangles, and also failed to back out changes."));
 		}
@@ -98,6 +103,45 @@ operation_failed:
 
 
 
+bool FDynamicMeshEditor::AddTriangleFan_OrderedVertexLoop(int CenterVertex, const TArray<int>& VertexLoop, int GroupID, FDynamicMeshEditResult& ResultOut)
+{
+	if (GroupID == -1)
+	{
+		GroupID = Mesh->AllocateTriangleGroup();
+		ResultOut.NewGroups.Add(GroupID);
+	}
+
+	int N = VertexLoop.Num();
+	ResultOut.NewTriangles.Reserve(N);
+
+	int i = 0;
+	for (i = 0; i < N; ++i)
+	{
+		int A = VertexLoop[i];
+		int B = VertexLoop[(i + 1) % N];
+
+		FIndex3i NewT(CenterVertex, B, A);
+		int NewTID = Mesh->AppendTriangle(NewT, GroupID);
+		if (NewTID < 0)
+		{
+			goto operation_failed;
+		}
+
+		ResultOut.NewTriangles.Add(NewTID);
+	}
+
+	return true;
+
+operation_failed:
+	// remove what we added so far
+	if (!RemoveTriangles(ResultOut.NewTriangles, false))
+	{
+		checkf(false, TEXT("FDynamicMeshEditor::AddTriangleFan: failed to add all triangles, and also failed to back out changes."));
+	}
+	return false;
+}
+
+
 
 bool FDynamicMeshEditor::RemoveTriangles(const TArray<int>& Triangles, bool bRemoveIsolatedVerts)
 {
diff --git a/Engine/Plugins/Experimental/GeometryProcessing/Source/DynamicMesh/Private/Operations/MeshPlaneCut.cpp b/Engine/Plugins/Experimental/GeometryProcessing/Source/DynamicMesh/Private/Operations/MeshPlaneCut.cpp
new file mode 100644
--- /dev/null
+++ b/Engine/Plugins/Experimental/GeometryProcessing/Source/DynamicMesh/Private/Operations/MeshPlaneCut.cpp
@@ -0,0 +1,220 @@
+// Copyright 1998-2019 Epic Games, Inc. All Rights Reserved.
+
+#include "Operations/MeshPlaneCut.h"
+#include "Operations/SimpleHoleFiller.h"
+#include "MeshNormals.h"
+#include "DynamicMeshEditor.h"
+#include "MathUtil.h"
+
+#include "Async/ParallelFor.h"
+
+bool FMeshPlaneCut::Cut()
+{
+	double InvalidDist = -FMathd::MaxReal;
+
+	// TODO: handle selections
+	//MeshEdgeSelection CutEdgeSet = null;
+	//MeshVertexSelection CutVertexSet = null;
+	//if (CutFaceSet != null) {
+	//	CutEdgeSet = new MeshEdgeSelection(Mesh, CutFaceSet);
+	//	CutVertexSet = new MeshVertexSelection(Mesh, CutEdgeSet);
+	//}
+
+	// compute Signs
+	int MaxVID = Mesh->MaxVertexID();
+	TArray<double> Signs;
+	Signs.SetNum(MaxVID);
+
+	bool bNoParallel = false;
+	ParallelFor(MaxVID, [&](int32 VID)
+	{
+		if (Mesh->IsVertex(VID))
+		{
+			Signs[VID] = (Mesh->GetVertex(VID) - PlaneOrigin).Dot(PlaneNormal);
+		}
+		else
+		{
+			Signs[VID] = InvalidDist;
+		}
+	},
+	bNoParallel);
+
+	TSet<int> ZeroEdges, ZeroVertices, OnCutEdges;
+
+	// have to skip processing of new edges. If edge id
+	// is > max at start, is new. Otherwise if in NewEdges list, also new.
+	int MaxEID = Mesh->MaxEdgeID();
+	TSet<int> NewEdges;
+
+	FDynamicMesh3::edge_iterator EdgeItr = Mesh->EdgeIndicesItr();
+	// TODO: selection logic
+	//IEnumerable<int> edgeItr = Interval1i.Range(MaxEID);
+	//if (CutEdgeSet != null)
+	//	edgeItr = CutEdgeSet;
+
+	// cut existing edges with plane, using edge split
+	for (int EID : EdgeItr)
+	{
+		if (!Mesh->IsEdge(EID))
+		{
+			continue;
+		}
+		if (EID >= MaxEID || NewEdges.Contains(EID))
+		{
+			continue;
+		}
+
+		FIndex2i ev = Mesh->GetEdgeV(EID);
+		double f0 = Signs[ev.A];
+		double f1 = Signs[ev.B];
+
+		// If both Signs are 0, this edge is on-contour
+		// If one sign is 0, that vertex is on-contour
+		int n0 = (FMathd::Abs(f0) < FMathd::Epsilon) ? 1 : 0;
+		int n1 = (FMathd::Abs(f1) < FMathd::Epsilon) ? 1 : 0;
+		if (n0 + n1 > 0)
+		{
+			if (n0 + n1 == 2)
+			{
+				ZeroEdges.Add(EID);
+			}
+			else
+			{
+				ZeroVertices.Add((n0 == 1) ? ev[0] : ev[1]);
+			}
+			continue;
+		}
+
+		// no crossing
+		if (f0 * f1 > 0)
+		{
+			continue;
+		}
+
+		FDynamicMesh3::FEdgeSplitInfo splitInfo;
+		EMeshResult result = Mesh->SplitEdge(EID, splitInfo);
+		ensureMsgf(result == EMeshResult::Ok, TEXT("FMeshPlaneCut::Cut: failed to SplitEdge")); // edge split really shouldn't fail
+
+		// SplitEdge just bisects edge; now move created vertex to plane intersection pt
+		double t = f0 / (f0 - f1);
+		FVector3d newPos = (1 - t) * Mesh->GetVertex(ev.A) + (t) * Mesh->GetVertex(ev.B);
+		Mesh->SetVertex(splitInfo.NewVertex, newPos);
+
+		NewEdges.Add(splitInfo.NewEdges.A); // TODO: why not OnCutEdges.Add(splitInfo.NewEdges.A);  ???
+		NewEdges.Add(splitInfo.NewEdges.B); OnCutEdges.Add(splitInfo.NewEdges.B);
+		if (splitInfo.NewEdges.C != FDynamicMesh3::InvalidID)
+		{
+			NewEdges.Add(splitInfo.NewEdges.C); OnCutEdges.Add(splitInfo.NewEdges.C);
+		}
+	}
+
+	// remove one-rings of all positive-side vertices. 
+	// @todo handle selection logic
+	//IEnumerable<int> vertexSet = Interval1i.Range(MaxVID);
+	//if (CutVertexSet != null)
+	//	vertexSet = CutVertexSet;
+	for (int VID : Mesh->VertexIndicesItr())
+	{ 
+		if (VID < Signs.Num() && Signs[VID] > FMathd::Epsilon)
+		{
+			Mesh->RemoveVertex(VID, true, false);
+		}
+	}
+
+	// collapse degenerate edges if we got em
+	if (bCollapseDegenerateEdgesOnCut)
+	{
+		CollapseDegenerateEdges(OnCutEdges, ZeroEdges);
+	}
+
+
+	// ok now we extract boundary loops, but restricted
+	// to either the zero-edges we found, or the edges we created! bang!!
+	
+	FMeshBoundaryLoops Loops(Mesh, false);
+	Loops.EdgeFilterFunc = [&OnCutEdges, &ZeroEdges](int EID)
+	{
+		return OnCutEdges.Contains(EID) || ZeroEdges.Contains(EID);
+	};
+	bool bFoundLoops = Loops.Compute();
+
+	if (bFoundLoops)
+	{
+		CutLoops = Loops.Loops;
+		CutSpans = Loops.Spans;
+		CutLoopsFailed = false;
+		FoundOpenSpans = CutSpans.Num() > 0;
+	}
+	else
+	{
+		CutLoops.Empty();
+		CutLoopsFailed = true;
+	}
+
+	return !CutLoopsFailed;
+}
+
+
+
+void FMeshPlaneCut::CollapseDegenerateEdges(const TSet<int>& OnCutEdges, const TSet<int>& ZeroEdges)
+{
+	TSet<int> Sets[2] { OnCutEdges, ZeroEdges };
+
+	double Tol2 = DegenerateEdgeTol * DegenerateEdgeTol;
+	FVector3d A, B;
+	int Collapsed = 0;
+	do
+	{
+		Collapsed = 0;
+		for (int SetIdx = 0; SetIdx < 2; SetIdx++)
+		{
+			for (int EID : Sets[SetIdx])
+			{
+				if (!Mesh->IsEdge(EID))
+				{
+					continue;
+				}
+				Mesh->GetEdgeV(EID, A, B);
+				if (A.DistanceSquared(B) > Tol2)
+				{
+					continue;
+				}
+
+				FIndex2i EV = Mesh->GetEdgeV(EID);
+				FDynamicMesh3::FEdgeCollapseInfo CollapseInfo;
+				EMeshResult Result = Mesh->CollapseEdge(EV.A, EV.B, CollapseInfo);
+				if (Result == EMeshResult::Ok)
+				{
+					Collapsed++;
+				}
+			}
+		}
+	} while (Collapsed != 0);
+}
+
+
+
+
+
+bool FMeshPlaneCut::FillHoles(int ConstantGroupID)
+{
+	bool bAllOk = true;
+
+	HoleFillTriangles.Empty();
+
+	for (const FEdgeLoop& Loop : CutLoops)
+	{
+		FSimpleHoleFiller Filler(Mesh, Loop);
+		int GID = ConstantGroupID >= 0 ? ConstantGroupID : Mesh->AllocateTriangleGroup();
+		if (Filler.Fill(GID))
+		{
+			HoleFillTriangles.Add(Filler.NewTriangles);
+		}
+		else
+		{
+			bAllOk = false;
+		}
+	}
+
+	return bAllOk;
+}
\ No newline at end of file
diff --git a/Engine/Plugins/Experimental/GeometryProcessing/Source/DynamicMesh/Private/Operations/SimpleHoleFiller.cpp b/Engine/Plugins/Experimental/GeometryProcessing/Source/DynamicMesh/Private/Operations/SimpleHoleFiller.cpp
new file mode 100644
--- /dev/null
+++ b/Engine/Plugins/Experimental/GeometryProcessing/Source/DynamicMesh/Private/Operations/SimpleHoleFiller.cpp
@@ -0,0 +1,52 @@
+// Copyright 1998-2019 Epic Games, Inc. All Rights Reserved.
+
+#include "Operations/SimpleHoleFiller.h"
+#include "DynamicMeshEditor.h"
+
+bool FSimpleHoleFiller::Fill(int GroupID)
+{
+	if (Loop.GetVertexCount() < 3)
+	{
+		return false;
+	}
+
+	// this just needs one triangle
+	if (Loop.GetVertexCount() == 3)
+	{
+		FIndex3i Tri(Loop.Vertices[0], Loop.Vertices[2], Loop.Vertices[1]);
+		int NewTID = Mesh->AppendTriangle(Tri, GroupID);
+		if (NewTID < 0)
+		{
+			return false;
+		}
+		NewTriangles = { NewTID };
+		NewVertex = FDynamicMesh3::InvalidID;
+		return true;
+	}
+
+	// [TODO] 4-case? could check nbr normals to figure out best internal edge...
+
+	// compute centroid
+	FVector3d c = FVector3d::Zero();
+	for (int i = 0; i < Loop.GetVertexCount(); ++i)
+	{
+		c += Mesh->GetVertex(Loop.Vertices[i]);
+	}
+	c *= 1.0 / Loop.GetVertexCount();
+
+	// add centroid vtx
+	NewVertex = Mesh->AppendVertex(c);
+
+	// stitch triangles
+	FDynamicMeshEditor Editor(Mesh);
+	FDynamicMeshEditResult AddFanResult;
+	if (!Editor.AddTriangleFan_OrderedVertexLoop(NewVertex, Loop.Vertices, GroupID, AddFanResult))
+	{
+		Mesh->RemoveVertex(NewVertex, true, false);
+		NewVertex = FDynamicMesh3::InvalidID;
+		return false;
+	}
+	NewTriangles = AddFanResult.NewTriangles;
+
+	return true;
+}
\ No newline at end of file
diff --git a/Engine/Plugins/Experimental/GeometryProcessing/Source/DynamicMesh/Public/DynamicMeshEditor.h b/Engine/Plugins/Experimental/GeometryProcessing/Source/DynamicMesh/Public/DynamicMeshEditor.h
--- a/Engine/Plugins/Experimental/GeometryProcessing/Source/DynamicMesh/Public/DynamicMeshEditor.h
+++ b/Engine/Plugins/Experimental/GeometryProcessing/Source/DynamicMesh/Public/DynamicMeshEditor.h
@@ -141,6 +141,17 @@ public:
 	bool StitchVertexLoopsMinimal(const TArray<int>& VertexLoop1, const TArray<int>& VertexLoop2, FDynamicMeshEditResult& ResultOut);
 
 
+	/**
+	 * Fill hole with a triangle fan given an existing (unconnected) center vertex and
+	 * an ordered loop of boundary vertices on the hole border.
+	 * @param CenterVertex Index of floating vertex in the center of the hole
+	 * @param VertexLoop Indices of vertices on the boundary of the hole, in order
+	 * @param ResultOut lists of newly created triangles
+	 * @return true if operation succeeded.  If a failure occurs, any added triangles are removed via RemoveTriangles.
+	 */
+	bool AddTriangleFan_OrderedVertexLoop(int CenterVertex, const TArray<int>& VertexLoop, int GroupID, FDynamicMeshEditResult& ResultOut);
+
+
 	/**
 	 * Duplicate triangles of a mesh. This duplicates the current groups and also any attributes existing on the triangles.
 	 * @param Triangles the triangles to duplicate
diff --git a/Engine/Plugins/Experimental/GeometryProcessing/Source/DynamicMesh/Public/Operations/MeshPlaneCut.h b/Engine/Plugins/Experimental/GeometryProcessing/Source/DynamicMesh/Public/Operations/MeshPlaneCut.h
new file mode 100644
--- /dev/null
+++ b/Engine/Plugins/Experimental/GeometryProcessing/Source/DynamicMesh/Public/Operations/MeshPlaneCut.h
@@ -0,0 +1,91 @@
+// Copyright 1998-2019 Epic Games, Inc. All Rights Reserved.
+
+// Port of geometry3Sharp MeshPlaneCut
+
+#pragma once
+
+#include "MathUtil.h"
+#include "VectorTypes.h"
+#include "GeometryTypes.h"
+#include "MeshBoundaryLoops.h"
+
+
+class FDynamicMesh3;
+
+
+/**
+ * Cut the Mesh with the Plane. The *positive* side, ie (p-o).n > 0, is removed.
+ * If possible, returns boundary loop(s) along cut
+ * (this will fail if cut intersected with holes in mesh).
+ * Also FillHoles() for a topological fill. Or use CutLoops and fill yourself.
+ * 
+ * Algorithm is:
+ *    1) find all edge crossings
+ *    2) Do edge splits at crossings
+ *    3) delete all vertices on positive side
+ *    4) (optionally) collapse any degenerate boundary edges 
+ *    5) find loops through valid boundary edges (ie connected to splits, or on-plane edges)
+ * 
+ * @todo Also discard any triangles with all vertex distances < epsilon.
+ */
+class DYNAMICMESH_API FMeshPlaneCut
+{
+public:
+
+	//
+	// Inputs
+	//
+	FDynamicMesh3 *Mesh;
+	FVector3d PlaneOrigin, PlaneNormal;
+	
+
+	bool bCollapseDegenerateEdgesOnCut = true;
+	double DegenerateEdgeTol = FMathd::ZeroTolerance;
+
+	// TODO support optionally restricting plane cut to a mesh selection
+	//MeshFaceSelection CutFaceSet;
+
+	//
+	// Outputs
+	//
+	TArray<FEdgeLoop> CutLoops;
+	TArray<FEdgeSpan> CutSpans;
+	bool CutLoopsFailed = false;		// set to true if we could not compute cut loops/spans
+	bool FoundOpenSpans = false;     // set to true if we found open spans in cut
+
+	// Note: In practice this is 1:1 with CutLoops if all the hole fills succeed due to simple hole filling, but this will not be true if/when we add better hole filling methods, or if any of the hole fills fails
+	// please do not write anything relying on a correspondence between the two
+	TArray<TArray<int>> HoleFillTriangles;
+
+
+public:
+
+	/**
+	 *  Cut mesh with plane. Assumption is that plane normal is Z value.
+	 */
+	FMeshPlaneCut(FDynamicMesh3* Mesh, FVector3d Origin, FVector3d Normal) : Mesh(Mesh), PlaneOrigin(Origin), PlaneNormal(Normal)
+	{
+	}
+	virtual ~FMeshPlaneCut() {}
+	
+	/**
+	 * @return EOperationValidationResult::Ok if we can apply operation, or error code if we cannot
+	 */
+	virtual EOperationValidationResult Validate()
+	{
+		// @todo validate inputs
+		return EOperationValidationResult::Ok;
+	}
+
+	virtual bool Cut();
+
+	/**
+	 *  Fill all loops with FSimpleHoleFiller. @todo: Add more advanced hole-fill options
+	 */
+	virtual bool FillHoles(int constantGroupID = -1);
+
+	
+protected:
+
+	void CollapseDegenerateEdges(const TSet<int>& OnCutEdges, const TSet<int>& ZeroEdges);
+};
diff --git a/Engine/Plugins/Experimental/GeometryProcessing/Source/DynamicMesh/Public/Operations/SimpleHoleFiller.h b/Engine/Plugins/Experimental/GeometryProcessing/Source/DynamicMesh/Public/Operations/SimpleHoleFiller.h
new file mode 100644
--- /dev/null
+++ b/Engine/Plugins/Experimental/GeometryProcessing/Source/DynamicMesh/Public/Operations/SimpleHoleFiller.h
@@ -0,0 +1,58 @@
+// Copyright 1998-2019 Epic Games, Inc. All Rights Reserved.
+
+// Port of geometry3Sharp SimpleHoleFiller
+
+#pragma once
+
+#include "MathUtil.h"
+#include "VectorTypes.h"
+#include "GeometryTypes.h"
+#include "MeshBoundaryLoops.h"
+
+
+class FDynamicMesh3;
+
+
+/**
+ * Fill an EdgeLoop hole with a triangle fan connected to a new vertex at the centroid
+ */
+class DYNAMICMESH_API FSimpleHoleFiller
+{
+public:
+	//
+	// Inputs
+	//
+	FDynamicMesh3 *Mesh;
+	FEdgeLoop Loop;
+
+	//
+	// Outputs
+	//
+	int NewVertex;
+	TArray<int> NewTriangles;
+
+public:
+	/**
+	 *  Construct simple hole filler (just adds a central vertex and a triangle fan)
+	 */
+	FSimpleHoleFiller(FDynamicMesh3* Mesh, FEdgeLoop Loop) : Mesh(Mesh), Loop(Loop)
+	{
+	}
+	virtual ~FSimpleHoleFiller() {}
+	
+	
+	/**
+	 * @return EOperationValidationResult::Ok if we can apply operation, or error code if we cannot
+	 */
+	virtual EOperationValidationResult Validate()
+	{
+		if (!Loop.IsBoundaryLoop(Mesh))
+		{
+			return EOperationValidationResult::Failed_UnknownReason;
+		}
+
+		return EOperationValidationResult::Ok;
+	}
+
+	virtual bool Fill(int GroupID = -1);	
+};
