Removed unused Engine dependency from MeshConversion module

diff --git a/Engine/Plugins/Experimental/GeometryProcessing/Source/MeshConversion/MeshConversion.Build.cs b/Engine/Plugins/Experimental/GeometryProcessing/Source/MeshConversion/MeshConversion.Build.cs
--- a/Engine/Plugins/Experimental/GeometryProcessing/Source/MeshConversion/MeshConversion.Build.cs
+++ b/Engine/Plugins/Experimental/GeometryProcessing/Source/MeshConversion/MeshConversion.Build.cs
@@ -20,8 +20,7 @@ public class MeshConversion : ModuleRules
         PrivateDependencyModuleNames.AddRange(
             new string[]
             {
-                "CoreUObject",
-                "Engine"
+                "CoreUObject"
 			}
 		);
     }
