add initial TTransform3 template, supports basic transforms, along with test comparing against FTransform. Simplify FProgressCancel. Add FVector2/3 vector divide operator.
#rnx

diff --git a/Engine/Plugins/Experimental/GeometryProcessing/Source/GeometricObjects/Private/Tests/TransformTypes.spec.cpp b/Engine/Plugins/Experimental/GeometryProcessing/Source/GeometricObjects/Private/Tests/TransformTypes.spec.cpp
new file mode 100644
--- /dev/null
+++ b/Engine/Plugins/Experimental/GeometryProcessing/Source/GeometricObjects/Private/Tests/TransformTypes.spec.cpp
@@ -0,0 +1,69 @@
+// Copyright 1998-2019 Epic Games, Inc. All Rights Reserved.
+
+#include "Misc/AutomationTest.h"
+#include "TransformTypes.h"
+
+#if WITH_DEV_AUTOMATION_TESTS
+
+BEGIN_DEFINE_SPEC(FTransformTypesSpec,
+	"GeometryProcessing.Unit",
+	EAutomationTestFlags::ProductFilter | EAutomationTestFlags::ApplicationContextMask)
+END_DEFINE_SPEC(FTransformTypesSpec)
+
+
+FVector MakeRandomVector()
+{
+	FVector Temp(FMath::RandRange(-1.0f, 1.0f), FMath::RandRange(-1.0f, 1.0f), FMath::RandRange(-1.0f, 1.0f));
+	Temp.Normalize();
+	return Temp;
+}
+
+void FTransformTypesSpec::Define()
+{
+	Describe("FTransform3<float>", [this]()
+	{
+		It("FTransform_Comparisons", [this]()
+		{
+			FMath::RandInit(313377);
+			for (int k = 0; k < 1000; ++k)
+			{
+				FQuat Rotation(MakeRandomVector(), FMath::RandRange(-100.0f, 100.0f));
+				FVector Translation = MakeRandomVector();
+				FVector Scale = FVector(1.1f,1.1f,1.1f) + MakeRandomVector();		// need to stay away from zero scale or floats will explode and precision is too low
+				FTransform UETransform(Rotation, Translation, Scale);
+				FTransform3f GPTransform(UETransform);
+
+				float tol = 0.001f;		// float precision is very bad...some float comparisons will fail at 0.0001f
+				for (int j = 0; j < 1000; ++j)
+				{
+					FVector TestVector = MakeRandomVector();
+
+					FVector UETransformP = UETransform.TransformPosition(TestVector);
+					FVector GPTransformP = (FVector)GPTransform.TransformPosition((FVector3f)TestVector);
+					TestEqual("TransformPosition", UETransformP, GPTransformP, tol);
+
+					FVector UETransformV = UETransform.TransformVector(TestVector);
+					FVector GPTransformV = (FVector)GPTransform.TransformVector((FVector3f)TestVector);
+					TestEqual("TransformVector", UETransformV, GPTransformV, tol);
+
+					FVector UETransformNS = UETransform.TransformVectorNoScale(TestVector);
+					FVector GPTransformNS = (FVector)GPTransform.TransformVectorNoScale((FVector3f)TestVector);
+					TestEqual("TransformVectorNoScale", UETransformNS, GPTransformNS, tol);
+
+					FVector UEInverseP = UETransform.InverseTransformPosition(TestVector);
+					FVector GPInverseP = (FVector)GPTransform.InverseTransformPosition((FVector3f)TestVector);
+					TestEqual("InverseTransformPosition", UEInverseP, GPInverseP, tol);
+
+					FVector UEInverseV = UETransform.InverseTransformVector(TestVector);
+					FVector GPInverseV = (FVector)GPTransform.InverseTransformVector((FVector3f)TestVector);
+					TestEqual("InverseTransformVector", UEInverseV, GPInverseV, tol);
+
+					FVector UEInverseNS = UETransform.InverseTransformVectorNoScale(TestVector);
+					FVector GPInverseNS = (FVector)GPTransform.InverseTransformVectorNoScale((FVector3f)TestVector);
+					TestEqual("InverseTransformVectorNoScale", UEInverseNS, GPInverseNS, tol);
+				}
+			}
+		});
+	});
+}
+#endif
diff --git a/Engine/Plugins/Experimental/GeometryProcessing/Source/GeometricObjects/Public/TransformTypes.h b/Engine/Plugins/Experimental/GeometryProcessing/Source/GeometricObjects/Public/TransformTypes.h
new file mode 100644
--- /dev/null
+++ b/Engine/Plugins/Experimental/GeometryProcessing/Source/GeometricObjects/Public/TransformTypes.h
@@ -0,0 +1,164 @@
+// Copyright 1998-2019 Epic Games, Inc. All Rights Reserved.
+
+#pragma once
+
+#include "CoreMinimal.h"
+#include "VectorTypes.h"
+#include "Quaternion.h"
+
+/**
+ * TTransform3 is a double/float templated version of standard UE FTransform.
+ */
+template<typename RealType>
+class TTransform3
+{
+protected:
+	TQuaternion<RealType> Rotation;
+	FVector3<RealType> Translation;
+	FVector3<RealType> Scale3D;
+
+public:
+
+	TTransform3()
+	{
+		Rotation = TQuaternion<RealType>::Identity();
+		Translation = FVector3<RealType>::Zero();
+		Scale3D = FVector3<RealType>::One();
+	}
+
+	TTransform3(const TQuaternion<RealType>& RotationIn, const FVector3<RealType>& TranslationIn, const FVector3<RealType>& ScaleIn)
+	{
+		Rotation = RotationIn;
+		Translation = TranslationIn;
+		Scale3D = ScaleIn;
+	}
+
+	explicit TTransform3(const TQuaternion<RealType>& RotationIn, const FVector3<RealType>& TranslationIn)
+	{
+		Rotation = RotationIn;
+		Translation = TranslationIn;
+		Scale3D = FVector3<RealType>::One();
+	}
+
+	explicit TTransform3(const FVector3<RealType>& TranslationIn)
+	{
+		Rotation = TQuaternion<RealType>::Identity();
+		Translation = TranslationIn;
+		Scale3D = FVector3<RealType>::One();
+	}
+
+	explicit TTransform3(const FTransform& Transform)
+	{
+		Rotation = TQuaternion<RealType>(Transform.GetRotation());
+		Translation = FVector3<RealType>(Transform.GetTranslation());
+		Scale3D = FVector3<RealType>(Transform.GetScale3D());
+	}
+
+	/**
+	 * @return identity transform, IE no rotation, zero origin, unit scale
+	 */
+	static TTransform3<RealType> Identity()
+	{
+		return TTransform3<RealType>(TQuaternion<RealType>::Identity(), FVector3<RealType>::Zero(), FVector3<RealType>::One());
+	}
+
+	/**
+	 * @return this transform converted to FTransform 
+	 */
+	explicit operator FTransform() const
+	{
+		return FTransform((FQuat)Rotation, (FVector)Translation, (FVector)Scale3D);
+	}
+
+
+	/**
+	 * @return input point with QST transformation applied, ie QST(P) = Rotate(Scale*P) + Translate
+	 */
+	FVector3<RealType> TransformPosition(const FVector3<RealType>& P) const
+	{
+		//Transform using QST is following
+		//QST(P) = Q.Rotate(S*P) + T where Q = quaternion, S = scale, T = translation
+		return Rotation * (Scale3D*P) + Translation;
+	}
+
+	/**
+	 * @return input vector with QS transformation applied, ie QS(V) = Rotate(Scale*V)
+	 */
+	FVector3<RealType> TransformVector(const FVector3<RealType>& V) const
+	{
+		return Rotation * (Scale3D*V);
+	}
+
+	/**
+	 * @return input vector with Q transformation applied, ie Q(V) = Rotate(V)
+	 */
+	FVector3<RealType> TransformVectorNoScale(const FVector3<RealType>& V) const
+	{
+		return Rotation * V;
+	}
+
+	/**
+	 * @return input vector with inverse-QST transformation applied, ie QSTinv(P) = InverseScale(InverseRotate(P - Translate))
+	 */
+	FVector3<RealType> InverseTransformPosition(const FVector3<RealType> &P) const
+	{
+		return GetSafeScaleReciprocal(Scale3D) * Rotation.InverseMultiply(P - Translation);
+	}
+
+	/**
+	 * @return input vector with inverse-QS transformation applied, ie QSinv(V) = InverseScale(InverseRotate(V))
+	 */
+	FVector3<RealType> InverseTransformVector(const FVector3<RealType> &V) const
+	{
+		return GetSafeScaleReciprocal(Scale3D) * Rotation.InverseMultiply(V);
+	}
+
+
+	/**
+	 * @return input vector with inverse-Q transformation applied, ie Qinv(V) = InverseRotate(V)
+	 */
+	FVector3<RealType> InverseTransformVectorNoScale(const FVector3<RealType> &V) const
+	{
+		return Rotation.InverseMultiply(V);
+	}
+
+	
+
+	static FVector3<RealType> GetSafeScaleReciprocal(const FVector3<RealType>& InScale, RealType Tolerance = TMathUtil<RealType>::ZeroTolerance)
+	{
+		FVector3<RealType> SafeReciprocalScale;
+		if (TMathUtil<RealType>::Abs(InScale.X) <= Tolerance)
+		{
+			SafeReciprocalScale.X = (RealType)0;
+		}
+		else
+		{
+			SafeReciprocalScale.X = (RealType)1 / InScale.X;
+		}
+
+		if (TMathUtil<RealType>::Abs(InScale.Y) <= Tolerance)
+		{
+			SafeReciprocalScale.Y = (RealType)0;
+		}
+		else
+		{
+			SafeReciprocalScale.Y = (RealType)1 / InScale.Y;
+		}
+
+		if (TMathUtil<RealType>::Abs(InScale.Z) <= Tolerance)
+		{
+			SafeReciprocalScale.Z = (RealType)0;
+		}
+		else
+		{
+			SafeReciprocalScale.Z = (RealType)1 / InScale.Z;
+		}
+
+		return SafeReciprocalScale;
+	}
+
+
+
+};
+typedef TTransform3<float> FTransform3f;
+typedef TTransform3<double> FTransform3d;
diff --git a/Engine/Plugins/Experimental/GeometryProcessing/Source/GeometricObjects/Public/Util/ProgressCancel.h b/Engine/Plugins/Experimental/GeometryProcessing/Source/GeometricObjects/Public/Util/ProgressCancel.h
--- a/Engine/Plugins/Experimental/GeometryProcessing/Source/GeometricObjects/Public/Util/ProgressCancel.h
+++ b/Engine/Plugins/Experimental/GeometryProcessing/Source/GeometricObjects/Public/Util/ProgressCancel.h
@@ -5,61 +5,19 @@
 #pragma once
 
 
-/**
- * ICancelSource is an object that provides a cancel function
- */
-class ICancelSource
-{
-public:
-	virtual ~ICancelSource() {}
-
-	/**
-	 * @return true if object wishes to cancel expensive operation
-	 */
-	virtual bool Cancelled() = 0;
-};
-
-
-/**
- * FCancelFunction uses a TFunction to implement the ICancelSource interface
- */
-class FCancelFunction : public ICancelSource
-{
-public:
-	virtual ~FCancelFunction() {}
-
-	/** function that returns true if object wishes to cancel operation */
-	TFunction<bool()> CancelF;
-
-	FCancelFunction(const TFunction<bool()> & cancelF) : CancelF(cancelF)
-	{
-	}
-
-	/**
-	 * @return true if object wishes to cancel expensive operation
-	 */
-	bool Cancelled() { return CancelF(); }
-};
-
-
-
 /**
  * FProgressCancel is an obejct that is intended to be passed to long-running
  * computes to do two things:
- * 1) provide progress info back to caller (not ported yet)
+ * 1) provide progress info back to caller (not implemented yet)
  * 2) allow caller to cancel the computation
  */
 class FProgressCancel
 {
-public:
-	TSharedPtr<ICancelSource> Source;
-
+private:
 	bool WasCancelled = false;  // will be set to true if CancelF() ever returns true
 
-	FProgressCancel(TSharedPtr<ICancelSource> source)
-	{
-		Source = source;
-	}
+public:
+	TFunction<bool()> CancelF = []() { return false; };
 
 	/**
 	 * @return true if client would like to cancel operation
@@ -70,7 +28,7 @@ public:
 		{
 			return true;
 		}
-		WasCancelled = Source->Cancelled();
+		WasCancelled = CancelF();
 		return WasCancelled;
 	}
 };
diff --git a/Engine/Plugins/Experimental/GeometryProcessing/Source/GeometricObjects/Public/VectorTypes.h b/Engine/Plugins/Experimental/GeometryProcessing/Source/GeometricObjects/Public/VectorTypes.h
--- a/Engine/Plugins/Experimental/GeometryProcessing/Source/GeometricObjects/Public/VectorTypes.h
+++ b/Engine/Plugins/Experimental/GeometryProcessing/Source/GeometricObjects/Public/VectorTypes.h
@@ -184,6 +184,11 @@ struct FVector3
 		return FVector3<T>(X / Scalar, Y / Scalar, Z / Scalar);
 	}
 
+	inline FVector3<T> operator/(const FVector3<T>& V2) const // component-wise
+	{
+		return FVector3<T>(X / V2.X, Y / V2.Y, Z / V2.Z);
+	}
+
 	inline FVector3<T>& operator+=(const FVector3<T>& V2)
 	{
 		X += V2.X;
@@ -487,6 +492,11 @@ struct FVector2
 		return FVector2(X / Scalar, Y / Scalar);
 	}
 
+	inline FVector2<T> operator/(const FVector2<T>& V2) const // component-wise
+	{
+		return FVector2<T>(X / V2.X, Y / V2.Y);
+	}
+
 	inline FVector2<T>& operator+=(const FVector2<T>& V2)
 	{
 		X += V2.X;
