Properly initialize AABBTree::MeshTimestamp, add AABBTree::IsValid(), and do check/error-handling for unbuilt tree (should not happen but just in case)
#rnx

diff --git a/Engine/Plugins/Experimental/GeometryProcessing/Source/GeometricObjects/Public/Spatial/MeshAABBTree3.h b/Engine/Plugins/Experimental/GeometryProcessing/Source/GeometricObjects/Public/Spatial/MeshAABBTree3.h
--- a/Engine/Plugins/Experimental/GeometryProcessing/Source/GeometricObjects/Public/Spatial/MeshAABBTree3.h
+++ b/Engine/Plugins/Experimental/GeometryProcessing/Source/GeometricObjects/Public/Spatial/MeshAABBTree3.h
@@ -20,7 +20,7 @@ class TMeshAABBTree3 : public IMeshSpatial
 
 protected:
 	TriangleMeshType* Mesh;
-	int MeshTimestamp;
+	int MeshTimestamp = -1;
 	int TopDownLeafMaxTriCount = 4;
 
 public:
@@ -44,6 +44,8 @@ public:
 	void SetMesh(TriangleMeshType* SourceMesh, bool bAutoBuild = true)
 	{
 		Mesh = SourceMesh;
+		MeshTimestamp = -1;
+
 		if (bAutoBuild)
 		{
 			Build();
@@ -55,6 +57,18 @@ public:
 		return Mesh;
 	}
 
+	/**
+	 * @return true if internal timestamp matches mesh timestamp
+	 */
+	bool IsValid() const 
+	{
+		if (MeshTimestamp == Mesh->GetShapeTimestamp())
+		{
+			check(RootIndex >= 0);
+		}
+		return MeshTimestamp == Mesh->GetShapeTimestamp();
+	}
+
 	void Build()
 	{
 		BuildTopDown(false);
@@ -73,6 +87,11 @@ public:
 	virtual int FindNearestTriangle(const FVector3d& P, double& NearestDistSqr, double MaxDist = TNumericLimits<double>::Max()) override
 	{
 		check(MeshTimestamp == Mesh->GetShapeTimestamp());
+		check(RootIndex >= 0);
+		if (RootIndex < 0)
+		{
+			return IndexConstants::InvalidID;
+		}
 
 		NearestDistSqr = (MaxDist < DOUBLE_MAX) ? MaxDist * MaxDist : DOUBLE_MAX;
 		int tNearID = IndexConstants::InvalidID;
@@ -153,6 +172,11 @@ public:
 	virtual int FindNearestHitTriangle(const FRay3d& Ray, double MaxDist = TNumericLimits<double>::Max()) override
 	{
 		check(MeshTimestamp == Mesh->GetShapeTimestamp());
+		check(RootIndex >= 0);
+		if (RootIndex < 0)
+		{
+			return IndexConstants::InvalidID;
+		}
 		// TODO: check( ray_is_normalized)
 
 		// [RMS] note: using float.MaxValue here because we need to use <= to compare Box hit
@@ -247,6 +271,11 @@ public:
 	virtual FIndex2i FindNearestTriangles(TMeshAABBTree3& OtherTree, TFunction<FVector3d(const FVector3d&)> TransformF, double& Distance, double MaxDist = FMathd::MaxReal)
 	{
 		check(MeshTimestamp == Mesh->GetShapeTimestamp());
+		check(RootIndex >= 0);
+		if (RootIndex < 0)
+		{
+			return FIndex2i::Invalid();
+		}
 
 		double NearestSqr = FMathd::MaxReal;
 		if (MaxDist < FMathd::MaxReal)
@@ -287,6 +316,12 @@ public:
 	virtual void DoTraversal(FTreeTraversal& Traversal)
 	{
 		check(MeshTimestamp == Mesh->GetShapeTimestamp());
+		check(RootIndex >= 0);
+		if (RootIndex < 0)
+		{
+			return;
+		}
+
 		TreeTraversalImpl(RootIndex, 0, Traversal);
 	}
 
