rename MeshProjectionTarget to FMeshProjectionTarget, make mesh pointer const, handle spatial queries returning InvalidID
#rnx

diff --git a/Engine/Plugins/Experimental/GeometryProcessing/Source/DynamicMesh/Public/ProjectionTargets.h b/Engine/Plugins/Experimental/GeometryProcessing/Source/DynamicMesh/Public/ProjectionTargets.h
--- a/Engine/Plugins/Experimental/GeometryProcessing/Source/DynamicMesh/Public/ProjectionTargets.h
+++ b/Engine/Plugins/Experimental/GeometryProcessing/Source/DynamicMesh/Public/ProjectionTargets.h
@@ -9,20 +9,24 @@
 
 
 /**
- * MeshProjectionTarget provides an IProjectionTarget interface to a FDynamicMesh + FDynamicMeshAABBTree3
+ * FMeshProjectionTarget provides an IProjectionTarget interface to a FDynamicMesh + FDynamicMeshAABBTree3
  * Use to project points to mesh surface.
  */
-class MeshProjectionTarget : public IOrientedProjectionTarget
+class FMeshProjectionTarget : public IOrientedProjectionTarget
 {
 public:
 	/** The mesh to project onto */
-	FDynamicMesh3* Mesh;
+	const FDynamicMesh3* Mesh = nullptr;
 	/** An AABBTree for Mesh */
-	FDynamicMeshAABBTree3* Spatial;
+	FDynamicMeshAABBTree3* Spatial = nullptr;
 
-	~MeshProjectionTarget() {}
+	~FMeshProjectionTarget() {}
 
-	MeshProjectionTarget(FDynamicMesh3* MeshIn, FDynamicMeshAABBTree3* SpatialIn)
+	FMeshProjectionTarget()
+	{
+	}
+
+	FMeshProjectionTarget(const FDynamicMesh3* MeshIn, FDynamicMeshAABBTree3* SpatialIn)
 	{
 		Mesh = MeshIn;
 		Spatial = SpatialIn;
@@ -36,6 +40,10 @@ public:
 	{
 		double fDistSqr;
 		int tNearestID = Spatial->FindNearestTriangle(Point, fDistSqr);
+		if (tNearestID <= 0)
+		{
+			return Point;
+		}
 		FTriangle3d Triangle;
 		Mesh->GetTriVertices(tNearestID, Triangle.V[0], Triangle.V[1], Triangle.V[2]);
 
@@ -51,6 +59,10 @@ public:
 	{
 		double fDistSqr;
 		int tNearestID = Spatial->FindNearestTriangle(Point, fDistSqr);
+		if (tNearestID <= 0)
+		{
+			return Point;
+		}
 		FTriangle3d Triangle;
 		Mesh->GetTriVertices(tNearestID, Triangle.V[0], Triangle.V[1], Triangle.V[2]);
 
