FFrame::RayPlaneIntersection now returns bool to indicate hit/miss, instead of invalid hit point or point behind ray origin. Updated call sites to handle new behavior.
#rnx

diff --git a/Engine/Plugins/Experimental/GeometryProcessing/Source/GeometricObjects/Public/FrameTypes.h b/Engine/Plugins/Experimental/GeometryProcessing/Source/GeometricObjects/Public/FrameTypes.h
--- a/Engine/Plugins/Experimental/GeometryProcessing/Source/GeometricObjects/Public/FrameTypes.h
+++ b/Engine/Plugins/Experimental/GeometryProcessing/Source/GeometricObjects/Public/FrameTypes.h
@@ -357,17 +357,27 @@ struct TFrame3
 	 * @param RayOrigin origin of ray
 	 * @param RayDirection direction of ray
 	 * @param PlaneNormalAxis which axis of frame to use as plane normal
-	 * @return intersection point, or FVector3::Max() if ray is parallel to plane
+	 * @param HitPointOut intersection point, or FVector3::Max() if ray does not hit plane or is parallel to plane
+	 * @return true if ray intersects plane and HitPointOut is valid
 	 */
-	FVector3<RealType> RayPlaneIntersection(const FVector3<RealType>& RayOrigin, const FVector3<RealType>& RayDirection, int PlaneNormalAxis)
+	bool RayPlaneIntersection(const FVector3<RealType>& RayOrigin, const FVector3<RealType>& RayDirection, int PlaneNormalAxis, FVector3<RealType>& HitPointOut)
 	{
 		FVector3<RealType> Normal = GetAxis(PlaneNormalAxis);
 		RealType PlaneD = -Origin.Dot(Normal);
 		RealType NormalDot = RayDirection.Dot(Normal);
 		if (VectorUtil::EpsilonEqual(NormalDot, (RealType)0, TMathUtil<RealType>::ZeroTolerance))
-			return FVector3<RealType>::Max();
+		{
+			HitPointOut = FVector3<RealType>::Max();
+			return false;
+		}
 		RealType t = -( RayOrigin.Dot(Normal) + PlaneD) / NormalDot;
-		return RayOrigin + t * RayDirection;
+		if (t < 0)
+		{
+			HitPointOut = FVector3<RealType>::Max();
+			return false;
+		}
+		HitPointOut = RayOrigin + t * RayDirection;
+		return true;
 	}
 
 };
