initial version of a plane cut tool

diff --git a/Engine/Plugins/Experimental/GeometryProcessing/Source/DynamicMesh/Private/DynamicMeshEditor.cpp b/Engine/Plugins/Experimental/GeometryProcessing/Source/DynamicMesh/Private/DynamicMeshEditor.cpp
--- a/Engine/Plugins/Experimental/GeometryProcessing/Source/DynamicMesh/Private/DynamicMeshEditor.cpp
+++ b/Engine/Plugins/Experimental/GeometryProcessing/Source/DynamicMesh/Private/DynamicMeshEditor.cpp
@@ -367,10 +367,55 @@ void FDynamicMeshEditor::SetTriangleNormals(const TArray<int>& Triangles, const
 }
 
 
+void FDynamicMeshEditor::SetTriangleUVsFromProjection(const TArray<int>& Triangles, const FFrame3d& ProjectionFrame, float UVScaleFactor, int UVLayerIndex)
+{
+	if (!Triangles.Num())
+	{
+		return;
+	}
+
+	check(Mesh->HasAttributes() && Mesh->Attributes()->NumUVLayers() > UVLayerIndex);
+	FDynamicMeshUVOverlay* UVs = Mesh->Attributes()->GetUVLayer(UVLayerIndex);
+
+	TMap<int, int> BaseToOverlayVIDMap;
+	TArray<int> AllUVIndices;
 
+	FAxisAlignedBox2f UVBounds(FAxisAlignedBox2f::Empty());
+
+	for (int TID : Triangles)
+	{
+		FIndex3i BaseTri = Mesh->GetTriangle(TID);
+		FIndex3i ElemTri;
+		for (int j = 0; j < 3; ++j)
+		{
+			const int* FoundElementID = BaseToOverlayVIDMap.Find(BaseTri[j]);
+			if (FoundElementID == nullptr)
+			{
+				FVector2f UV = ProjectionFrame.ToPlaneUV(Mesh->GetVertex(BaseTri[j]), 2);
+				UVBounds.Contain(UV);
+				ElemTri[j] = UVs->AppendElement(UV, BaseTri[j]);
+				AllUVIndices.Add(ElemTri[j]);
+				BaseToOverlayVIDMap.Add(BaseTri[j], ElemTri[j]);
+			}
+			else
+			{
+				ElemTri[j] = *FoundElementID;
+			}
+		}
+		UVs->SetTriangle(TID, ElemTri);
+	}
+
+	// shift UVs so that their bbox min-corner is at origin and scaled by external scale factor
+	for (int UVID : AllUVIndices)
+	{
+		FVector2f UV = UVs->GetElement(UVID);
+		FVector2f TransformedUV = (UV - UVBounds.Min) * UVScaleFactor;
+		UVs->SetElement(UVID, TransformedUV);
+	}
+}
 
 
-void FDynamicMeshEditor::SetQuadUVsFromProjection(const FIndex2i& QuadTris, const FFrame3f& ProjectionFrame, float UVScaleFactor, int UVLayerIndex)
+void FDynamicMeshEditor::SetQuadUVsFromProjection(const FIndex2i& QuadTris, const FFrame3d& ProjectionFrame, float UVScaleFactor, int UVLayerIndex)
 {
 	check(Mesh->HasAttributes() && Mesh->Attributes()->NumUVLayers() > UVLayerIndex );
 	FDynamicMeshUVOverlay* UVs = Mesh->Attributes()->GetUVLayer(UVLayerIndex);
@@ -383,7 +428,7 @@ void FDynamicMeshEditor::SetQuadUVsFromProjection(const FIndex2i& QuadTris, cons
 	FIndex3i UVTriangle1;
 	for (int j = 0; j < 3; ++j)
 	{
-		FVector2f UV = ProjectionFrame.ToPlaneUV( (FVector3f)Mesh->GetVertex(Triangle1[j]), 2);
+		FVector2f UV = ProjectionFrame.ToPlaneUV(Mesh->GetVertex(Triangle1[j]), 2);
 		UVTriangle1[j] = UVs->AppendElement(UV, Triangle1[j]);
 		AllUVs[j] = UV;
 		AllUVIndices[j] = UVTriangle1[j];
@@ -400,7 +445,7 @@ void FDynamicMeshEditor::SetQuadUVsFromProjection(const FIndex2i& QuadTris, cons
 			int i = Triangle1.IndexOf(Triangle2[j]);
 			if (i == -1)
 			{
-				FVector2f UV = ProjectionFrame.ToPlaneUV( (FVector3f)Mesh->GetVertex(Triangle2[j]), 2);
+				FVector2f UV = ProjectionFrame.ToPlaneUV(Mesh->GetVertex(Triangle2[j]), 2);
 				UVTriangle2[j] = UVs->AppendElement(UV, Triangle2[j]);
 				AllUVs[3] = UV;
 				AllUVIndices[3] = UVTriangle2[j];
diff --git a/Engine/Plugins/Experimental/GeometryProcessing/Source/DynamicMesh/Private/Operations/ExtrudeMesh.cpp b/Engine/Plugins/Experimental/GeometryProcessing/Source/DynamicMesh/Private/Operations/ExtrudeMesh.cpp
--- a/Engine/Plugins/Experimental/GeometryProcessing/Source/DynamicMesh/Private/Operations/ExtrudeMesh.cpp
+++ b/Engine/Plugins/Experimental/GeometryProcessing/Source/DynamicMesh/Private/Operations/ExtrudeMesh.cpp
@@ -108,14 +108,14 @@ bool FExtrudeMesh::Apply()
 
 				// @todo is there a simpler way to construct rotation from 3 known axes (third id Normal.Cross(UnitY))
 				//  (converting from matrix might end up being more efficient due to trig in ConstrainedAlignAxis?)
-				FFrame3f ProjectFrame(FVector3f::Zero(), Normal);
-				if (FMathd::Abs(ProjectFrame.Y().Dot(FVector3f::UnitY())) < 0.01)
+				FFrame3d ProjectFrame(FVector3d::Zero(), (FVector3d)Normal);
+				if (FMathd::Abs(ProjectFrame.Y().Dot(FVector3d::UnitY())) < 0.01)
 				{
-					ProjectFrame.ConstrainedAlignAxis(0, FVector3f::UnitX(), ProjectFrame.Z());
+					ProjectFrame.ConstrainedAlignAxis(0, FVector3d::UnitX(), ProjectFrame.Z());
 				}
 				else
 				{
-					ProjectFrame.ConstrainedAlignAxis(1, FVector3f::UnitY(), ProjectFrame.Z());
+					ProjectFrame.ConstrainedAlignAxis(1, FVector3d::UnitY(), ProjectFrame.Z());
 				}
 				Editor.SetQuadUVsFromProjection(StitchResult.NewQuads[k], ProjectFrame, UVScaleFactor);
 			}
diff --git a/Engine/Plugins/Experimental/GeometryProcessing/Source/DynamicMesh/Private/Operations/MeshPlaneCut.cpp b/Engine/Plugins/Experimental/GeometryProcessing/Source/DynamicMesh/Private/Operations/MeshPlaneCut.cpp
--- a/Engine/Plugins/Experimental/GeometryProcessing/Source/DynamicMesh/Private/Operations/MeshPlaneCut.cpp
+++ b/Engine/Plugins/Experimental/GeometryProcessing/Source/DynamicMesh/Private/Operations/MeshPlaneCut.cpp
@@ -202,6 +202,8 @@ bool FMeshPlaneCut::FillHoles(int ConstantGroupID)
 
 	HoleFillTriangles.Empty();
 
+	FFrame3d ProjectionFrame(PlaneOrigin, PlaneNormal);
+
 	for (const FEdgeLoop& Loop : CutLoops)
 	{
 		FSimpleHoleFiller Filler(Mesh, Loop);
@@ -209,6 +211,14 @@ bool FMeshPlaneCut::FillHoles(int ConstantGroupID)
 		if (Filler.Fill(GID))
 		{
 			HoleFillTriangles.Add(Filler.NewTriangles);
+
+			if (Mesh->HasAttributes())
+			{
+				FDynamicMeshEditor Editor(Mesh);
+				Editor.SetTriangleNormals(Filler.NewTriangles, (FVector3f)PlaneNormal);
+				Editor.SetTriangleUVsFromProjection(Filler.NewTriangles, ProjectionFrame, UVScaleFactor);
+			}
+
 		}
 		else
 		{
diff --git a/Engine/Plugins/Experimental/GeometryProcessing/Source/DynamicMesh/Public/DynamicMeshEditor.h b/Engine/Plugins/Experimental/GeometryProcessing/Source/DynamicMesh/Public/DynamicMeshEditor.h
--- a/Engine/Plugins/Experimental/GeometryProcessing/Source/DynamicMesh/Public/DynamicMeshEditor.h
+++ b/Engine/Plugins/Experimental/GeometryProcessing/Source/DynamicMesh/Public/DynamicMeshEditor.h
@@ -239,14 +239,24 @@ public:
 
 
 	/**
-	 * Project the two triangles of the quad onto a plane defined by the normal and use that to create/set new shared per-triangle UVs.
+	 * Project the two triangles of the quad onto a plane defined by the ProjectionFrame and use that to create/set new shared per-triangle UVs.
 	 * UVs are translated so that their bbox min-corner is at origin, and scaled by given scale factor
 	 * @param QuadTris pair of triangle IDs. If second ID is invalid, it is ignored
 	 * @param ProjectFrame vertices are projected into XY axes of this frame
 	 * @param UVScaleFactor UVs are scaled
 	 * @param UVLayerIndex which UV layer to operate on (must exist)
 	 */
-	void SetQuadUVsFromProjection(const FIndex2i& QuadTris, const FFrame3f& ProjectFrame, float UVScaleFactor = 1.0f, int UVLayerIndex = 0);
+	void SetQuadUVsFromProjection(const FIndex2i& QuadTris, const FFrame3d& ProjectionFrame, float UVScaleFactor = 1.0f, int UVLayerIndex = 0);
+
+	/**
+	* Project triangles onto a plane defined by the ProjectionFrame and use that to create/set new shared per-triangle UVs.
+	* UVs are translated so that their bbox min-corner is at origin, and scaled by given scale factor
+	* @param Triangles TArray of triangle IDs
+	* @param ProjectFrame vertices are projected into XY axes of this frame
+	* @param UVScaleFactor UVs are scaled
+	* @param UVLayerIndex which UV layer to operate on (must exist)
+	*/
+	void SetTriangleUVsFromProjection(const TArray<int>& Triangles, const FFrame3d& ProjectionFrame, float UVScaleFactor, int UVLayerIndex = 0);
 
 
 
diff --git a/Engine/Plugins/Experimental/GeometryProcessing/Source/DynamicMesh/Public/Operations/MeshPlaneCut.h b/Engine/Plugins/Experimental/GeometryProcessing/Source/DynamicMesh/Public/Operations/MeshPlaneCut.h
--- a/Engine/Plugins/Experimental/GeometryProcessing/Source/DynamicMesh/Public/Operations/MeshPlaneCut.h
+++ b/Engine/Plugins/Experimental/GeometryProcessing/Source/DynamicMesh/Public/Operations/MeshPlaneCut.h
@@ -38,10 +38,12 @@ public:
 	FDynamicMesh3 *Mesh;
 	FVector3d PlaneOrigin, PlaneNormal;
 	
-
 	bool bCollapseDegenerateEdgesOnCut = true;
 	double DegenerateEdgeTol = FMathd::ZeroTolerance;
 
+	/** UVs on any hole fill surfaces are scaled by this amount */
+	float UVScaleFactor = 1.0f;
+
 	// TODO support optionally restricting plane cut to a mesh selection
 	//MeshFaceSelection CutFaceSet;
 
