add function to convert TFrame3 to TTransform3
#rnx

diff --git a/Engine/Plugins/Experimental/GeometryProcessing/Source/GeometricObjects/Public/FrameTypes.h b/Engine/Plugins/Experimental/GeometryProcessing/Source/GeometricObjects/Public/FrameTypes.h
--- a/Engine/Plugins/Experimental/GeometryProcessing/Source/GeometricObjects/Public/FrameTypes.h
+++ b/Engine/Plugins/Experimental/GeometryProcessing/Source/GeometricObjects/Public/FrameTypes.h
@@ -7,6 +7,7 @@
 #include "VectorUtil.h"
 #include "RayTypes.h"
 #include "Quaternion.h"
+#include "TransformTypes.h"
 
 /**
  * TFrame3 is an object that represents an oriented 3D coordinate frame, ie orthogonal X/Y/Z axes at a point in space.
@@ -131,6 +132,11 @@ struct TFrame3
 		return FTransform(Rotation, Origin);
 	}
 
+	/** @return conversion of this Frame to TTransform */
+	TTransform3<RealType> ToTransform() const
+	{
+		return TTransform3<RealType>(Rotation, Origin);
+	}
 
 	/** @return point at distances along frame axes */
 	FVector3<RealType> PointAt(RealType X, RealType Y, RealType Z) const
