add required casts to fix linux compile errors
#rnx

diff --git a/Engine/Plugins/Experimental/GeometryProcessing/Source/DynamicMesh/Private/DynamicMeshEditor.cpp b/Engine/Plugins/Experimental/GeometryProcessing/Source/DynamicMesh/Private/DynamicMeshEditor.cpp
--- a/Engine/Plugins/Experimental/GeometryProcessing/Source/DynamicMesh/Private/DynamicMeshEditor.cpp
+++ b/Engine/Plugins/Experimental/GeometryProcessing/Source/DynamicMesh/Private/DynamicMeshEditor.cpp
@@ -391,7 +391,7 @@ void FDynamicMeshEditor::SetTriangleUVsFromProjection(const TArray<int>& Triangl
 			const int* FoundElementID = BaseToOverlayVIDMap.Find(BaseTri[j]);
 			if (FoundElementID == nullptr)
 			{
-				FVector2f UV = ProjectionFrame.ToPlaneUV(Mesh->GetVertex(BaseTri[j]), 2);
+				FVector2f UV = (FVector2f)ProjectionFrame.ToPlaneUV(Mesh->GetVertex(BaseTri[j]), 2);
 				UVBounds.Contain(UV);
 				ElemTri[j] = UVs->AppendElement(UV, BaseTri[j]);
 				AllUVIndices.Add(ElemTri[j]);
@@ -428,7 +428,7 @@ void FDynamicMeshEditor::SetQuadUVsFromProjection(const FIndex2i& QuadTris, cons
 	FIndex3i UVTriangle1;
 	for (int j = 0; j < 3; ++j)
 	{
-		FVector2f UV = ProjectionFrame.ToPlaneUV(Mesh->GetVertex(Triangle1[j]), 2);
+		FVector2f UV = (FVector2f)ProjectionFrame.ToPlaneUV(Mesh->GetVertex(Triangle1[j]), 2);
 		UVTriangle1[j] = UVs->AppendElement(UV, Triangle1[j]);
 		AllUVs[j] = UV;
 		AllUVIndices[j] = UVTriangle1[j];
@@ -445,7 +445,7 @@ void FDynamicMeshEditor::SetQuadUVsFromProjection(const FIndex2i& QuadTris, cons
 			int i = Triangle1.IndexOf(Triangle2[j]);
 			if (i == -1)
 			{
-				FVector2f UV = ProjectionFrame.ToPlaneUV(Mesh->GetVertex(Triangle2[j]), 2);
+				FVector2f UV = (FVector2f)ProjectionFrame.ToPlaneUV(Mesh->GetVertex(Triangle2[j]), 2);
 				UVTriangle2[j] = UVs->AppendElement(UV, Triangle2[j]);
 				AllUVs[3] = UV;
 				AllUVIndices[3] = UVTriangle2[j];
