const and explicit constructor fixes
#rnx

diff --git a/Engine/Plugins/Experimental/GeometryProcessing/Source/GeometricObjects/Public/FrameTypes.h b/Engine/Plugins/Experimental/GeometryProcessing/Source/GeometricObjects/Public/FrameTypes.h
--- a/Engine/Plugins/Experimental/GeometryProcessing/Source/GeometricObjects/Public/FrameTypes.h
+++ b/Engine/Plugins/Experimental/GeometryProcessing/Source/GeometricObjects/Public/FrameTypes.h
@@ -41,7 +41,7 @@ struct TFrame3
 	/**
 	 * Construct a frame at the given Origin aligned to the unit axes
 	 */
-	TFrame3(const FVector3<RealType>& OriginIn)
+	explicit TFrame3(const FVector3<RealType>& OriginIn)
 	{
 		Origin = OriginIn;
 		Rotation = TQuaternion<RealType>::Identity();
@@ -81,7 +81,7 @@ struct TFrame3
 	}
 
 	/** Construct a Frame from an FTransform */
-	TFrame3(const FTransform& Transform)
+	explicit TFrame3(const FTransform& Transform)
 	{
 		Origin = Transform.GetTranslation();
 		Rotation = Transform.GetRotation();
@@ -366,7 +366,7 @@ struct TFrame3
 	 * @param HitPointOut intersection point, or FVector3::Max() if ray does not hit plane or is parallel to plane
 	 * @return true if ray intersects plane and HitPointOut is valid
 	 */
-	bool RayPlaneIntersection(const FVector3<RealType>& RayOrigin, const FVector3<RealType>& RayDirection, int PlaneNormalAxis, FVector3<RealType>& HitPointOut)
+	bool RayPlaneIntersection(const FVector3<RealType>& RayOrigin, const FVector3<RealType>& RayDirection, int PlaneNormalAxis, FVector3<RealType>& HitPointOut) const
 	{
 		FVector3<RealType> Normal = GetAxis(PlaneNormalAxis);
 		RealType PlaneD = -Origin.Dot(Normal);
