Geometry Processing - updating the construction of a lapacian operator to better handle edges.

#rnx

diff --git a/Engine/Plugins/Experimental/GeometryProcessing/Source/MeshSolverUtilities/Private/LaplacianMeshSmoother.cpp b/Engine/Plugins/Experimental/GeometryProcessing/Source/MeshSolverUtilities/Private/LaplacianMeshSmoother.cpp
--- a/Engine/Plugins/Experimental/GeometryProcessing/Source/MeshSolverUtilities/Private/LaplacianMeshSmoother.cpp
+++ b/Engine/Plugins/Experimental/GeometryProcessing/Source/MeshSolverUtilities/Private/LaplacianMeshSmoother.cpp
@@ -63,10 +63,15 @@ FConstrainedMeshOperator::FConstrainedMeshOperator(const FDynamicMesh3& DynamicM
 	: VertexCount(DynamicMesh.VertexCount())
 {
 
-	Laplacian = ConstructLaplacian(Scheme, DynamicMesh, VtxLinearization, &EdgeVerts);
-	FSparseMatrixD& LMatrix = *(Laplacian);
+	FSparseMatrixD LaplacianInternal;
+	FSparseMatrixD LaplacianBoundary;
+	ConstructLaplacian(Scheme, DynamicMesh, VtxLinearization, LaplacianInternal, LaplacianBoundary);
 
-	TUniquePtr<FSparseMatrixD> LTL(new FSparseMatrixD(Laplacian->rows(), Laplacian->cols()));
+	// capture the number interior (non boundary) verts.
+	InternalVertexCount = VertexCount - VtxLinearization.NumBoundaryVerts();
+
+
+	TUniquePtr<FSparseMatrixD> LTL( new FSparseMatrixD(LaplacianInternal.rows(), LaplacianInternal.cols()) );
 	FSparseMatrixD& LTLMatrix = *(LTL);
 
 	bool bIsLaplacianSymmetric = (Scheme == ELaplacianWeightScheme::Valence || Scheme == ELaplacianWeightScheme::Uniform);
@@ -74,14 +79,14 @@ FConstrainedMeshOperator::FConstrainedMeshOperator(const FDynamicMesh3& DynamicM
 	if (bIsLaplacianSymmetric)
 	{
 		// Laplacian is symmetric, i.e. equal to its transpose
-		LTLMatrix = LMatrix * LMatrix;
+		LTLMatrix = LaplacianInternal * LaplacianInternal;
 	
 		ConstrainedSolver.Reset(new FConstrainedSolver(LTL, MatrixSolverType));
 	}
 	else
 	{
 		// the laplacian 
-		LTLMatrix = LMatrix.transpose() * LMatrix;
+		LTLMatrix = LaplacianInternal.transpose() * LaplacianInternal;
 		ConstrainedSolver.Reset(new FConstrainedSolver(LTL, MatrixSolverType));
 	}
 
@@ -89,36 +94,52 @@ FConstrainedMeshOperator::FConstrainedMeshOperator(const FDynamicMesh3& DynamicM
 
 void FConstrainedMeshOperator::AddConstraint(const int32 VtxId, const double Weight, const FVector3d& Pos, const bool bPostFix)
 {
-
-	bConstraintPositionsDirty = true;
-	bConstraintWeightsDirty   = true;
 	int32 Index = VtxLinearization.ToIndex()[VtxId];
 
-	ConstraintPositionMap.Add(TTuple<int32, FConstraintPosition>(Index, FConstraintPosition(Pos, bPostFix)));
-	ConstraintWeightMap.Add(TTuple<int32, double>(Index, Weight));
+	// Only add the constraint if the vertex is actually in the interior.  We aren't solving for edge vertices.
+	if (Index < InternalVertexCount)
+	{
+		bConstraintPositionsDirty = true;
+		bConstraintWeightsDirty = true;
+
 
+		ConstraintPositionMap.Add(TTuple<int32, FConstraintPosition>(Index, FConstraintPosition(Pos, bPostFix)));
+		ConstraintWeightMap.Add(TTuple<int32, double>(Index, Weight));
+	}
 }
 
 
 bool FConstrainedMeshOperator::UpdateConstraintPosition(const int32 VtxId, const FVector3d& Pos, const bool bPostFix)
 {
-	bConstraintPositionsDirty = true;
 	int32 Index = VtxLinearization.ToIndex()[VtxId];
-	// Add should over-write any existing value for this key
-	ConstraintPositionMap.Add(TTuple<int32, FConstraintPosition>(Index, FConstraintPosition(Pos, bPostFix)));
+	bool Result = false;
 
-	return ConstraintWeightMap.Contains(VtxId);
+	if (Index < InternalVertexCount)
+	{
+		bConstraintPositionsDirty = true;
+
+		// Add should over-write any existing value for this key
+		ConstraintPositionMap.Add(TTuple<int32, FConstraintPosition>(Index, FConstraintPosition(Pos, bPostFix)));
+
+		Result = ConstraintWeightMap.Contains(VtxId);
+	}
+	return Result;
 }
 
 bool FConstrainedMeshOperator::UpdateConstraintWeight(const int32 VtxId, const double Weight)
 {
-
-	bConstraintWeightsDirty = true;
+	bool Result = false;
 	int32 Index = VtxLinearization.ToIndex()[VtxId];
-	// Add should over-write any existing value for this key
-	ConstraintWeightMap.Add(TTuple<int32, double>(Index, Weight));
+	if (Index < InternalVertexCount)
+	{
+		bConstraintWeightsDirty = true;
 
-	return ConstraintPositionMap.Contains(VtxId);
+		// Add should over-write any existing value for this key
+		ConstraintWeightMap.Add(TTuple<int32, double>(Index, Weight));
+
+		Result = ConstraintPositionMap.Contains(VtxId);
+	}
+	return Result;
 }
 
 
@@ -127,8 +148,13 @@ bool FConstrainedMeshOperator::IsConstrained(const int32 VtxId) const
 	if (VtxId > VtxLinearization.ToIndex().Num()) return false;
 
 	int32 Index = VtxLinearization.ToIndex()[VtxId];
+	bool Result = false;
+	if (Index < InternalVertexCount)
+	{
+		Result = ConstraintWeightMap.Contains(Index);
+	}
 
-	return ConstraintWeightMap.Contains(Index);
+	return Result;
 }
 
 void FConstrainedMeshOperator::UpdateSolverConstraints()
@@ -158,6 +184,8 @@ bool FConstrainedMeshOperator::Linearize(const FSOAPositions& PositionalVector,
 		return false;
 	}
 
+	checkSlow(Num == InternalVertexCount);
+
 	// 
 	const auto& IndexToVtxId = VtxLinearization.ToId();
 	const int32 MaxVtxId = IndexToVtxId.Num(); // NB: this is really max_used + 1 in the mesh.  See  FDynamicMesh3::MaxVertexID()
@@ -179,7 +207,7 @@ bool FConstrainedMeshOperator::Linearize(const FSOAPositions& PositionalVector,
 
 void FConstrainedMeshOperator::ExtractVertexPositions(const FDynamicMesh3& DynamicMesh, FSOAPositions& VertexPositions) const
 {
-	VertexPositions.SetZero(VertexCount);
+	VertexPositions.SetZero(InternalVertexCount);
 
 
 	const TArray<int32>& ToIndex = VtxLinearization.ToIndex();
@@ -190,12 +218,15 @@ void FConstrainedMeshOperator::ExtractVertexPositions(const FDynamicMesh3& Dynam
 		const int32 i = ToIndex[VtxId];
 
 		checkSlow(i != FDynamicMesh3::InvalidID);
-
-		const FVector3d& Vertex = DynamicMesh.GetVertex(VtxId);
-		// coeffRef - coeff access without range checking
-		VertexPositions.XVector.coeffRef(i) = Vertex.X;
-		VertexPositions.YVector.coeffRef(i) = Vertex.Y;
-		VertexPositions.ZVector.coeffRef(i) = Vertex.Z;
+		
+		if (i < InternalVertexCount) // Skipping boundary verts
+		{
+			const FVector3d& Vertex = DynamicMesh.GetVertex(VtxId);
+			// coeffRef - coeff access without range checking
+			VertexPositions.XVector.coeffRef(i) = Vertex.X;
+			VertexPositions.YVector.coeffRef(i) = Vertex.Y;
+			VertexPositions.ZVector.coeffRef(i) = Vertex.Z;
+		}
 	}
 }
 
@@ -206,6 +237,8 @@ void FConstrainedMeshOperator::UpdateWithPostFixConstraints(FSOAPositions& Posit
 		const int32 Index = ConstraintPosition.Key;
 		const FConstraintPosition& Constraint = ConstraintPosition.Value;
 
+		checkSlow(Index < InternalVertexCount);
+
 		// we only care about post-fix constraints
 
 		if (Constraint.bPostFix)
@@ -222,8 +255,9 @@ void FConstrainedMeshOperator::UpdateWithPostFixConstraints(FSOAPositions& Posit
 
 FConstrainedMeshDeformer::FConstrainedMeshDeformer(const FDynamicMesh3& DynamicMesh, const ELaplacianWeightScheme LaplacianType)
 	: FConstrainedMeshOperator(DynamicMesh, LaplacianType, EMatrixSolverType::LU)
-	, LaplacianVectors(FConstrainedMeshOperator::VertexCount)
+	, LaplacianVectors(FConstrainedMeshOperator::InternalVertexCount)
 {
+	
 
 	// The current vertex positions 
 	
@@ -241,6 +275,8 @@ FConstrainedMeshDeformer::FConstrainedMeshDeformer(const FDynamicMesh3& DynamicM
 	// Compute the Laplacian Vectors
 	//    := Biharmonic * VertexPostion
 	// In the case of the cotangent laplacian this can be identified as the mean curvature * normal.
+	checkSlow(LaplacianVectors.Num() == OriginalVertexPositions.Num());
+
 	for (int32 i = 0; i < 3; ++i)
 	{
 		LaplacianVectors.Array(i) = Biharmonic * OriginalVertexPositions.Array(i);
@@ -254,7 +290,7 @@ bool FConstrainedMeshDeformer::Deform(TArray<FVector3d>& PositionBuffer)
 	UpdateSolverConstraints();
 
 	// Allocate space for the result as a struct of arrays
-	FSOAPositions SolutionVector(VertexCount);
+	FSOAPositions SolutionVector(InternalVertexCount);
 
 	// Solve the linear system
 	// NB: the original positions will only be used if the underlying solver type is iterative	
@@ -282,7 +318,7 @@ bool FBiHarmonicMeshSmoother::ComputeSmoothedMeshPositions(TArray<FVector3d>& Up
 
 	// Solves the constrained system and updates the mesh 
 
-	FSOAPositions SolutionVector(VertexCount);
+	FSOAPositions SolutionVector(InternalVertexCount);
 
 	bool bSuccess = ConstrainedSolver->Solve(SolutionVector);
 
@@ -307,7 +343,7 @@ bool FCGBiHarmonicMeshSmoother::ComputeSmoothedMeshPositions(TArray<FVector3d>&
 
 	// Solves the constrained system
 
-	FSOAPositions SolutionVector(VertexCount);
+	FSOAPositions SolutionVector(InternalVertexCount);
 
 	bool bSuccess = ConstrainedSolver->Solve(SolutionVector);
 
diff --git a/Engine/Plugins/Experimental/GeometryProcessing/Source/MeshSolverUtilities/Private/LaplacianMeshSmoother.h b/Engine/Plugins/Experimental/GeometryProcessing/Source/MeshSolverUtilities/Private/LaplacianMeshSmoother.h
--- a/Engine/Plugins/Experimental/GeometryProcessing/Source/MeshSolverUtilities/Private/LaplacianMeshSmoother.h
+++ b/Engine/Plugins/Experimental/GeometryProcessing/Source/MeshSolverUtilities/Private/LaplacianMeshSmoother.h
@@ -21,13 +21,13 @@ public:
 	FConstrainedMeshOperator(const FDynamicMesh3& DynamicMesh, const ELaplacianWeightScheme Scheme, const EMatrixSolverType MatrixSolverType);
 	virtual ~FConstrainedMeshOperator() override {}
 
-	// Add constraint associated with given vertex id.
+	// Add constraint associated with given vertex id.  Boundary vertices will be ignored
 	void AddConstraint(const int32 VtxId, const double Weight, const FVector3d& Pos, const bool bPostFix) override;
 
-
+	// Update the position of an existing constraint.  Bool return if a corresponding constraint weight exists. Boundary vertices will be ignored (and return false).  
 	bool UpdateConstraintPosition(const int32 VtxId, const FVector3d& Position, const bool bPostFix) override;
 
-	// The underlying solver will have to refactor the matrix if this is done
+	// The underlying solver will have to refactor the matrix if this is done. Bool return if a corresponding constraint position exists. Boundary vertices will be ignored (and return false).  
 	bool UpdateConstraintWeight(const int32 VtxId, const double Weight) override;
 
 	// Clear all constraints associated with this smoother
@@ -38,7 +38,7 @@ public:
 	void ClearConstraintPositions() override { ConstraintPositionMap.Empty();  bConstraintPositionsDirty = true; }
 
 
-	// Test if for constraint associated with given vertex id.
+	// Test if for constraint associated with given vertex id. Will return false for any boundary vert.
 	bool IsConstrained(const int32 VtxId) const override;
 
 	
@@ -74,16 +74,15 @@ protected:
 	bool bConstraintPositionsDirty = true;
 	bool bConstraintWeightsDirty   = true;
 
-	// Cache the vertex count.
+	// Cache the vertex count. boundary + internal
 	int32                 VertexCount;
 
+	// Cache the number of internal vertices
+	int32                 InternalVertexCount;
+
 	// Used to map between VtxId and vertex Index in linear vector..
 	FVertexLinearization  VtxLinearization;
 
-	// I don't know if we want to keep this after the constructor
-	TUniquePtr<FSparseMatrixD>             Laplacian;
-	TArray<int32>                          EdgeVerts;
-
 	// Actual solver that manages the various linear algebra bits.
 	TUniquePtr<FConstrainedSolver>         ConstrainedSolver;
 };
diff --git a/Engine/Plugins/Experimental/GeometryProcessing/Source/MeshSolverUtilities/Private/LaplacianOperators.cpp b/Engine/Plugins/Experimental/GeometryProcessing/Source/MeshSolverUtilities/Private/LaplacianOperators.cpp
--- a/Engine/Plugins/Experimental/GeometryProcessing/Source/MeshSolverUtilities/Private/LaplacianOperators.cpp
+++ b/Engine/Plugins/Experimental/GeometryProcessing/Source/MeshSolverUtilities/Private/LaplacianOperators.cpp
@@ -83,266 +83,6 @@
 	return NumMatrixEntries;
 }
 
-TUniquePtr<FSparseMatrixD> ConstructUniformLaplacian(const FDynamicMesh3& DynamicMesh, FVertexLinearization& VertexMap, TArray<int32>* BoundaryVerts)
-{
-
-
-	typedef FSparseMatrixD::Scalar    ScalarT;
-	typedef Eigen::Triplet<ScalarT>  MatrixTripletT;
-
-	if (BoundaryVerts)
-	{
-		// empty
-		BoundaryVerts->Empty();
-	}
-
-	// Sync the mapping between the mesh vertex ids and their offsets in a nomimal linear array.
-	VertexMap.Reset(DynamicMesh);
-
-	const TArray<int32>& ToMeshV = VertexMap.ToId();
-	const TArray<int32>& ToIndex = VertexMap.ToIndex();
-	const int32 NumVerts = VertexMap.NumVerts();
-
-	// Eigen constructs a sparse matrix from a linearized array of matrix entries.
-
-	std::vector<MatrixTripletT> MatrixTripelList;
-	{
-		int32 NumMatrixEntries = ComputeNumMatrixElements(DynamicMesh, ToMeshV);
-		MatrixTripelList.reserve(NumMatrixEntries);
-	}
-
-	// Construct Laplacian Matrix: loop over verts constructing the corresponding matrix row.
-
-	for (int32 i = 0; i < NumVerts; ++i)
-	{
-		const int32 VertId = ToMeshV[i];
-
-		if (DynamicMesh.IsBoundaryVertex(VertId))
-		{
-			if (BoundaryVerts)
-			{
-				BoundaryVerts->Add(VertId);
-			}
-#if LAPLACIAN_SKIP_BOUNDARY == 1  
-			MatrixTripelList.push_back(MatrixTripletT(i, i, 0.));
-			continue;
-#endif 
-		}
-
-
-		ScalarT CenterWeight = ScalarT(0); // equal and opposite the sum of the neighbor weights
-		for (int NeighborVertId : DynamicMesh.VtxVerticesItr(VertId))
-		{
-			const int32 j = ToIndex[NeighborVertId];
-
-
-			ScalarT NeighborWeight = ScalarT(1);
-			CenterWeight += NeighborWeight;
-
-			// add the neighbor 
-			MatrixTripelList.push_back(MatrixTripletT(i, j, NeighborWeight));
-
-		}
-		// add the center
-		MatrixTripelList.push_back(MatrixTripletT(i, i, -CenterWeight));
-
-	}
-
-	// SparseMatrix doesn't have (SparseMatrix&& ) constructor 
-
-	TUniquePtr<FSparseMatrixD> LaplacianMatrix(new FSparseMatrixD(NumVerts, NumVerts));
-	LaplacianMatrix->setFromTriplets(MatrixTripelList.begin(), MatrixTripelList.end());
-
-	LaplacianMatrix->makeCompressed();
-
-	return LaplacianMatrix;
-
-}
-
-TUniquePtr<FSparseMatrixD> ConstructUmbrellaLaplacian(const FDynamicMesh3& DynamicMesh, FVertexLinearization& VertexMap, TArray<int32>* BoundaryVerts)
-{
-
-
-	typedef FSparseMatrixD::Scalar    ScalarT;
-	typedef Eigen::Triplet<ScalarT>  MatrixTripletT;
-
-	if (BoundaryVerts)
-	{
-		// empty
-		BoundaryVerts->Empty();
-	}
-
-	// Sync the mapping between the mesh vertex ids and their offsets in a nomimal linear array.
-	VertexMap.Reset(DynamicMesh);
-
-	const TArray<int32>& ToMeshV = VertexMap.ToId();
-	const TArray<int32>& ToIndex = VertexMap.ToIndex();
-	const int32 NumVerts = VertexMap.NumVerts();
-
-	// Eigen constructs a sparse matrix from a linearized array of matrix entries.
-
-	std::vector<MatrixTripletT> MatrixTripelList;
-	{
-		int32 NumMatrixEntries = ComputeNumMatrixElements(DynamicMesh, ToMeshV);
-		MatrixTripelList.reserve(NumMatrixEntries);
-	}
-
-	// Cache valency of each vertex.
-	// Number of non-zero elements in the i'th row = 1 + OneRingSize(i)
-	TArray<int32> OneRingSize;
-	{
-		OneRingSize.SetNumUninitialized(NumVerts);
-
-		for (int32 i = 0; i < NumVerts; ++i)
-		{
-			const int32 VertId = ToMeshV[i];
-			OneRingSize[i] = DynamicMesh.GetVtxEdgeCount(VertId);
-		}
-	}
-
-	// Construct Laplacian Matrix: loop over verts constructing the corresponding matrix row.
-
-	for (int32 i = 0; i < NumVerts; ++i)
-	{
-		const int32 VertId = ToMeshV[i];
-		const int32 Valence = OneRingSize[i];
-		double InvValence = (Valence != 0) ? 1. / double(Valence) : 0.;
-
-		if (DynamicMesh.IsBoundaryVertex(VertId))
-		{
-			if (BoundaryVerts)
-			{
-				BoundaryVerts->Add(VertId);
-			}
-#if LAPLACIAN_SKIP_BOUNDARY == 1  
-			MatrixTripelList.push_back(MatrixTripletT(i, i, 0.));
-			continue;
-#endif 
-		}
-
-		for (int NeighborVertId : DynamicMesh.VtxVerticesItr(VertId))
-		{
-			const int32 j = ToIndex[NeighborVertId];
-
-			// add the neighbor 
-			MatrixTripelList.push_back(MatrixTripletT(i, j, InvValence));
-
-		}
-		// add the center
-		MatrixTripelList.push_back(MatrixTripletT(i, i, -ScalarT(1)));
-
-	}
-
-	// SparseMatrix doesn't have (SparseMatrix&& ) constructor 
-
-	TUniquePtr<FSparseMatrixD> LaplacianMatrix(new FSparseMatrixD(NumVerts, NumVerts));
-	LaplacianMatrix->setFromTriplets(MatrixTripelList.begin(), MatrixTripelList.end());
-
-	LaplacianMatrix->makeCompressed();
-
-	return LaplacianMatrix;
-
-}
-
-TUniquePtr<FSparseMatrixD> ConstructValenceWeightedLaplacian(const FDynamicMesh3& DynamicMesh, FVertexLinearization& VertexMap, TArray<int32>* BoundaryVerts)
-{
-
-
-	typedef FSparseMatrixD::Scalar    ScalarT;
-	typedef Eigen::Triplet<ScalarT>  MatrixTripletT;
-
-	if (BoundaryVerts)
-	{
-		// empty
-		BoundaryVerts->Empty();
-	}
-
-	// Sync the mapping between the mesh vertex ids and their offsets in a nomimal linear array.
-	VertexMap.Reset(DynamicMesh);
-
-	const TArray<int32>& ToMeshV = VertexMap.ToId();
-	const TArray<int32>& ToIndex = VertexMap.ToIndex();
-	const int32 NumVerts = VertexMap.NumVerts();
-
-	// Cache valency of each vertex.
-	// Number of non-zero elements in the i'th row = 1 + OneRingSize(i)
-	TArray<int32> OneRingSize;
-	{
-		OneRingSize.SetNumUninitialized(NumVerts);
-
-		for (int32 i = 0; i < NumVerts; ++i)
-		{
-			const int32 VertId = ToMeshV[i];
-			OneRingSize[i] = DynamicMesh.GetVtxEdgeCount(VertId);
-		}
-	}
-
-
-	// Compute the total number of entries in the sparse matrix
-	int32 NumMatrixEntries = 0;
-	{
-		for (int32 i = 0; i < NumVerts; ++i)
-		{
-			NumMatrixEntries += 1 + OneRingSize[i]; // myself plus my neighbors
-		}
-
-	}
-	std::vector<MatrixTripletT> MatrixTripelList;
-	MatrixTripelList.reserve(NumMatrixEntries);
-
-
-
-	// Construct Laplacian Matrix: loop over verts constructing the corresponding matrix row.
-
-	for (int32 i = 0; i < NumVerts; ++i)
-	{
-		const int32 VertId = ToMeshV[i];
-		const int32 IOneRingSize = OneRingSize[i];
-
-		if (DynamicMesh.IsBoundaryVertex(VertId))
-		{
-			if (BoundaryVerts)
-			{
-				BoundaryVerts->Add(VertId);
-			}
-#if LAPLACIAN_SKIP_BOUNDARY == 1  
-			MatrixTripelList.push_back(MatrixTripletT(i, i, 0.));
-			continue;
-#endif 
-		}
-
-
-		ScalarT CenterWeight = ScalarT(0); // equal and opposite the sum of the neighbor weights
-		for (int NeighborVertId : DynamicMesh.VtxVerticesItr(VertId))
-		{
-			const int32 j = ToIndex[NeighborVertId];
-			const int32 JOneRingSize = OneRingSize[j];
-
-
-			ScalarT NeighborWeight = ScalarT(1) / std::sqrt(IOneRingSize + JOneRingSize);
-			CenterWeight += NeighborWeight;
-
-			// add the neighbor 
-			MatrixTripelList.push_back(MatrixTripletT(i, j, NeighborWeight));
-
-		}
-		// add the center
-		MatrixTripelList.push_back(MatrixTripletT(i, i, -CenterWeight));
-
-	}
-
-	// SparseMatrix doesn't have (SparseMatrix&& ) constructor 
-
-	TUniquePtr<FSparseMatrixD> LaplacianMatrix(new FSparseMatrixD(NumVerts, NumVerts));
-	LaplacianMatrix->setFromTriplets(MatrixTripelList.begin(), MatrixTripelList.end());
-
-	LaplacianMatrix->makeCompressed();
-
-	return LaplacianMatrix;
-
-}
-
-
 /**
 * The per-triangle data used in constructing the cotangent weighted laplacian.
 *
@@ -422,9 +162,9 @@
 
 			// Cotangent[i] is cos(theta)/sin(theta) at the i'th vertex.
 
-			Cotangent[0] =  EdgeAB.Dot(EdgeAC) / TwiceArea;
+			Cotangent[0] = EdgeAB.Dot(EdgeAC) / TwiceArea;
 			Cotangent[1] = -EdgeAB.Dot(EdgeBC) / TwiceArea;
-			Cotangent[2] =  EdgeAC.Dot(EdgeBC) / TwiceArea;
+			Cotangent[2] = EdgeAC.Dot(EdgeBC) / TwiceArea;
 
 			if (bIsObtuse())
 			{
@@ -607,9 +347,9 @@
 
 		// Compute tan(angle/2) = Sqrt[ (1-cos) / (1 + cos)]
 
-		const double ABdotAC =  EdgeAB.Dot(EdgeAC);
+		const double ABdotAC = EdgeAB.Dot(EdgeAC);
 		const double BCdotBA = -EdgeBC.Dot(EdgeAB);
-		const double CAdotCB =  EdgeAC.Dot(EdgeBC);
+		const double CAdotCB = EdgeAC.Dot(EdgeBC);
 
 
 		const double RegularizingConst = 1.e-6; // keeps us from dividing by zero when making tan[180/2] = sin[90]/cos[90] = inf
@@ -625,7 +365,7 @@
 #if 0
 		// testing
 		double totalAngle = 2. * (std::atan(TanHalfAngle[0]) + std::atan(TanHalfAngle[1]) + std::atan(TanHalfAngle[2]));
-		
+
 		double angleError = M_PI - totalAngle;
 #endif
 	}
@@ -698,4 +438,991 @@
 	return TriangleDataArray;
 }
 
+void ConstructUniformLaplacian(const FDynamicMesh3& DynamicMesh, FVertexLinearization& VertexMap, FSparseMatrixD& LaplacianInterior, FSparseMatrixD& LaplacianBoundary)
+{
+	typedef FSparseMatrixD::Scalar    ScalarT;
+	typedef Eigen::Triplet<ScalarT>  MatrixTripletT;
 
+	// Sync the mapping between the mesh vertex ids and their offsets in a nominal linear array.
+	VertexMap.Reset(DynamicMesh);
+
+	const TArray<int32>& ToMeshV = VertexMap.ToId();
+	const TArray<int32>& ToIndex = VertexMap.ToIndex();
+	const int32 NumVerts = VertexMap.NumVerts();
+
+	const int32 NumBoundaryVerts = VertexMap.NumBoundaryVerts();
+	const int32 NumInteriorVerts = NumVerts - NumBoundaryVerts;
+
+
+	// Eigen constructs a sparse matrix from a linearized array of matrix entries.
+	// pre-allocate space for the temporary triplet list 
+	std::vector<MatrixTripletT> LaplacianAsTripletList;
+	std::vector<MatrixTripletT> BoundaryAsTripletList;
+	{
+		int32 NumMatrixEntries = ComputeNumMatrixElements(DynamicMesh, ToMeshV);
+		LaplacianAsTripletList.reserve(NumMatrixEntries);
+	}
+
+	
+	// Construct Laplacian Matrix: loop over verts constructing the corresponding matrix row.
+	// NB: the vertices are ordered with the interior verts first.
+	for (int32 i = 0; i < NumInteriorVerts; ++i)
+	{
+		const int32 VertId = ToMeshV[i];
+		ScalarT CenterWeight = ScalarT(0); // equal and opposite the sum of the neighbor weights
+
+		checkSlow(!DynamicMesh.IsBoundaryVertex(VertId));  // we should only be looping over the internal verts
+
+		for (int NeighborVertId : DynamicMesh.VtxVerticesItr(VertId))
+		{
+			const int32 j = ToIndex[NeighborVertId];
+
+
+			ScalarT NeighborWeight = ScalarT(1);
+			CenterWeight += NeighborWeight;
+
+			if (j < NumInteriorVerts)
+			{
+				// add the neighbor 
+				LaplacianAsTripletList.push_back(MatrixTripletT(i, j, NeighborWeight));
+			}
+			else
+			{
+				int32 jBoundary = j - NumInteriorVerts;
+				BoundaryAsTripletList.push_back(MatrixTripletT(i, jBoundary, NeighborWeight));
+			}
+		}
+		// add the center
+		LaplacianAsTripletList.push_back(MatrixTripletT(i, i, -CenterWeight));
+
+	}
+
+	// populate the Laplacian for the interior.
+	{
+		// [djh] - can we construct with (0,0) ? or do we need to check for it?
+		FSparseMatrixD InteriorMatrix(NumInteriorVerts, NumInteriorVerts);
+		InteriorMatrix.setFromTriplets(LaplacianAsTripletList.begin(), LaplacianAsTripletList.end());
+		LaplacianInterior.swap(InteriorMatrix);
+		LaplacianInterior.makeCompressed();
+	}
+	
+	// populate the boundary matrix
+	{
+		FSparseMatrixD BoundaryMatrix(NumInteriorVerts, NumBoundaryVerts);
+		BoundaryMatrix.setFromTriplets(BoundaryAsTripletList.begin(), BoundaryAsTripletList.end());
+		LaplacianBoundary.swap(BoundaryMatrix);
+		LaplacianBoundary.makeCompressed();
+	}
+
+}
+
+
+void ConstructUmbrellaLaplacian(const FDynamicMesh3& DynamicMesh, FVertexLinearization& VertexMap, FSparseMatrixD& LaplacianInterior, FSparseMatrixD& LaplacianBoundary)
+{
+
+	typedef FSparseMatrixD::Scalar    ScalarT;
+	typedef Eigen::Triplet<ScalarT>  MatrixTripletT;
+
+	// Sync the mapping between the mesh vertex ids and their offsets in a nominal linear array.
+	VertexMap.Reset(DynamicMesh);
+
+	const TArray<int32>& ToMeshV = VertexMap.ToId();
+	const TArray<int32>& ToIndex = VertexMap.ToIndex();
+	const int32 NumVerts = VertexMap.NumVerts();
+
+	const int32 NumBoundaryVerts = VertexMap.NumBoundaryVerts();
+	const int32 NumInteriorVerts = NumVerts - NumBoundaryVerts;
+
+	// Eigen constructs a sparse matrix from a linearized array of matrix entries.
+	// pre-allocate space for the temporary triplet list 
+	std::vector<MatrixTripletT> LaplacianAsTripletList;
+	std::vector<MatrixTripletT> BoundaryAsTripletList;
+	{
+		int32 NumMatrixEntries = ComputeNumMatrixElements(DynamicMesh, ToMeshV);
+		LaplacianAsTripletList.reserve(NumMatrixEntries);
+	}
+
+	// Cache valency of each vertex.
+	// Number of non-zero elements in the i'th row = 1 + OneRingSize(i)
+	TArray<int32> OneRingSize;
+	{
+		OneRingSize.SetNumUninitialized(NumVerts);
+
+		for (int32 i = 0; i < NumVerts; ++i)
+		{
+			const int32 VertId = ToMeshV[i];
+			OneRingSize[i] = DynamicMesh.GetVtxEdgeCount(VertId);
+		}
+	}
+
+	// Construct Laplacian Matrix: loop over verts constructing the corresponding matrix row.
+
+	for (int32 i = 0; i < NumInteriorVerts; ++i)
+	{
+		const int32 VertId = ToMeshV[i];
+		const int32 Valence = OneRingSize[i];
+		double InvValence = (Valence != 0) ? 1. / double(Valence) : 0.;
+
+		checkSlow(!DynamicMesh.IsBoundaryVertex(VertId));
+
+		for (int NeighborVertId : DynamicMesh.VtxVerticesItr(VertId))
+		{
+			const int32 j = ToIndex[NeighborVertId];
+
+			// add the neighbor 
+			if (j < NumInteriorVerts)
+			{
+				LaplacianAsTripletList.push_back(MatrixTripletT(i, j, InvValence));
+			}
+			else
+			{
+				int32 jBoundary = j - NumInteriorVerts;
+				BoundaryAsTripletList.push_back(MatrixTripletT(i, jBoundary, InvValence));
+			}
+		}
+		// add the center
+		LaplacianAsTripletList.push_back(MatrixTripletT(i, i, -ScalarT(1)));
+
+	}
+
+	// populate the Laplacian for the interior.
+	{
+		// [djh] - can we construct with (0,0) ? or do we need to check for it?
+		FSparseMatrixD InteriorMatrix(NumInteriorVerts, NumInteriorVerts);
+		InteriorMatrix.setFromTriplets(LaplacianAsTripletList.begin(), LaplacianAsTripletList.end());
+		LaplacianInterior.swap(InteriorMatrix);
+		LaplacianInterior.makeCompressed();
+	}
+
+	// populate the boundary matrix
+	{
+		FSparseMatrixD BoundaryMatrix(NumInteriorVerts, NumBoundaryVerts);
+		BoundaryMatrix.setFromTriplets(BoundaryAsTripletList.begin(), BoundaryAsTripletList.end());
+		LaplacianBoundary.swap(BoundaryMatrix);
+		LaplacianBoundary.makeCompressed();
+	}
+
+}
+
+
+void ConstructValenceWeightedLaplacian(const FDynamicMesh3& DynamicMesh, FVertexLinearization& VertexMap, FSparseMatrixD& LaplacianInterior, FSparseMatrixD& LaplacianBoundary)
+{
+
+
+	typedef FSparseMatrixD::Scalar    ScalarT;
+	typedef Eigen::Triplet<ScalarT>  MatrixTripletT;
+
+
+	// Sync the mapping between the mesh vertex ids and their offsets in a nomimal linear array.
+	VertexMap.Reset(DynamicMesh);
+
+	const TArray<int32>& ToMeshV = VertexMap.ToId();
+	const TArray<int32>& ToIndex = VertexMap.ToIndex();
+	const int32 NumVerts = VertexMap.NumVerts();
+
+	const int32 NumBoundaryVerts = VertexMap.NumBoundaryVerts();
+	const int32 NumInteriorVerts = NumVerts - NumBoundaryVerts;
+
+	// Eigen constructs a sparse matrix from a linearized array of matrix entries.
+	// pre-allocate space for the temporary triplet list 
+	std::vector<MatrixTripletT> LaplacianAsTripletList;
+	std::vector<MatrixTripletT> BoundaryAsTripletList;
+
+	// Cache valency of each vertex.
+	// Number of non-zero elements in the i'th row = 1 + OneRingSize(i)
+	TArray<int32> OneRingSize;
+	{
+		OneRingSize.SetNumUninitialized(NumVerts);
+
+		for (int32 i = 0; i < NumVerts; ++i)
+		{
+			const int32 VertId = ToMeshV[i];
+			OneRingSize[i] = DynamicMesh.GetVtxEdgeCount(VertId);
+		}
+	}
+
+
+	// Compute the total number of entries in the sparse matrix
+	int32 NumMatrixEntries = 0;
+	{
+		for (int32 i = 0; i < NumVerts; ++i)
+		{
+			NumMatrixEntries += 1 + OneRingSize[i]; // myself plus my neighbors
+		}
+
+	}
+	LaplacianAsTripletList.reserve(NumMatrixEntries);
+
+
+
+	// Construct Laplacian Matrix: loop over verts constructing the corresponding matrix row.
+
+	for (int32 i = 0; i < NumInteriorVerts; ++i)
+	{
+		const int32 VertId       = ToMeshV[i];
+		const int32 IOneRingSize = OneRingSize[i];
+
+		ScalarT CenterWeight = ScalarT(0); // equal and opposite the sum of the neighbor weights
+		for (int NeighborVertId : DynamicMesh.VtxVerticesItr(VertId))
+		{
+			const int32 j = ToIndex[NeighborVertId];
+			const int32 JOneRingSize = OneRingSize[j];
+
+
+			ScalarT NeighborWeight = ScalarT(1) / std::sqrt(IOneRingSize + JOneRingSize);
+			CenterWeight += NeighborWeight;
+
+			if (j < NumInteriorVerts)
+			{
+				// add the neighbor 
+				LaplacianAsTripletList.push_back(MatrixTripletT(i, j, NeighborWeight));
+			}
+			else
+			{
+				int32 jBoundary = j - NumInteriorVerts;
+				BoundaryAsTripletList.push_back(MatrixTripletT(i, jBoundary, NeighborWeight));
+			}
+		}
+		// add the center
+		LaplacianAsTripletList.push_back(MatrixTripletT(i, i, -CenterWeight));
+
+	}
+
+
+	// populate the Laplacian for the interior.
+	{
+		// [djh] - can we construct with (0,0) ? or do we need to check for it?
+		FSparseMatrixD InteriorMatrix(NumInteriorVerts, NumInteriorVerts);
+		InteriorMatrix.setFromTriplets(LaplacianAsTripletList.begin(), LaplacianAsTripletList.end());
+		LaplacianInterior.swap(InteriorMatrix);
+		LaplacianInterior.makeCompressed();
+	}
+
+	// populate the boundary matrix
+	{
+		FSparseMatrixD BoundaryMatrix(NumInteriorVerts, NumBoundaryVerts);
+		BoundaryMatrix.setFromTriplets(BoundaryAsTripletList.begin(), BoundaryAsTripletList.end());
+		LaplacianBoundary.swap(BoundaryMatrix);
+		LaplacianBoundary.makeCompressed();
+	}
+
+}
+
+
+void ConstructCotangentLaplacian(const FDynamicMesh3& DynamicMesh, FVertexLinearization& VertexMap, FSparseMatrixD& AreaMatrix, FSparseMatrixD& LaplacianInterior, FSparseMatrixD& LaplacianBoundary)
+{
+
+
+	typedef FSparseMatrixD::Scalar  ScalarT;
+	typedef Eigen::Triplet<ScalarT> MatrixTripletT;
+
+
+	// Sync the mapping between the mesh vertex ids and their offsets in a nomimal linear array.
+	VertexMap.Reset(DynamicMesh);
+
+	const TArray<int32>& ToMeshV = VertexMap.ToId();
+	const TArray<int32>& ToIndex = VertexMap.ToIndex();
+	const int32 NumVerts = VertexMap.NumVerts();
+
+	const int32 NumBoundaryVerts = VertexMap.NumBoundaryVerts();
+	const int32 NumInteriorVerts = NumVerts - NumBoundaryVerts;
+
+	// Eigen constructs a sparse matrix from a linearized array of matrix entries.
+	// pre-allocate space for the temporary triplet list 
+	
+	std::vector<MatrixTripletT> LaplacianAsTripletList;
+	std::vector<MatrixTripletT> BoundaryAsTripletList;
+	
+	// Eigen constructs a sparse matrix from a linearized array of matrix entries.
+	{
+		int32 NumMatrixEntries = ComputeNumMatrixElements(DynamicMesh, ToMeshV);
+		LaplacianAsTripletList.reserve(NumMatrixEntries);
+	}
+
+	// Create the mapping of triangles
+
+	FTriangleLinearization TriangleMap(DynamicMesh);
+
+	const TArray<int32>& ToMeshTri = TriangleMap.ToId();
+	const TArray<int32>& ToTriIdx = TriangleMap.ToIndex();
+	const int32 NumTris = TriangleMap.NumTris();
+
+
+	// Clear space for the areas
+	std::vector<MatrixTripletT> DiagonalTriplets;
+	DiagonalTriplets.reserve(NumVerts);
+
+	FSparseMatrixD Diagonals;
+	Diagonals.reserve(NumVerts);
+
+
+	// Create an array that holds all the geometric information we need for each triangle.
+
+	TArray<CotanTriangleData>  CotangentTriangleDataArray = ConstructTriangleDataArray<CotanTriangleData>(DynamicMesh, TriangleMap);
+
+
+
+	// Construct Laplacian Matrix: loop over verts constructing the corresponding matrix row.
+	//                             store the id of the boundary verts for later use.
+
+	for (int32 i = 0; i < NumInteriorVerts; ++i)
+	{
+		const int32 IVertId = ToMeshV[i]; // I - the row
+
+
+		// Compute the Voronoi area for this vertex.
+		double WeightArea = 0.;
+		for (int32 TriId : DynamicMesh.VtxTrianglesItr(IVertId))
+		{
+			const int32 TriIdx = ToTriIdx[TriId];
+			const CotanTriangleData& TriData = CotangentTriangleDataArray[TriIdx];
+
+
+			// The three VertIds for this triangle.
+			const FIndex3i TriVertIds = DynamicMesh.GetTriangle(TriId);
+
+			// Which of the corners is IVertId?
+			int32 Offset = 0;
+			while (TriVertIds[Offset] != IVertId)
+			{
+				Offset++;
+				checkSlow(Offset < 3);
+			}
+
+			WeightArea += TriData.VoronoiArea[Offset];
+		}
+
+
+		double WeightII = 0.; // accumulate to equal and opposite the sum of the neighbor weights
+
+		// for each connecting edge
+
+		for (int32 EdgeId : DynamicMesh.VtxEdgesItr(IVertId))
+		{
+			// [v0, v1, t0, t1]:  NB: both t0 & t1 exist since IVert isn't a boundary vert.
+			FIndex4i Edge = DynamicMesh.GetEdge(EdgeId);
+
+
+			// the other vert in the edge - identifies the matrix column
+			const int32 JVertId = (Edge[0] == IVertId) ? Edge[1] : Edge[0];  // J - the column
+
+			checkSlow(JVertId != IVertId);
+
+			// Get the cotangents for this edge.
+
+			const int32 Tri0Idx = ToTriIdx[Edge[2]];
+			const CotanTriangleData& Tri0Data = CotangentTriangleDataArray[Tri0Idx];
+			const double CotanAlpha = Tri0Data.GetOpposingCotangent(EdgeId);
+
+
+			// The second triangle will be invalid if this is an edge!
+
+			const double CotanBeta = (Edge[3] != FDynamicMesh3::InvalidID) ? CotangentTriangleDataArray[ToTriIdx[Edge[3]]].GetOpposingCotangent(EdgeId) : 0.;
+
+			double WeightIJ = 0.5 * (CotanAlpha + CotanBeta);
+			WeightII += WeightIJ;
+
+			const int32 j = ToIndex[JVertId];
+
+			if (j < NumInteriorVerts)
+			{
+				LaplacianAsTripletList.push_back(MatrixTripletT(i, j, WeightIJ));
+			}
+			else
+			{
+				int32 jBoundary = j - NumInteriorVerts;
+				BoundaryAsTripletList.push_back(MatrixTripletT(i, j, WeightIJ));
+			}
+
+		}
+
+		LaplacianAsTripletList.push_back(MatrixTripletT(i, i, -WeightII));
+
+		DiagonalTriplets.push_back(MatrixTripletT(i, i, WeightArea));
+
+
+	}
+
+	Diagonals.setFromTriplets(DiagonalTriplets.begin(), DiagonalTriplets.end());
+	AreaMatrix.swap(Diagonals);
+	AreaMatrix.makeCompressed();
+
+	// populate the Laplacian for the interior.
+	{
+		// [djh] - can we construct with (0,0) ? or do we need to check for it?
+		FSparseMatrixD InteriorMatrix(NumInteriorVerts, NumInteriorVerts);
+		InteriorMatrix.setFromTriplets(LaplacianAsTripletList.begin(), LaplacianAsTripletList.end());
+		LaplacianInterior.swap(InteriorMatrix);
+		LaplacianInterior.makeCompressed();
+	}
+
+	// populate the boundary matrix
+	{
+		FSparseMatrixD BoundaryMatrix(NumInteriorVerts, NumBoundaryVerts);
+		BoundaryMatrix.setFromTriplets(BoundaryAsTripletList.begin(), BoundaryAsTripletList.end());
+		LaplacianBoundary.swap(BoundaryMatrix);
+		LaplacianBoundary.makeCompressed();
+	}
+
+}
+
+
+void ConstructCotangentLaplacian(const FDynamicMesh3& DynamicMesh, FVertexLinearization& VertexMap, FSparseMatrixD& LaplacianInterior, FSparseMatrixD& LaplacianBoundary, const bool bClampWeights)
+{
+
+
+	typedef FSparseMatrixD::Scalar  ScalarT;
+	typedef Eigen::Triplet<ScalarT> MatrixTripletT;
+
+
+	
+	// Sync the mapping between the mesh vertex ids and their offsets in a nominal linear array.
+	VertexMap.Reset(DynamicMesh);
+
+	const TArray<int32>& ToMeshV = VertexMap.ToId();
+	const TArray<int32>& ToIndex = VertexMap.ToIndex();
+	const int32 NumVerts = VertexMap.NumVerts();
+
+	const int32 NumBoundaryVerts = VertexMap.NumBoundaryVerts();
+	const int32 NumInteriorVerts = NumVerts - NumBoundaryVerts;
+
+	// Eigen constructs a sparse matrix from a linearized array of matrix entries.
+	// pre-allocate space for the temporary triplet list 
+
+	std::vector<MatrixTripletT> LaplacianAsTripletList;
+	std::vector<MatrixTripletT> BoundaryAsTripletList;
+	{
+		int32 NumMatrixEntries = ComputeNumMatrixElements(DynamicMesh, ToMeshV);
+		LaplacianAsTripletList.reserve(NumMatrixEntries);
+	}
+
+	// Map the triangles.
+
+	FTriangleLinearization TriangleMap(DynamicMesh);
+
+	const TArray<int32>& ToMeshTri = TriangleMap.ToId();
+	const TArray<int32>& ToTriIdx = TriangleMap.ToIndex();
+	const int32 NumTris = TriangleMap.NumTris();
+
+	// Create an array that holds all the geometric information we need for each triangle.
+
+	TArray<CotanTriangleData>  CotangentTriangleDataArray = ConstructTriangleDataArray<CotanTriangleData>(DynamicMesh, TriangleMap);
+
+
+	// Construct Laplacian Matrix: loop over verts constructing the corresponding matrix row.
+	//                             skipping the boundary verts for later use.
+
+	for (int32 i = 0; i < NumInteriorVerts; ++i)
+	{
+		const int32 IVertId = ToMeshV[i]; // I - the row
+
+
+		// Compute the Voronoi area for this vertex.
+		double WeightArea = 0.;
+		for (int32 TriId : DynamicMesh.VtxTrianglesItr(IVertId))
+		{
+			const int32 TriIdx = ToTriIdx[TriId];
+			const CotanTriangleData& TriData = CotangentTriangleDataArray[TriIdx];
+
+
+			// The three VertIds for this triangle.
+			const FIndex3i TriVertIds = DynamicMesh.GetTriangle(TriId);
+
+			// Which of the corners is IVertId?
+			int32 Offset = 0;
+			while (TriVertIds[Offset] != IVertId)
+			{
+				Offset++;
+				checkSlow(Offset < 3);
+			}
+
+			WeightArea += TriData.VoronoiArea[Offset];
+		}
+
+
+		double WeightII = 0.; // accumulate to equal and opposite the sum of the neighbor weights
+
+		// for each connecting edge
+
+		for (int32 EdgeId : DynamicMesh.VtxEdgesItr(IVertId))
+		{
+			// [v0, v1, t0, t1]:  NB: both t0 & t1 exist since IVert isn't a boundary vert.
+			FIndex4i Edge = DynamicMesh.GetEdge(EdgeId);
+
+
+			// the other vert in the edge - identifies the matrix column
+			const int32 JVertId = (Edge[0] == IVertId) ? Edge[1] : Edge[0];  // J - the column
+
+			checkSlow(JVertId != IVertId);
+
+			// Get the cotangents for this edge.
+
+			const int32 Tri0Idx = ToTriIdx[Edge[2]];
+			const CotanTriangleData& Tri0Data = CotangentTriangleDataArray[Tri0Idx];
+			const double CotanAlpha = Tri0Data.GetOpposingCotangent(EdgeId);
+
+			// The second triangle will be invalid if this is an edge!
+
+			const double CotanBeta = (Edge[3] != FDynamicMesh3::InvalidID) ? CotangentTriangleDataArray[ToTriIdx[Edge[3]]].GetOpposingCotangent(EdgeId) : 0.;
+
+			double WeightIJ = 0.5 * (CotanAlpha + CotanBeta);
+
+			// clamp the weight
+			if (bClampWeights)
+			{
+				WeightIJ = FMath::Clamp(WeightIJ, -1.e5 * WeightArea, 1.e5 * WeightArea);
+			}
+
+			WeightII += WeightIJ;
+
+			const int32 j = ToIndex[JVertId];
+
+			if (j < NumInteriorVerts)
+			{ 
+				LaplacianAsTripletList.push_back(MatrixTripletT(i, j, WeightIJ / WeightArea));
+			}
+			else
+			{
+				int32 jBoundary = j - NumInteriorVerts;
+				BoundaryAsTripletList.push_back(MatrixTripletT(i, jBoundary, WeightIJ / WeightArea));
+			}
+			
+
+		}
+
+		LaplacianAsTripletList.push_back(MatrixTripletT(i, i, -WeightII / WeightArea));
+
+	}
+
+	// populate the Laplacian for the interior.
+	{
+		// [djh] - can we construct with (0,0) ? or do we need to check for it?
+		FSparseMatrixD InteriorMatrix(NumInteriorVerts, NumInteriorVerts);
+		InteriorMatrix.setFromTriplets(LaplacianAsTripletList.begin(), LaplacianAsTripletList.end());
+		LaplacianInterior.swap(InteriorMatrix);
+		LaplacianInterior.makeCompressed();
+	}
+
+	// populate the boundary matrix
+	{
+		FSparseMatrixD BoundaryMatrix(NumInteriorVerts, NumBoundaryVerts);
+		BoundaryMatrix.setFromTriplets(BoundaryAsTripletList.begin(), BoundaryAsTripletList.end());
+		LaplacianBoundary.swap(BoundaryMatrix);
+		LaplacianBoundary.makeCompressed();
+	}
+
+}
+
+
+void ConstructMeanValueWeightLaplacian(const FDynamicMesh3& DynamicMesh, FVertexLinearization& VertexMap, FSparseMatrixD& LaplacianInterior, FSparseMatrixD& LaplacianBoundary)
+{
+
+	typedef FSparseMatrixD::Scalar  ScalarT;
+	typedef Eigen::Triplet<ScalarT> MatrixTripletT;
+
+
+	// Sync the mapping between the mesh vertex ids and their offsets in a nominal linear array.
+	VertexMap.Reset(DynamicMesh);
+
+	const TArray<int32>& ToMeshV = VertexMap.ToId();
+	const TArray<int32>& ToIndex = VertexMap.ToIndex();
+	const int32 NumVerts = VertexMap.NumVerts();
+
+	const int32 NumBoundaryVerts = VertexMap.NumBoundaryVerts();
+	const int32 NumInteriorVerts = NumVerts - NumBoundaryVerts;
+
+	// Eigen constructs a sparse matrix from a linearized array of matrix entries.
+	// pre-allocate space for the temporary triplet list 
+
+	std::vector<MatrixTripletT> LaplacianAsTripletList;
+	std::vector<MatrixTripletT> BoundaryAsTripletList;
+	{
+		int32 NumMatrixEntries = ComputeNumMatrixElements(DynamicMesh, ToMeshV);
+		LaplacianAsTripletList.reserve(NumMatrixEntries);
+	}
+
+	// Map the triangles.
+
+	FTriangleLinearization TriangleMap(DynamicMesh);
+
+	const TArray<int32>& ToMeshTri = TriangleMap.ToId();
+	const TArray<int32>& ToTriIdx = TriangleMap.ToIndex();
+	const int32 NumTris = TriangleMap.NumTris();
+
+
+	// Create an array that holds all the geometric information we need for each triangle.
+
+	TArray<MeanValueTriangleData>  TriangleDataArray = ConstructTriangleDataArray<MeanValueTriangleData>(DynamicMesh, TriangleMap);
+
+
+	// Construct Laplacian Matrix: loop over verts constructing the corresponding matrix row.
+	//                             skipping the boundary verts for later use.
+
+	for (int32 i = 0; i < NumInteriorVerts; ++i)
+	{
+		const int32 IVertId = ToMeshV[i]; // I - the row
+
+
+		double WeightII = 0.; // accumulate to equal and opposite the sum of the neighbor weights
+
+		// for each connecting edge
+
+		for (int32 EdgeId : DynamicMesh.VtxEdgesItr(IVertId))
+		{
+			// [v0, v1, t0, t1]:  NB: both t0 & t1 exist since IVert isn't a boundary vert.
+			FIndex4i Edge = DynamicMesh.GetEdge(EdgeId);
+
+			// the other vert in the edge - identifies the matrix column
+			const int32 JVertId = (Edge[0] == IVertId) ? Edge[1] : Edge[0];  // J - the column
+
+			// Get the cotangents for this edge.
+
+			const int32 Tri0Idx = ToTriIdx[Edge[2]];
+			const auto& Tri0Data = TriangleDataArray[Tri0Idx];
+			double TanHalfAngleSum = Tri0Data.GetTanHalfAngle(IVertId);
+			double EdgeLength = FMath::Max(1.e-5, Tri0Data.GetEdgeLenght(EdgeId)); // Clamp the length
+
+			// The second triangle will be invalid if this is an edge!
+
+			TanHalfAngleSum += (Edge[3] != FDynamicMesh3::InvalidID) ? TriangleDataArray[ToTriIdx[Edge[3]]].GetTanHalfAngle(IVertId) : 0.;
+
+			double WeightIJ = TanHalfAngleSum / EdgeLength;
+			WeightII += WeightIJ;
+
+			const int32 j = ToIndex[JVertId];
+
+			
+			if (j < NumInteriorVerts)
+			{
+				LaplacianAsTripletList.push_back(MatrixTripletT(i, j, WeightIJ));
+			}
+			else
+			{
+				int32 jBoundary = j - NumInteriorVerts;
+				BoundaryAsTripletList.push_back(MatrixTripletT(i, jBoundary, WeightIJ));
+			}
+
+		}
+
+		LaplacianAsTripletList.push_back(MatrixTripletT(i, i, -WeightII));
+
+
+	}
+
+	// populate the Laplacian for the interior.
+	{
+		// [djh] - can we construct with (0,0) ? or do we need to check for it?
+		FSparseMatrixD InteriorMatrix(NumInteriorVerts, NumInteriorVerts);
+		InteriorMatrix.setFromTriplets(LaplacianAsTripletList.begin(), LaplacianAsTripletList.end());
+		LaplacianInterior.swap(InteriorMatrix);
+		LaplacianInterior.makeCompressed();
+	}
+
+	// populate the boundary matrix
+	{
+		FSparseMatrixD BoundaryMatrix(NumInteriorVerts, NumBoundaryVerts);
+		BoundaryMatrix.setFromTriplets(BoundaryAsTripletList.begin(), BoundaryAsTripletList.end());
+		LaplacianBoundary.swap(BoundaryMatrix);
+		LaplacianBoundary.makeCompressed();
+	}
+
+}
+
+
+void ConstructLaplacian(const ELaplacianWeightScheme Scheme, const FDynamicMesh3& DynamicMesh, FVertexLinearization& VertexMap, FSparseMatrixD& LaplacianInterior, FSparseMatrixD& LaplacianBoundary)
+{
+
+	switch (Scheme)
+	{
+	default:
+	case ELaplacianWeightScheme::Uniform:
+		ConstructUniformLaplacian(DynamicMesh, VertexMap, LaplacianInterior, LaplacianBoundary);
+		break;
+	case ELaplacianWeightScheme::Umbrella:
+		ConstructUmbrellaLaplacian(DynamicMesh, VertexMap, LaplacianInterior, LaplacianBoundary);
+		break;
+	case ELaplacianWeightScheme::Valence:
+		ConstructValenceWeightedLaplacian(DynamicMesh, VertexMap, LaplacianInterior, LaplacianBoundary);
+		break;
+	case ELaplacianWeightScheme::Cotangent:
+	{
+		bool bClampWeights = false;
+		ConstructCotangentLaplacian(DynamicMesh, VertexMap, LaplacianInterior, LaplacianBoundary, bClampWeights);
+		break;
+	}
+	case ELaplacianWeightScheme::ClampedCotangent:
+	{
+		bool bClampWeights = true;
+		ConstructCotangentLaplacian(DynamicMesh, VertexMap, LaplacianInterior, LaplacianBoundary, bClampWeights);
+		break;
+	}
+	case ELaplacianWeightScheme::MeanValue:
+		ConstructMeanValueWeightLaplacian(DynamicMesh, VertexMap, LaplacianInterior, LaplacianBoundary);
+		break;
+	}
+}
+
+
+
+TUniquePtr<FSparseMatrixD> ConstructUniformLaplacian(const FDynamicMesh3& DynamicMesh, FVertexLinearization& VertexMap, TArray<int32>* BoundaryVerts)
+{
+
+
+	typedef FSparseMatrixD::Scalar    ScalarT;
+	typedef Eigen::Triplet<ScalarT>  MatrixTripletT;
+
+	if (BoundaryVerts)
+	{
+		// empty
+		BoundaryVerts->Empty();
+	}
+
+	// Sync the mapping between the mesh vertex ids and their offsets in a nomimal linear array.
+	VertexMap.Reset(DynamicMesh);
+
+	const TArray<int32>& ToMeshV = VertexMap.ToId();
+	const TArray<int32>& ToIndex = VertexMap.ToIndex();
+	const int32 NumVerts = VertexMap.NumVerts();
+
+	// Eigen constructs a sparse matrix from a linearized array of matrix entries.
+
+	std::vector<MatrixTripletT> MatrixTripelList;
+	{
+		int32 NumMatrixEntries = ComputeNumMatrixElements(DynamicMesh, ToMeshV);
+		MatrixTripelList.reserve(NumMatrixEntries);
+	}
+
+	// Construct Laplacian Matrix: loop over verts constructing the corresponding matrix row.
+
+	for (int32 i = 0; i < NumVerts; ++i)
+	{
+		const int32 VertId = ToMeshV[i];
+
+		if (DynamicMesh.IsBoundaryVertex(VertId))
+		{
+			if (BoundaryVerts)
+			{
+				BoundaryVerts->Add(VertId);
+			}
+#if LAPLACIAN_SKIP_BOUNDARY == 1  
+			MatrixTripelList.push_back(MatrixTripletT(i, i, 0.));
+			continue;
+#endif 
+		}
+
+
+		ScalarT CenterWeight = ScalarT(0); // equal and opposite the sum of the neighbor weights
+		for (int NeighborVertId : DynamicMesh.VtxVerticesItr(VertId))
+		{
+			const int32 j = ToIndex[NeighborVertId];
+
+
+			ScalarT NeighborWeight = ScalarT(1);
+			CenterWeight += NeighborWeight;
+
+			// add the neighbor 
+			MatrixTripelList.push_back(MatrixTripletT(i, j, NeighborWeight));
+
+		}
+		// add the center
+		MatrixTripelList.push_back(MatrixTripletT(i, i, -CenterWeight));
+
+	}
+
+	// SparseMatrix doesn't have (SparseMatrix&& ) constructor 
+
+	TUniquePtr<FSparseMatrixD> LaplacianMatrix(new FSparseMatrixD(NumVerts, NumVerts));
+	LaplacianMatrix->setFromTriplets(MatrixTripelList.begin(), MatrixTripelList.end());
+
+	LaplacianMatrix->makeCompressed();
+
+	return LaplacianMatrix;
+
+}
+
+
+TUniquePtr<FSparseMatrixD> ConstructUmbrellaLaplacian(const FDynamicMesh3& DynamicMesh, FVertexLinearization& VertexMap, TArray<int32>* BoundaryVerts)
+{
+
+
+	typedef FSparseMatrixD::Scalar    ScalarT;
+	typedef Eigen::Triplet<ScalarT>  MatrixTripletT;
+
+	if (BoundaryVerts)
+	{
+		// empty
+		BoundaryVerts->Empty();
+	}
+
+	// Sync the mapping between the mesh vertex ids and their offsets in a nomimal linear array.
+	VertexMap.Reset(DynamicMesh);
+
+	const TArray<int32>& ToMeshV = VertexMap.ToId();
+	const TArray<int32>& ToIndex = VertexMap.ToIndex();
+	const int32 NumVerts = VertexMap.NumVerts();
+
+	// Eigen constructs a sparse matrix from a linearized array of matrix entries.
+
+	std::vector<MatrixTripletT> MatrixTripelList;
+	{
+		int32 NumMatrixEntries = ComputeNumMatrixElements(DynamicMesh, ToMeshV);
+		MatrixTripelList.reserve(NumMatrixEntries);
+	}
+
+	// Cache valency of each vertex.
+	// Number of non-zero elements in the i'th row = 1 + OneRingSize(i)
+	TArray<int32> OneRingSize;
+	{
+		OneRingSize.SetNumUninitialized(NumVerts);
+
+		for (int32 i = 0; i < NumVerts; ++i)
+		{
+			const int32 VertId = ToMeshV[i];
+			OneRingSize[i] = DynamicMesh.GetVtxEdgeCount(VertId);
+		}
+	}
+
+	// Construct Laplacian Matrix: loop over verts constructing the corresponding matrix row.
+
+	for (int32 i = 0; i < NumVerts; ++i)
+	{
+		const int32 VertId = ToMeshV[i];
+		const int32 Valence = OneRingSize[i];
+		double InvValence = (Valence != 0) ? 1. / double(Valence) : 0.;
+
+		if (DynamicMesh.IsBoundaryVertex(VertId))
+		{
+			if (BoundaryVerts)
+			{
+				BoundaryVerts->Add(VertId);
+			}
+#if LAPLACIAN_SKIP_BOUNDARY == 1  
+			MatrixTripelList.push_back(MatrixTripletT(i, i, 0.));
+			continue;
+#endif 
+		}
+
+		for (int NeighborVertId : DynamicMesh.VtxVerticesItr(VertId))
+		{
+			const int32 j = ToIndex[NeighborVertId];
+
+			// add the neighbor 
+			MatrixTripelList.push_back(MatrixTripletT(i, j, InvValence));
+
+		}
+		// add the center
+		MatrixTripelList.push_back(MatrixTripletT(i, i, -ScalarT(1)));
+
+	}
+
+	// SparseMatrix doesn't have (SparseMatrix&& ) constructor 
+
+	TUniquePtr<FSparseMatrixD> LaplacianMatrix(new FSparseMatrixD(NumVerts, NumVerts));
+	LaplacianMatrix->setFromTriplets(MatrixTripelList.begin(), MatrixTripelList.end());
+
+	LaplacianMatrix->makeCompressed();
+
+	return LaplacianMatrix;
+
+}
+
+
+TUniquePtr<FSparseMatrixD> ConstructValenceWeightedLaplacian(const FDynamicMesh3& DynamicMesh, FVertexLinearization& VertexMap, TArray<int32>* BoundaryVerts)
+{
+
+
+	typedef FSparseMatrixD::Scalar    ScalarT;
+	typedef Eigen::Triplet<ScalarT>  MatrixTripletT;
+
+	if (BoundaryVerts)
+	{
+		// empty
+		BoundaryVerts->Empty();
+	}
+
+	// Sync the mapping between the mesh vertex ids and their offsets in a nomimal linear array.
+	VertexMap.Reset(DynamicMesh);
+
+	const TArray<int32>& ToMeshV = VertexMap.ToId();
+	const TArray<int32>& ToIndex = VertexMap.ToIndex();
+	const int32 NumVerts = VertexMap.NumVerts();
+
+	// Cache valency of each vertex.
+	// Number of non-zero elements in the i'th row = 1 + OneRingSize(i)
+	TArray<int32> OneRingSize;
+	{
+		OneRingSize.SetNumUninitialized(NumVerts);
+
+		for (int32 i = 0; i < NumVerts; ++i)
+		{
+			const int32 VertId = ToMeshV[i];
+			OneRingSize[i] = DynamicMesh.GetVtxEdgeCount(VertId);
+		}
+	}
+
+
+	// Compute the total number of entries in the sparse matrix
+	int32 NumMatrixEntries = 0;
+	{
+		for (int32 i = 0; i < NumVerts; ++i)
+		{
+			NumMatrixEntries += 1 + OneRingSize[i]; // myself plus my neighbors
+		}
+
+	}
+	std::vector<MatrixTripletT> MatrixTripelList;
+	MatrixTripelList.reserve(NumMatrixEntries);
+
+
+
+	// Construct Laplacian Matrix: loop over verts constructing the corresponding matrix row.
+
+	for (int32 i = 0; i < NumVerts; ++i)
+	{
+		const int32 VertId = ToMeshV[i];
+		const int32 IOneRingSize = OneRingSize[i];
+
+		if (DynamicMesh.IsBoundaryVertex(VertId))
+		{
+			if (BoundaryVerts)
+			{
+				BoundaryVerts->Add(VertId);
+			}
+#if LAPLACIAN_SKIP_BOUNDARY == 1  
+			MatrixTripelList.push_back(MatrixTripletT(i, i, 0.));
+			continue;
+#endif 
+		}
+
+
+		ScalarT CenterWeight = ScalarT(0); // equal and opposite the sum of the neighbor weights
+		for (int NeighborVertId : DynamicMesh.VtxVerticesItr(VertId))
+		{
+			const int32 j = ToIndex[NeighborVertId];
+			const int32 JOneRingSize = OneRingSize[j];
+
+
+			ScalarT NeighborWeight = ScalarT(1) / std::sqrt(IOneRingSize + JOneRingSize);
+			CenterWeight += NeighborWeight;
+
+			// add the neighbor 
+			MatrixTripelList.push_back(MatrixTripletT(i, j, NeighborWeight));
+
+		}
+		// add the center
+		MatrixTripelList.push_back(MatrixTripletT(i, i, -CenterWeight));
+
+	}
+
+	// SparseMatrix doesn't have (SparseMatrix&& ) constructor 
+
+	TUniquePtr<FSparseMatrixD> LaplacianMatrix(new FSparseMatrixD(NumVerts, NumVerts));
+	LaplacianMatrix->setFromTriplets(MatrixTripelList.begin(), MatrixTripelList.end());
+
+	LaplacianMatrix->makeCompressed();
+
+	return LaplacianMatrix;
+
+}
+
+
@@ -853,6 +1580,7 @@
 
 }
 
+
 TUniquePtr<FSparseMatrixD> ConstructCotangentLaplacian(const FDynamicMesh3& DynamicMesh, FVertexLinearization& VertexMap, const bool bClampWeights, TArray<int32>* BoundaryVerts)
 {
 
@@ -1013,11 +1741,6 @@
 }
 
 
-
-
-
-
-
 TUniquePtr<FSparseMatrixD> ConstructMeanValueWeightLaplacian(const FDynamicMesh3& DynamicMesh, FVertexLinearization& VertexMap, TArray<int32>* BoundaryVerts)
 {
 
@@ -1129,6 +1852,7 @@
 
 }
 
+
 TUniquePtr<FSparseMatrixD> ConstructLaplacian(const ELaplacianWeightScheme Scheme, const FDynamicMesh3& DynamicMesh, FVertexLinearization& VertexMap, TArray<int32>* BoundaryVerts)
 {
 
@@ -1161,4 +1885,3 @@
 		break;
 	}
 }
-
diff --git a/Engine/Plugins/Experimental/GeometryProcessing/Source/MeshSolverUtilities/Private/LaplacianOperators.h b/Engine/Plugins/Experimental/GeometryProcessing/Source/MeshSolverUtilities/Private/LaplacianOperators.h
--- a/Engine/Plugins/Experimental/GeometryProcessing/Source/MeshSolverUtilities/Private/LaplacianOperators.h
+++ b/Engine/Plugins/Experimental/GeometryProcessing/Source/MeshSolverUtilities/Private/LaplacianOperators.h
@@ -12,6 +12,188 @@
 #include "MeshSmoothingUtilities.h" // for ELaplacianWeightScheme
 
 
+
+/**
+* Construct a sparse matrix representation of the requested "Laplacian" type operator.
+*
+* The mesh itself is assumed to have N interior vertices, and M boundary vertices.
+*
+*
+* @param Scheme             Selects the actual type of "laplacian" constructed.
+* @param DynamicMesh        The triangle mesh
+* @param VertexMap          On return, Additional arrays used to map between vertexID and offset in a linear array (i.e. the row).
+*                           The vertices are ordered so that last M ( = VertexMap.NumBoundaryVerts() )  correspond to those on the boundary.
+* @param LaplacianInterior  On return, the laplacian operator that acts on the interior vertices: sparse  N x N matrix
+* @param LaplacianBoundary  On return, the portion of the operator that acts on the boundary vertices: sparse  N x M matrix
+*
+*   LaplacianInterior * Vector_InteriorVerts + LaplacianBoundary * Vector_BoundaryVerts = Full Laplacian applied to interior vertices.
+*/
+void ConstructLaplacian(const ELaplacianWeightScheme Scheme,
+	                   const FDynamicMesh3& DynamicMesh,
+	                   FVertexLinearization& VertexMap,
+	                   FSparseMatrixD& LaplacianInterior,
+	                   FSparseMatrixD& LaplacianBoundary);
+
+
+/**
+* Construct a sparse matrix representation of a uniform weighted Laplacian.
+* The uniform weighted Laplacian is defined solely in terms of the connectivity
+* of the mesh.  Note, by construction this should be a symmetric matrix.
+*
+* The mesh itself is assumed to have N interior vertices, and M boundary vertices.
+*
+* Row i represents the Laplacian at vert_i, the non-zero entries correspond to
+* the incident one-ring vertices vert_j.
+*
+* L_{ij} = 1                      if vert_j is in the one-ring of vert_i
+* L_{ii} = -Sum{ L_{ij}, j != i}
+*
+*
+* @param DynamicMesh        The triangle mesh
+* @param VertexMap          On return, Additional arrays used to map between vertexID and offset in a linear array (i.e. the row). 
+*                           The vertices are ordered so that last M ( = VertexMap.NumBoundaryVerts() )  correspond to those on the boundary. 
+* @param LaplacianInterior  On return, the laplacian operator that acts on the interior vertices: sparse  N x N matrix     
+* @param LaplacianBoundary  On return, the portion of the operator that acts on the boundary vertices: sparse  N x M matrix
+*
+*   LaplacianInterior * Vector_InteriorVerts + LaplacianBoundary * Vector_BoundaryVerts = Full Laplacian applied to interior vertices.
+*/
+void ConstructUniformLaplacian(const FDynamicMesh3& DynamicMesh, 
+	                           FVertexLinearization& VertexMap, 
+	                           FSparseMatrixD& LaplacianInterior, 
+	                           FSparseMatrixD& LaplacianBoundary);
+
+/**
+* Construct a sparse matrix representation of an umbrella weighted Laplacian.
+* This Laplacian is defined solely in terms of the connectivity
+* of the mesh.  Note, there is no expectation that the resulting matrix will be symmetric.
+*
+* The mesh itself is assumed to have N interior vertices, and M boundary vertices.
+*
+* Row i represents the Laplacian at vert_i, the non-zero entries correspond to
+* the incident one-ring vertices vert_j.
+*
+* L_{ij} = 1 / valence(of i)                      if vert_j is in the one-ring of vert_i
+* L_{ii} = -Sum{ L_{ij}, j != i} = -1
+*
+*
+* @param DynamicMesh        The triangle mesh
+* @param VertexMap          Additional arrays used to map between vertexID and offset in a linear array (i.e. the row).
+*                           The vertices are ordered so that last M ( = VertexMap.NumBoundaryVerts() )  correspond to those on the boundary.
+* @param LaplacianInterior  On return, the laplacian operator that acts on the interior vertices: sparse  N x N matrix
+* @param LaplacianBoundary  On return, the portion of the operator that acts on the boundary vertices: sparse  N x M matrix
+*
+*   LaplacianInterior * Vector_InteriorVerts + LaplacianBoundary * Vector_BoundaryVerts = Full Laplacian applied to interior vertices.
+*/
+void ConstructUmbrellaLaplacian(const FDynamicMesh3& DynamicMesh, 
+	                            FVertexLinearization& VertexMap,  
+	                            FSparseMatrixD& LaplacianInterior, 
+	                            FSparseMatrixD& LaplacianBoundary);
+
+/**
+* Construct a sparse matrix representation of a valence-weighted Laplacian.
+* The valence weighted Laplacian is defined solely in terms of the connectivity
+* of the mesh.  Note, by construction this should be a symmetric matrix.
+*
+* The mesh itself is assumed to have N interior vertices, and M boundary vertices.
+*
+* Row i represents the Laplacian at vert_i, the non-zero entries correspond to
+* the incident one-ring vertices vert_j.
+*
+* L_{ij} = 1/\sqrt(valence(i) + valence(j))   if vert_j is in the one-ring of vert_i
+* L_{ii} = -Sum{ L_{ij}, j != i}
+*
+* @param DynamicMesh        The triangle mesh
+* @param VertexMap          Additional arrays used to map between vertexID and offset in a linear array (i.e. the row).
+*                           The vertices are ordered so that last M ( = VertexMap.NumBoundaryVerts() )  correspond to those on the boundary.
+* @param LaplacianInterior  On return, the laplacian operator that acts on the interior vertices: sparse  N x N matrix
+* @param LaplacianBoundary  On return, the portion of the operator that acts on the boundary vertices: sparse  N x M matrix
+*
+*   LaplacianInterior * Vector_InteriorVerts + LaplacianBoundary * Vector_BoundaryVerts = Full Laplacian applied to interior vertices.
+*/
+void ConstructValenceWeightedLaplacian(const FDynamicMesh3& DynamicMesh, 
+	                                   FVertexLinearization& VertexMap, 
+	                                   FSparseMatrixD& LaplacianInterior, 
+	                                   FSparseMatrixD& LaplacianBoundary);
+
+
+/**
+* Construct a sparse matrix representation using a cotangent-weighted  Laplacian.
+* NB: there is no reason to expect this to be a symmetric matrix.
+*
+* The mesh itself is assumed to have N interior vertices, and M boundary vertices.
+*
+* @param DynamicMesh        The triangle mesh
+* @param VertexMap          Additional arrays used to map between vertexID and offset in a linear array (i.e. the row).
+*                           The vertices are ordered so that last M ( = VertexMap.NumBoundaryVerts() )  correspond to those on the boundary.
+* @param LaplacianInterior  On return, the laplacian operator that acts on the interior vertices: sparse  N x N matrix
+* @param LaplacianBoundary  On return, the portion of the operator that acts on the boundary vertices: sparse  N x M matrix
+* @param bClampWeights      Indicates if the off-diagonal weights should be clamped on construction.
+*                           in practice this is desirable when creating the biharmonic operator, but not the mean curvature flow operator
+*
+*   LaplacianInterior * Vector_InteriorVerts + LaplacianBoundary * Vector_BoundaryVerts = Full Laplacian applied to interior vertices.
+*/
+void  ConstructCotangentLaplacian(const FDynamicMesh3& DynamicMesh, 
+	                              FVertexLinearization& VertexMap, 
+	                              FSparseMatrixD& LaplacianInterior, 
+	                              FSparseMatrixD& LaplacianBoundary, 
+	                              const bool bClampWeights = false);
+
+
+/**
+* Construct a sparse matrix representation using a cotangent-weighted  Laplacian.
+* but returns the result in two symmetric parts.
+*
+* The mesh itself is assumed to have N interior vertices, and M boundary vertices.
+*
+*  (AreaMatrix^{-1}) * L_hat  = Cotangent weighted Laplacian.
+*
+*
+* @param DynamicMesh        The triangle mesh
+* @param VertexMap          Additional arrays used to map between vertexID and offset in a linear array (i.e. the row).
+*                           The vertices are ordered so that last M ( = VertexMap.NumBoundaryVerts() )  correspond to those on the boundary.
+* @param AreaMatrix         On return, the mass matrix for the internal vertices.  sparse N x N matrix
+* @param LaplacianInterior  On return, the laplacian operator that acts on the interior vertices: sparse  N x N matrix - symmetric
+* @param LaplacianBoundary  On return, the portion of the operator that acts on the boundary vertices: sparse  N x M matrix
+*
+*   AreaMatrix^{-1} * ( LaplacianInterior * Vector_InteriorVerts + LaplacianBoundary * Vector_BoundaryVerts) = Full Laplacian applied to interior vertices.
+*/
+void  ConstructCotangentLaplacian(const FDynamicMesh3& DynamicMesh, 
+	                              FVertexLinearization& VertexMap, 
+	                              FSparseMatrixD& AreaMatrix, 
+	                              FSparseMatrixD& LaplacianInterior, 
+	                              FSparseMatrixD& LaplacianBoundary);
+
+
+
+/**
+* Construct a sparse matrix representation using a meanvalue-weighted  Laplacian.
+* NB: there is no reason to expect this to be a symmetric matrix.
+*
+* The mesh itself is assumed to have N interior vertices, and M boundary vertices.
+*
+* @param DynamicMesh        The triangle mesh
+* @param VertexMap          Additional arrays used to map between vertexID and offset in a linear array (i.e. the row).
+*                           The vertices are ordered so that last M ( = VertexMap.NumBoundaryVerts() )  correspond to those on the boundary.
+* @param LaplacianInterior  On return, the laplacian operator that acts on the interior vertices: sparse  N x N matrix
+* @param LaplacianBoundary  On return, the portion of the operator that acts on the boundary vertices: sparse  N x M matrix
+*
+*   LaplacianInterior * Vector_InteriorVerts + LaplacianBoundary * Vector_BoundaryVerts = Full Laplacian applied to interior vertices.
+*/
+void ConstructMeanValueWeightLaplacian(const FDynamicMesh3& DynamicMesh, 
+	                                   FVertexLinearization& VertexMap, 
+	                                   FSparseMatrixD& LaplacianInterior, 
+	                                   FSparseMatrixD& LaplacianBoundary);
+
+
+
+
+
+TUniquePtr<FSparseMatrixD> ConstructLaplacian(const ELaplacianWeightScheme Scheme,
+	const FDynamicMesh3& DynamicMesh,
+	FVertexLinearization& VertexMap,
+	TArray<int32>* BoundaryVerts = NULL);
+
+
 /**
 * Construct a sparse matrix representation of a uniform weighted Laplacian.
 * The uniform weighted Laplacian is defined solely in terms of the connectivity
@@ -25,10 +207,12 @@
 *
 *
 * @param DynamicMesh3   The triangle mesh
-* @param VertexMap      Additional arrays used to map between vertexID and offset in a linear array (i.e. the row)
+* @param VertexMap      On return - Additional arrays used to map between vertexID and offset in a linear array (i.e. the row)
 * @param BoundaryVerts  Optionally capture the verts that are on mesh edges.
 */
-TUniquePtr<FSparseMatrixD> ConstructUniformLaplacian(const FDynamicMesh3& DynamicMesh, FVertexLinearization& VertexMap, TArray<int32>* BoundaryVerts = NULL);
+TUniquePtr<FSparseMatrixD> ConstructUniformLaplacian(const FDynamicMesh3& DynamicMesh,
+	FVertexLinearization& VertexMap,
+	TArray<int32>* BoundaryVerts = NULL);
 
 /**
 * Construct a sparse matrix representation of an umbrella weighted Laplacian.
@@ -43,10 +227,12 @@ TUniquePtr<FSparseMatrixD> ConstructUniformLaplacian(const FDynamicMesh3& Dynami
 *
 *
 * @param DynamicMesh3   The triangle mesh
-* @param VertexMap      Additional arrays used to map between vertexID and offset in a linear array (i.e. the row)
+* @param VertexMap      On return, Additional arrays used to map between vertexID and offset in a linear array (i.e. the row)
 * @param BoundaryVerts  Optionally capture the verts that are on mesh edges.
 */
-TUniquePtr<FSparseMatrixD> ConstructUmbrellaLaplacian(const FDynamicMesh3& DynamicMesh, FVertexLinearization& VertexMap, TArray<int32>* BoundaryVerts = NULL);
+TUniquePtr<FSparseMatrixD> ConstructUmbrellaLaplacian(const FDynamicMesh3& DynamicMesh,
+	FVertexLinearization& VertexMap,
+	TArray<int32>* BoundaryVerts = NULL);
 
 /**
 * Construct a sparse matrix representation of a valence-weighted Laplacian.
@@ -64,7 +250,9 @@ TUniquePtr<FSparseMatrixD> ConstructUmbrellaLaplacian(const FDynamicMesh3& Dynam
 * @param VertexMap      Additional arrays used to map between vertexID and offset in a linear array (i.e. the row)
 * @param BoundaryVerts  Optionally capture the verts that are on mesh edges.
 */
-TUniquePtr<FSparseMatrixD>  ConstructValenceWeightedLaplacian(const FDynamicMesh3& DynamicMesh3, FVertexLinearization& VertexMap, TArray<int32>* BoundaryVerts = NULL);
+TUniquePtr<FSparseMatrixD>  ConstructValenceWeightedLaplacian(const FDynamicMesh3& DynamicMesh3,
+	FVertexLinearization& VertexMap,
+	TArray<int32>* BoundaryVerts = NULL);
 
 /**
 * Construct a sparse matrix representation using a cotangent-weighted  Laplacian.
@@ -76,8 +264,10 @@ TUniquePtr<FSparseMatrixD>  ConstructValenceWeightedLaplacian(const FDynamicMesh
 *                       in practice this is desirable when creating the biharmonic operator, but not the mean curvature flow operator
 * @param BoundaryVerts  Optionally capture the verts that are on mesh edges.
 */
-
-TUniquePtr<FSparseMatrixD> ConstructCotangentLaplacian(const FDynamicMesh3& DynamicMesh, FVertexLinearization& VertexMap, const bool bClampWeights = false, TArray<int32>* BoundaryVerts = NULL);
+TUniquePtr<FSparseMatrixD> ConstructCotangentLaplacian(const FDynamicMesh3& DynamicMesh,
+	FVertexLinearization& VertexMap,
+	const bool bClampWeights = false,
+	TArray<int32>* BoundaryVerts = NULL);
 
 /**
 * Construct a sparse matrix representation using a cotangent-weighted  Laplacian.
@@ -94,10 +284,10 @@ TUniquePtr<FSparseMatrixD> ConstructCotangentLaplacian(const FDynamicMesh3& Dyna
 * @param VertexMap      Additional arrays used to map between vertexID and offset in a linear array (i.e. the row)
 * @param BoundaryVerts  Optionally capture the verts that are on mesh edges.
 */
-
-TUniquePtr<FSparseMatrixD> ConstructCotangentLaplacian(const FDynamicMesh3& DynamicMesh, FVertexLinearization& VertexMap, FSparseMatrixD& AreaMatrix, TArray<int32>* BoundaryVerts = NULL);
-
-
+TUniquePtr<FSparseMatrixD> ConstructCotangentLaplacian(const FDynamicMesh3& DynamicMesh,
+	FVertexLinearization& VertexMap,
+	FSparseMatrixD& AreaMatrix,
+	TArray<int32>* BoundaryVerts = NULL);
 
 /**
 * Construct a sparse matrix representation using a meanvalue-weighted  Laplacian.
@@ -109,10 +299,11 @@ TUniquePtr<FSparseMatrixD> ConstructCotangentLaplacian(const FDynamicMesh3& Dyna
 *                       in practice this is desirable when creating the biharmonic operator, but not the mean curvature flow operator
 * @param BoundaryVerts  Optionally capture the verts that are on mesh edges.
 */
+TUniquePtr<FSparseMatrixD> ConstructMeanValueWeightLaplacian(const FDynamicMesh3& DynamicMesh,
+	FVertexLinearization& VertexMap,
+	TArray<int32>* BoundaryVerts = NULL);
 
-TUniquePtr<FSparseMatrixD> ConstructMeanValueWeightLaplacian(const FDynamicMesh3& DynamicMesh, FVertexLinearization& VertexMap, TArray<int32>* BoundaryVerts = NULL);
 
-TUniquePtr<FSparseMatrixD> ConstructLaplacian(const ELaplacianWeightScheme Scheme, const FDynamicMesh3& DynamicMesh, FVertexLinearization& VertexMap, TArray<int32>* BoundaryVerts = NULL);
 
 
 /**
@@ -122,7 +313,7 @@ FString LaplacianSchemeName(const ELaplacianWeightScheme Scheme);
 
 
 /**
-* Only certian laplacian operators are symmetric...
+* Only certain laplacian operators are symmetric...
 */
 static bool bIsSymmetricLaplacian(const ELaplacianWeightScheme Scheme)
 {
diff --git a/Engine/Plugins/Experimental/GeometryProcessing/Source/MeshSolverUtilities/Private/MeshElementLinearizations.h b/Engine/Plugins/Experimental/GeometryProcessing/Source/MeshSolverUtilities/Private/MeshElementLinearizations.h
--- a/Engine/Plugins/Experimental/GeometryProcessing/Source/MeshSolverUtilities/Private/MeshElementLinearizations.h
+++ b/Engine/Plugins/Experimental/GeometryProcessing/Source/MeshSolverUtilities/Private/MeshElementLinearizations.h
@@ -62,7 +62,8 @@ private:
 
 /**
 * Used linearize the VtxIds in a mesh as a single array and allow mapping from array offset to mesh VtxId.
-* Generally, the array offset will correspond to a matrix row when forming a Laplacian
+* Generally, the array offset will correspond to a matrix row when forming a Laplacian.
+* The last NumBoundaryVerts are the boundary verts.  This may be zero.
 */
 class FVertexLinearization : public FMeshElementLinearization
 {
@@ -77,12 +78,64 @@ public:
 	{
 		Empty();
 		FMeshElementLinearization::Populate(DynamicMesh.MaxVertexID(), DynamicMesh.VertexCount(), DynamicMesh.VertexIndicesItr());
+		RemapBoundaryVerts(DynamicMesh);
 	}
 
 	int32 NumVerts() const { return FMeshElementLinearization::NumIds(); }
 
+	int32 NumBoundaryVerts() const { return NumBndryVerts; }
+
+private:
+
+	// Moves the boundary verts to the end of the arrays and records the number of boundary verts
+	void RemapBoundaryVerts(const FDynamicMesh3& DynamicMesh)
+	{
+		int32 VertCount = NumVerts();
+
+		// Collect the BoundaryVerts and the internal verts in two array
+		TArray<int32> BoundaryVertIds;
+		TArray<int32> TmpToIdMap;
+		TmpToIdMap.Reserve(VertCount);
+		for (int32 i = 0, I = ToIdMap.Num(); i < I; ++i)
+		{
+			int32 VtxId = ToIdMap[i];
+			if (DynamicMesh.IsBoundaryVertex(VtxId))
+			{
+				BoundaryVertIds.Add(VtxId);
+			}
+			else
+			{
+				TmpToIdMap.Add(VtxId);
+			}
+		}
+
+		// The number of boundary verts
+		NumBndryVerts = BoundaryVertIds.Num();
+
+		// Merge the boundary verts at the tail 
+		// Add the Boundary verts at the end of the array.
+		TmpToIdMap.Append(BoundaryVertIds);
+	
+		// rebuild the 'to' Index
+		for (int32 i = 0, I = ToIndexMap.Num(); i < I; ++i)
+		{
+			ToIndexMap[i] = FDynamicMesh3::InvalidID;
+		}
+		for (int32 i = 0, I = TmpToIdMap.Num(); i < I; ++i)
+		{
+			int32 Id = TmpToIdMap[i];
+			ToIndexMap[Id] = i;
+		}
+
+		// swap the temp
+		Swap(TmpToIdMap, ToIdMap);
+	}
+
 private:
+
 	FVertexLinearization(const FVertexLinearization&);
+
+	int32 NumBndryVerts = 0;
 };
 
 /**
