handle several missed cases in GroupTopology (boundary edge, loop with a single corner). Add ::GetGroupEdgeTangent() which returns estimate of tangent direction for GroupEdge, useful for local frame in move tools.
#rnx

diff --git a/Engine/Plugins/Experimental/GeometryProcessing/Source/DynamicMesh/Private/GroupTopology.cpp b/Engine/Plugins/Experimental/GeometryProcessing/Source/DynamicMesh/Private/GroupTopology.cpp
--- a/Engine/Plugins/Experimental/GeometryProcessing/Source/DynamicMesh/Private/GroupTopology.cpp
+++ b/Engine/Plugins/Experimental/GeometryProcessing/Source/DynamicMesh/Private/GroupTopology.cpp
@@ -426,16 +426,13 @@ void FGroupTopology::ExtractGroupEdges(FGroup& Group)
 
 			FGroupEdge Edge = { EdgeID };
 
-			int i1 = CornerIndices[k+1];
-			check(i0 != i1);	// this case will happen on a cylinder w/ a cut side, I think...
-								// (but that case won't exist because that's wouldn't be a group boundary edge!)
+			int i1 = CornerIndices[k+1];		// note: i1 == i0 on a closed loop, ie NumSpans == 1
 			TArray<int> SpanVertices;
-			while (i0 != i1)
-			{
+			do {
 				SpanVertices.Add(Loop.Vertices[i0]);
 				i0 = (i0 + 1) % NumV;
-			}
-			SpanVertices.Add(Loop.Vertices[i0]);
+			} while (i0 != i1);
+			SpanVertices.Add(Loop.Vertices[i0]);	// add last vertex
 
 			Edge.Span = FEdgeSpan(Mesh);
 			Edge.Span.InitializeFromVertices(SpanVertices);
@@ -448,10 +445,9 @@ void FGroupTopology::ExtractGroupEdges(FGroup& Group)
 
 
 
-
-
 int FGroupTopology::FindExistingGroupEdge(int GroupID, int OtherGroupID, int FirstVertexID)
 {
+	// if this is a boundary edge, we cannot have created it already
 	if (OtherGroupID < 0)
 	{
 		return -1;
@@ -481,3 +477,21 @@ int FGroupTopology::FindExistingGroupEdge(int GroupID, int OtherGroupID, int Fir
 }
 
 
+
+bool FGroupTopology::GetGroupEdgeTangent(int GroupEdgeID, FVector3d& TangentOut) const
+{
+	check(GroupEdgeID >= 0 && GroupEdgeID < Edges.Num());
+	const FGroupEdge& Edge = Edges[GroupEdgeID];
+	FVector3d StartPos = Mesh->GetVertex(Edge.Span.Vertices[0]);
+	FVector3d EndPos = Mesh->GetVertex(Edge.Span.Vertices[Edge.Span.Vertices.Num()-1]);
+	if (StartPos.DistanceSquared(EndPos) > 100 * FMathd::ZeroTolerance)
+	{
+		TangentOut = (EndPos - StartPos).Normalized();
+		return true;
+	}
+	else
+	{
+		TangentOut = FVector3d::UnitX();
+		return false;
+	}
+}
\ No newline at end of file
diff --git a/Engine/Plugins/Experimental/GeometryProcessing/Source/DynamicMesh/Public/GroupTopology.h b/Engine/Plugins/Experimental/GeometryProcessing/Source/DynamicMesh/Public/GroupTopology.h
--- a/Engine/Plugins/Experimental/GeometryProcessing/Source/DynamicMesh/Public/GroupTopology.h
+++ b/Engine/Plugins/Experimental/GeometryProcessing/Source/DynamicMesh/Public/GroupTopology.h
@@ -163,6 +163,11 @@ public:
 	/** Add all the group boundary vertices of the given GroupID to the Vertices set */
 	void CollectGroupBoundaryVertices(int GroupID, TSet<int>& Vertices) const;
 
+	/** 
+	 * Calculate tangent of group edge, as direction from start to end endpoints
+	 * @return false if edge is degenerate or start == end 
+	 */
+	bool GetGroupEdgeTangent(int GroupEdgeID, FVector3d& TangentOut) const;
 
 protected:
 	const FDynamicMesh3* Mesh = nullptr;
@@ -179,11 +184,24 @@ protected:
 	FIndex2i MakeEdgeID(int MeshEdgeID)
 	{
 		FIndex2i EdgeTris = Mesh->GetEdgeT(MeshEdgeID);
-		return MakeEdgeID(GetGroupID(EdgeTris.A), GetGroupID(EdgeTris.B));
+
+		if (EdgeTris.A == FDynamicMesh3::InvalidID)
+		{
+			return FIndex2i(GetGroupID(EdgeTris.B), FDynamicMesh3::InvalidGroupID);
+		}
+		else if (EdgeTris.B == FDynamicMesh3::InvalidID)
+		{
+			return FIndex2i(GetGroupID(EdgeTris.A), FDynamicMesh3::InvalidGroupID);
+		}
+		else
+		{
+			return MakeEdgeID(GetGroupID(EdgeTris.A), GetGroupID(EdgeTris.B));
+		}
 	}
 	FIndex2i MakeEdgeID(int Group1, int Group2)
 	{
 		check(Group1 != Group2);
+		check(Group1 >= 0 && Group2 >= 0);
 		return (Group1 < Group2) ? FIndex2i(Group1, Group2) : FIndex2i(Group2, Group1);
 	}
 
