Add TCircle2 and TCircle3, TMathUtil::ACos, FVector2::SignedAngleR, GTEUtil.h with functions to convert between GP and GTE types, DistLine3Circle3 wrapper for GTEngine implementation
#rnx

diff --git a/Engine/Plugins/Experimental/GeometryProcessing/Source/GeometricObjects/Public/CircleTypes.h b/Engine/Plugins/Experimental/GeometryProcessing/Source/GeometricObjects/Public/CircleTypes.h
new file mode 100644
--- /dev/null
+++ b/Engine/Plugins/Experimental/GeometryProcessing/Source/GeometricObjects/Public/CircleTypes.h
@@ -0,0 +1,273 @@
+// Copyright 1998-2019 Epic Games, Inc. All Rights Reserved.
+
+// Port of WildMagic DistPoint3TCircle3
+
+#pragma once
+
+#include "VectorTypes.h"
+#include "VectorUtil.h"
+#include "BoxTypes.h"
+#include "OrientedBoxTypes.h"
+#include "FrameTypes.h"
+
+template<typename RealType>
+struct TCircle2
+{
+	FVector2<RealType> Center = FVector2<RealType>(0, 0);
+	RealType Radius = (RealType)1;
+	bool bIsReversed = false;
+
+	TCircle2() {}
+
+	explicit TCircle2(const RealType& RadiusIn)
+	{
+		Center = FVector2<RealType>(0, 0);
+		Radius = RadiusIn;
+	}
+
+	TCircle2(const FVector2<RealType>& CenterIn, const RealType& RadiusIn)
+	{
+		Center = CenterIn;
+		Radius = RadiusIn;
+	}
+
+
+	RealType GetCircumference() const
+	{
+		return TMathUtil<RealType>::TwoPi * Radius;
+	}
+	void SetCircumference(RealType NewCircumference)
+	{
+		Radius = NewCircumference / TMathUtil<RealType>::TwoPi;
+	}
+
+
+	RealType GetDiameter() const
+	{
+		return (RealType)2 * Radius;
+	}
+	void SetDiameter(RealType NewDiameter)
+	{
+		Radius = NewDiameter * (RealType)0.5;
+	}
+
+	RealType GetArea() const
+	{
+		return TMathUtil<RealType>::Pi * Radius * Radius;
+	}
+	void SetArea(RealType NewArea)
+	{
+		Radius = TMathUtil<RealType>::Sqrt(NewArea / TMathUtil<RealType>::Pi);
+	}
+
+
+	RealType GetCurvature() const
+	{
+		return (RealType)1 / Radius;
+	}
+
+	RealType GetSignedCurvature() const
+	{
+		return (bIsReversed) ? ((RealType)-1 / Radius) : ((RealType)1 / Radius);
+	}
+
+
+	FVector2<RealType> GetPointFromAngleD(RealType AngleDeg) const
+	{
+		return GetPointFromAngleR(AngleDeg * TMathUtil<RealType>::DegToRad);
+	}
+
+	FVector2<RealType> GetPointFromAngleR(RealType AngleRad) const
+	{
+		RealType c = TMathUtil<RealType>::Cos(AngleRad), s = TMathUtil<RealType>::Sin(AngleRad);
+		return FVector2<RealType>(Center.X + c*Radius, Center.Y + s*Radius);
+	}
+
+
+	FVector2<RealType> GetPointFromUnitParameter(RealType UnitParam) const
+	{
+		RealType AngleRad = ((bIsReversed) ? (-UnitParam) : (UnitParam)) * TMathUtil<RealType>::TwoPi;
+		return GetPointFromAngleR(AngleRad);
+	}
+
+
+	bool IsInside(const FVector2<RealType>& Point) const
+	{
+		return Center.DistanceSquared(Point) < Radius*Radius;
+	}
+
+
+	RealType SignedDistance(const FVector2<RealType>& Point) const
+	{
+		return Center.Distance(Point) - Radius;
+	}
+
+	RealType Distance(const FVector2<RealType>& Point) const
+	{
+		return TMathUtil<RealType>::Abs(Center.Distance(Point) - Radius);
+	}
+
+
+	TAxisAlignedBox2<RealType> GetBoundingBox() const
+	{
+		FVector2<RealType>(Center.X + Radius, Center.Y + Radius);
+		return TAxisAlignedBox2<RealType>(
+			FVector2<RealType>(Center.X - Radius, Center.Y - Radius),
+			FVector2<RealType>(Center.X + Radius, Center.Y + Radius));
+	}
+
+
+	RealType GetBoundingPolygonRadius(int NumSides) const
+	{
+		RealType DeltaAngle = (TMathUtil<RealType>::TwoPi / (RealType)NumSides) / (RealType)2;
+		return Radius / TMathUtil<RealType>::Cos(DeltaAngle);
+	}
+
+};
+
+typedef TCircle2<float> FCircle2f;
+typedef TCircle2<double> FCircle2d;
+
+
+
+
+
+
+
+template<typename RealType>
+struct TCircle3
+{
+	TFrame3<RealType> Frame;
+	RealType Radius = (RealType)1;
+	bool bIsReversed = false;
+
+	TCircle3() {}
+
+	explicit TCircle3(const RealType& RadiusIn)
+	{
+		Radius = RadiusIn;
+	}
+
+	TCircle3(const FVector3<RealType>& CenterIn, const RealType& RadiusIn)
+	{
+		Frame.Origin = CenterIn;
+		Radius = RadiusIn;
+	}
+
+	TCircle3(const TFrame3<RealType>& FrameIn, const RealType& RadiusIn)
+	{
+		Frame = FrameIn;
+		Radius = RadiusIn;
+	}
+
+
+	const FVector3<RealType>& GetCenter() const
+	{
+		return Frame.Origin;
+	}
+
+	FVector3<RealType> GetNormal() const
+	{
+		return Frame.Z();
+	}
+
+
+	RealType GetCircumference() const
+	{
+		return TMathUtil<RealType>::TwoPi * Radius;
+	}
+	void SetCircumference(RealType NewCircumference)
+	{
+		Radius = NewCircumference / TMathUtil<RealType>::TwoPi;
+	}
+
+
+	RealType GetDiameter() const
+	{
+		return (RealType)2 * Radius;
+	}
+	void SetDiameter(RealType NewDiameter)
+	{
+		Radius = NewDiameter * (RealType)0.5;
+	}
+
+	RealType GetArea() const
+	{
+		return TMathUtil<RealType>::Pi * Radius * Radius;
+	}
+	void SetArea(RealType NewArea)
+	{
+		Radius = TMathUtil<RealType>::Sqrt(NewArea / TMathUtil<RealType>::Pi);
+	}
+
+
+	RealType GetCurvature() const
+	{
+		return (RealType)1 / Radius;
+	}
+
+	RealType GetSignedCurvature() const
+	{
+		return (bIsReversed) ? ((RealType)-1 / Radius) : ((RealType)1 / Radius);
+	}
+
+
+	FVector3<RealType> GetPointFromAngleD(RealType AngleDeg) const
+	{
+		return GetPointFromAngleR(AngleDeg * TMathUtil<RealType>::DegToRad);
+	}
+
+	FVector3<RealType> GetPointFromAngleR(RealType AngleRad) const
+	{
+		RealType c = TMathUtil<RealType>::Cos(AngleRad), s = TMathUtil<RealType>::Sin(AngleRad);
+		return Frame.FromPlaneUV(FVector2<RealType>(Radius*c, Radius*s), 2);
+	}
+
+	FVector3<RealType> GetPointFromUnitParameter(RealType UnitParam) const
+	{
+		RealType AngleRad = ((bIsReversed) ? (-UnitParam) : (UnitParam)) * TMathUtil<RealType>::TwoPi;
+		return GetPointFromAngleR(AngleRad);
+	}
+
+
+
+
+	FVector3<RealType> ClosestPoint(const FVector3<RealType>& QueryPoint) const
+	{
+		const FVector3<RealType>& Center = Frame.Origin;
+		FVector3<RealType> Normal = Frame.GetAxis(2);
+
+		FVector3<RealType> PointDelta = QueryPoint - Center;
+		FVector3<RealType> DeltaInPlane = PointDelta - Normal.Dot(PointDelta)*Normal;
+		RealType OriginDist = DeltaInPlane.Length();
+		if (OriginDist > (RealType)0)
+		{
+			return Center + (Radius / OriginDist)*DeltaInPlane;
+		}
+		else    // all points equidistant, use any one
+		{
+			return Frame.Origin + Radius * Frame.GetAxis(0);
+		}
+	}
+
+
+	RealType DistanceSquared(const FVector3<RealType>& Point) const
+	{
+		return Point.DistanceSquared(ClosestPoint(Point));
+	}
+	
+	RealType Distance(const FVector3<RealType>& Point) const
+	{
+		return TMathUtil<RealType>::Sqrt(DistanceSquared(Point));
+	}
+
+	TOrientedBox3<RealType> GetBoundingBox() const
+	{
+		return TOrientedBox3<RealType>(Frame, FVector3<RealType>(Radius, Radius, 0));
+	}
+
+};
+
+typedef TCircle3<float> FCircle3f;
+typedef TCircle3<double> FCircle3d;
+
diff --git a/Engine/Plugins/Experimental/GeometryProcessing/Source/GeometricObjects/Public/MathUtil.h b/Engine/Plugins/Experimental/GeometryProcessing/Source/GeometricObjects/Public/MathUtil.h
--- a/Engine/Plugins/Experimental/GeometryProcessing/Source/GeometricObjects/Public/MathUtil.h
+++ b/Engine/Plugins/Experimental/GeometryProcessing/Source/GeometricObjects/Public/MathUtil.h
@@ -60,6 +60,7 @@ public:
 	static inline RealType Atan2(const RealType ValueY, const RealType ValueX);
 	static inline RealType Sin(const RealType Value);
 	static inline RealType Cos(const RealType Value);
+	static inline RealType ACos(const RealType Value);
 
 
 
@@ -163,6 +164,12 @@ RealType TMathUtil<RealType>::Cos(const RealType Value)
 	return cos(Value);
 }
 
+template<typename RealType>
+RealType TMathUtil<RealType>::ACos(const RealType Value)
+{
+	return acos(Value);
+}
+
 
 template<typename RealType>
 RealType TMathUtil<RealType>::Atan2Positive(const RealType Y, const RealType X)
diff --git a/Engine/Plugins/Experimental/GeometryProcessing/Source/GeometricObjects/Public/VectorTypes.h b/Engine/Plugins/Experimental/GeometryProcessing/Source/GeometricObjects/Public/VectorTypes.h
--- a/Engine/Plugins/Experimental/GeometryProcessing/Source/GeometricObjects/Public/VectorTypes.h
+++ b/Engine/Plugins/Experimental/GeometryProcessing/Source/GeometricObjects/Public/VectorTypes.h
@@ -546,6 +546,23 @@ struct FVector2
 		return (T)acos(ClampedDot);
 	}
 
+	// Angle in Radians
+	T SignedAngleR(const FVector2<T>& V2) const
+	{
+		T DotVal = Dot(V2);
+		T ClampedDot = (DotVal < (T)-1) ? (T)-1 : ((DotVal > (T)1) ? (T)1 : DotVal);
+		T Direction = Cross(V2);
+		if (Direction*Direction < TMathUtil<T>::ZeroTolerance)
+		{
+			return (DotVal < 0) ? TMathUtil<T>::Pi : (T)0;
+		}
+		else
+		{
+			T Sign = Direction < 0 ? (T)-1 : (T)1;
+			return Sign * TMathUtil<T>::ACos(ClampedDot);
+		}
+	}
+
 	T Normalize(const T Epsilon = 0)
 	{
 		T length = Length();
diff --git a/Engine/Plugins/Experimental/GeometryProcessing/Source/GeometryAlgorithms/Private/Distance/DistLine3Circle3.cpp b/Engine/Plugins/Experimental/GeometryProcessing/Source/GeometryAlgorithms/Private/Distance/DistLine3Circle3.cpp
new file mode 100644
--- /dev/null
+++ b/Engine/Plugins/Experimental/GeometryProcessing/Source/GeometryAlgorithms/Private/Distance/DistLine3Circle3.cpp
@@ -0,0 +1,38 @@
+// Copyright 1998-2019 Epic Games, Inc. All Rights Reserved.
+
+#include "Distance/DistLine3Circle3.h"
+#include "ThirdParty/GTEngine/Mathematics/GteDistLine3Circle3.h"
+#include "GteUtil.h"
+
+
+template <typename Real>
+Real TDistLine3Circle3<Real>::ComputeResult()
+{
+	if (DistanceSquared >= 0)
+	{
+		return DistanceSquared;
+	}
+
+	gte::Line3<Real> gteLine = Convert(Line);
+	gte::Circle3<Real> gteCircle = Convert(Circle);
+
+	gte::DCPQuery<Real, gte::Line3<Real>, gte::Circle3<Real>> Query;
+	auto Result = Query(gteLine, gteCircle);
+
+	DistanceSquared = Result.sqrDistance;
+	NumClosestPairs = Result.numClosestPairs;
+	LineClosest[0] = Convert(Result.lineClosest[0]);
+	LineClosest[1] = Convert(Result.lineClosest[1]);
+	CircleClosest[0] = Convert(Result.circleClosest[0]);
+	CircleClosest[1] = Convert(Result.circleClosest[1]);
+	bIsEquiDistant = Result.equidistant;
+
+	return DistanceSquared;
+}
+
+
+
+
+
+template class GEOMETRYALGORITHMS_API TDistLine3Circle3<float>;
+template class GEOMETRYALGORITHMS_API TDistLine3Circle3<double>;
diff --git a/Engine/Plugins/Experimental/GeometryProcessing/Source/GeometryAlgorithms/Private/GteUtil.h b/Engine/Plugins/Experimental/GeometryProcessing/Source/GeometryAlgorithms/Private/GteUtil.h
new file mode 100644
--- /dev/null
+++ b/Engine/Plugins/Experimental/GeometryProcessing/Source/GeometryAlgorithms/Private/GteUtil.h
@@ -0,0 +1,35 @@
+// Copyright 1998-2019 Epic Games, Inc. All Rights Reserved.
+
+#include "VectorTypes.h"
+#include "LineTypes.h"
+#include "CircleTypes.h"
+
+#include "ThirdParty/GTEngine/Mathematics/GteVector3.h"
+#include "ThirdParty/GTEngine/Mathematics/GteLine.h"
+#include "ThirdParty/GTEngine/Mathematics/GteCircle3.h"
+
+
+template<typename Real>
+gte::Vector3<Real> Convert(const FVector3<Real>& Vec)
+{
+	return gte::Vector3<Real>({ Vec.X, Vec.Y, Vec.Z });
+}
+
+template<typename Real>
+FVector3<Real> Convert(const gte::Vector3<Real>& Vec)
+{
+	return FVector3<Real>(Vec[0], Vec[1], Vec[2]);
+}
+
+template<typename Real>
+gte::Line3<Real> Convert(const TLine3<Real>& Line)
+{
+	return gte::Line3<Real>(Convert(Line.Origin), Convert(Line.Direction));
+}
+
+template<typename Real>
+gte::Circle3<Real> Convert(const TCircle3<Real>& Circle)
+{
+	return gte::Circle3<Real>(
+		Convert(Circle.Frame.Origin), Convert(Circle.GetNormal()), Circle.Radius);
+}
diff --git a/Engine/Plugins/Experimental/GeometryProcessing/Source/GeometryAlgorithms/Public/Distance/DistLine3Circle3.h b/Engine/Plugins/Experimental/GeometryProcessing/Source/GeometryAlgorithms/Public/Distance/DistLine3Circle3.h
new file mode 100644
--- /dev/null
+++ b/Engine/Plugins/Experimental/GeometryProcessing/Source/GeometryAlgorithms/Public/Distance/DistLine3Circle3.h
@@ -0,0 +1,50 @@
+// Copyright 1998-2019 Epic Games, Inc. All Rights Reserved.
+
+// Port of geometry3Sharp DistLine3Circle3
+// which was ported from WildMagic 5 
+
+#pragma once
+
+#include "VectorTypes.h"
+#include "CircleTypes.h"
+#include "LineTypes.h"
+
+
+/**
+* Compute unsigned distance and closest-points between 3D line and 3D circle
+*/
+template <typename Real>
+class TDistLine3Circle3
+{
+public:
+	// Input
+	TLine3<Real> Line;
+	TCircle3<Real> Circle;
+
+	// Output
+	Real DistanceSquared = -1.0;
+
+	int NumClosestPairs = 0;
+	FVector3<Real> LineClosest[2];
+	FVector3<Real> CircleClosest[2];
+	bool bIsEquiDistant;
+
+	TDistLine3Circle3(const TLine3<Real>& LineIn, const TCircle3<Real>& CircleIn) : Line(LineIn), Circle(CircleIn)
+	{
+	}
+
+	Real Get()
+	{
+		return TMathUtil<Real>::Sqrt(ComputeResult());
+	}
+	Real GetSquared()
+	{
+		return ComputeResult();
+	}
+
+	Real ComputeResult();
+};
+
+typedef TDistLine3Circle3<float> FDistLine3Circle3f;
+typedef TDistLine3Circle3<double> FDistLine3Circle3d;
+
