general ConstrainedDelaunay2d triangulation + tests

diff --git a/Engine/Plugins/Experimental/GeometryProcessing/Source/GeometricObjects/Public/Polygon2.h b/Engine/Plugins/Experimental/GeometryProcessing/Source/GeometricObjects/Public/Polygon2.h
--- a/Engine/Plugins/Experimental/GeometryProcessing/Source/GeometricObjects/Public/Polygon2.h
+++ b/Engine/Plugins/Experimental/GeometryProcessing/Source/GeometricObjects/Public/Polygon2.h
@@ -228,7 +228,7 @@ public:
 		inline TSegment2<T> operator*() const
 		{
 			check(polygon != nullptr && i < polygon->VertexCount());
-			return TSegment2<T>(polygon->Vertices[i], polygon->Vertices[i+1 % polygon->VertexCount()]);
+			return TSegment2<T>(polygon->Vertices[i], polygon->Vertices[(i+1) % polygon->VertexCount()]);
 		}
 		//inline TSegment2<T> & operator*();
 		inline SegmentIterator & operator++() 		// prefix
@@ -422,6 +422,46 @@ public:
 		return nWindingNumber != 0;
 	}
 
+	/**
+	 * Check for polygon overlap, aka solid intersection.  (In contrast, note that the "Intersects" method checks for edge intersection)
+	 *
+	 * @return true if the Polygon overlaps the OtherPolygon.
+	 */
+	bool Overlaps(const TPolygon2<T>& OtherPoly) const
+	{
+		if (!Bounds().Intersects(OtherPoly.Bounds()))
+		{
+			return false;
+		}
+
+		for (int i = 0, N = OtherPoly.VertexCount(); i < N; ++i) 
+		{
+			if (Contains(OtherPoly[i]))
+			{
+				return true;
+			}
+		}
+
+		for (int i = 0, N = VertexCount(); i < N; ++i) 
+		{
+			if (OtherPoly.Contains(Vertices[i]))
+			{
+				return true;
+			}
+		}
+
+		for (TSegment2<T> seg : Segments()) 
+		{
+			for (TSegment2<T> oseg : OtherPoly.Segments()) 
+			{
+				if (seg.Intersects(oseg))
+				{
+					return true;
+				}
+			}
+		}
+		return false;
+	}
 
 	/**
 	 * @return true if the Polygon fully contains the OtherPolygon
diff --git a/Engine/Plugins/Experimental/GeometryProcessing/Source/GeometryAlgorithms/Private/ConstrainedDelaunay2.cpp b/Engine/Plugins/Experimental/GeometryProcessing/Source/GeometryAlgorithms/Private/ConstrainedDelaunay2.cpp
new file mode 100644
--- /dev/null
+++ b/Engine/Plugins/Experimental/GeometryProcessing/Source/GeometryAlgorithms/Private/ConstrainedDelaunay2.cpp
@@ -0,0 +1,234 @@
+// Copyright 1998-2019 Epic Games, Inc. All Rights Reserved.
+
+#include "ConstrainedDelaunay2.h"
+#include "ThirdParty/GTEngine/Mathematics/GteConstrainedDelaunay2.h"
+#include "ThirdParty/GTEngine/Mathematics/GteBSNumber.h"
+#include "ThirdParty/GTEngine/Mathematics/GteUIntegerFP32.h"
+#include <vector>
+
+template<class RealType>
+template<class InputRealType>
+void TConstrainedDelaunay2<RealType>::Add(const FDynamicGraph2<InputRealType>& Graph)
+{
+	// TODO: implement
+	check(false);
+}
+
+template<class RealType>
+template<class InputRealType>
+void TConstrainedDelaunay2<RealType>::Add(const TPolygon2<InputRealType>& Polygon, bool bIsHole)
+{
+	int32 VertexStart = Vertices.Num();
+	int32 VertexEnd = VertexStart + Polygon.VertexCount();
+	for (const FVector2<InputRealType> &Vertex : Polygon.GetVertices())
+	{
+		Vertices.Add((FVector2<RealType>)Vertex);
+	}
+
+	TArray<FIndex2i>* EdgeArr;
+	if (bIsHole)
+	{
+		EdgeArr = &HoleEdges;
+	}
+	else
+	{
+		EdgeArr = &Edges;
+	}
+	for (int32 A = VertexEnd - 1, B = VertexStart; B < VertexEnd; A = B++)
+	{
+		EdgeArr->Add(FIndex2i(A, B));
+	}
+}
+
+template<class RealType>
+template<class InputRealType>
+void TConstrainedDelaunay2<RealType>::Add(const TGeneralPolygon2<InputRealType>& GPolygon)
+{
+	Add(GPolygon.GetOuter(), false);
+	const TArray<TPolygon2<InputRealType>>& Holes = GPolygon.GetHoles();
+	for (int HoleIdx = 0, HolesNum = Holes.Num(); HoleIdx < HolesNum; HoleIdx++)
+	{
+		Add(Holes[HoleIdx], true);
+	}
+}
+
+void AddOrderedEdge(TMap<TPair<int, int>, bool>& EdgeMap, int VertA, int VertB)
+{
+	bool bReversed = VertA > VertB;
+	if (bReversed)
+	{
+		Swap(VertA, VertB);
+	}
+	EdgeMap.Add(TPair<int, int>(VertA, VertB), bReversed);
+}
+
+/**
+ * Compute the change in winding number from crossing an oriented edge connecting VertA to VertB
+ *
+ * @param EdgeMap Map of known edges & orientations
+ * @param VertA First vertex on edge
+ * @param VertB Second vertex on edge
+ * @return  -1 if reverse edge (B-A) found, 1 if forward edge (A-B) found, 0 otherwise
+ */
+int WindingAcross(const TMap<TPair<int, int>, bool>& EdgeMap, int VertA, int VertB)
+{
+	bool bReversed = VertA > VertB;
+	if (bReversed)
+	{
+		Swap(VertA, VertB);
+	}
+	TPair<int, int> EdgeKey(VertA, VertB);
+	const bool *bFoundReversed = EdgeMap.Find(EdgeKey);
+	if (!bFoundReversed)
+	{
+		return 0;
+	}
+	bool bSameDir = bReversed == *bFoundReversed;
+	return bSameDir ? 1 : -1;
+}
+
+/**
+ * Check if any edge in edge map connects VertA to VertB (or VertB to VertA)
+ * 
+ * @param EdgeMap Map of known edges & orientations
+ * @param VertA First vertex on edge
+ * @param VertB Second vertex on edge
+ * @return  true if VertA and VertB are connected by an edge (in either direction)
+ */
+bool HasUnorderedEdge(const TMap<TPair<int, int>, bool>& EdgeMap, int VertA, int VertB)
+{
+	return EdgeMap.Contains(TPair<int, int>(FMath::Min(VertA, VertB), FMath::Max(VertA, VertB)));
+}
+
+template<class RealType>
+bool TConstrainedDelaunay2<RealType>::Triangulate()
+{
+	Triangles.Empty();
+
+	check(FillRule <= EFillRule::Odd || bOrientedEdges);
+
+	gte::ConstrainedDelaunay2<double, gte::BSNumber<gte::UIntegerFP32<263>>> Delaunay; // Value of 263 is from comment in GTEngine/Mathematics/GteDelaunay2.h
+
+	std::vector<gte::Vector2<double>> InputVertices;
+	for (int i = 0; i < Vertices.Num(); i++)
+	{
+		InputVertices.push_back(gte::Vector2<double> {{Vertices[i].X, Vertices[i].Y}});
+	}
+
+	if (!Delaunay(Vertices.Num(), &InputVertices[0], 0))
+	{
+		return false;
+	}
+
+	std::vector<int> OutEdges;
+	bool InsertConstraintFailure = false;
+	TMap<TPair<int, int>, bool> BoundaryMap, HoleMap; // tracks all the boundary edges as they are added, so we can later flood fill across them for inside/outside decisions
+	TMap<TPair<int, int>, bool>* EdgeAndHoleMaps[2] = { &BoundaryMap, &HoleMap };
+	bool bBoundaryTrackingFailure = false;
+
+	TArray<FIndex2i>* InputEdgesAndHoles[2] = { &Edges, &HoleEdges };
+	for (int EdgeOrHole = 0; EdgeOrHole < 2; EdgeOrHole++)
+	{
+		TArray<FIndex2i>& Input = *InputEdgesAndHoles[EdgeOrHole];
+		TMap<TPair<int, int>, bool>& InputMap = *EdgeAndHoleMaps[EdgeOrHole];
+		for (const FIndex2i& Edge : Input)
+		{
+			if (!Delaunay.Insert({{Edge.A, Edge.B}}, OutEdges))
+			{
+				// Note the failed edge; we will try to proceed anyway, just without this edge.  Hopefully the CDT is robust and this never happens!
+				ensureMsgf(false, TEXT("CDT edge insertion failed"));
+				InsertConstraintFailure = true;
+				bBoundaryTrackingFailure = true;
+			}
+			else
+			{
+				for (size_t i = 0; i + 1 < OutEdges.size(); i++)
+				{
+					AddOrderedEdge(InputMap, OutEdges[i], OutEdges[i + 1]);
+				}
+			}
+		}
+	}
+
+	
+	const std::vector<int>& Indices = Delaunay.GetIndices();
+	
+	const std::vector<int>& Adj = Delaunay.GetAdjacencies();
+	int TriNum = int(Adj.size() / 3);
+	TArray<int8> Keep;  // values: 0->unprocessed (delete), 1->yes keep, -1->processed, delete
+	Keep.SetNumZeroed(TriNum);
+
+	TArray<TPair<int, int>> ToWalkQ; // Pair of tri index, winding number
+	// seed the queue with all triangles that are on the boundary of the convex hull
+	// note: need *all* not just *one* because of the strategy of refusing to cross hole edges; if using pure winding number classification would just need one boundary triangle to start
+	for (int TriIdx = 0; TriIdx < TriNum; TriIdx++)
+	{
+		int BaseIdx = TriIdx * 3;
+		for (int SubIdx = 0, NextIdx = 2; SubIdx < 3; NextIdx = SubIdx++)
+		{
+			if (Adj[BaseIdx + NextIdx] == -1) // on hull
+			{
+				int VertA = Indices[BaseIdx + SubIdx], VertB = Indices[BaseIdx + NextIdx];
+				if (HasUnorderedEdge(HoleMap, VertA, VertB))
+				{
+					continue; // cannot cross hole edges
+				}
+				// note we negate the winding across for these hull triangles because we're not actually crossing the edge; we're already on the 'inside' of the hull edge
+				int Winding = -WindingAcross(BoundaryMap, VertA, VertB);
+				ToWalkQ.Add(TPair<int, int>(TriIdx, Winding));
+				Keep[TriIdx] = ClassifyFromRule(Winding) ? 1 : -1;
+				break; // don't check any more edges once in queue
+			}
+		}
+	}
+
+	while (ToWalkQ.Num())
+	{
+		TPair<int, int> TriWithWinding = ToWalkQ.Pop();
+		int BaseIdx = TriWithWinding.Key * 3;
+		int LastWinding = TriWithWinding.Value;
+		for (int SubIdx = 0, NextIdx = 2; SubIdx < 3; NextIdx = SubIdx++)
+		{
+			int VertA = Indices[BaseIdx + SubIdx], VertB = Indices[BaseIdx + NextIdx];
+			if (HasUnorderedEdge(HoleMap, VertA, VertB))
+			{
+				continue; // cannot cross hole edges
+			}
+			int AdjTri = Adj[BaseIdx + NextIdx];
+			if (AdjTri != -1 && Keep[AdjTri] == 0)
+			{
+				int WindingChange = WindingAcross(BoundaryMap, VertA, VertB);
+				int Winding = LastWinding + WindingChange;
+				ToWalkQ.Add(TPair<int, int>(AdjTri, Winding));
+				Keep[AdjTri] = ClassifyFromRule(Winding) ? 1 : -1;
+			}
+		}
+	}
+
+	for (int i = 0; i < TriNum; i++)
+	{
+		if (Keep[i] > 0)
+		{
+			Triangles.Add(FIntVector(Indices[i * 3], Indices[i * 3 + 1], Indices[i * 3 + 2]));
+		}
+	}
+	
+	return !bBoundaryTrackingFailure;
+}
+
+
+
+template void TConstrainedDelaunay2<double>::Add(const FDynamicGraph2<double>& Graph);
+template void TConstrainedDelaunay2<double>::Add(const TGeneralPolygon2<double>& Polygon);
+template void TConstrainedDelaunay2<double>::Add(const TGeneralPolygon2<float>& Polygon);
+template void TConstrainedDelaunay2<double>::Add(const TPolygon2<double>& Polygon, bool bIsHole);
+template void TConstrainedDelaunay2<double>::Add(const TPolygon2<float>& Polygon, bool bIsHole);
+
+template void TConstrainedDelaunay2<float>::Add(const FDynamicGraph2<double>& Graph);
+template void TConstrainedDelaunay2<float>::Add(const TGeneralPolygon2<double>& Polygon);
+template void TConstrainedDelaunay2<float>::Add(const TGeneralPolygon2<float>& Polygon);
+template void TConstrainedDelaunay2<float>::Add(const TPolygon2<double>& Polygon, bool bIsHole);
+template void TConstrainedDelaunay2<float>::Add(const TPolygon2<float>& Polygon, bool bIsHole);
+
+template TConstrainedDelaunay2<float>;
+template TConstrainedDelaunay2<double>;
diff --git a/Engine/Plugins/Experimental/GeometryProcessing/Source/GeometryAlgorithms/Public/ConstrainedDelaunay2.h b/Engine/Plugins/Experimental/GeometryProcessing/Source/GeometryAlgorithms/Public/ConstrainedDelaunay2.h
new file mode 100644
--- /dev/null
+++ b/Engine/Plugins/Experimental/GeometryProcessing/Source/GeometryAlgorithms/Public/ConstrainedDelaunay2.h
@@ -0,0 +1,78 @@
+// Copyright 1998-2019 Epic Games, Inc. All Rights Reserved.
+
+#pragma once
+
+#include "BoxTypes.h"
+#include "Polygon2.h"
+#include "Curve/DynamicGraph2.h"
+#include "Curve/GeneralPolygon2.h"
+
+#include "CoreMinimal.h"
+
+
+template<typename RealType>
+struct TConstrainedDelaunay2
+{
+	//
+	// inputs
+	//
+	TArray<FVector2<RealType>> Vertices;
+
+	// Edges and HoleEdges must not be intersecting; use Arrangment2D to pre-process any input w/ intersecting edges
+	TArray<FIndex2i> Edges;	// Edges can be boundaries or not based on the EFillRule setting
+	TArray<FIndex2i> HoleEdges; // Any triangles inside 'hole' edges *must* be cut out
+
+	bool bOrientedEdges = true;
+
+	enum class EFillRule {
+		Odd,
+		// bOrientedEdges must be true for the below
+		NonZero,
+		Positive,
+		Negative
+	};
+	EFillRule FillRule = EFillRule::Odd;
+
+	inline bool ClassifyFromRule(int Winding)
+	{
+		switch (FillRule)
+		{
+		case EFillRule::Odd:
+			return Winding % 2 != 0;
+		case EFillRule::NonZero:
+			return Winding != 0;
+		case EFillRule::Positive:
+			return Winding > 0;
+		case EFillRule::Negative:
+			return Winding < 0;
+		default:
+			check(false);
+			return false;
+		}
+	}
+
+	template<class InputRealType>
+	void GEOMETRYALGORITHMS_API Add(const FDynamicGraph2<InputRealType>& Graph);
+	template<class InputRealType>
+	void GEOMETRYALGORITHMS_API Add(const TPolygon2<InputRealType>& Polygon, bool bIsHole = false);
+	template<class InputRealType>
+	void GEOMETRYALGORITHMS_API Add(const TGeneralPolygon2<InputRealType>& Polygon);
+
+	//
+	// outputs
+	//
+	TArray<FIntVector> Triangles;
+
+	/**
+	 * Populate Triangles
+	 *
+	 * @return false if Triangulation failed
+	 */
+	bool GEOMETRYALGORITHMS_API Triangulate();
+};
+
+
+
+typedef TConstrainedDelaunay2<float> FConstrainedDelaunay2f;
+typedef TConstrainedDelaunay2<double> FConstrainedDelaunay2d;
+
