add generalpolygon2 #rb none

diff --git a/Engine/Plugins/Experimental/GeometryProcessing/Source/GeometricObjects/Public/Curve/GeneralPolygon2.h b/Engine/Plugins/Experimental/GeometryProcessing/Source/GeometricObjects/Public/Curve/GeneralPolygon2.h
new file mode 100644
--- /dev/null
+++ b/Engine/Plugins/Experimental/GeometryProcessing/Source/GeometricObjects/Public/Curve/GeneralPolygon2.h
@@ -0,0 +1,322 @@
+// Copyright 1998-2019 Epic Games, Inc. All Rights Reserved.
+
+// port of geometry3Sharp Polygon
+
+#pragma once
+
+#include "Templates/UnrealTemplate.h"
+#include "Math/UnrealMath.h"
+#include "VectorTypes.h"
+#include "Polygon2.h"
+#include "BoxTypes.h"
+#include "MatrixTypes.h"
+#include "MathUtil.h"
+
+
+/**
+ * TGeneralPolygon2 is a 2D polygon with holes
+ */
+template<typename T>
+class TGeneralPolygon2
+{
+protected:
+	/* The Outer boundary of the polygon */
+	TPolygon2<T> Outer;
+
+	/* if true, Outer polygon winding is clockwise */
+	bool bOuterIsCW;
+
+	/** The list of Holes in the polygon */
+	TArray<TPolygon2<T>> Holes;
+
+public:
+
+	TGeneralPolygon2()
+	{
+	}
+
+	/**
+	* Construct a copy of another general polygon
+	*/
+	TGeneralPolygon2(const TGeneralPolygon2& Copy) : Outer(Copy.Outer), bOuterIsCW(Copy.bOuterIsCW), Holes(Copy.Holes)
+	{
+	}
+
+	/**
+	* Construct a general polygon with the given polygon as boundary
+	*/
+	TGeneralPolygon2(const TPolygon2<T>& Outer) : Outer(Outer)
+	{
+		bOuterIsCW = Outer.IsClockwise();
+	}
+
+	void SetOuter(const TPolygon2<T>& Outer)
+	{
+		this->Outer = Outer;
+		bOuterIsCW = Outer.IsClockwise();
+	}
+
+	void SetOuterWithOrientation(const TPolygon2<T>& Outer, bool bOuterIsCW)
+	{
+		checkSlow(Outer.IsClockwise() == bOuterIsCW);
+		this->Outer = Outer;
+		this->bOuterIsCW = bOuterIsCW;
+	}
+
+	const TPolygon2<T>& GetOuter() const
+	{
+		return this->Outer;
+	}
+
+	const TArray<TPolygon2<T>>& GetHoles() const
+	{
+		return Holes;
+	}
+
+
+	bool AddHole(TPolygon2<T> Hole, bool bCheckContainment = true, bool bCheckOrientation = true)
+	{
+		if (bCheckContainment)
+		{
+			if (!Outer.Contains(Hole))
+			{
+				return false;
+			}
+
+			// [RMS] segment/segment intersection broken?
+			for (const TPolygon2<T>& ExistingHole : Holes)
+			{
+				if (Hole.Overlaps(ExistingHole))
+				{
+					return false;
+				}
+			}
+		}
+
+
+        if (bCheckOrientation)
+		{
+			bool bHoleIsClockwise = Hole.IsClockwise();
+			if (bOuterIsCW == bHoleIsClockwise)
+			{
+				return false;
+			}
+        }
+
+		Holes.Add(Hole);
+		return true;
+	}
+
+    void ClearHoles()
+	{
+        Holes.Empty();
+    }
+
+
+	bool HasHoles() const
+	{
+		return Holes.Num() > 0;
+	}
+
+
+    double SignedArea() const
+    {
+        double Sign = (bOuterIsCW) ? -1.0 : 1.0;
+        double AreaSum = Sign * Outer.SignedArea();
+		for (const TPolygon2<T>& Hole : Holes)
+		{
+			AreaSum += Sign * Hole.SignedArea();
+		}
+        return AreaSum;
+    }
+
+
+    double HoleUnsignedArea() const
+    {
+		double AreaSum = 0;
+		for (const TPolygon2<T>& Hole : Holes)
+		{
+			AreaSum += Math.Abs(Hole.SignedArea());
+		}
+		return AreaSum;
+    }
+
+
+    double Perimeter()
+    {
+		double PerimSum = Outer.Perimeter();
+		for (const TPolygon2<T> &Hole : Holes)
+		{
+			PerimSum += Hole.Perimeter();
+		}
+		return PerimSum;
+    }
+
+
+    TAxisAlignedBox2<T> Bounds()
+    {
+		TAxisAlignedBox2<T> Box = Outer.Bounds();
+		for (const TPolygon2<T> Hole : Holes)
+		{
+			Box.Contain(Hole.Bounds());
+		}
+		return Box;
+    }
+
+
+	void Translate(FVector2<T> translate) {
+		Outer.Translate(translate);
+		for (const TPolygon2<T>& Hole : Holes)
+		{
+			Hole.Translate(translate);
+		}
+	}
+
+    void Rotate(FMatrix2d rotation, FVector2<T> origin) {
+        Outer.Rotate(rotation, origin);
+		for (const TPolygon2<T>& Hole : Holes)
+		{
+			Hole.Rotate(rotation, origin);
+		}
+    }
+
+
+    void Scale(FVector2<T> scale, FVector2<T> origin) {
+		Outer.Scale(scale, origin);
+		for (const TPolygon2<T>& Hole : Holes)
+		{
+			Hole.Scale(scale, origin);
+		}
+	}
+
+    void Transform(const TFunction<FVector2<T> (const FVector2<T>&)>& TransformFunc)
+    {
+        Outer.Transform(TransformFunc);
+		for (const TPolygon2<T>& Hole : Holes)
+		{
+			Hole.Transform(TransformFunc);
+		}
+    }
+
+    void Reverse()
+    {
+        Outer.Reverse();
+        bOuterIsCW = Outer.IsClockwise;
+		for (const TPolygon2<T>& Hole : Holes)
+		{
+			Hole.Reverse();
+		}
+    }
+
+
+	bool Contains(FVector2<T> vTest)
+	{
+		if (Outer.Contains(vTest) == false)
+		{
+			return false;
+		}
+		for (const TPolygon2<T>& Hole : Holes)
+		{
+			if (Hole.Contains(vTest))
+			{
+				return false;
+			}
+		}
+		return true;
+	}
+
+    bool Contains(TPolygon2<T> Poly) {
+		if (Outer.Contains(Poly) == false)
+		{
+			return false;
+		}
+        for (const TPolygon2<T>& Hole : Holes)
+		{
+			if (Hole.Overlaps(Poly))
+			{
+				return false;
+			}
+        }
+        return true;
+    }
+
+
+    bool Intersects(TPolygon2<T> Poly)
+    {
+		if (Outer.Intersects(Poly))
+		{
+			return true;
+		}
+        for (const TPolygon2<T>& Hole : Holes)
+		{
+			if (Hole.Intersects(Poly))
+			{
+				return true;
+			}
+        }
+        return false;
+    }
+
+
+    FVector2<T> PointAt(int iSegment, double fSegT, int iHoleIndex = -1)
+	{
+		if (iHoleIndex == -1)
+		{
+			return Outer.PointAt(iSegment, fSegT);
+		}
+		return Holes[iHoleIndex].PointAt(iSegment, fSegT);
+	}
+
+	TSegment2<T> Segment(int iSegment, int iHoleIndex = -1)
+	{
+		if (iHoleIndex == -1)
+		{
+			return Outer.Segment(iSegment);
+		}
+		return Holes[iHoleIndex].Segment(iSegment);			
+	}
+
+	FVector2<T> GetNormal(int iSegment, double segT, int iHoleIndex = -1)
+	{
+		if (iHoleIndex == -1)
+		{
+			return Outer.GetNormal(iSegment, segT);
+		}
+		return Holes[iHoleIndex].GetNormal(iSegment, segT);
+	}
+
+	// this should be more efficient when there are Holes...
+	double DistanceSquared(FVector2<T> p, int &iHoleIndex, int &iNearSeg, double &fNearSegT)
+	{
+		iNearSeg = iHoleIndex = -1;
+		fNearSegT = double.MaxValue;
+		double dist = Outer.DistanceSquared(p, out iNearSeg, out fNearSegT);
+		for (int i = 0; i < Holes.Num(); ++i )
+		{
+			int seg; double segt;
+			double holedist = Holes[i].DistanceSquared(p, out seg, out segt);
+			if (holedist < dist)
+			{
+				dist = holedist;
+				iHoleIndex = i;
+				iNearSeg = seg;
+				fNearSegT = segt;
+			}
+		}
+		return dist;
+	}
+
+
+    void Simplify(double ClusterTol = 0.0001, double LineDeviationTol = 0.01, bool bSimplifyStraightLines = true)
+    {
+        // [TODO] should make sure that Holes stay inside Outer!!
+        Outer.Simplify(ClusterTol, LineDeviationTol, bSimplifyStraightLines);
+		for (const TPolygon2<T>& Hole : Holes) {
+			Hole.Simplify(ClusterTol, LineDeviationTol, bSimplifyStraightLines);
+		}
+    }
+
+};
+
+typedef TGeneralPolygon2<double> FGeneralPolygon2d;
+typedef TGeneralPolygon2<float> FGeneralPolygon2f;
