Fix porting issues with GeneralPolygon2.h #rb none #rnx

diff --git a/Engine/Plugins/Experimental/GeometryProcessing/Source/GeometricObjects/Public/Curve/GeneralPolygon2.h b/Engine/Plugins/Experimental/GeometryProcessing/Source/GeometricObjects/Public/Curve/GeneralPolygon2.h
--- a/Engine/Plugins/Experimental/GeometryProcessing/Source/GeometricObjects/Public/Curve/GeneralPolygon2.h
+++ b/Engine/Plugins/Experimental/GeometryProcessing/Source/GeometricObjects/Public/Curve/GeneralPolygon2.h
@@ -45,22 +45,22 @@ public:
 	/**
 	* Construct a general polygon with the given polygon as boundary
 	*/
-	TGeneralPolygon2(const TPolygon2<T>& Outer) : Outer(Outer)
+	TGeneralPolygon2(const TPolygon2<T>& ToSetOuter) : Outer(ToSetOuter)
 	{
-		bOuterIsCW = Outer.IsClockwise();
+		bOuterIsCW = ToSetOuter.IsClockwise();
 	}
 
-	void SetOuter(const TPolygon2<T>& Outer)
+	void SetOuter(const TPolygon2<T>& ToSetOuter)
 	{
-		this->Outer = Outer;
-		bOuterIsCW = Outer.IsClockwise();
+		this->Outer = ToSetOuter;
+		bOuterIsCW = ToSetOuter.IsClockwise();
 	}
 
-	void SetOuterWithOrientation(const TPolygon2<T>& Outer, bool bOuterIsCW)
+	void SetOuterWithOrientation(const TPolygon2<T>& ToSetOuter, bool bToSetOuterIsCW)
 	{
-		checkSlow(Outer.IsClockwise() == bOuterIsCW);
-		this->Outer = Outer;
-		this->bOuterIsCW = bOuterIsCW;
+		checkSlow(ToSetOuter.IsClockwise() == bToSetOuterIsCW);
+		this->Outer = ToSetOuter;
+		this->bOuterIsCW = bToSetOuterIsCW;
 	}
 
 	const TPolygon2<T>& GetOuter() const
@@ -136,7 +136,7 @@ public:
 		double AreaSum = 0;
 		for (const TPolygon2<T>& Hole : Holes)
 		{
-			AreaSum += Math.Abs(Hole.SignedArea());
+			AreaSum += FMath::Abs(Hole.SignedArea());
 		}
 		return AreaSum;
     }
@@ -166,24 +166,16 @@ public:
 
 	void Translate(FVector2<T> translate) {
 		Outer.Translate(translate);
-		for (const TPolygon2<T>& Hole : Holes)
+		for (TPolygon2<T>& Hole : Holes)
 		{
 			Hole.Translate(translate);
 		}
 	}
 
-    void Rotate(FMatrix2d rotation, FVector2<T> origin) {
-        Outer.Rotate(rotation, origin);
-		for (const TPolygon2<T>& Hole : Holes)
-		{
-			Hole.Rotate(rotation, origin);
-		}
-    }
-
 
     void Scale(FVector2<T> scale, FVector2<T> origin) {
 		Outer.Scale(scale, origin);
-		for (const TPolygon2<T>& Hole : Holes)
+		for (TPolygon2<T>& Hole : Holes)
 		{
 			Hole.Scale(scale, origin);
 		}
@@ -192,7 +184,7 @@ public:
     void Transform(const TFunction<FVector2<T> (const FVector2<T>&)>& TransformFunc)
     {
         Outer.Transform(TransformFunc);
-		for (const TPolygon2<T>& Hole : Holes)
+		for (TPolygon2<T>& Hole : Holes)
 		{
 			Hole.Transform(TransformFunc);
 		}
@@ -201,8 +193,8 @@ public:
     void Reverse()
     {
         Outer.Reverse();
-        bOuterIsCW = Outer.IsClockwise;
-		for (const TPolygon2<T>& Hole : Holes)
+		bOuterIsCW = !bOuterIsCW;
+		for (TPolygon2<T>& Hole : Holes)
 		{
 			Hole.Reverse();
 		}
@@ -258,13 +250,13 @@ public:
     }
 
 
-    FVector2<T> PointAt(int iSegment, double fSegT, int iHoleIndex = -1)
+    FVector2<T> GetSegmentPoint(int iSegment, double fSegT, int iHoleIndex = -1)
 	{
 		if (iHoleIndex == -1)
 		{
-			return Outer.PointAt(iSegment, fSegT);
+			return Outer.GetSegmentPoint(iSegment, fSegT);
 		}
-		return Holes[iHoleIndex].PointAt(iSegment, fSegT);
+		return Holes[iHoleIndex].GetSegmentPoint(iSegment, fSegT);
 	}
 
 	TSegment2<T> Segment(int iSegment, int iHoleIndex = -1)
@@ -289,12 +281,12 @@ public:
 	double DistanceSquared(FVector2<T> p, int &iHoleIndex, int &iNearSeg, double &fNearSegT)
 	{
 		iNearSeg = iHoleIndex = -1;
-		fNearSegT = double.MaxValue;
-		double dist = Outer.DistanceSquared(p, out iNearSeg, out fNearSegT);
+		fNearSegT = TMathUtil<T>::MaxReal;
+		double dist = Outer.DistanceSquared(p, iNearSeg, fNearSegT);
 		for (int i = 0; i < Holes.Num(); ++i )
 		{
 			int seg; double segt;
-			double holedist = Holes[i].DistanceSquared(p, out seg, out segt);
+			double holedist = Holes[i].DistanceSquared(p, seg, segt);
 			if (holedist < dist)
 			{
 				dist = holedist;
@@ -307,12 +299,13 @@ public:
 	}
 
 
-    void Simplify(double ClusterTol = 0.0001, double LineDeviationTol = 0.01, bool bSimplifyStraightLines = true)
+    void Simplify(double ClusterTol = 0.0001, double LineDeviationTol = 0.01)
     {
         // [TODO] should make sure that Holes stay inside Outer!!
-        Outer.Simplify(ClusterTol, LineDeviationTol, bSimplifyStraightLines);
-		for (const TPolygon2<T>& Hole : Holes) {
-			Hole.Simplify(ClusterTol, LineDeviationTol, bSimplifyStraightLines);
+        Outer.Simplify(ClusterTol, LineDeviationTol);
+		for (TPolygon2<T>& Hole : Holes)
+		{
+			Hole.Simplify(ClusterTol, LineDeviationTol);
 		}
     }
 
@@ -320,3 +313,4 @@ public:
 
 typedef TGeneralPolygon2<double> FGeneralPolygon2d;
 typedef TGeneralPolygon2<float> FGeneralPolygon2f;
+
