add PlanarComplex method to find solids w/ holes from an array of simple polygons  #rb none #rnx

diff --git a/Engine/Plugins/Experimental/GeometryProcessing/Source/GeometricObjects/Public/BoxTypes.h b/Engine/Plugins/Experimental/GeometryProcessing/Source/GeometricObjects/Public/BoxTypes.h
--- a/Engine/Plugins/Experimental/GeometryProcessing/Source/GeometricObjects/Public/BoxTypes.h
+++ b/Engine/Plugins/Experimental/GeometryProcessing/Source/GeometricObjects/Public/BoxTypes.h
@@ -57,6 +57,11 @@ struct TInterval1
 		return D >= Min && D <= Max;
 	}
 
+	bool Contains(const TInterval1<RealType>& O) const
+	{
+		return Contains(O.Min) && Contains(O.Max);
+	}
+
 	bool Overlaps(const TInterval1<RealType>& O) const
 	{
 		return !(O.Min > Max || O.Max < Min);
@@ -297,6 +302,11 @@ struct TAxisAlignedBox3
 		return (Min.X <= V.X) && (Min.Y <= V.Y) && (Min.Z <= V.Z) && (Max.X >= V.X) && (Max.Y >= V.Y) && (Max.Z >= V.Z);
 	}
 
+	bool Contains(const TAxisAlignedBox3<RealType>& Box) const
+	{
+		return Contains(Box.Min) && Contains(Box.Max);
+	}
+
 	TAxisAlignedBox3<RealType> Intersect(const TAxisAlignedBox3<RealType>& Box) const
 	{
 		TAxisAlignedBox3<RealType> Intersection(
@@ -478,6 +488,11 @@ struct TAxisAlignedBox2
 		return (Min.X <= V.X) && (Min.Y <= V.Y) && (Max.X >= V.X) && (Max.Y >= V.Y);
 	}
 
+	bool Contains(const TAxisAlignedBox2<RealType>& Box) const
+	{
+		return Contains(Box.Min) && Contains(Box.Max);
+	}
+
 	bool Intersects(const TAxisAlignedBox2<RealType>& Box) const
 	{
 		return !((Box.Max.X < Min.X) || (Box.Min.X > Max.X) || (Box.Max.Y < Min.Y) || (Box.Min.Y > Max.Y));
diff --git a/Engine/Plugins/Experimental/GeometryProcessing/Source/GeometricObjects/Public/Curve/PlanarComplex.h b/Engine/Plugins/Experimental/GeometryProcessing/Source/GeometricObjects/Public/Curve/PlanarComplex.h
new file mode 100644
--- /dev/null
+++ b/Engine/Plugins/Experimental/GeometryProcessing/Source/GeometricObjects/Public/Curve/PlanarComplex.h
@@ -0,0 +1,167 @@
+// Copyright 1998-2019 Epic Games, Inc. All Rights Reserved.
+
+// Very approximately ported from geometry3sharp's PlanarComplex
+// Convert a set of polygons into general polygon form, tracking the outer-polygon and hole-polygon indices
+
+#include "SegmentTypes.h"
+#include "VectorTypes.h"
+#include "BoxTypes.h"
+#include "Polygon2.h"
+#include "Curve/GeneralPolygon2.h"
+
+#pragma once
+
+
+template <typename RealType>
+struct TPlanarComplex
+{
+	//
+	// Inputs
+	//
+	TArray<TPolygon2<RealType>> Polygons;
+
+	bool bTrustOrientations = true;
+	bool bAllowOverlappingHoles = false;
+
+	//
+	// Outputs
+	//
+	struct FPolygonNesting
+	{
+		int OuterIndex;
+		TArray<int> HoleIndices;
+	};
+	TArray<FPolygonNesting> Solids;
+
+	// Finds set of "solid" regions -- e.g. boundary loops with interior holes.
+	// Result has outer loops being clockwise, and holes counter-clockwise
+	void FindSolidRegions()
+	{
+		Solids.Empty();
+
+		int32 N = Polygons.Num();
+		TArray<int> PolygonIndices;
+		PolygonIndices.SetNum(N);
+		for (int i = 0; i < N; i++)
+		{
+			PolygonIndices[i] = i;
+		}
+
+		// precomputed useful stats on input polygons
+		struct FPolygonInfo
+		{
+			RealType Area;
+			bool Orientation;
+			TAxisAlignedBox2<RealType> Bounds;
+		};
+		TArray<FPolygonInfo> Info;
+		Info.SetNum(N);
+		for (int i = 0; i < N; i++)
+		{
+			RealType SignedArea = Polygons[i].SignedArea();
+			Info[i].Orientation = SignedArea > 0;
+			Info[i].Area = Info[i].Orientation ? SignedArea : -SignedArea;
+			Info[i].Bounds = Polygons[i].Bounds();
+		}
+
+		PolygonIndices.Sort([&Info](int Ind1, int Ind2)
+		{
+			return Info[Ind1].Area > Info[Ind2].Area;
+		}
+		);
+
+		TArray<int> Parent;
+		Parent.SetNumUninitialized(N);
+		for (int i = 0; i < N; i++)
+		{
+			Parent[i] = -1;
+		}
+
+		// From largest polygon to smallest, greedily assign parents where possible
+		// If a polygon already has a parent, but a smaller containing parent is found, switch to the new parent (unless doing so would create an intersection)
+		for (int ParentCand = 0; ParentCand+1 < N; ParentCand++)
+		{
+			int PIdx = PolygonIndices[ParentCand];
+			const TPolygon2<RealType>& PPoly = Polygons[PIdx];
+			const FPolygonInfo& PInfo = Info[PIdx];
+			if (bTrustOrientations && !PInfo.Orientation)
+			{
+				// orientation tells us this is a hole; skip checking for holes in it
+				continue;
+			}
+			for (int ChildCand = ParentCand + 1; ChildCand < N; ChildCand++)
+			{
+				int CIdx = PolygonIndices[ChildCand];
+				const TPolygon2<RealType>& CPoly = Polygons[CIdx];
+				const FPolygonInfo& CInfo = Info[CIdx];
+				if (bTrustOrientations && CInfo.Orientation)
+				{
+					// orientation tells us this is an outer; skip checking if it could be a hole
+					continue;
+				}
+				if (PInfo.Bounds.Contains(CInfo.Bounds) && PPoly.Contains(CPoly))
+				{
+					bool bOverlapFound = false;
+					if (!bAllowOverlappingHoles)
+					{
+						for (int ParentRecheckIdx = ParentCand + 1; ParentRecheckIdx < ChildCand; ParentRecheckIdx++)
+						{
+							if (Parent[ParentRecheckIdx] == ParentCand)
+							{
+								if (Polygons[PolygonIndices[ParentRecheckIdx]].Overlaps(CPoly))
+								{
+									bOverlapFound = true;
+									break;
+								}
+							}
+						}
+					}
+					if (!bOverlapFound)
+					{
+						Parent[ChildCand] = ParentCand;
+					}
+				}
+			}
+		}
+
+		// Build nests from parent relationship
+		TArray<bool> Taken;
+		Taken.SetNumZeroed(N);
+		for (int NestCand = 0; NestCand < N; NestCand++)
+		{
+			int NIdx = PolygonIndices[NestCand];
+			if (Taken[NestCand] || (bTrustOrientations && !Info[NIdx].Orientation))
+			{
+				continue;
+			}
+			FPolygonNesting &Nest = Solids.Emplace_GetRef();
+			Nest.OuterIndex = NIdx;
+			for (int HoleCand = NestCand + 1; HoleCand < N; HoleCand++)
+			{
+				int HIdx = PolygonIndices[HoleCand];
+				if (Parent[HoleCand] == NestCand)
+				{
+					Nest.HoleIndices.Add(HIdx);
+					Taken[HoleCand] = true;
+				}
+			}
+		}
+	}
+
+	TArray<TGeneralPolygon2<RealType>> ConvertOutputToGeneralPolygons() const
+	{
+		TArray<TGeneralPolygon2<RealType>> Output;
+		for (const FPolygonNesting& Nest : Solids)
+		{
+			TGeneralPolygon2<RealType>& GPoly = Output.Emplace_GetRef(TGeneralPolygon2<RealType>(Polygons[Nest.OuterIndex]));
+			for (int HoleIdx : Nest.HoleIndices)
+			{
+				GPoly.AddHole(Polygons[HoleIdx], false, false);
+			}
+		}
+		return Output;
+	}
+};
+
+typedef TPlanarComplex<double> FPlanarComplexd;
+typedef TPlanarComplex<float> FPlanarComplexf;
