Fix Clang compile warning in ConstrainedDelaunay2.cpp
#rnx

diff --git a/Engine/Plugins/Experimental/GeometryProcessing/Source/GeometryAlgorithms/Private/ConstrainedDelaunay2.cpp b/Engine/Plugins/Experimental/GeometryProcessing/Source/GeometryAlgorithms/Private/ConstrainedDelaunay2.cpp
--- a/Engine/Plugins/Experimental/GeometryProcessing/Source/GeometryAlgorithms/Private/ConstrainedDelaunay2.cpp
+++ b/Engine/Plugins/Experimental/GeometryProcessing/Source/GeometryAlgorithms/Private/ConstrainedDelaunay2.cpp
@@ -229,6 +229,3 @@ template void TConstrainedDelaunay2<float>::Add(const TGeneralPolygon2<double>&
 template void TConstrainedDelaunay2<float>::Add(const TGeneralPolygon2<float>& Polygon);
 template void TConstrainedDelaunay2<float>::Add(const TPolygon2<double>& Polygon, bool bIsHole);
 template void TConstrainedDelaunay2<float>::Add(const TPolygon2<float>& Polygon, bool bIsHole);
-
-template TConstrainedDelaunay2<float>;
-template TConstrainedDelaunay2<double>;
