add corrected explicit instantiation for constrained delaunay class #rb none #rnx

diff --git a/Engine/Plugins/Experimental/GeometryProcessing/Source/GeometryAlgorithms/Private/ConstrainedDelaunay2.cpp b/Engine/Plugins/Experimental/GeometryProcessing/Source/GeometryAlgorithms/Private/ConstrainedDelaunay2.cpp
--- a/Engine/Plugins/Experimental/GeometryProcessing/Source/GeometryAlgorithms/Private/ConstrainedDelaunay2.cpp
+++ b/Engine/Plugins/Experimental/GeometryProcessing/Source/GeometryAlgorithms/Private/ConstrainedDelaunay2.cpp
@@ -229,3 +229,6 @@ template void TConstrainedDelaunay2<float>::Add(const TGeneralPolygon2<double>&
 template void TConstrainedDelaunay2<float>::Add(const TGeneralPolygon2<float>& Polygon);
 template void TConstrainedDelaunay2<float>::Add(const TPolygon2<double>& Polygon, bool bIsHole);
 template void TConstrainedDelaunay2<float>::Add(const TPolygon2<float>& Polygon, bool bIsHole);
+
+template struct TConstrainedDelaunay2<float>;
+template struct TConstrainedDelaunay2<double>;
