planar mesh cut now uses planar hole fill via planarcomplex + cdt; can fill plane cut geometry more robustly
#rnx

diff --git a/Engine/Plugins/Experimental/GeometryProcessing/Source/DynamicMesh/Private/Operations/MeshPlaneCut.cpp b/Engine/Plugins/Experimental/GeometryProcessing/Source/DynamicMesh/Private/Operations/MeshPlaneCut.cpp
--- a/Engine/Plugins/Experimental/GeometryProcessing/Source/DynamicMesh/Private/Operations/MeshPlaneCut.cpp
+++ b/Engine/Plugins/Experimental/GeometryProcessing/Source/DynamicMesh/Private/Operations/MeshPlaneCut.cpp
@@ -2,6 +2,7 @@
 
 #include "Operations/MeshPlaneCut.h"
 #include "Operations/SimpleHoleFiller.h"
+#include "Operations/PlanarHoleFiller.h"
 #include "MeshNormals.h"
 #include "DynamicMeshEditor.h"
 #include "MathUtil.h"
@@ -193,10 +194,7 @@ void FMeshPlaneCut::CollapseDegenerateEdges(const TSet<int>& OnCutEdges, const T
 }
 
 
-
-
-
-bool FMeshPlaneCut::FillHoles(int ConstantGroupID)
+bool FMeshPlaneCut::SimpleHoleFill(int ConstantGroupID)
 {
 	bool bAllOk = true;
 
@@ -227,4 +225,42 @@ bool FMeshPlaneCut::FillHoles(int ConstantGroupID)
 	}
 
 	return bAllOk;
-}
\ No newline at end of file
+}
+
+
+bool FMeshPlaneCut::HoleFill(TFunction<TArray<FIndex3i>(const FGeneralPolygon2d&)> PlanarTriangulationFunc, bool bFillSpans, int ConstantGroupID)
+{
+	HoleFillTriangles.Empty();
+
+	TArray<TArray<int>> LoopVertices;
+	for (const FEdgeLoop& Loop : CutLoops)
+	{
+		LoopVertices.Add(Loop.Vertices);
+	}
+	if (bFillSpans)
+	{
+		for (const FEdgeSpan& Span : CutSpans)
+		{
+			LoopVertices.Add(Span.Vertices);
+		}
+	}
+	FPlanarHoleFiller Filler(Mesh, &LoopVertices, PlanarTriangulationFunc, PlaneOrigin, PlaneNormal);
+
+	int GID = ConstantGroupID >= 0 ? ConstantGroupID : Mesh->AllocateTriangleGroup();
+	if (!Filler.Fill(GID))
+	{
+		return false;
+	}
+
+	HoleFillTriangles.Add(Filler.NewTriangles);
+	if (Mesh->HasAttributes())
+	{
+		FDynamicMeshEditor Editor(Mesh);
+		Editor.SetTriangleNormals(Filler.NewTriangles, (FVector3f)PlaneNormal);
+
+		FFrame3d ProjectionFrame(PlaneOrigin, PlaneNormal);
+		Editor.SetTriangleUVsFromProjection(Filler.NewTriangles, ProjectionFrame, UVScaleFactor);
+	}
+
+	return true;
+}
diff --git a/Engine/Plugins/Experimental/GeometryProcessing/Source/DynamicMesh/Private/Operations/PlanarHoleFiller.cpp b/Engine/Plugins/Experimental/GeometryProcessing/Source/DynamicMesh/Private/Operations/PlanarHoleFiller.cpp
new file mode 100644
--- /dev/null
+++ b/Engine/Plugins/Experimental/GeometryProcessing/Source/DynamicMesh/Private/Operations/PlanarHoleFiller.cpp
@@ -0,0 +1,54 @@
+// Copyright 1998-2019 Epic Games, Inc. All Rights Reserved.
+
+#include "Operations/PlanarHoleFiller.h"
+#include "Curve/PlanarComplex.h"
+
+
+
+bool FPlanarHoleFiller::Fill(int GroupID)
+{
+	FVector3d PlaneX, PlaneY;
+	VectorUtil::MakePerpVectors(PlaneNormal, PlaneX, PlaneY);
+
+	
+	FPlanarComplexd Complex;
+	for (const TArray<int>& Loop : *VertexLoops)
+	{
+		FPolygon2d& Polygon = Complex.Polygons.Emplace_GetRef();
+		for (int VID : Loop)
+		{
+			FVector3d VertMinusOrigin = Mesh->GetVertex(VID) - PlaneOrigin;
+			Polygon.AppendVertex(FVector2d(PlaneX.Dot(VertMinusOrigin), PlaneY.Dot(VertMinusOrigin)));
+		}
+	}
+
+	Complex.bTrustOrientations = true;
+	Complex.bAllowOverlappingHoles = false;
+
+	Complex.FindSolidRegions();
+
+	bool bAddedAll = true;
+	for (const FPlanarComplexd::FPolygonNesting& Nest : Complex.Solids)
+	{
+		FGeneralPolygon2d GPoly = Complex.ConvertNestToGeneralPolygon(Nest);
+		TArray<FIndex3i> Triangles = PlanarTriangulationFunc(GPoly);
+		// create mapping from local triangulation problem to larger mesh vertex indices
+		TArray<int> VertexMapping = (*VertexLoops)[Nest.OuterIndex]; // init to copy of outer loop vertices
+		for (int HoleIdx : Nest.HoleIndices) // append hole loop indices
+		{
+			VertexMapping.Append((*VertexLoops)[HoleIdx]);
+		}
+		for (const FIndex3i& SourceTri : Triangles)
+		{
+			FIndex3i OutTri(VertexMapping[SourceTri.A], VertexMapping[SourceTri.B], VertexMapping[SourceTri.C]);
+			int NewTID = Mesh->AppendTriangle(OutTri, GroupID);
+			if (NewTID < 0) // e.g. if triangle would create non-manifold geometry in mesh, it will be skipped
+			{
+				bAddedAll = false;
+			}
+			NewTriangles.Add(NewTID);
+		}
+	}
+	
+	return bAddedAll;
+}
diff --git a/Engine/Plugins/Experimental/GeometryProcessing/Source/DynamicMesh/Public/Operations/MeshPlaneCut.h b/Engine/Plugins/Experimental/GeometryProcessing/Source/DynamicMesh/Public/Operations/MeshPlaneCut.h
--- a/Engine/Plugins/Experimental/GeometryProcessing/Source/DynamicMesh/Public/Operations/MeshPlaneCut.h
+++ b/Engine/Plugins/Experimental/GeometryProcessing/Source/DynamicMesh/Public/Operations/MeshPlaneCut.h
@@ -8,6 +8,7 @@
 #include "VectorTypes.h"
 #include "GeometryTypes.h"
 #include "MeshBoundaryLoops.h"
+#include "Curve/GeneralPolygon2.h"
 
 
 class FDynamicMesh3;
@@ -55,7 +56,7 @@ public:
 	bool CutLoopsFailed = false;		// set to true if we could not compute cut loops/spans
 	bool FoundOpenSpans = false;     // set to true if we found open spans in cut
 
-	// Note: In practice this is 1:1 with CutLoops if all the hole fills succeed due to simple hole filling, but this will not be true if/when we add better hole filling methods, or if any of the hole fills fails
+	// Note: In practice this is 1:1 with CutLoops if all the hole fills succeed w/ simple hole filling, but this will not be true for other hole filling methods, or if any of the hole fills fails
 	// please do not write anything relying on a correspondence between the two
 	TArray<TArray<int>> HoleFillTriangles;
 
@@ -82,9 +83,14 @@ public:
 	virtual bool Cut();
 
 	/**
-	 *  Fill all loops with FSimpleHoleFiller. @todo: Add more advanced hole-fill options
+	 *  Fill cut loops with FSimpleHoleFiller
 	 */
-	virtual bool FillHoles(int constantGroupID = -1);
+	virtual bool SimpleHoleFill(int ConstantGroupID = -1);
+
+	/**
+	 *  Fill cut loops with FPlanarHoleFiller, using a caller-provided triangulation function
+	 */
+	virtual bool HoleFill(TFunction<TArray<FIndex3i>(const FGeneralPolygon2d&)> PlanarTriangulationFunc, bool bFillSpans, int ConstantGroupID = -1);
 
 	
 protected:
diff --git a/Engine/Plugins/Experimental/GeometryProcessing/Source/DynamicMesh/Public/Operations/PlanarHoleFiller.h b/Engine/Plugins/Experimental/GeometryProcessing/Source/DynamicMesh/Public/Operations/PlanarHoleFiller.h
new file mode 100644
--- /dev/null
+++ b/Engine/Plugins/Experimental/GeometryProcessing/Source/DynamicMesh/Public/Operations/PlanarHoleFiller.h
@@ -0,0 +1,63 @@
+// Copyright 1998-2019 Epic Games, Inc. All Rights Reserved.
+
+// Port of geometry3Sharp PlanarHoleFiller
+
+#pragma once
+
+#include "CoreMinimal.h"
+#include "MathUtil.h"
+#include "VectorTypes.h"
+#include "IndexTypes.h"
+#include "GeometryTypes.h"
+#include "Curve/GeneralPolygon2.h"
+#include "MeshBoundaryLoops.h"
+
+
+class FDynamicMesh3;
+
+
+/**
+ * Fill a set of boundary loops with planar surfaces.  User must provide the triangulation function.
+ */
+class DYNAMICMESH_API FPlanarHoleFiller
+{
+public:
+	//
+	// Inputs
+	//
+	FDynamicMesh3* Mesh;
+	const TArray<TArray<int>> *VertexLoops;
+
+	TFunction<TArray<FIndex3i>(const FGeneralPolygon2d&)> PlanarTriangulationFunc;
+
+	FVector3d PlaneOrigin;
+	FVector3d PlaneNormal;
+
+	//
+	// Outputs
+	//
+	TArray<int> NewTriangles;
+
+public:
+	/**
+	 *  Construct simple hole filler (just adds a central vertex and a triangle fan)
+	 */
+	FPlanarHoleFiller(FDynamicMesh3* Mesh, const TArray<TArray<int>> *VertexLoops, TFunction<TArray<FIndex3i>(const FGeneralPolygon2d&)> PlanarTriangulationFunc, FVector3d PlaneOrigin, FVector3d PlaneNormal) 
+		: Mesh(Mesh), VertexLoops(VertexLoops), PlanarTriangulationFunc(PlanarTriangulationFunc), PlaneOrigin(PlaneOrigin), PlaneNormal(PlaneNormal)
+	{
+	}
+	virtual ~FPlanarHoleFiller() {}
+	
+	
+	/**
+	 * @return EOperationValidationResult::Ok if we can apply operation, or error code if we cannot
+	 */
+	virtual EOperationValidationResult Validate()
+	{
+		// TODO: validation
+
+		return EOperationValidationResult::Ok;
+	}
+
+	virtual bool Fill(int GroupID = -1);
+};
diff --git a/Engine/Plugins/Experimental/GeometryProcessing/Source/GeometricObjects/Public/Curve/PlanarComplex.h b/Engine/Plugins/Experimental/GeometryProcessing/Source/GeometricObjects/Public/Curve/PlanarComplex.h
--- a/Engine/Plugins/Experimental/GeometryProcessing/Source/GeometricObjects/Public/Curve/PlanarComplex.h
+++ b/Engine/Plugins/Experimental/GeometryProcessing/Source/GeometricObjects/Public/Curve/PlanarComplex.h
@@ -148,6 +148,15 @@ struct TPlanarComplex
 		}
 	}
 
+	TGeneralPolygon2<RealType> ConvertNestToGeneralPolygon(const FPolygonNesting& Nest) const
+	{
+		TGeneralPolygon2<RealType> GPoly(Polygons[Nest.OuterIndex]);
+		for (int HoleIdx : Nest.HoleIndices)
+		{
+			GPoly.AddHole(Polygons[HoleIdx], false, false);
+		}
+		return GPoly;
+	}
 	TArray<TGeneralPolygon2<RealType>> ConvertOutputToGeneralPolygons() const
 	{
 		TArray<TGeneralPolygon2<RealType>> Output;
diff --git a/Engine/Plugins/Experimental/GeometryProcessing/Source/GeometryAlgorithms/Private/ConstrainedDelaunay2.cpp b/Engine/Plugins/Experimental/GeometryProcessing/Source/GeometryAlgorithms/Private/ConstrainedDelaunay2.cpp
--- a/Engine/Plugins/Experimental/GeometryProcessing/Source/GeometryAlgorithms/Private/ConstrainedDelaunay2.cpp
+++ b/Engine/Plugins/Experimental/GeometryProcessing/Source/GeometryAlgorithms/Private/ConstrainedDelaunay2.cpp
@@ -6,6 +6,64 @@
 #include "ThirdParty/GTEngine/Mathematics/GteUIntegerFP32.h"
 #include <vector>
 
+//namespace
+//{
+//#define DEBUG_FILE_DUMPING 1
+//#ifndef DEBUG_FILE_DUMPING
+//	void DumpDelaunayInputForDebugAsOBJ(const FConstrainedDelaunay2d& Delaunay, const FString& PathBase)
+//	{
+//	}
+//	void DumpDelaunayTriangulationForDebug(const FConstrainedDelaunay2d& Delaunay, const FString& PathBase)
+//	{
+//	}
+//#else
+//#include <fstream>
+//	static int num = 0;
+//	template <typename RealType>
+//	void DumpDelaunayInputForDebugAsOBJ(const TConstrainedDelaunay2<RealType>& Delaunay, const FString& PathBase)
+//	//void DumpGraphForDebugAsOBJ(const FDynamicGraph2d& Graph, const FString& PathBase)
+//	{
+//		num++;
+//		FString Path = PathBase + FString::FromInt(num) + ".obj";
+//		std::ofstream f(*Path);
+//
+//		for (int32 VertexIdx = 0; VertexIdx < Delaunay.Vertices.Num(); VertexIdx++)
+//		{
+//			const FVector2<RealType>& Vertex = Delaunay.Vertices[VertexIdx];
+//			f << "v " << Vertex.X << " " << Vertex.Y << " 0" << std::endl;
+//		}
+//		for (int32 VertexIdx = 0; VertexIdx < Delaunay.Vertices.Num(); VertexIdx++)
+//		{
+//			const FVector2<RealType>& Vertex = Delaunay.Vertices[VertexIdx];
+//			f << "v " << Vertex.X << " " << Vertex.Y << " .5" << std::endl;
+//		}
+//		for (const FIndex2i& Edge : Delaunay.Edges)
+//		{
+//			f << "f " << Edge.A + 1 << " " << Edge.B + 1 << " " << 1 + Edge.A + Delaunay.Vertices.Num() << std::endl;
+//		}
+//		f.close();
+//	}
+//	//void DumpTriangulationForDebug(const FDynamicGraph2d& Graph, const TArray<FIntVector>& Triangles, const FString& PathBase)
+//	template <typename RealType>
+//	void DumpDelaunayTriangulationForDebug(const TConstrainedDelaunay2<RealType>& Delaunay, const FString& PathBase)
+//	{
+//		num++;
+//		FString Path = PathBase + FString::FromInt(num) + ".obj";
+//		std::ofstream f(*Path);
+//		for (int32 VertexIdx = 0; VertexIdx < Delaunay.Vertices.Num(); VertexIdx++)
+//		{
+//			const FVector2<RealType>& Vertex = Delaunay.Vertices[VertexIdx];
+//			f << "v " << Vertex.X << " " << Vertex.Y << " 0" << std::endl;
+//		}
+//		for (const FIndex3i& Tri : Delaunay.Triangles)
+//		{
+//			f << "f " << 1 + Tri.A << " " << 1 + Tri.B << " " << 1 + Tri.C << std::endl;
+//		}
+//		f.close();
+//	}
+//#endif
+//}
+
 template<class RealType>
 template<class InputRealType>
 void TConstrainedDelaunay2<RealType>::Add(const FDynamicGraph2<InputRealType>& Graph)
@@ -182,9 +240,12 @@ bool TConstrainedDelaunay2<RealType>::Triangulate()
 		}
 	}
 
+	int SelIdx = 0; // Index of item to Pop next; used to make the traversal less depth-first in shape, so a little more robust to bad data
 	while (ToWalkQ.Num())
 	{
-		TPair<int, int> TriWithWinding = ToWalkQ.Pop();
+		SelIdx = (SelIdx + 1) % ToWalkQ.Num();
+		TPair<int, int> TriWithWinding = ToWalkQ[SelIdx];
+		ToWalkQ.RemoveAtSwap(SelIdx);
 		int BaseIdx = TriWithWinding.Key * 3;
 		int LastWinding = TriWithWinding.Value;
 		for (int SubIdx = 0, NextIdx = 2; SubIdx < 3; NextIdx = SubIdx++)
@@ -209,14 +270,30 @@ bool TConstrainedDelaunay2<RealType>::Triangulate()
 	{
 		if (Keep[i] > 0)
 		{
-			Triangles.Add(FIntVector(Indices[i * 3], Indices[i * 3 + 1], Indices[i * 3 + 2]));
+			FIndex3i& Tri = Triangles.Emplace_GetRef(Indices[i * 3], Indices[i * 3 + 1], Indices[i * 3 + 2]);
+			if (!bOutputCCW)
+			{
+				Swap(Tri.B, Tri.C);
+			}
 		}
 	}
 	
 	return !bBoundaryTrackingFailure;
 }
+//
+template<typename RealType>
+TArray<FIndex3i> GEOMETRYALGORITHMS_API ConstrainedDelaunayTriangulate(const TGeneralPolygon2<RealType>& GeneralPolygon)
+{
+	TConstrainedDelaunay2<RealType> Triangulation;
+	Triangulation.FillRule = TConstrainedDelaunay2<RealType>::EFillRule::Positive;
+	Triangulation.Add(GeneralPolygon);
+	Triangulation.Triangulate();
+	return Triangulation.Triangles;
+}
 
 
+template TArray<FIndex3i> GEOMETRYALGORITHMS_API ConstrainedDelaunayTriangulate(const TGeneralPolygon2<double>& GeneralPolygon);
+template TArray<FIndex3i> GEOMETRYALGORITHMS_API ConstrainedDelaunayTriangulate(const TGeneralPolygon2<float>& GeneralPolygon);
 
 template void TConstrainedDelaunay2<double>::Add(const FDynamicGraph2<double>& Graph);
 template void TConstrainedDelaunay2<double>::Add(const TGeneralPolygon2<double>& Polygon);
diff --git a/Engine/Plugins/Experimental/GeometryProcessing/Source/GeometryAlgorithms/Public/ConstrainedDelaunay2.h b/Engine/Plugins/Experimental/GeometryProcessing/Source/GeometryAlgorithms/Public/ConstrainedDelaunay2.h
--- a/Engine/Plugins/Experimental/GeometryProcessing/Source/GeometryAlgorithms/Public/ConstrainedDelaunay2.h
+++ b/Engine/Plugins/Experimental/GeometryProcessing/Source/GeometryAlgorithms/Public/ConstrainedDelaunay2.h
@@ -3,6 +3,7 @@
 #pragma once
 
 #include "BoxTypes.h"
+#include "IndexTypes.h"
 #include "Polygon2.h"
 #include "Curve/DynamicGraph2.h"
 #include "Curve/GeneralPolygon2.h"
@@ -23,9 +24,10 @@ struct TConstrainedDelaunay2
 	TArray<FIndex2i> HoleEdges; // Any triangles inside 'hole' edges *must* be cut out
 
 	bool bOrientedEdges = true;
+	bool bOutputCCW = false;
 
 	enum class EFillRule {
-		Odd,
+		Odd = 0,
 		// bOrientedEdges must be true for the below
 		NonZero,
 		Positive,
@@ -61,7 +63,7 @@ struct TConstrainedDelaunay2
 	//
 	// outputs
 	//
-	TArray<FIntVector> Triangles;
+	TArray<FIndex3i> Triangles;
 
 	/**
 	 * Populate Triangles
@@ -71,7 +73,8 @@ struct TConstrainedDelaunay2
 	bool GEOMETRYALGORITHMS_API Triangulate();
 };
 
-
+template<typename RealType>
+TArray<FIndex3i> GEOMETRYALGORITHMS_API ConstrainedDelaunayTriangulate(const TGeneralPolygon2<RealType>& GeneralPolygon);
 
 typedef TConstrainedDelaunay2<float> FConstrainedDelaunay2f;
 typedef TConstrainedDelaunay2<double> FConstrainedDelaunay2d;
