Mesh Processing: Laplacian Mesh Deformation.  Fixed a bug that caused the laplacian deformation to crash when all mesh vertices were on the edge of the mesh, also added related unit test

diff --git a/Engine/Plugins/Experimental/GeometryProcessing/Source/MeshSolverUtilities/Private/LaplacianMeshSmoother.cpp b/Engine/Plugins/Experimental/GeometryProcessing/Source/MeshSolverUtilities/Private/LaplacianMeshSmoother.cpp
--- a/Engine/Plugins/Experimental/GeometryProcessing/Source/MeshSolverUtilities/Private/LaplacianMeshSmoother.cpp
+++ b/Engine/Plugins/Experimental/GeometryProcessing/Source/MeshSolverUtilities/Private/LaplacianMeshSmoother.cpp
@@ -67,12 +67,29 @@ FConstrainedMeshOperator::FConstrainedMeshOperator(const FDynamicMesh3& DynamicM
 	FSparseMatrixD LaplacianBoundary;
 	ConstructLaplacian(Scheme, DynamicMesh, VtxLinearization, LaplacianInternal, LaplacianBoundary);
 
-	// capture the number interior (non boundary) verts.
-	InternalVertexCount = VertexCount - VtxLinearization.NumBoundaryVerts();
+	// Copy the original boundary vertex locations
+	const int32 BoundaryVertexCount = VtxLinearization.NumBoundaryVerts();
+	
+	// Number of vertices in the interior of the mesh
+	InternalVertexCount = VertexCount - BoundaryVertexCount;
+
+	// Copy the original boundary vertex locations
+	BoundaryPositions.Empty(BoundaryVertexCount);
+	BoundaryPositions.AddUninitialized(BoundaryVertexCount);
+	{
+		const auto& ToIdx = VtxLinearization.ToIndex();
+		for (int32 i = 0; i < BoundaryVertexCount; ++i)
+		{
+			int32 VtxId = ToIdx[i + InternalVertexCount];
+			FVector3d Pos = DynamicMesh.GetVertex(VtxId);
+			BoundaryPositions[i] = Pos;
+		}
+	}
 
+	checkSlow(LaplacianInternal.rows() == LaplacianInternal.cols());
 
-	TUniquePtr<FSparseMatrixD> LTL( new FSparseMatrixD(LaplacianInternal.rows(), LaplacianInternal.cols()) );
-	FSparseMatrixD& LTLMatrix = *(LTL);
+	TUniquePtr<FSparseMatrixD> LTLPtr( new FSparseMatrixD(LaplacianInternal.rows(), LaplacianInternal.cols()) );
+	FSparseMatrixD& LTLMatrix = *(LTLPtr);
 
 	bool bIsLaplacianSymmetric = (Scheme == ELaplacianWeightScheme::Valence || Scheme == ELaplacianWeightScheme::Uniform);
 
@@ -81,23 +98,30 @@ FConstrainedMeshOperator::FConstrainedMeshOperator(const FDynamicMesh3& DynamicM
 		// Laplacian is symmetric, i.e. equal to its transpose
 		LTLMatrix = LaplacianInternal * LaplacianInternal;
 	
-		ConstrainedSolver.Reset(new FConstrainedSolver(LTL, MatrixSolverType));
+		ConstrainedSolver.Reset(new FConstrainedSolver(LTLPtr, MatrixSolverType));
 	}
 	else
 	{
 		// the laplacian 
 		LTLMatrix = LaplacianInternal.transpose() * LaplacianInternal;
-		ConstrainedSolver.Reset(new FConstrainedSolver(LTL, MatrixSolverType));
+		ConstrainedSolver.Reset(new FConstrainedSolver(LTLPtr, MatrixSolverType));
 	}
 
 }
 
 void FConstrainedMeshOperator::AddConstraint(const int32 VtxId, const double Weight, const FVector3d& Pos, const bool bPostFix)
 {
-	int32 Index = VtxLinearization.ToIndex()[VtxId];
+	const auto& ToIndex = VtxLinearization.ToIndex();
+	
+	if (VtxId > ToIndex.Num())
+	{
+		return;
+	}
+
+	const int32 Index = ToIndex[VtxId];
 
 	// Only add the constraint if the vertex is actually in the interior.  We aren't solving for edge vertices.
-	if (Index < InternalVertexCount)
+	if (Index != FDynamicMesh3::InvalidID && Index < InternalVertexCount)
 	{
 		bConstraintPositionsDirty = true;
 		bConstraintWeightsDirty = true;
@@ -111,10 +135,18 @@ void FConstrainedMeshOperator::AddConstraint(const int32 VtxId, const double Wei
 
 bool FConstrainedMeshOperator::UpdateConstraintPosition(const int32 VtxId, const FVector3d& Pos, const bool bPostFix)
 {
-	int32 Index = VtxLinearization.ToIndex()[VtxId];
 	bool Result = false;
+	const auto& ToIndex = VtxLinearization.ToIndex();
 
-	if (Index < InternalVertexCount)
+	if (VtxId > ToIndex.Num())
+	{
+		return Result;
+	}
+
+	const int32 Index = ToIndex[VtxId];
+	
+
+	if (Index != FDynamicMesh3::InvalidID && Index < InternalVertexCount)
 	{
 		bConstraintPositionsDirty = true;
 
@@ -129,8 +161,15 @@ bool FConstrainedMeshOperator::UpdateConstraintPosition(const int32 VtxId, const
 bool FConstrainedMeshOperator::UpdateConstraintWeight(const int32 VtxId, const double Weight)
 {
 	bool Result = false;
-	int32 Index = VtxLinearization.ToIndex()[VtxId];
-	if (Index < InternalVertexCount)
+	const auto& ToIndex = VtxLinearization.ToIndex();
+
+	if (VtxId > ToIndex.Num())
+	{
+		return Result;
+	}
+
+	const int32 Index = ToIndex[VtxId];
+	if (Index != FDynamicMesh3::InvalidID && Index < InternalVertexCount)
 	{
 		bConstraintWeightsDirty = true;
 
@@ -145,11 +184,18 @@ bool FConstrainedMeshOperator::UpdateConstraintWeight(const int32 VtxId, const d
 
 bool FConstrainedMeshOperator::IsConstrained(const int32 VtxId) const
 {
-	if (VtxId > VtxLinearization.ToIndex().Num()) return false;
-
-	int32 Index = VtxLinearization.ToIndex()[VtxId];
 	bool Result = false;
-	if (Index < InternalVertexCount)
+
+	const auto& ToIndex = VtxLinearization.ToIndex();
+
+	if (VtxId > ToIndex.Num())
+	{
+		return Result;
+	}
+
+	const int32 Index = ToIndex[VtxId];
+	
+	if (Index != FDynamicMesh3::InvalidID && Index < InternalVertexCount)
 	{
 		Result = ConstraintWeightMap.Contains(Index);
 	}
@@ -194,8 +240,19 @@ bool FConstrainedMeshOperator::Linearize(const FSOAPositions& PositionalVector,
 
 	LinearArray.Empty(MaxVtxId);
 	LinearArray.AddUninitialized(MaxVtxId);
+	
+	// Copy over the boundary positions.
+	{
+		int32 BoundaryVertexCount = VertexCount - InternalVertexCount;
+		for (int32 i = 0; i < BoundaryVertexCount; ++i)
+		{
+			const int32 VtxId = IndexToVtxId[i + InternalVertexCount];
+			LinearArray[VtxId] = BoundaryPositions[i];
+		}
+	}
 
-	for (int32 i = 0; i < Num; ++i)
+	// Update the internal positions.
+	for (int32 i = 0; i < InternalVertexCount; ++i)
 	{
 		const int32 VtxId = IndexToVtxId[i];
 
@@ -205,29 +262,22 @@ bool FConstrainedMeshOperator::Linearize(const FSOAPositions& PositionalVector,
 	return true;
 }
 
-void FConstrainedMeshOperator::ExtractVertexPositions(const FDynamicMesh3& DynamicMesh, FSOAPositions& VertexPositions) const
+void FConstrainedMeshOperator::ExtractInteriorVertexPositions(const FDynamicMesh3& DynamicMesh, FSOAPositions& VertexPositions) const
 {
 	VertexPositions.SetZero(InternalVertexCount);
 
+	const auto& ToVtxId = VtxLinearization.ToId();
 
-	const TArray<int32>& ToIndex = VtxLinearization.ToIndex();
-
-
-	for (int32 VtxId : DynamicMesh.VertexIndicesItr())
+	for (int32 i = 0; i < InternalVertexCount; ++i)
 	{
-		const int32 i = ToIndex[VtxId];
-
-		checkSlow(i != FDynamicMesh3::InvalidID);
-		
-		if (i < InternalVertexCount) // Skipping boundary verts
-		{
-			const FVector3d& Vertex = DynamicMesh.GetVertex(VtxId);
-			// coeffRef - coeff access without range checking
-			VertexPositions.XVector.coeffRef(i) = Vertex.X;
-			VertexPositions.YVector.coeffRef(i) = Vertex.Y;
-			VertexPositions.ZVector.coeffRef(i) = Vertex.Z;
-		}
+		const int32 VtxId = ToVtxId[i];
+		const FVector3d& Pos = DynamicMesh.GetVertex(VtxId);
+		VertexPositions.XVector.coeffRef(i) = Pos.X;
+		VertexPositions.YVector.coeffRef(i) = Pos.Y;
+		VertexPositions.ZVector.coeffRef(i) = Pos.Z;
 	}
+
+	
 }
 
 void FConstrainedMeshOperator::UpdateWithPostFixConstraints(FSOAPositions& PositionVector) const
@@ -261,10 +311,10 @@ FConstrainedMeshDeformer::FConstrainedMeshDeformer(const FDynamicMesh3& DynamicM
 
 	// The current vertex positions 
 	
-	// Note: the OriginalVertexPositions are being stored as member data 
+	// Note: the OriginalInteriorPositions are being stored as member data 
 	// for use if the solver is iterative.
-	// FSOAPositions OriginalVertexPositions; 
-	ExtractVertexPositions(DynamicMesh, OriginalVertexPositions);
+	// FSOAPositions OriginalInteriorPositions; 
+	ExtractInteriorVertexPositions(DynamicMesh, OriginalInteriorPositions);
 	
 	
 	// The biharmonic part of the constrained solver
@@ -275,17 +325,17 @@ FConstrainedMeshDeformer::FConstrainedMeshDeformer(const FDynamicMesh3& DynamicM
 	// Compute the Laplacian Vectors
 	//    := Biharmonic * VertexPostion
 	// In the case of the cotangent laplacian this can be identified as the mean curvature * normal.
-	checkSlow(LaplacianVectors.Num() == OriginalVertexPositions.Num());
+	checkSlow(LaplacianVectors.Num() == OriginalInteriorPositions.Num());
 
 	for (int32 i = 0; i < 3; ++i)
 	{
-		LaplacianVectors.Array(i) = Biharmonic * OriginalVertexPositions.Array(i);
+		LaplacianVectors.Array(i) = Biharmonic * OriginalInteriorPositions.Array(i);
 	}
 }
 
 bool FConstrainedMeshDeformer::Deform(TArray<FVector3d>& PositionBuffer) 
 {
-	
+
 	// Update constraints.  This only trigger solver rebuild if the weights were updated.
 	UpdateSolverConstraints();
 
@@ -294,13 +344,13 @@ bool FConstrainedMeshDeformer::Deform(TArray<FVector3d>& PositionBuffer)
 
 	// Solve the linear system
 	// NB: the original positions will only be used if the underlying solver type is iterative	
-	bool bSuccess = ConstrainedSolver->SolveWithGuess(OriginalVertexPositions, LaplacianVectors, SolutionVector);
+	bool bSuccess = ConstrainedSolver->SolveWithGuess(OriginalInteriorPositions, LaplacianVectors, SolutionVector);
 	
 	// Move any vertices to match bPostFix constraints
 
 	UpdateWithPostFixConstraints(SolutionVector);
 
-	// Copy the result into the array of stucts form.  
+	// Copy the result into the array of structs form.  
 	// NB: this re-indexes so the results can be looked up using VtxId
 
 	Linearize(SolutionVector, PositionBuffer);
@@ -553,13 +603,13 @@ TUniquePtr<FSparseMatrixD> FLaplacianDiffusionMeshSmoother::ConstructDiffusionOp
 																						const FDynamicMesh3& DynamicMesh,
 																						bool& bIsOperatorSymmetric,
 																						FVertexLinearization& Linearization,
-																						TArray<int32>* EdgeVtx ) 
+																						TArray<int32>* EdgeVtxs ) 
 {
 	bIsOperatorSymmetric = bIsSymmetricLaplacian(Scheme);
 
 
 	// Construct the laplacian, and extract the mapping for vertices (VtxLinearization)
-	TUniquePtr<FSparseMatrixD> Laplacian = ConstructLaplacian(Scheme, DynamicMesh, VtxLinearization, &EdgeVerts);
+	TUniquePtr<FSparseMatrixD> Laplacian = ConstructLaplacian(Scheme, DynamicMesh, VtxLinearization, EdgeVtxs);
 	Laplacian->makeCompressed();
 	return Laplacian;
 };
@@ -568,12 +618,12 @@ TUniquePtr<FSparseMatrixD> FBiHarmonicDiffusionMeshSmoother::ConstructDiffusionO
 																						 const FDynamicMesh3& DynamicMesh,
 																						 bool& bIsOperatorSymmetric,
 																						 FVertexLinearization& Linearization,
-																						 TArray<int32>* EdgeVtx ) 
+																						 TArray<int32>* EdgeVtxs ) 
 {
 	bIsOperatorSymmetric = true;
 	
 	// Construct the laplacian, and extract the mapping for vertices (VtxLinearization)
-	TUniquePtr<FSparseMatrixD> Laplacian = ConstructLaplacian(Scheme, DynamicMesh, VtxLinearization, &EdgeVerts);
+	TUniquePtr<FSparseMatrixD> Laplacian = ConstructLaplacian(Scheme, DynamicMesh, VtxLinearization, EdgeVtxs);
 	const FSparseMatrixD& L = *Laplacian;
 	TUniquePtr<FSparseMatrixD> MatrixOperator(new FSparseMatrixD());
 	
diff --git a/Engine/Plugins/Experimental/GeometryProcessing/Source/MeshSolverUtilities/Private/LaplacianMeshSmoother.h b/Engine/Plugins/Experimental/GeometryProcessing/Source/MeshSolverUtilities/Private/LaplacianMeshSmoother.h
--- a/Engine/Plugins/Experimental/GeometryProcessing/Source/MeshSolverUtilities/Private/LaplacianMeshSmoother.h
+++ b/Engine/Plugins/Experimental/GeometryProcessing/Source/MeshSolverUtilities/Private/LaplacianMeshSmoother.h
@@ -41,17 +41,16 @@ public:
 	// Test if for constraint associated with given vertex id. Will return false for any boundary vert.
 	bool IsConstrained(const int32 VtxId) const override;
 
-	
-
 	virtual bool Deform(TArray<FVector3d>& PositionBuffer) override { return false;  }
 
-protected:
-
-	// Sync constraints with internal solver.
+	// Sync constraints with internal solver.  If in the process any internal matrix factoring is dirty, it will be rebuilt.
+	// Note: this is called from within the Deform() method.   Only call this method if you want to trigger the matrix refactor yourself.
 	void UpdateSolverConstraints();
 
+protected:
+
 
-	void ExtractVertexPositions(const FDynamicMesh3& DynamicMesh, FSOAPositions& Positions) const ;
+	void ExtractInteriorVertexPositions(const FDynamicMesh3& DynamicMesh, FSOAPositions& Positions) const ;
 
 	//void UpdateMeshWithConstraints();
 	// Respect any bPostFix constraints by moving those vertices to position defined by said constraint.
@@ -83,6 +82,8 @@ protected:
 	// Used to map between VtxId and vertex Index in linear vector..
 	FVertexLinearization  VtxLinearization;
 
+	TArray<FVector3d> BoundaryPositions;
+
 	// Actual solver that manages the various linear algebra bits.
 	TUniquePtr<FConstrainedSolver>         ConstrainedSolver;
 };
@@ -99,7 +100,7 @@ public:
 private:
 
 	FSOAPositions LaplacianVectors;
-	FSOAPositions OriginalVertexPositions;
+	FSOAPositions OriginalInteriorPositions;
 };
 
 
diff --git a/Engine/Plugins/Experimental/GeometryProcessing/Source/MeshSolverUtilities/Private/MatrixSolver.h b/Engine/Plugins/Experimental/GeometryProcessing/Source/MeshSolverUtilities/Private/MatrixSolver.h
--- a/Engine/Plugins/Experimental/GeometryProcessing/Source/MeshSolverUtilities/Private/MatrixSolver.h
+++ b/Engine/Plugins/Experimental/GeometryProcessing/Source/MeshSolverUtilities/Private/MatrixSolver.h
@@ -135,15 +135,20 @@ public:
 
 	void Solve(const VectorType& BVector, VectorType& SolVector) const override
 	{
-		checkSlow(bSetup);
-		SolVector = MatrixSolver.solve(BVector);
+		if (bSetup)
+		{
+			SolVector = MatrixSolver.solve(BVector);
+		}
 	}
 
 	void Solve(const FSOAPositions& BVectors, FSOAPositions& SolVectors) const override
 	{
 		const bool bForceSingleThreaded = false;
-		ParallelFor(3, [&](int Dir)
-		{ SolVectors.Array(Dir) = MatrixSolver.solve(BVectors.Array(Dir)); }, bForceSingleThreaded);
+		if (bSetup)
+		{
+			ParallelFor(3, [&](int Dir)
+			{ SolVectors.Array(Dir) = MatrixSolver.solve(BVectors.Array(Dir)); }, bForceSingleThreaded);
+		}
 	}
 
 	void SetUp(const FSparseMatrixD& SparseMatrix, bool bIsSymmetric) override
@@ -164,7 +169,9 @@ public:
 
 	bool bSucceeded() const override
 	{
-		return (MatrixSolver.info() == Eigen::ComputationInfo::Success);
+		const auto Info = MatrixSolver.info();
+		bool bSuccess = bSetup && (Info == Eigen::ComputationInfo::Success);
+		return bSuccess;
 	}
 
 
@@ -172,6 +179,11 @@ private:
 
 	void SetUp(const FSparseMatrixD& SparseMatrix)
 	{
+		if (SparseMatrix.rows() == 0)
+		{
+			return;
+		}
+		
 		// The analyzePattern could be done just once if
 		// the sparsity pattern of the matrix is fixed.
 		// But testing indicates that takes little time compared 
@@ -180,7 +192,7 @@ private:
 		MatrixSolver.analyzePattern(SparseMatrix);
 		MatrixSolver.factorize(SparseMatrix);
 
-		bSetup = true;
+		bSetup = true;	
 	}
 private:
 
@@ -236,28 +248,38 @@ public:
 
 	void Solve(const VectorType& BVector, VectorType& SolVector) const override
 	{
-		checkSlow(bSetup);
-		SolVector = MatrixSolver.solve(BVector);
+		if (bSetup)
+		{
+			SolVector = MatrixSolver.solve(BVector);
+		}
 	}
 	void Solve(const FSOAPositions& BVectors, FSOAPositions& SolVectors) const override
 	{
 		const bool bForceSingleThreaded = false;
-		ParallelFor(3, [&](int Dir)
-		{ SolVectors.Array(Dir) = MatrixSolver.solve(BVectors.Array(Dir)); }, bForceSingleThreaded);
+		if (bSetup)
+		{
+			ParallelFor(3, [&](int Dir)
+			{ SolVectors.Array(Dir) = MatrixSolver.solve(BVectors.Array(Dir)); }, bForceSingleThreaded);
+		}
 	}
 
 	void SolveWithGuess(const VectorType& GuessVector, const VectorType& BVector, VectorType& SolVector)  const override
 	{
 
-		checkSlow(bSetup);
-		SolVector = MatrixSolver.solveWithGuess(BVector, GuessVector);
+		if (bSetup)
+		{
+			SolVector = MatrixSolver.solveWithGuess(BVector, GuessVector);
+		}
 	}
 
 	void SolveWithGuess(const FSOAPositions& GuessVectors, const FSOAPositions& BVectors, FSOAPositions& SolVectors) const override
 	{
 		const bool bForceSingleThreaded = false;
-		ParallelFor(3, [&](int Dir)
-		{ SolVectors.Array(Dir) = MatrixSolver.solveWithGuess(BVectors.Array(Dir), GuessVectors.Array(Dir)); }, bForceSingleThreaded);
+		if (bSetup)
+		{
+			ParallelFor(3, [&](int Dir)
+			{ SolVectors.Array(Dir) = MatrixSolver.solveWithGuess(BVectors.Array(Dir), GuessVectors.Array(Dir)); }, bForceSingleThreaded);
+		}
 	}
 	void SetUp(const FSparseMatrixD& SparseMatrix, bool bIsSymmetric) override
 	{
@@ -275,7 +297,9 @@ public:
 
 	bool bSucceeded() const override
 	{
-		return (MatrixSolver.info() == Eigen::ComputationInfo::Success);
+		const auto Info = MatrixSolver.info();
+		bool bSuccess = bSetup && (Info == Eigen::ComputationInfo::Success);
+		return bSuccess;
 	}
 
 private:
