Geometry Tools.  - Fix bug with mesh laplacian mesh deformation that scrambled the edge vertices of the mesh.

diff --git a/Engine/Plugins/Experimental/GeometryProcessing/Source/MeshSolverUtilities/Private/LaplacianMeshSmoother.cpp b/Engine/Plugins/Experimental/GeometryProcessing/Source/MeshSolverUtilities/Private/LaplacianMeshSmoother.cpp
--- a/Engine/Plugins/Experimental/GeometryProcessing/Source/MeshSolverUtilities/Private/LaplacianMeshSmoother.cpp
+++ b/Engine/Plugins/Experimental/GeometryProcessing/Source/MeshSolverUtilities/Private/LaplacianMeshSmoother.cpp
@@ -77,10 +77,10 @@ FConstrainedMeshOperator::FConstrainedMeshOperator(const FDynamicMesh3& DynamicM
 	BoundaryPositions.Empty(BoundaryVertexCount);
 	BoundaryPositions.AddUninitialized(BoundaryVertexCount);
 	{
-		const auto& ToIdx = VtxLinearization.ToIndex();
+		const auto& ToId = VtxLinearization.ToId();
 		for (int32 i = 0; i < BoundaryVertexCount; ++i)
 		{
-			int32 VtxId = ToIdx[i + InternalVertexCount];
+			int32 VtxId = ToId[i + InternalVertexCount];
 			FVector3d Pos = DynamicMesh.GetVertex(VtxId);
 			BoundaryPositions[i] = Pos;
 		}
