add generalized cylinders and curveutil to geometryprocessing

#rnx

diff --git a/Engine/Plugins/Experimental/GeometryProcessing/Source/DynamicMesh/Public/Generators/SweepGenerator.h b/Engine/Plugins/Experimental/GeometryProcessing/Source/DynamicMesh/Public/Generators/SweepGenerator.h
--- a/Engine/Plugins/Experimental/GeometryProcessing/Source/DynamicMesh/Public/Generators/SweepGenerator.h
+++ b/Engine/Plugins/Experimental/GeometryProcessing/Source/DynamicMesh/Public/Generators/SweepGenerator.h
@@ -12,6 +12,7 @@
 #include "FrameTypes.h"
 #include "MatrixTypes.h"
 #include "Polygon2.h"
+#include "Curve/CurveUtil.h"
 
 /**
  * ECapType indicates the type of cap to use on a sweep
@@ -284,3 +285,81 @@ public:
 		return *this;
 	}
 };
+
+/**
+ * Sweep a 2D Profile Polygon along a 3D Path.
+ * 
+ * TODO: 
+ *  - Loop path support
+ *  - Mitering cross sections support?
+ */
+class /*DYNAMICMESH_API*/ FGeneralizedCylinderGenerator : public FSweepGeneratorBase
+{
+public:
+	FPolygon2d CrossSection;
+	TArray<FVector3d> Path;
+
+	FFrame3d InitialFrame;
+
+	bool bCapped = false;
+
+public:
+	/** Generate the mesh */
+	virtual FMeshShapeGenerator& Generate() override
+	{
+		const bool bLoop = false; // TODO: support loops -- note this requires updating FSweepGeneratorBase to set up the mesh topology correctly!
+		
+		const TArray<FVector2d>& XVerts = CrossSection.GetVertices();
+		ECapType Caps[2] = {ECapType::None, ECapType::None};
+
+		if (bCapped)
+		{
+			Caps[0] = ECapType::FlatTriangulation;
+			Caps[1] = ECapType::FlatTriangulation;
+		}
+		int PathNum = Path.Num();
+		ConstructMeshTopology(CrossSection, {}, {}, PathNum, Caps, FVector2d(1, 1), FVector2d(0, 0));
+
+		int XNum = CrossSection.VertexCount();
+		TArray<FVector2d> XNormals; XNormals.SetNum(XNum);
+		for (int Idx = 0; Idx < XNum; Idx++)
+		{
+			XNormals[Idx] = CrossSection.GetNormal_FaceAvg(Idx);
+		}
+
+		FFrame3d CrossSectionFrame = InitialFrame;
+		for (int PathIdx = 0; PathIdx < PathNum; ++PathIdx)
+		{
+			FVector3d Tangent = TCurveUtil<double>::Tangent(Path, PathIdx, bLoop);
+			CrossSectionFrame.AlignAxis(2, Tangent);
+			FVector3d C = Path[PathIdx];
+			FVector3d X = CrossSectionFrame.X();
+			FVector3d Y = CrossSectionFrame.Y();
+			for (int SubIdx = 0; SubIdx < XNum; SubIdx++)
+			{
+				FVector2d XP = CrossSection[SubIdx];
+				FVector2d XN = XNormals[SubIdx];
+				Vertices[SubIdx + PathIdx * XNum] = C + X * XP.X + Y * XP.Y;
+				Normals[SubIdx + PathIdx * XNum] = FVector3f(X * XN.X + Y * XN.Y);
+			}
+		}
+		if (bCapped && !bLoop)
+		{
+			for (int CapIdx = 0; CapIdx < 2; CapIdx++)
+			{
+				FVector3f Normal(TCurveUtil<double>::Tangent(Path, CapIdx * (PathNum-1), bLoop) * (CapIdx * 2 - 1));
+				for (int SubIdx = 0; SubIdx < XNum; SubIdx++)
+				{
+					Normals[CapNormalStart[CapIdx] + SubIdx] = Normal;
+				}
+			}
+		}
+
+		for (int k = 0; k < Normals.Num(); ++k)
+		{
+			Normals[k].Normalize();
+		}
+
+		return *this;
+	}
+};
diff --git a/Engine/Plugins/Experimental/GeometryProcessing/Source/GeometricObjects/Public/Curve/CurveUtil.h b/Engine/Plugins/Experimental/GeometryProcessing/Source/GeometricObjects/Public/Curve/CurveUtil.h
new file mode 100644
--- /dev/null
+++ b/Engine/Plugins/Experimental/GeometryProcessing/Source/GeometricObjects/Public/Curve/CurveUtil.h
@@ -0,0 +1,166 @@
+// Copyright 1998-2019 Epic Games, Inc. All Rights Reserved.
+
+#pragma once
+
+#include "VectorTypes.h"
+#include "SegmentTypes.h"
+#include "LineTypes.h"
+#include "RayTypes.h"
+#include "BoxTypes.h"
+#include "SegmentTypes.h"
+
+/**
+ * Curve utility functions
+ */
+template<typename RealType>
+struct TCurveUtil
+{
+
+	static FVector3<RealType> Tangent(const TArrayView<const FVector3<RealType>>& Vertices, int32 Idx, bool bLoop = false)
+	{
+		int32 EndIdx = Idx + 1;
+		int32 StartIdx = Idx - 1;
+		int32 NV = Vertices.Num();
+		if (bLoop)
+		{
+			EndIdx = EndIdx % NV;
+			StartIdx = (StartIdx + NV) % NV;
+		}
+		else
+		{
+			EndIdx = FMath::Min(NV - 1, EndIdx);
+			StartIdx = FMath::Max(0, StartIdx);
+		}
+		return (Vertices[EndIdx] - Vertices[StartIdx]).Normalized();
+	}
+
+
+	static RealType ArcLength(const TArrayView<const FVector3<RealType>>& Vertices, bool bLoop = false) {
+		RealType Sum = 0;
+		int32 NV = Vertices.Num();
+		for (int i = 1; i < NV; ++i)
+		{
+			Sum += Vertices[i].Distance(Vertices[i - 1]);
+		}
+		if (bLoop)
+		{
+			Sum += Vertices[NV - 1].Distance(Vertices[0]);
+		}
+		return Sum;
+	}
+
+	static int FindNearestIndex(const TArrayView<const FVector3<RealType>>& Vertices, FVector3<RealType> V)
+	{
+		int iNearest = -1;
+		double dNear = TMathUtil<double>::MaxReal;
+		int N = Vertices.Num();
+		for ( int i = 0; i < N; ++i )
+		{
+			double dSqr = Vertices[i].DistanceSquared(V);
+			if (dSqr < dNear)
+			{
+				dNear = dSqr;
+				iNearest = i;
+			}
+		}
+		return iNearest;
+	}
+
+
+	/**
+	 * smooth vertices in-place (will not produce a symmetric result, but does not require extra buffer)
+	 */
+	static void InPlaceSmooth(TArrayView<FVector3<RealType>> Vertices, int StartIdx, int EndIdx, double Alpha, int NumIterations, bool bClosed)
+	{
+		int N = Vertices.Num();
+		if ( bClosed )
+		{
+			for (int Iter = 0; Iter < NumIterations; ++Iter)
+			{
+				for (int ii = StartIdx; ii < EndIdx; ++ii)
+				{
+					int i = (ii % N);
+					int iPrev = (ii == 0) ? N - 1 : ii - 1;
+					int iNext = (ii + 1) % N;
+					FVector3<RealType> prev = Vertices[iPrev], next = Vertices[iNext];
+					FVector3<RealType> c = (prev + next) * 0.5f;
+					Vertices[i] = (1 - Alpha) * Vertices[i] + (Alpha) * c;
+				}
+			}
+		}
+		else
+		{
+			for (int Iter = 0; Iter < NumIterations; ++Iter)
+			{
+				for (int i = StartIdx; i < EndIdx; ++i)
+				{
+					if (i == 0 || i >= N - 1)
+					{
+						continue;
+					}
+					FVector3<RealType> prev = Vertices[i - 1], next = Vertices[i + 1];
+					FVector3<RealType> c = (prev + next) * 0.5f;
+					Vertices[i] = (1 - Alpha) * Vertices[i] + (Alpha) * c;
+				}
+			}
+		}
+	}
+
+
+
+	/**
+	 * smooth set of vertices using extra buffer
+	 */
+	static void IterativeSmooth(TArrayView<FVector3<RealType>> Vertices, int StartIdx, int EndIdx, double Alpha, int NumIterations, bool bClosed)
+	{
+		int N = Vertices.Num();
+		TArray<FVector3<RealType>> Buffer;
+		Buffer.SetNumZeroed(N);
+
+		if (bClosed)
+		{
+			for (int Iter = 0; Iter < NumIterations; ++Iter)
+			{
+				for (int ii = StartIdx; ii < EndIdx; ++ii)
+				{
+					int i = (ii % N);
+					int iPrev = (ii == 0) ? N - 1 : ii - 1;
+					int iNext = (ii + 1) % N;
+					FVector3<RealType> prev = Vertices[iPrev], next = Vertices[iNext];
+					FVector3<RealType> c = (prev + next) * 0.5f;
+					Buffer[i] = (1 - Alpha) * Vertices[i] + (Alpha) * c;
+				}
+				for (int ii = StartIdx; ii < EndIdx; ++ii)
+				{
+					int i = (ii % N);
+					Vertices[i] = Buffer[i];
+				}
+			}
+		}
+		else
+		{
+			for (int Iter = 0; Iter < NumIterations; ++Iter)
+			{
+				for (int i = StartIdx; i < EndIdx && i < N; ++i)
+				{
+					if (i == 0 || i == N - 1)
+					{
+						continue;
+					}
+					FVector3<RealType> prev = Vertices[i - 1], next = Vertices[i + 1];
+					FVector3<RealType> c = (prev + next) * 0.5f;
+					Buffer[i] = (1 - Alpha) * Vertices[i] + (Alpha) * c;
+				}
+				for (int i = StartIdx; i < EndIdx && i < N; ++i)
+				{
+					if (i == 0 || i == N - 1)
+					{
+						continue;
+					}
+					Vertices[i] = Buffer[i];
+				}
+			}
+		}
+	}
+};
+
diff --git a/Engine/Plugins/Experimental/GeometryProcessing/Source/GeometricObjects/Public/Polyline3.h b/Engine/Plugins/Experimental/GeometryProcessing/Source/GeometricObjects/Public/Polyline3.h
--- a/Engine/Plugins/Experimental/GeometryProcessing/Source/GeometricObjects/Public/Polyline3.h
+++ b/Engine/Plugins/Experimental/GeometryProcessing/Source/GeometricObjects/Public/Polyline3.h
@@ -209,7 +209,7 @@ public:
 
 
 	/**
-	 * @return edge of the poyline starting at vertex SegmentIndex
+	 * @return edge of the polyline starting at vertex SegmentIndex
 	 */
 	TSegment3<T> GetSegment(int SegmentIndex) const
 	{
