Geometry tools - Make the generation of the laplacian more robust to bad meshes..  This will exclude free (no edges) vertices from the matrix.

diff --git a/Engine/Plugins/Experimental/GeometryProcessing/Source/MeshSolverUtilities/Private/MeshElementLinearizations.h b/Engine/Plugins/Experimental/GeometryProcessing/Source/MeshSolverUtilities/Private/MeshElementLinearizations.h
--- a/Engine/Plugins/Experimental/GeometryProcessing/Source/MeshSolverUtilities/Private/MeshElementLinearizations.h
+++ b/Engine/Plugins/Experimental/GeometryProcessing/Source/MeshSolverUtilities/Private/MeshElementLinearizations.h
@@ -99,7 +99,17 @@ private:
 		for (int32 i = 0, I = ToIdMap.Num(); i < I; ++i)
 		{
 			int32 VtxId = ToIdMap[i];
-			if (DynamicMesh.IsBoundaryVertex(VtxId))
+			// Test if the vertex has a one ring
+			
+			bool bEmptyOneRing = true;
+			for (int NeighborVertId : DynamicMesh.VtxVerticesItr(VtxId))
+			{
+				bEmptyOneRing = false;
+				break;
+			};
+				
+
+			if (bEmptyOneRing || DynamicMesh.IsBoundaryVertex(VtxId))
 			{
 				BoundaryVertIds.Add(VtxId);
 			}
