GeometryProcessing:  Provide new area-clamped version of cotan-laplacian. Also simplify and remove some code duplication in the laplacian matrix generation.

diff --git a/Engine/Plugins/Experimental/GeometryProcessing/Source/MeshSolverUtilities/Private/LaplacianOperators.cpp b/Engine/Plugins/Experimental/GeometryProcessing/Source/MeshSolverUtilities/Private/LaplacianOperators.cpp
--- a/Engine/Plugins/Experimental/GeometryProcessing/Source/MeshSolverUtilities/Private/LaplacianOperators.cpp
+++ b/Engine/Plugins/Experimental/GeometryProcessing/Source/MeshSolverUtilities/Private/LaplacianOperators.cpp
@@ -752,8 +752,8 @@ void ConstructCotangentLaplacian(const FDynamicMesh3& DynamicMesh, FVertexLinear
 	std::vector<MatrixTripletT> DiagonalTriplets;
 	DiagonalTriplets.reserve(NumVerts);
 
-	FSparseMatrixD Diagonals;
-	Diagonals.reserve(NumVerts);
+	FSparseMatrixD Diagonals(NumInteriorVerts, NumInteriorVerts);
+	Diagonals.reserve(NumInteriorVerts);
 
 
 	// Create an array that holds all the geometric information we need for each triangle.
@@ -831,7 +831,7 @@ void ConstructCotangentLaplacian(const FDynamicMesh3& DynamicMesh, FVertexLinear
 			else
 			{
 				int32 jBoundary = j - NumInteriorVerts;
-				BoundaryAsTripletList.push_back(MatrixTripletT(i, j, WeightIJ));
+				BoundaryAsTripletList.push_back(MatrixTripletT(i, jBoundary, WeightIJ));
 			}
 
 		}
@@ -866,6 +866,53 @@ void ConstructCotangentLaplacian(const FDynamicMesh3& DynamicMesh, FVertexLinear
 
 }
 
+double ConstructScaledCotangentLaplacian(const FDynamicMesh3& DynamicMesh, FVertexLinearization& VertexMap, FSparseMatrixD& LaplacianInterior, FSparseMatrixD& LaplacianBoundary, const bool bClampAreas)
+{
+	typedef FSparseMatrixD::Scalar  ScalarT;
+	typedef Eigen::Triplet<ScalarT> MatrixTripletT;
+
+	// diagonal mass matrix.
+	FSparseMatrixD AreaMatrix;
+	FSparseMatrixD CotangentInterior;
+	FSparseMatrixD CotangentBoundary;
+	ConstructCotangentLaplacian(DynamicMesh, VertexMap, AreaMatrix, CotangentInterior, CotangentBoundary);
+
+	// Find average entry in the area matrix
+	const int32 Rank = AreaMatrix.cols();
+	double AveArea = 0.;
+	for (int32 i = 0; i < Rank; ++i)
+	{
+		double Area = AreaMatrix.coeff(i, i);
+		checkSlow(Area > 0.);  // Area must be positive.
+		AveArea += Area;
+	}
+	AveArea /= Rank;
+	
+	std::vector<MatrixTripletT> ScaledInvAreaTriplets;
+	ScaledInvAreaTriplets.reserve(Rank);
+	for (int32 i = 0; i < Rank; ++i)
+	{
+		double Area = AreaMatrix.coeff(i, i);
+		double ScaledInvArea = AveArea / Area;
+		if (bClampAreas)
+		{
+			ScaledInvArea = FMath::Clamp(ScaledInvArea, 0.5, 5.); // when  squared this gives largest scales 100 x smallest
+		}
+
+		ScaledInvAreaTriplets.push_back(MatrixTripletT(i, i, ScaledInvArea));
+	}
+
+	FSparseMatrixD ScaledInvAreaMatrix(Rank, Rank);
+	ScaledInvAreaMatrix.setFromTriplets(ScaledInvAreaTriplets.begin(), ScaledInvAreaTriplets.end());
+	ScaledInvAreaMatrix.makeCompressed();
+
+	LaplacianBoundary = ScaledInvAreaMatrix * CotangentBoundary;
+	LaplacianBoundary.makeCompressed();
+	LaplacianInterior = ScaledInvAreaMatrix * CotangentInterior;
+	LaplacianInterior.makeCompressed();	
+
+	return AveArea;
+}
 
 void ConstructCotangentLaplacian(const FDynamicMesh3& DynamicMesh, FVertexLinearization& VertexMap, FSparseMatrixD& LaplacianInterior, FSparseMatrixD& LaplacianBoundary, const bool bClampWeights)
 {
@@ -1147,13 +1194,13 @@ void ConstructLaplacian(const ELaplacianWeightScheme Scheme, const FDynamicMesh3
 	case ELaplacianWeightScheme::Cotangent:
 	{
 		bool bClampWeights = false;
-		ConstructCotangentLaplacian(DynamicMesh, VertexMap, LaplacianInterior, LaplacianBoundary, bClampWeights);
+		ConstructScaledCotangentLaplacian(DynamicMesh, VertexMap, LaplacianInterior, LaplacianBoundary, bClampWeights);
 		break;
 	}
 	case ELaplacianWeightScheme::ClampedCotangent:
 	{
 		bool bClampWeights = true;
-		ConstructCotangentLaplacian(DynamicMesh, VertexMap, LaplacianInterior, LaplacianBoundary, bClampWeights);
+		ConstructScaledCotangentLaplacian(DynamicMesh, VertexMap, LaplacianInterior, LaplacianBoundary, bClampWeights);
 		break;
 	}
 	case ELaplacianWeightScheme::MeanValue:
@@ -1162,265 +1209,67 @@ void ConstructLaplacian(const ELaplacianWeightScheme Scheme, const FDynamicMesh3
 	}
 }
 
-
-
-TUniquePtr<FSparseMatrixD> ConstructUniformLaplacian(const FDynamicMesh3& DynamicMesh, FVertexLinearization& VertexMap, TArray<int32>* BoundaryVerts)
+static void ExtractBoundaryVerts(const FVertexLinearization& VertexMap, TArray<int32>& BoundaryVerts)
 {
+	int32 NumBoundaryVerts = VertexMap.NumBoundaryVerts();
+	int32 NumInternalVerts = VertexMap.NumVerts() - NumBoundaryVerts;
+	BoundaryVerts.Empty(NumBoundaryVerts);
 
-
-	typedef FSparseMatrixD::Scalar    ScalarT;
-	typedef Eigen::Triplet<ScalarT>  MatrixTripletT;
-
-	if (BoundaryVerts)
+	const auto& ToId = VertexMap.ToId();
+	for (int32 i = NumInternalVerts; i < VertexMap.NumVerts(); ++i)
 	{
-		// empty
-		BoundaryVerts->Empty();
+		int32 VtxId = ToId[i];
+		BoundaryVerts.Add(VtxId);
 	}
+}
 
-	// Sync the mapping between the mesh vertex ids and their offsets in a nomimal linear array.
-	VertexMap.Reset(DynamicMesh);
-
-	const TArray<int32>& ToMeshV = VertexMap.ToId();
-	const TArray<int32>& ToIndex = VertexMap.ToIndex();
-	const int32 NumVerts = VertexMap.NumVerts();
-
-	// Eigen constructs a sparse matrix from a linearized array of matrix entries.
+TUniquePtr<FSparseMatrixD> ConstructUniformLaplacian(const FDynamicMesh3& DynamicMesh, FVertexLinearization& VertexMap, TArray<int32>* BoundaryVerts)
+{
 
-	std::vector<MatrixTripletT> MatrixTripelList;
-	{
-		int32 NumMatrixEntries = ComputeNumMatrixElements(DynamicMesh, ToMeshV);
-		MatrixTripelList.reserve(NumMatrixEntries);
-	}
+	TUniquePtr<FSparseMatrixD> LaplacianMatrix(new FSparseMatrixD);
+	FSparseMatrixD BoundaryMatrix;
 
-	// Construct Laplacian Matrix: loop over verts constructing the corresponding matrix row.
+	ConstructUniformLaplacian(DynamicMesh, VertexMap, *LaplacianMatrix, BoundaryMatrix);
 
-	for (int32 i = 0; i < NumVerts; ++i)
+	if (BoundaryVerts)
 	{
-		const int32 VertId = ToMeshV[i];
-
-		if (DynamicMesh.IsBoundaryVertex(VertId))
-		{
-			if (BoundaryVerts)
-			{
-				BoundaryVerts->Add(VertId);
-			}
-#if LAPLACIAN_SKIP_BOUNDARY == 1  
-			MatrixTripelList.push_back(MatrixTripletT(i, i, 0.));
-			continue;
-#endif 
-		}
-
-
-		ScalarT CenterWeight = ScalarT(0); // equal and opposite the sum of the neighbor weights
-		for (int NeighborVertId : DynamicMesh.VtxVerticesItr(VertId))
-		{
-			const int32 j = ToIndex[NeighborVertId];
-
-
-			ScalarT NeighborWeight = ScalarT(1);
-			CenterWeight += NeighborWeight;
-
-			// add the neighbor 
-			MatrixTripelList.push_back(MatrixTripletT(i, j, NeighborWeight));
-
-		}
-		// add the center
-		MatrixTripelList.push_back(MatrixTripletT(i, i, -CenterWeight));
-
+		ExtractBoundaryVerts(VertexMap, *BoundaryVerts);
 	}
 
-	// SparseMatrix doesn't have (SparseMatrix&& ) constructor 
-
-	TUniquePtr<FSparseMatrixD> LaplacianMatrix(new FSparseMatrixD(NumVerts, NumVerts));
-	LaplacianMatrix->setFromTriplets(MatrixTripelList.begin(), MatrixTripelList.end());
-
-	LaplacianMatrix->makeCompressed();
-
 	return LaplacianMatrix;
-
 }
+	
 
 
 TUniquePtr<FSparseMatrixD> ConstructUmbrellaLaplacian(const FDynamicMesh3& DynamicMesh, FVertexLinearization& VertexMap, TArray<int32>* BoundaryVerts)
 {
+	TUniquePtr<FSparseMatrixD> LaplacianMatrix(new FSparseMatrixD);
+	FSparseMatrixD BoundaryMatrix;
 
-
-	typedef FSparseMatrixD::Scalar    ScalarT;
-	typedef Eigen::Triplet<ScalarT>  MatrixTripletT;
+	ConstructUmbrellaLaplacian(DynamicMesh, VertexMap, *LaplacianMatrix, BoundaryMatrix);
 
 	if (BoundaryVerts)
 	{
-		// empty
-		BoundaryVerts->Empty();
-	}
-
-	// Sync the mapping between the mesh vertex ids and their offsets in a nomimal linear array.
-	VertexMap.Reset(DynamicMesh);
-
-	const TArray<int32>& ToMeshV = VertexMap.ToId();
-	const TArray<int32>& ToIndex = VertexMap.ToIndex();
-	const int32 NumVerts = VertexMap.NumVerts();
-
-	// Eigen constructs a sparse matrix from a linearized array of matrix entries.
-
-	std::vector<MatrixTripletT> MatrixTripelList;
-	{
-		int32 NumMatrixEntries = ComputeNumMatrixElements(DynamicMesh, ToMeshV);
-		MatrixTripelList.reserve(NumMatrixEntries);
-	}
-
-	// Cache valency of each vertex.
-	// Number of non-zero elements in the i'th row = 1 + OneRingSize(i)
-	TArray<int32> OneRingSize;
-	{
-		OneRingSize.SetNumUninitialized(NumVerts);
-
-		for (int32 i = 0; i < NumVerts; ++i)
-		{
-			const int32 VertId = ToMeshV[i];
-			OneRingSize[i] = DynamicMesh.GetVtxEdgeCount(VertId);
-		}
-	}
-
-	// Construct Laplacian Matrix: loop over verts constructing the corresponding matrix row.
-
-	for (int32 i = 0; i < NumVerts; ++i)
-	{
-		const int32 VertId = ToMeshV[i];
-		const int32 Valence = OneRingSize[i];
-		double InvValence = (Valence != 0) ? 1. / double(Valence) : 0.;
-
-		if (DynamicMesh.IsBoundaryVertex(VertId))
-		{
-			if (BoundaryVerts)
-			{
-				BoundaryVerts->Add(VertId);
-			}
-#if LAPLACIAN_SKIP_BOUNDARY == 1  
-			MatrixTripelList.push_back(MatrixTripletT(i, i, 0.));
-			continue;
-#endif 
-		}
-
-		for (int NeighborVertId : DynamicMesh.VtxVerticesItr(VertId))
-		{
-			const int32 j = ToIndex[NeighborVertId];
-
-			// add the neighbor 
-			MatrixTripelList.push_back(MatrixTripletT(i, j, InvValence));
-
-		}
-		// add the center
-		MatrixTripelList.push_back(MatrixTripletT(i, i, -ScalarT(1)));
-
+		ExtractBoundaryVerts(VertexMap, *BoundaryVerts);
 	}
 
-	// SparseMatrix doesn't have (SparseMatrix&& ) constructor 
-
-	TUniquePtr<FSparseMatrixD> LaplacianMatrix(new FSparseMatrixD(NumVerts, NumVerts));
-	LaplacianMatrix->setFromTriplets(MatrixTripelList.begin(), MatrixTripelList.end());
-
-	LaplacianMatrix->makeCompressed();
-
 	return LaplacianMatrix;
-
 }
-
+	
 
 TUniquePtr<FSparseMatrixD> ConstructValenceWeightedLaplacian(const FDynamicMesh3& DynamicMesh, FVertexLinearization& VertexMap, TArray<int32>* BoundaryVerts)
 {
 
+	TUniquePtr<FSparseMatrixD> LaplacianMatrix(new FSparseMatrixD);
+	FSparseMatrixD BoundaryMatrix;
 
-	typedef FSparseMatrixD::Scalar    ScalarT;
-	typedef Eigen::Triplet<ScalarT>  MatrixTripletT;
+	ConstructValenceWeightedLaplacian(DynamicMesh, VertexMap, *LaplacianMatrix, BoundaryMatrix);
 
 	if (BoundaryVerts)
 	{
-		// empty
-		BoundaryVerts->Empty();
-	}
-
-	// Sync the mapping between the mesh vertex ids and their offsets in a nomimal linear array.
-	VertexMap.Reset(DynamicMesh);
-
-	const TArray<int32>& ToMeshV = VertexMap.ToId();
-	const TArray<int32>& ToIndex = VertexMap.ToIndex();
-	const int32 NumVerts = VertexMap.NumVerts();
-
-	// Cache valency of each vertex.
-	// Number of non-zero elements in the i'th row = 1 + OneRingSize(i)
-	TArray<int32> OneRingSize;
-	{
-		OneRingSize.SetNumUninitialized(NumVerts);
-
-		for (int32 i = 0; i < NumVerts; ++i)
-		{
-			const int32 VertId = ToMeshV[i];
-			OneRingSize[i] = DynamicMesh.GetVtxEdgeCount(VertId);
-		}
-	}
-
-
-	// Compute the total number of entries in the sparse matrix
-	int32 NumMatrixEntries = 0;
-	{
-		for (int32 i = 0; i < NumVerts; ++i)
-		{
-			NumMatrixEntries += 1 + OneRingSize[i]; // myself plus my neighbors
-		}
-
-	}
-	std::vector<MatrixTripletT> MatrixTripelList;
-	MatrixTripelList.reserve(NumMatrixEntries);
-
-
-
-	// Construct Laplacian Matrix: loop over verts constructing the corresponding matrix row.
-
-	for (int32 i = 0; i < NumVerts; ++i)
-	{
-		const int32 VertId = ToMeshV[i];
-		const int32 IOneRingSize = OneRingSize[i];
-
-		if (DynamicMesh.IsBoundaryVertex(VertId))
-		{
-			if (BoundaryVerts)
-			{
-				BoundaryVerts->Add(VertId);
-			}
-#if LAPLACIAN_SKIP_BOUNDARY == 1  
-			MatrixTripelList.push_back(MatrixTripletT(i, i, 0.));
-			continue;
-#endif 
-		}
-
-
-		ScalarT CenterWeight = ScalarT(0); // equal and opposite the sum of the neighbor weights
-		for (int NeighborVertId : DynamicMesh.VtxVerticesItr(VertId))
-		{
-			const int32 j = ToIndex[NeighborVertId];
-			const int32 JOneRingSize = OneRingSize[j];
-
-
-			ScalarT NeighborWeight = ScalarT(1) / std::sqrt(IOneRingSize + JOneRingSize);
-			CenterWeight += NeighborWeight;
-
-			// add the neighbor 
-			MatrixTripelList.push_back(MatrixTripletT(i, j, NeighborWeight));
-
-		}
-		// add the center
-		MatrixTripelList.push_back(MatrixTripletT(i, i, -CenterWeight));
-
+		ExtractBoundaryVerts(VertexMap, *BoundaryVerts);
 	}
-
-	// SparseMatrix doesn't have (SparseMatrix&& ) constructor 
-
-	TUniquePtr<FSparseMatrixD> LaplacianMatrix(new FSparseMatrixD(NumVerts, NumVerts));
-	LaplacianMatrix->setFromTriplets(MatrixTripelList.begin(), MatrixTripelList.end());
-
-	LaplacianMatrix->makeCompressed();
-
+	
 	return LaplacianMatrix;
 
 }
@@ -1428,429 +1277,56 @@ TUniquePtr<FSparseMatrixD> ConstructValenceWeightedLaplacian(const FDynamicMesh3
 
 TUniquePtr<FSparseMatrixD> ConstructCotangentLaplacian(const FDynamicMesh3& DynamicMesh, FVertexLinearization& VertexMap, FSparseMatrixD& AreaMatrix, TArray<int32>* BoundaryVerts)
 {
+	TUniquePtr<FSparseMatrixD> LaplacianMatrix(new FSparseMatrixD);
+	FSparseMatrixD BoundaryMatrix;
+
+	ConstructCotangentLaplacian(DynamicMesh, VertexMap, AreaMatrix, *LaplacianMatrix, BoundaryMatrix);
 
 	if (BoundaryVerts)
 	{
-		// empty
-		BoundaryVerts->Empty();
+		ExtractBoundaryVerts(VertexMap, *BoundaryVerts);
 	}
 
-	typedef FSparseMatrixD::Scalar  ScalarT;
-	typedef Eigen::Triplet<ScalarT> MatrixTripletT;
-
-
-	// Create a mapping between the ordering of the vertex indices in the Dynamic Mesh
-	// and the actual storage. 
-	VertexMap.Reset(DynamicMesh);
-
-	const TArray<int32>& ToMeshV = VertexMap.ToId();
-	const TArray<int32>& ToIndex = VertexMap.ToIndex();
-	const int32 NumVerts = VertexMap.NumVerts();
-
-
-	// Create the mapping of triangles
-
-	FTriangleLinearization TriangleMap(DynamicMesh);
-
-	const TArray<int32>& ToMeshTri = TriangleMap.ToId();
-	const TArray<int32>& ToTriIdx = TriangleMap.ToIndex();
-	const int32 NumTris = TriangleMap.NumTris();
-
-
-	// Clear space for the areas
-	std::vector<MatrixTripletT> DiagonalTriplets;
-	DiagonalTriplets.reserve(NumVerts);
-
-	FSparseMatrixD Diagonals;
-	Diagonals.reserve(NumVerts);
+	return LaplacianMatrix;
+}
 
 
-	// Create an array that holds all the geometric information we need for each triangle.
 
-	TArray<CotanTriangleData>  CotangentTriangleDataArray = ConstructTriangleDataArray<CotanTriangleData>(DynamicMesh, TriangleMap);
+TUniquePtr<FSparseMatrixD> ConstructCotangentLaplacian(const FDynamicMesh3& DynamicMesh, FVertexLinearization& VertexMap, const bool bClampWeights, TArray<int32>* BoundaryVerts)
+{
 
+	TUniquePtr<FSparseMatrixD> LaplacianMatrix(new FSparseMatrixD);
+	FSparseMatrixD BoundaryMatrix;
 
-	// Eigen constructs a sparse matrix from a linearized array of matrix entries.
+	ConstructCotangentLaplacian(DynamicMesh, VertexMap, *LaplacianMatrix, BoundaryMatrix, bClampWeights);
 
-	std::vector<MatrixTripletT> MatrixTripelList;
+	if (BoundaryVerts)
 	{
-		int32 NumMatrixEntries = ComputeNumMatrixElements(DynamicMesh, ToMeshV);
-		MatrixTripelList.reserve(NumMatrixEntries);
+		ExtractBoundaryVerts(VertexMap, *BoundaryVerts);
 	}
 
+	return LaplacianMatrix;
+}
 
-	// Construct Laplacian Matrix: loop over verts constructing the corresponding matrix row.
-	//                             store the id of the boundary verts for later use.
-
-	for (int32 i = 0; i < NumVerts; ++i)
-	{
-		const int32 IVertId = ToMeshV[i]; // I - the row
 
 
-		if (DynamicMesh.IsBoundaryVertex(IVertId))
-		{
-			if (BoundaryVerts)
-			{
-				BoundaryVerts->Add(IVertId);
-			}
-#if LAPLACIAN_SKIP_BOUNDARY == 1  
-			MatrixTripelList.push_back(MatrixTripletT(i, i, 0.));
-			DiagonalTriplets.push_back(MatrixTripletT(i, i, 1.));
-			continue;
-#endif 
-		}
+TUniquePtr<FSparseMatrixD> ConstructMeanValueWeightLaplacian(const FDynamicMesh3& DynamicMesh, FVertexLinearization& VertexMap, TArray<int32>* BoundaryVerts)
+{
 
+	TUniquePtr<FSparseMatrixD> LaplacianMatrix(new FSparseMatrixD);
+	FSparseMatrixD BoundaryMatrix;
 
-		// Compute the Voronoi area for this vertex.
-		double WeightArea = 0.;
-		for (int32 TriId : DynamicMesh.VtxTrianglesItr(IVertId))
-		{
-			const int32 TriIdx = ToTriIdx[TriId];
-			const CotanTriangleData& TriData = CotangentTriangleDataArray[TriIdx];
+	ConstructMeanValueWeightLaplacian(DynamicMesh, VertexMap, *LaplacianMatrix, BoundaryMatrix);
 
+	if (BoundaryVerts)
+	{
+		ExtractBoundaryVerts(VertexMap, *BoundaryVerts);
+	}
 
-			// The three VertIds for this triangle.
-			const FIndex3i TriVertIds = DynamicMesh.GetTriangle(TriId);
+	return LaplacianMatrix;
+}
 
-			// Which of the corners is IVertId?
-			int32 Offset = 0;
-			while (TriVertIds[Offset] != IVertId)
-			{
-				Offset++;
-				checkSlow(Offset < 3);
-			}
-
-			WeightArea += TriData.VoronoiArea[Offset];
-		}
-
-
-		double WeightII = 0.; // accumulate to equal and opposite the sum of the neighbor weights
-
-		// for each connecting edge
-
-		for (int32 EdgeId : DynamicMesh.VtxEdgesItr(IVertId))
-		{
-			// [v0, v1, t0, t1]:  NB: both t0 & t1 exist since IVert isn't a boundary vert.
-			FIndex4i Edge = DynamicMesh.GetEdge(EdgeId);
-
-
-			// the other vert in the edge - identifies the matrix column
-			const int32 JVertId = (Edge[0] == IVertId) ? Edge[1] : Edge[0];  // J - the column
-
-			checkSlow(JVertId != IVertId);
-
-			// Get the cotangents for this edge.
-
-			const int32 Tri0Idx = ToTriIdx[Edge[2]];
-			const CotanTriangleData& Tri0Data = CotangentTriangleDataArray[Tri0Idx];
-			const double CotanAlpha = Tri0Data.GetOpposingCotangent(EdgeId);
-
-
-			// The second triangle will be invalid if this is an edge!
-
-			const double CotanBeta = (Edge[3] != FDynamicMesh3::InvalidID) ? CotangentTriangleDataArray[ToTriIdx[Edge[3]]].GetOpposingCotangent(EdgeId) : 0.;
-
-			double WeightIJ = 0.5 * (CotanAlpha + CotanBeta);
-			WeightII += WeightIJ;
-
-			const int32 j = ToIndex[JVertId];
-
-			MatrixTripelList.push_back(MatrixTripletT(i, j, WeightIJ));
-
-
-		}
-
-		MatrixTripelList.push_back(MatrixTripletT(i, i, -WeightII));
-
-		DiagonalTriplets.push_back(MatrixTripletT(i, i, WeightArea));
-
-
-	}
-
-	Diagonals.setFromTriplets(DiagonalTriplets.begin(), DiagonalTriplets.end());
-	AreaMatrix.swap(Diagonals);
-	AreaMatrix.makeCompressed();
-
-	TUniquePtr<FSparseMatrixD> LaplacianMatrix(new FSparseMatrixD(NumVerts, NumVerts));
-	LaplacianMatrix->setFromTriplets(MatrixTripelList.begin(), MatrixTripelList.end());
-	LaplacianMatrix->makeCompressed();
-
-
-	return LaplacianMatrix;
-
-}
-
-
-TUniquePtr<FSparseMatrixD> ConstructCotangentLaplacian(const FDynamicMesh3& DynamicMesh, FVertexLinearization& VertexMap, const bool bClampWeights, TArray<int32>* BoundaryVerts)
-{
-
-	if (BoundaryVerts)
-	{
-		// empty
-		BoundaryVerts->Empty();
-	}
-
-	typedef FSparseMatrixD::Scalar  ScalarT;
-	typedef Eigen::Triplet<ScalarT> MatrixTripletT;
-
-
-	// Create a mapping between the ordering of the vertex indices in the Dynamic Mesh
-	// and the actual storage. 
-	VertexMap.Reset(DynamicMesh);
-
-	const TArray<int32>& ToMeshV = VertexMap.ToId();
-	const TArray<int32>& ToIndex = VertexMap.ToIndex();
-	const int32 NumVerts = VertexMap.NumVerts();
-
-
-	// Map the triangles.
-
-	FTriangleLinearization TriangleMap(DynamicMesh);
-
-	const TArray<int32>& ToMeshTri = TriangleMap.ToId();
-	const TArray<int32>& ToTriIdx = TriangleMap.ToIndex();
-	const int32 NumTris = TriangleMap.NumTris();
-
-	// Create an array that holds all the geometric information we need for each triangle.
-
-	TArray<CotanTriangleData>  CotangentTriangleDataArray = ConstructTriangleDataArray<CotanTriangleData>(DynamicMesh, TriangleMap);
-
-	// Eigen constructs a sparse matrix from a linearized array of matrix entries.
-
-	std::vector<MatrixTripletT> MatrixTripelList;
-	{
-		int32 NumMatrixEntries = ComputeNumMatrixElements(DynamicMesh, ToMeshV);
-		MatrixTripelList.reserve(NumMatrixEntries);
-	}
-
-
-	// Construct Laplacian Matrix: loop over verts constructing the corresponding matrix row.
-	//                             skipping the boundary verts for later use.
-
-	for (int32 i = 0; i < NumVerts; ++i)
-	{
-		const int32 IVertId = ToMeshV[i]; // I - the row
-
-
-		if (DynamicMesh.IsBoundaryVertex(IVertId))
-		{
-			if (BoundaryVerts)
-			{
-				BoundaryVerts->Add(IVertId);
-			}
-#if LAPLACIAN_SKIP_BOUNDARY == 1  
-			MatrixTripelList.push_back(MatrixTripletT(i, i, 0.));
-			continue;
-#endif 
-		}
-
-
-		// Compute the Voronoi area for this vertex.
-		double WeightArea = 0.;
-		for (int32 TriId : DynamicMesh.VtxTrianglesItr(IVertId))
-		{
-			const int32 TriIdx = ToTriIdx[TriId];
-			const CotanTriangleData& TriData = CotangentTriangleDataArray[TriIdx];
-
-
-			// The three VertIds for this triangle.
-			const FIndex3i TriVertIds = DynamicMesh.GetTriangle(TriId);
-
-			// Which of the corners is IVertId?
-			int32 Offset = 0;
-			while (TriVertIds[Offset] != IVertId)
-			{
-				Offset++;
-				checkSlow(Offset < 3);
-			}
-
-			WeightArea += TriData.VoronoiArea[Offset];
-		}
-
-
-		double WeightII = 0.; // accumulate to equal and opposite the sum of the neighbor weights
-
-		// for each connecting edge
-
-		for (int32 EdgeId : DynamicMesh.VtxEdgesItr(IVertId))
-		{
-			// [v0, v1, t0, t1]:  NB: both t0 & t1 exist since IVert isn't a boundary vert.
-			FIndex4i Edge = DynamicMesh.GetEdge(EdgeId);
-
-
-			// the other vert in the edge - identifies the matrix column
-			const int32 JVertId = (Edge[0] == IVertId) ? Edge[1] : Edge[0];  // J - the column
-
-			checkSlow(JVertId != IVertId);
-
-			// Get the cotangents for this edge.
-
-			const int32 Tri0Idx = ToTriIdx[Edge[2]];
-			const CotanTriangleData& Tri0Data = CotangentTriangleDataArray[Tri0Idx];
-			const double CotanAlpha = Tri0Data.GetOpposingCotangent(EdgeId);
-
-			// The second triangle will be invalid if this is an edge!
-
-			const double CotanBeta = (Edge[3] != FDynamicMesh3::InvalidID) ? CotangentTriangleDataArray[ToTriIdx[Edge[3]]].GetOpposingCotangent(EdgeId) : 0.;
-
-			double WeightIJ = 0.5 * (CotanAlpha + CotanBeta);
-
-			// clamp the weight
-			if (bClampWeights)
-			{
-				WeightIJ = FMath::Clamp(WeightIJ, -1.e5 * WeightArea, 1.e5 * WeightArea);
-			}
-
-			WeightII += WeightIJ;
-
-			const int32 j = ToIndex[JVertId];
-
-			MatrixTripelList.push_back(MatrixTripletT(i, j, WeightIJ / WeightArea));
-
-		}
-
-		MatrixTripelList.push_back(MatrixTripletT(i, i, -WeightII / WeightArea));
-
-	}
-
-#if 0
-	// used in degubing to insepct the matrix for testing
-	std::vector<MatrixTripletT> Tmp = MatrixTripelList;
-
-	auto SortFunctor = [](const MatrixTripletT& A, const MatrixTripletT& B)->bool
-	{
-		return (A.value() < B.value());
-	};
-
-	std::sort(Tmp.begin(), Tmp.end(), SortFunctor);
-
-	auto ReverseSortFunctor = [](const MatrixTripletT& A, const MatrixTripletT& B)->bool
-	{
-		return (A.value() > B.value());
-	};
-
-	std::sort(Tmp.begin(), Tmp.end(), ReverseSortFunctor);
-#endif 
-
-	TUniquePtr<FSparseMatrixD> LaplacianMatrix(new FSparseMatrixD(NumVerts, NumVerts));
-	LaplacianMatrix->setFromTriplets(MatrixTripelList.begin(), MatrixTripelList.end());
-	LaplacianMatrix->makeCompressed();
-
-	return LaplacianMatrix;
-
-}
-
-
-TUniquePtr<FSparseMatrixD> ConstructMeanValueWeightLaplacian(const FDynamicMesh3& DynamicMesh, FVertexLinearization& VertexMap, TArray<int32>* BoundaryVerts)
-{
-
-	if (BoundaryVerts)
-	{
-		// empty
-		BoundaryVerts->Empty();
-	}
-
-	typedef FSparseMatrixD::Scalar  ScalarT;
-	typedef Eigen::Triplet<ScalarT> MatrixTripletT;
-
-
-	// Create a mapping between the ordering of the vertex indices in the Dynamic Mesh
-	// and the actual storage. 
-	VertexMap.Reset(DynamicMesh);
-
-	const TArray<int32>& ToMeshV = VertexMap.ToId();
-	const TArray<int32>& ToIndex = VertexMap.ToIndex();
-	const int32 NumVerts = VertexMap.NumVerts();
-
-
-	// Map the triangles.
-
-	FTriangleLinearization TriangleMap(DynamicMesh);
-
-	const TArray<int32>& ToMeshTri = TriangleMap.ToId();
-	const TArray<int32>& ToTriIdx = TriangleMap.ToIndex();
-	const int32 NumTris = TriangleMap.NumTris();
-
-
-	// Create an array that holds all the geometric information we need for each triangle.
-
-	TArray<MeanValueTriangleData>  TriangleDataArray = ConstructTriangleDataArray<MeanValueTriangleData>(DynamicMesh, TriangleMap);
-
-
-	// Eigen constructs a sparse matrix from a linearized array of matrix entries.
-
-	std::vector<MatrixTripletT> MatrixTripelList;
-	{
-		int32 NumMatrixEntries = ComputeNumMatrixElements(DynamicMesh, ToMeshV);
-		MatrixTripelList.reserve(NumMatrixEntries);
-	}
-
-
-	// Construct Laplacian Matrix: loop over verts constructing the corresponding matrix row.
-	//                             skipping the boundary verts for later use.
-
-	for (int32 i = 0; i < NumVerts; ++i)
-	{
-		const int32 IVertId = ToMeshV[i]; // I - the row
-
-
-		if (DynamicMesh.IsBoundaryVertex(IVertId))
-		{
-			if (BoundaryVerts)
-			{
-				BoundaryVerts->Add(IVertId);
-			}
-#if LAPLACIAN_SKIP_BOUNDARY == 1  
-			MatrixTripelList.push_back(MatrixTripletT(i, i, 0.));
-			continue;
-#endif 
-		}
-
-
-
-		double WeightII = 0.; // accumulate to equal and opposite the sum of the neighbor weights
-
-		// for each connecting edge
-
-		for (int32 EdgeId : DynamicMesh.VtxEdgesItr(IVertId))
-		{
-			// [v0, v1, t0, t1]:  NB: both t0 & t1 exist since IVert isn't a boundary vert.
-			FIndex4i Edge = DynamicMesh.GetEdge(EdgeId);
-
-			// the other vert in the edge - identifies the matrix column
-			const int32 JVertId = (Edge[0] == IVertId) ? Edge[1] : Edge[0];  // J - the column
-
-			// Get the cotangents for this edge.
-
-			const int32 Tri0Idx = ToTriIdx[Edge[2]];
-			const auto& Tri0Data = TriangleDataArray[Tri0Idx];
-			double TanHalfAngleSum = Tri0Data.GetTanHalfAngle(IVertId);
-			double EdgeLength = FMath::Max(1.e-5, Tri0Data.GetEdgeLenght(EdgeId)); // Clamp the length
-
-			// The second triangle will be invalid if this is an edge!
-
-			TanHalfAngleSum += (Edge[3] != FDynamicMesh3::InvalidID) ? TriangleDataArray[ToTriIdx[Edge[3]]].GetTanHalfAngle(IVertId) : 0.;
-
-			double WeightIJ = TanHalfAngleSum / EdgeLength;
-			WeightII += WeightIJ;
-
-			const int32 j = ToIndex[JVertId];
-
-			MatrixTripelList.push_back(MatrixTripletT(i, j, WeightIJ));
-		}
-
-		MatrixTripelList.push_back(MatrixTripletT(i, i, -WeightII));
-
-
-	}
-
-	TUniquePtr<FSparseMatrixD> LaplacianMatrix(new FSparseMatrixD(NumVerts, NumVerts));
-	LaplacianMatrix->setFromTriplets(MatrixTripelList.begin(), MatrixTripelList.end());
-	LaplacianMatrix->makeCompressed();
-
-	return LaplacianMatrix;
-
-}
+	
 
 
 TUniquePtr<FSparseMatrixD> ConstructLaplacian(const ELaplacianWeightScheme Scheme, const FDynamicMesh3& DynamicMesh, FVertexLinearization& VertexMap, TArray<int32>* BoundaryVerts)
diff --git a/Engine/Plugins/Experimental/GeometryProcessing/Source/MeshSolverUtilities/Private/LaplacianOperators.h b/Engine/Plugins/Experimental/GeometryProcessing/Source/MeshSolverUtilities/Private/LaplacianOperators.h
--- a/Engine/Plugins/Experimental/GeometryProcessing/Source/MeshSolverUtilities/Private/LaplacianOperators.h
+++ b/Engine/Plugins/Experimental/GeometryProcessing/Source/MeshSolverUtilities/Private/LaplacianOperators.h
@@ -120,6 +120,9 @@ void ConstructValenceWeightedLaplacian(const FDynamicMesh3& DynamicMesh,
 * Construct a sparse matrix representation using a cotangent-weighted  Laplacian.
 * NB: there is no reason to expect this to be a symmetric matrix.
 *
+*
+* This computes the cotanget laplacian :  ie.  L =   1/(2A_i) ( Cot alpha_ij + Cot beta_ij ) 
+*
 * The mesh itself is assumed to have N interior vertices, and M boundary vertices.
 *
 * @param DynamicMesh        The triangle mesh
@@ -137,8 +140,29 @@ void  ConstructCotangentLaplacian(const FDynamicMesh3& DynamicMesh,
 	                              FSparseMatrixD& LaplacianInterior, 
 	                              FSparseMatrixD& LaplacianBoundary, 
 	                              const bool bClampWeights = false);
-
-
+/**
+* Construct a sparse matrix representation using a pre-multiplied cotangent-weighted Laplacian.
+* NB: there is no reason to expect this to be a symmetric matrix.
+*
+* This computes the laplacian scaled by the average area A_ave:  ie.  LScaled =   A_ave/(2A_i) ( Cot alpha_ij + Cot beta_ij ) 
+*
+* The mesh itself is assumed to have N interior vertices, and M boundary vertices.
+*
+* @param DynamicMesh        The triangle mesh
+* @param VertexMap          Additional arrays used to map between vertexID and offset in a linear array (i.e. the row).
+*                           The vertices are ordered so that last M ( = VertexMap.NumBoundaryVerts() )  correspond to those on the boundary.
+* @param LaplacianInterior  On return, scaled laplacian operator that acts on the interior vertices: sparse  N x N matrix
+* @param LaplacianBoundary  On return, scaled portion of the operator that acts on the boundary vertices: sparse  N x M matrix
+* @param bClampAreas        Indicates if (A_ave / A_i) should be clamped to (0.5, 5) range. 
+*                           in practice this is desirable when creating the biharmonic operator, but not the mean curvature flow operator
+*
+*   LaplacianInterior * Vector_InteriorVerts + LaplacianBoundary * Vector_BoundaryVerts = Full Laplacian applied to interior vertices.
+*/
+double ConstructScaledCotangentLaplacian(const FDynamicMesh3& DynamicMesh, 
+	                                     FVertexLinearization& VertexMap, 
+	                                     FSparseMatrixD& LaplacianInterior, 
+	                                     FSparseMatrixD& LaplacianBoundary, 
+	                                     const bool bClampAreas = false);
 /**
 * Construct a sparse matrix representation using a cotangent-weighted  Laplacian.
 * but returns the result in two symmetric parts.
