GeometryProcessing :  Mesh smoothing and diffusion operatiors - update to better handle mesh boundaries, this is needed if we want to run this on a local mesh that is a copy of a region of a larger mesh.

diff --git a/Engine/Plugins/Experimental/GeometryProcessing/Source/MeshSolverUtilities/Private/LaplacianMeshSmoother.cpp b/Engine/Plugins/Experimental/GeometryProcessing/Source/MeshSolverUtilities/Private/LaplacianMeshSmoother.cpp
--- a/Engine/Plugins/Experimental/GeometryProcessing/Source/MeshSolverUtilities/Private/LaplacianMeshSmoother.cpp
+++ b/Engine/Plugins/Experimental/GeometryProcessing/Source/MeshSolverUtilities/Private/LaplacianMeshSmoother.cpp
@@ -58,7 +58,6 @@ double ComputeDistSqrd(const FSOAPositions& VecA, const FSOAPositions& VecB)
 }
 
 
-
 FConstrainedMeshOperator::FConstrainedMeshOperator(const FDynamicMesh3& DynamicMesh, const ELaplacianWeightScheme Scheme, const EMatrixSolverType MatrixSolverType)
 	: VertexCount(DynamicMesh.VertexCount())
 {
@@ -74,15 +73,17 @@ FConstrainedMeshOperator::FConstrainedMeshOperator(const FDynamicMesh3& DynamicM
 	InternalVertexCount = VertexCount - BoundaryVertexCount;
 
 	// Copy the original boundary vertex locations
-	BoundaryPositions.Empty(BoundaryVertexCount);
-	BoundaryPositions.AddUninitialized(BoundaryVertexCount);
 	{
-		const auto& ToId = VtxLinearization.ToId();
+		const auto ToVertId = VtxLinearization.ToId();
+		BoundaryPositions.SetZero(BoundaryVertexCount);
 		for (int32 i = 0; i < BoundaryVertexCount; ++i)
 		{
-			int32 VtxId = ToId[i + InternalVertexCount];
-			FVector3d Pos = DynamicMesh.GetVertex(VtxId);
-			BoundaryPositions[i] = Pos;
+			const int32 VtxId = ToVertId[i + InternalVertexCount];
+			const FVector3d Pos = DynamicMesh.GetVertex(VtxId);
+
+			BoundaryPositions.XVector[i] = Pos.X;
+			BoundaryPositions.YVector[i] = Pos.Y;
+			BoundaryPositions.ZVector[i] = Pos.Z;
 		}
 	}
 
@@ -96,14 +97,17 @@ FConstrainedMeshOperator::FConstrainedMeshOperator(const FDynamicMesh3& DynamicM
 	if (bIsLaplacianSymmetric)
 	{
 		// Laplacian is symmetric, i.e. equal to its transpose
-		LTLMatrix = LaplacianInternal * LaplacianInternal;
-	
+		LTLMatrix        = LaplacianInternal * LaplacianInternal;
+		BoundaryOperator = -1. * LaplacianInternal * LaplacianBoundary;
+		
 		ConstrainedSolver.Reset(new FConstrainedSolver(LTLPtr, MatrixSolverType));
 	}
 	else
 	{
 		// the laplacian 
-		LTLMatrix = LaplacianInternal.transpose() * LaplacianInternal;
+		LTLMatrix        = LaplacianInternal.transpose() * LaplacianInternal;
+		BoundaryOperator = -1. * LaplacianInternal.transpose() * LaplacianBoundary;
+		
 		ConstrainedSolver.Reset(new FConstrainedSolver(LTLPtr, MatrixSolverType));
 	}
 
@@ -218,7 +222,7 @@ void FConstrainedMeshOperator::UpdateSolverConstraints()
 	}	
 }
 
-bool FConstrainedMeshOperator::Linearize(const FSOAPositions& PositionalVector, TArray<FVector3d>& LinearArray) const
+bool FConstrainedMeshOperator::CopyInternalPositions(const FSOAPositions& PositionalVector, TArray<FVector3d>& LinearArray) const
 {
 	// Number of positions
 
@@ -233,28 +237,18 @@ bool FConstrainedMeshOperator::Linearize(const FSOAPositions& PositionalVector,
 	checkSlow(Num == InternalVertexCount);
 
 	// 
-	const auto& IndexToVtxId = VtxLinearization.ToId();
-	const int32 MaxVtxId = IndexToVtxId.Num(); // NB: this is really max_used + 1 in the mesh.  See  FDynamicMesh3::MaxVertexID()
-
-
+	const auto& ToVtxId = VtxLinearization.ToId();
+	const int32 MaxVtxId = ToVtxId.Num(); // NB: this is really max_used + 1 in the mesh.  See  FDynamicMesh3::MaxVertexID()
 
-	LinearArray.Empty(MaxVtxId);
-	LinearArray.AddUninitialized(MaxVtxId);
-	
-	// Copy over the boundary positions.
+	if (LinearArray.Num() != MaxVtxId)
 	{
-		int32 BoundaryVertexCount = VertexCount - InternalVertexCount;
-		for (int32 i = 0; i < BoundaryVertexCount; ++i)
-		{
-			const int32 VtxId = IndexToVtxId[i + InternalVertexCount];
-			LinearArray[VtxId] = BoundaryPositions[i];
-		}
+		return false;
 	}
 
 	// Update the internal positions.
 	for (int32 i = 0; i < InternalVertexCount; ++i)
 	{
-		const int32 VtxId = IndexToVtxId[i];
+		const int32 VtxId = ToVtxId[i];
 
 		LinearArray[VtxId] = FVector3d(PositionalVector.XVector.coeff(i), PositionalVector.YVector.coeff(i), PositionalVector.ZVector.coeff(i));
 	}
@@ -262,6 +256,28 @@ bool FConstrainedMeshOperator::Linearize(const FSOAPositions& PositionalVector,
 	return true;
 }
 
+bool FConstrainedMeshOperator::CopyBoundaryPositions(TArray<FVector3d>& LinearArray) const
+{
+	const auto& ToVtxId = VtxLinearization.ToId();
+	const int32 MaxVtxId = ToVtxId.Num();
+
+	if (LinearArray.Num() != MaxVtxId)
+	{
+		return false;
+	}
+
+	int32 BoundaryVertexCount = VertexCount - InternalVertexCount;
+
+	for (int32 i = 0; i < BoundaryVertexCount; ++i)
+	{
+		const int32 VtxId = ToVtxId[i + InternalVertexCount];
+		LinearArray[VtxId] = FVector3d(BoundaryPositions.XVector.coeff(i), BoundaryPositions.YVector.coeff(i), BoundaryPositions.ZVector.coeff(i)); //BoundaryPositions[i];
+	}
+
+	return true;
+}
+
+
 void FConstrainedMeshOperator::ExtractInteriorVertexPositions(const FDynamicMesh3& DynamicMesh, FSOAPositions& VertexPositions) const
 {
 	VertexPositions.SetZero(InternalVertexCount);
@@ -276,8 +292,6 @@ void FConstrainedMeshOperator::ExtractInteriorVertexPositions(const FDynamicMesh
 		VertexPositions.YVector.coeffRef(i) = Pos.Y;
 		VertexPositions.ZVector.coeffRef(i) = Pos.Z;
 	}
-
-	
 }
 
 void FConstrainedMeshOperator::UpdateWithPostFixConstraints(FSOAPositions& PositionVector) const
@@ -350,10 +364,20 @@ bool FConstrainedMeshDeformer::Deform(TArray<FVector3d>& PositionBuffer)
 
 	UpdateWithPostFixConstraints(SolutionVector);
 
-	// Copy the result into the array of structs form.  
+	// Allocate Position Buffer for random access writes
+	int32 MaxVtxId = VtxLinearization.ToId().Num();
+	PositionBuffer.Empty(MaxVtxId);
+	PositionBuffer.AddUninitialized(MaxVtxId);
+
+	// Export the computed internal positions:
+	// Copy the results into the array of structs form.  
 	// NB: this re-indexes so the results can be looked up using VtxId
 
-	Linearize(SolutionVector, PositionBuffer);
+	CopyInternalPositions(SolutionVector, PositionBuffer);
+
+	// Copy the boundary
+	// NB: this re-indexes so the results can be looked up using VtxId
+	CopyBoundaryPositions(PositionBuffer);
 
 	// the matrix solve state
 	return bSuccess;
@@ -366,20 +390,46 @@ bool FBiHarmonicMeshSmoother::ComputeSmoothedMeshPositions(TArray<FVector3d>& Up
 
 	UpdateSolverConstraints();
 
+	// Compute the source vector
+	FSOAPositions SourceVector(InternalVertexCount);
+
+	if (InternalVertexCount != VertexCount) // have boundary points
+	{
+		for (int32 dir = 0; dir < 3; ++dir)
+		{
+			SourceVector.Array(dir) = BoundaryOperator * BoundaryPositions.Array(dir);
+		}
+	}
+	else
+	{
+		SourceVector.SetZero(InternalVertexCount);
+	}
+
 	// Solves the constrained system and updates the mesh 
 
 	FSOAPositions SolutionVector(InternalVertexCount);
 
-	bool bSuccess = ConstrainedSolver->Solve(SolutionVector);
+	bool bSuccess = ConstrainedSolver->Solve(SourceVector, SolutionVector);
 
 
 	// Move any vertices to match bPostFix constraints
 
 	UpdateWithPostFixConstraints(SolutionVector);
 
-	// Move vertices to solution positions
+	// Allocate Position Buffer for random access writes
+	int32 MaxVtxId = VtxLinearization.ToId().Num();
+	UpdatedPositions.Empty(MaxVtxId);
+	UpdatedPositions.AddUninitialized(MaxVtxId);
+
+	// Export the computed internal positions:
+	// Copy the results into the array of structs form.  
+	// NB: this re-indexes so the results can be looked up using VtxId
+
+	CopyInternalPositions(SolutionVector, UpdatedPositions);
 
-	Linearize(SolutionVector, UpdatedPositions);
+	// Copy the boundary
+	// NB: this re-indexes so the results can be looked up using VtxId
+	CopyBoundaryPositions(UpdatedPositions);
 
 	return bSuccess;
 }
@@ -389,22 +439,48 @@ bool FCGBiHarmonicMeshSmoother::ComputeSmoothedMeshPositions(TArray<FVector3d>&
 
 	UpdateSolverConstraints();
 
+	// Compute the source vector
+	FSOAPositions SourceVector(InternalVertexCount);
+
+	if (InternalVertexCount != VertexCount) // have boundary points
+	{
+		for (int32 dir = 0; dir < 3; ++dir)
+		{
+			SourceVector.Array(dir) = BoundaryOperator * BoundaryPositions.Array(dir);
+		}
+	}
+	else
+	{
+		SourceVector.SetZero(InternalVertexCount);
+	}
+
 	// Solves the constrained system and updates the mesh 
 
 	// Solves the constrained system
 
 	FSOAPositions SolutionVector(InternalVertexCount);
 
-	bool bSuccess = ConstrainedSolver->Solve(SolutionVector);
+	bool bSuccess = ConstrainedSolver->Solve(SourceVector, SolutionVector);
 
 
 	// Move any vertices to match bPostFix constraints
 
 	UpdateWithPostFixConstraints(SolutionVector);
 
-	// Move vertices to solution positions
+	// Allocate Position Buffer for random access writes
+	int32 MaxVtxId = VtxLinearization.ToId().Num();
+	UpdatedPositions.Empty(MaxVtxId);
+	UpdatedPositions.AddUninitialized(MaxVtxId);
 
-	Linearize(SolutionVector, UpdatedPositions);
+	// Export the computed internal postions:
+	// Copy the results into the array of structs form.  
+	// NB: this re-indexes so the results can be looked up using VtxId
+
+	CopyInternalPositions(SolutionVector, UpdatedPositions);
+
+	// Copy the boundary
+	// NB: this re-indexes so the results can be looked up using VtxId
+	CopyBoundaryPositions(UpdatedPositions);
 
 	return bSuccess;
 }
@@ -419,24 +495,29 @@ FDiffusionIntegrator::FDiffusionIntegrator(const FDynamicMesh3& DynamicMesh, con
 
 	VertexCount = DynamicMesh.VertexCount();
 	
-	// Allocate the buffers.
-	Tmp[0].SetZero(VertexCount);
-	Tmp[1].SetZero(VertexCount);
-
-
 }
 
 void FDiffusionIntegrator::Initialize(const FDynamicMesh3& DynamicMesh, const ELaplacianWeightScheme Scheme)
 {
 	// Construct the laplacian, and extract the mapping for vertices (VtxLinearization)
-	DiffusionOperator = ConstructDiffusionOperator(Scheme, DynamicMesh, bIsSymmetric, VtxLinearization, &EdgeVerts);
+	//DiffusionOperator = ConstructDiffusionOperator(Scheme, DynamicMesh, bIsSymmetric, VtxLinearization, &EdgeVerts);
+
+	
+	ConstructOperators(Scheme, DynamicMesh, bIsSymmetric, VtxLinearization, DiffusionOperator, BoundaryOperator);
+
+	const int32 BoundaryVertexCount = VtxLinearization.NumBoundaryVerts();
+	InternalVertexCount = VertexCount - BoundaryVertexCount;
+
+	// Allocate the double buffers.
+	Tmp[0].SetZero(InternalVertexCount);
+	Tmp[1].SetZero(InternalVertexCount);
 
 	const auto& ToVertId = VtxLinearization.ToId();
 
-	// Extract current positions.
-	for (int32 i = 0; i < VertexCount; ++i)
+	// Extract current internal positions.
+	for (int32 i = 0; i < InternalVertexCount; ++i)
 	{
-		int32 VtxId = ToVertId[i];
+		const int32 VtxId = ToVertId[i];
 
 		const FVector3d Pos = DynamicMesh.GetVertex(VtxId);
 		Tmp[0].XVector[i] = Pos.X;
@@ -444,7 +525,21 @@ void FDiffusionIntegrator::Initialize(const FDynamicMesh3& DynamicMesh, const EL
 		Tmp[0].ZVector[i] = Pos.Z;
 	}
 
-	const FSparseMatrixD&  M = *DiffusionOperator;
+	// backup the locations of the boundary verts.
+	{
+		BoundaryPositions.SetZero(BoundaryVertexCount);
+		for (int32 i = 0; i < BoundaryVertexCount; ++i)
+		{
+			const int32 VtxId = ToVertId[i + InternalVertexCount];
+			const FVector3d Pos = DynamicMesh.GetVertex(VtxId);
+
+			BoundaryPositions.XVector[i] = Pos.X;
+			BoundaryPositions.YVector[i] = Pos.Y;
+			BoundaryPositions.ZVector[i] = Pos.Z;
+		}
+	}
+
+	const FSparseMatrixD&  M = DiffusionOperator;
 
 	// Find the min diagonal entry (all should be negative).
 	int32 Rank = M.rows();
@@ -454,7 +549,8 @@ void FDiffusionIntegrator::Initialize(const FDynamicMesh3& DynamicMesh, const EL
 		auto Diag = M.coeff(i, i);
 		MinDiagonalValue = FMath::Min(Diag, MinDiagonalValue);
 	}
-
+	// The matrix should have a row for each internal vertex
+	checkSlow(Rank == InternalVertexCount);
 
 #if 0
 	// testing - how to print the matrix to debug output 
@@ -474,7 +570,6 @@ void FDiffusionIntegrator::Integrate_ForwardEuler(int32 NumSteps, double Alpha,
 {
 	Alpha = FMath::Clamp(Alpha, 0., 1.);
 
-	const FSparseMatrixD& M = *DiffusionOperator;
 	FSparseMatrixD::Scalar TimeStep = -Alpha / MinDiagonalValue;
 	Id = 0;
 	for (int32 s = 0; s < NumSteps; ++s)
@@ -482,9 +577,9 @@ void FDiffusionIntegrator::Integrate_ForwardEuler(int32 NumSteps, double Alpha,
 
 		int32 SrcBuffer = Id;
 		Id = 1 - Id;
-		Tmp[Id].XVector = Tmp[SrcBuffer].XVector + TimeStep * M * Tmp[SrcBuffer].XVector;
-		Tmp[Id].YVector = Tmp[SrcBuffer].YVector + TimeStep * M * Tmp[SrcBuffer].YVector;
-		Tmp[Id].ZVector = Tmp[SrcBuffer].ZVector + TimeStep * M * Tmp[SrcBuffer].ZVector;
+		Tmp[Id].XVector = Tmp[SrcBuffer].XVector + TimeStep * ( DiffusionOperator * Tmp[SrcBuffer].XVector + BoundaryOperator * BoundaryPositions.XVector );
+		Tmp[Id].YVector = Tmp[SrcBuffer].YVector + TimeStep * ( DiffusionOperator * Tmp[SrcBuffer].YVector + BoundaryOperator * BoundaryPositions.YVector );
+		Tmp[Id].ZVector = Tmp[SrcBuffer].ZVector + TimeStep * ( DiffusionOperator * Tmp[SrcBuffer].ZVector + BoundaryOperator * BoundaryPositions.ZVector ); 
 
 	}
 
@@ -498,35 +593,38 @@ void FDiffusionIntegrator::Integrate_BackwardEuler(const EMatrixSolverType Matri
 	//typedef typename TMatrixSolverTrait<EMatrixSolverType::LU>::MatrixSolverType   MatrixSolverType;
 
 	// We solve 
-	// p^{n+1} - dt * L[p^{n+1}] = p^{n}
+	// p^{n+1} - dt * L[p^{n+1}] = p^{n} + dt * B[boundaryPts]
 	// 
 	// i.e.
-	// [I - dt * L ] p^{n+1} = p^{n}
+	// [I - dt * L ] p^{n+1} = p^{n} + dt * B[boundaryPts]
 	//
 	// NB: in the case of the cotangent laplacian this would be better if we broke the L int
 	// L = (A^{-1}) H  where A is the "area matrix" (think "mass matrix"), then this would
 	// become
-	// [A - dt * H] p^{n+1} = Ap^{n}  
+	// [A - dt * H] p^{n+1} = Ap^{n}  dt * A *B[boundaryPts]
 	//  
 	// A - dt * H would be symmetric
 	//
 
 
-	const FSparseMatrixD& L = *DiffusionOperator;
+	
 	// Identity matrix
-	FSparseMatrixD Ident(L.rows(), L.cols());
+	FSparseMatrixD Ident(DiffusionOperator.rows(), DiffusionOperator.cols());
 	Ident.setIdentity();
 
 	FSparseMatrixD::Scalar TimeStep = Alpha * FMath::Min(Intensity, 1.e6);
 	
-	FSparseMatrixD M = Ident -TimeStep * L;
+	FSparseMatrixD SparseMatrix = Ident -TimeStep * DiffusionOperator;
 
 	
-	M.makeCompressed();
+	SparseMatrix.makeCompressed();
 
 	TUniquePtr<IMatrixSolverBase> MatrixSolver = ContructMatrixSolver(MatrixSolverType);
 
-	MatrixSolver->SetUp(M, bIsSymmetric);
+	MatrixSolver->SetUp(SparseMatrix, bIsSymmetric);
+
+	// We are going to solve the system 
+	FSOAPositions Source(InternalVertexCount);
 
 	if (MatrixSolver->bIsIterative())
 	{
@@ -540,8 +638,12 @@ void FDiffusionIntegrator::Integrate_BackwardEuler(const EMatrixSolverType Matri
 			int32 SrcBuffer = Id;
 			Id = 1 - Id;
 
+			for (int32 i = 0; i < 3; ++i)
+			{
+				Source.Array(i) = Tmp[SrcBuffer].Array(i) + TimeStep *  BoundaryOperator * BoundaryPositions.Array(i);
+			}
 			// Old solution is the guess.
-			IterativeSolver->SolveWithGuess(Tmp[SrcBuffer], Tmp[SrcBuffer], Tmp[Id]);
+			IterativeSolver->SolveWithGuess(Tmp[SrcBuffer], Source, Tmp[Id]);
 
 		}
 	}
@@ -555,19 +657,35 @@ void FDiffusionIntegrator::Integrate_BackwardEuler(const EMatrixSolverType Matri
 			int32 SrcBuffer = Id;
 			Id = 1 - Id;
 
-			MatrixSolver->Solve(Tmp[SrcBuffer], Tmp[Id]);
+			for (int32 i = 0; i < 3; ++i)
+			{
+				Source.Array(i) = Tmp[SrcBuffer].Array(i) + TimeStep * BoundaryOperator * BoundaryPositions.Array(i);
+			}
+
+			MatrixSolver->Solve(Source, Tmp[Id]);
 		}
 	}
 	
 
 }
 
-void FDiffusionIntegrator::GetPositions(TArray<FVector3d>& PositionArray) const
+void FDiffusionIntegrator::GetPositions(TArray<FVector3d>& PositionBuffer) const
 {
-	Linearize(Tmp[Id], PositionArray);
+	// Allocate Position Buffer for random access writes
+	int32 MaxVtxId = VtxLinearization.ToId().Num();
+	PositionBuffer.Empty(MaxVtxId);
+	PositionBuffer.AddUninitialized(MaxVtxId);
+
+	CopyInternalPositions(Tmp[Id], PositionBuffer);
+
+
+	// Copy the boundary
+	// NB: this re-indexes so the results can be looked up using VtxId
+	CopyBoundaryPositions(PositionBuffer);
+
 }
 
-bool FDiffusionIntegrator::Linearize(const FSOAPositions& PositionalVector, TArray<FVector3d>& LinearArray) const
+bool FDiffusionIntegrator::CopyInternalPositions(const FSOAPositions& PositionalVector, TArray<FVector3d>& LinearArray) const
 {
 	// Number of positions
 
@@ -580,68 +698,98 @@ bool FDiffusionIntegrator::Linearize(const FSOAPositions& PositionalVector, TArr
 	}
 
 	// 
-	const auto& IndexToVtxId = VtxLinearization.ToId();
-	const int32 MaxVtxId = IndexToVtxId.Num(); // NB: this is really max_used + 1 in the mesh.  See  FDynamicMesh3::MaxVertexID()
+	const auto& ToVtxId  = VtxLinearization.ToId();
+	const int32 MaxVtxId = ToVtxId.Num(); // NB: this is really max_used + 1 in the mesh.  See  FDynamicMesh3::MaxVertexID()
+	const int32 BoundaryVertexCount = VtxLinearization.NumBoundaryVerts();
+
+	if (MaxVtxId != LinearArray.Num())
+	{
+		return false;
+	}
 
+	// Copy the updated internal vertex locations over
+	for (int32 i = 0; i < InternalVertexCount; ++i)
+	{
+		const int32 VtxId = ToVtxId[i];
 
+		LinearArray[VtxId] = FVector3d(PositionalVector.XVector.coeff(i), PositionalVector.YVector.coeff(i), PositionalVector.ZVector.coeff(i));
+	}
 
-	LinearArray.Empty(MaxVtxId);
-	LinearArray.AddUninitialized(MaxVtxId);
+	return true;
+}
+                       
+bool FDiffusionIntegrator::CopyBoundaryPositions(TArray<FVector3d>& LinearArray) const
+{
+	const auto& ToVtxId = VtxLinearization.ToId();
+	const int32 MaxVtxId = ToVtxId.Num();
 
-	for (int32 i = 0; i < Num; ++i)
+	if (LinearArray.Num() != MaxVtxId)
 	{
-		const int32 VtxId = IndexToVtxId[i];
+		return false;
+	}
 
-		LinearArray[VtxId] = FVector3d(PositionalVector.XVector.coeff(i), PositionalVector.YVector.coeff(i), PositionalVector.ZVector.coeff(i));
+	int32 BoundaryVertexCount = VertexCount - InternalVertexCount;
+
+	for (int32 i = 0; i < BoundaryVertexCount; ++i)
+	{
+		const int32 VtxId = ToVtxId[i + InternalVertexCount];
+		LinearArray[VtxId] = FVector3d(BoundaryPositions.XVector.coeff(i), BoundaryPositions.YVector.coeff(i), BoundaryPositions.ZVector.coeff(i));
 	}
 
 	return true;
 }
 
 
-TUniquePtr<FSparseMatrixD> FLaplacianDiffusionMeshSmoother::ConstructDiffusionOperator( const ELaplacianWeightScheme Scheme,
-																						const FDynamicMesh3& DynamicMesh,
-																						bool& bIsOperatorSymmetric,
-																						FVertexLinearization& Linearization,
-																						TArray<int32>* EdgeVtxs ) 
+
+void FLaplacianDiffusionMeshSmoother::ConstructOperators( const ELaplacianWeightScheme Scheme,
+	                                                      const FDynamicMesh3& Mesh,
+	                                                      bool& bIsOperatorSymmetric,
+	                                                      FVertexLinearization& Linearization,
+	                                                      FSparseMatrixD& DiffusionOp,
+	                                                      FSparseMatrixD& BoundaryOp) 
 {
 	bIsOperatorSymmetric = bIsSymmetricLaplacian(Scheme);
+	ConstructLaplacian(Scheme, Mesh, VtxLinearization, DiffusionOp, BoundaryOp);
+}
 
 
-	// Construct the laplacian, and extract the mapping for vertices (VtxLinearization)
-	TUniquePtr<FSparseMatrixD> Laplacian = ConstructLaplacian(Scheme, DynamicMesh, VtxLinearization, EdgeVtxs);
-	Laplacian->makeCompressed();
-	return Laplacian;
-};
-
-TUniquePtr<FSparseMatrixD> FBiHarmonicDiffusionMeshSmoother::ConstructDiffusionOperator( const ELaplacianWeightScheme Scheme,
-																						 const FDynamicMesh3& DynamicMesh,
-																						 bool& bIsOperatorSymmetric,
-																						 FVertexLinearization& Linearization,
-																						 TArray<int32>* EdgeVtxs ) 
+void FBiHarmonicDiffusionMeshSmoother::ConstructOperators( const ELaplacianWeightScheme Scheme,
+														   const FDynamicMesh3& Mesh,
+														   bool& bIsOperatorSymmetric,
+														   FVertexLinearization& Linearization,
+														   FSparseMatrixD& DiffusionOp,
+	                                                       FSparseMatrixD& BoundaryOp) 
 {
 	bIsOperatorSymmetric = true;
 	
-	// Construct the laplacian, and extract the mapping for vertices (VtxLinearization)
-	TUniquePtr<FSparseMatrixD> Laplacian = ConstructLaplacian(Scheme, DynamicMesh, VtxLinearization, EdgeVtxs);
-	const FSparseMatrixD& L = *Laplacian;
-	TUniquePtr<FSparseMatrixD> MatrixOperator(new FSparseMatrixD());
-	
+	FSparseMatrixD Laplacian; 
+	FSparseMatrixD BoundaryTerms;
+	ConstructLaplacian(Scheme, Mesh, VtxLinearization, Laplacian, BoundaryTerms);
+
 	bool bIsLaplacianSymmetric = bIsSymmetricLaplacian(Scheme);
 
+	// It is actually unclear the best way to approximate the boundary conditions in this case.  
+	// because we are repeatedly applying the operator ( for example thing about the way ( f(x+d)-f(x-d) )/ d will spread if you apply it twice
+	// as opposed to (f(x+d)-2f(x) + f(x-d) ) / d*d
+
+	// Anyhow here is a guess..
+
 	if (bIsLaplacianSymmetric)
 	{
-		
-		*MatrixOperator = -1. * L * L;
+		DiffusionOp = -1. * Laplacian * Laplacian;
+		BoundaryOp  = -1. * Laplacian * BoundaryTerms;
 	}
 	else
 	{
-		*MatrixOperator = -1. * L.transpose() * L;
+		FSparseMatrixD LTran = Laplacian.transpose();
+		DiffusionOp = -1. * LTran * Laplacian;
+		BoundaryOp  = -1. * LTran * BoundaryTerms;
 	}
+	
+	DiffusionOp.makeCompressed();
+	BoundaryOp.makeCompressed();
 
-	MatrixOperator->makeCompressed();
-	return MatrixOperator;
-};
+}
 
 
 void MeshSmoothingOperators::ComputeSmoothing_BiHarmonic(const ELaplacianWeightScheme WeightScheme, const FDynamicMesh3& OriginalMesh, 
@@ -674,9 +822,10 @@ void MeshSmoothingOperators::ComputeSmoothing_BiHarmonic(const ELaplacianWeightS
 #endif
 
 
+
 	FString DebugLogString = FString::Printf(TEXT("Biharmonic Smoothing of mesh with %d verts "), OriginalMesh.VertexCount()) + LaplacianSchemeName(WeightScheme) + MatrixSolverName(MatrixSolverType);
 
-	FScopedDurationTimeLogger Timmer(DebugLogString);
+	FScopedDurationTimeLogger Timer(DebugLogString);
 
 	FBiHarmonicDiffusionMeshSmoother BiHarmonicDiffusionSmoother(OriginalMesh, WeightScheme);
 
@@ -684,7 +833,6 @@ void MeshSmoothingOperators::ComputeSmoothing_BiHarmonic(const ELaplacianWeightS
 
 	BiHarmonicDiffusionSmoother.GetPositions(PositionArray);
 
-
 }
 
 void MeshSmoothingOperators::ComputeSmoothing_ImplicitBiHarmonicPCG( const ELaplacianWeightScheme WeightScheme, const FDynamicMesh3& OriginalMesh, 
@@ -705,7 +853,7 @@ void MeshSmoothingOperators::ComputeSmoothing_ImplicitBiHarmonicPCG( const ELapl
 
 	FString DebugLogString = FString::Printf(TEXT("PCG Biharmonic Smoothing of mesh with %d verts "), OriginalMesh.VertexCount()) + LaplacianSchemeName(WeightScheme);
 
-	FScopedDurationTimeLogger Timmer(DebugLogString);
+	FScopedDurationTimeLogger Timer(DebugLogString);
 
 	if (MaxIterations < 1) return;
 
@@ -746,7 +894,7 @@ void  MeshSmoothingOperators::ComputeSmoothing_Diffusion( const ELaplacianWeight
 		DebugLogString += MatrixSolverName(MatrixSolverType);
 	}
 
-	FScopedDurationTimeLogger Timmer(DebugLogString);
+	FScopedDurationTimeLogger Timer(DebugLogString);
 
 	if (IterationCount < 1) return;
 
diff --git a/Engine/Plugins/Experimental/GeometryProcessing/Source/MeshSolverUtilities/Private/LaplacianMeshSmoother.h b/Engine/Plugins/Experimental/GeometryProcessing/Source/MeshSolverUtilities/Private/LaplacianMeshSmoother.h
--- a/Engine/Plugins/Experimental/GeometryProcessing/Source/MeshSolverUtilities/Private/LaplacianMeshSmoother.h
+++ b/Engine/Plugins/Experimental/GeometryProcessing/Source/MeshSolverUtilities/Private/LaplacianMeshSmoother.h
@@ -60,7 +60,9 @@ protected:
 	// converts the positionalvector to a TArray<FVector3d> where the offset in the array is implicitly the
 	// VtxId in the mesh, and not ness the matrix row id
 	// NB: the resulting array is treated as sparse and may have un-initialized elements.
-	bool Linearize(const FSOAPositions& PositionalVector, TArray<FVector3d>& LinearArray) const;
+	bool CopyInternalPositions(const FSOAPositions& PositionalVector, TArray<FVector3d>& LinearArray) const;
+
+	bool CopyBoundaryPositions(TArray<FVector3d>& LinearArray) const;
 
 
 protected:
@@ -82,10 +84,15 @@ protected:
 	// Used to map between VtxId and vertex Index in linear vector..
 	FVertexLinearization  VtxLinearization;
 
-	TArray<FVector3d> BoundaryPositions;
+	// Boundary points, split into 3 arrays (x,y,z).
+	FSOAPositions         BoundaryPositions;
+
 
 	// Actual solver that manages the various linear algebra bits.
 	TUniquePtr<FConstrainedSolver>         ConstrainedSolver;
+
+    // Sparse Matrix that holds L^T * B where B has the boundary terms.
+	FSparseMatrixD                         BoundaryOperator;
 };
 
 
@@ -180,29 +187,43 @@ public:
 
 protected:
 
+
 	// The derived class has to implement this.
-	virtual TUniquePtr<FSparseMatrixD> ConstructDiffusionOperator(const ELaplacianWeightScheme Scheme, 
-		                                                          const FDynamicMesh3& Mesh,
-		                                                          bool& bIsOperatorSymmetric,
-		                                                          FVertexLinearization& Linearization,
-		                                                          TArray<int32>* EdgeVtx) = 0 ;
+	// responsible for constructing the Diffusion Operator, the Boundary Operator etc
+	virtual void ConstructOperators(const ELaplacianWeightScheme Scheme,
+		                            const FDynamicMesh3& Mesh,
+	 	                            bool& bIsOperatorSymmetric,
+		                            FVertexLinearization& Linearization,
+									FSparseMatrixD& DiffusionOp,
+									FSparseMatrixD& BoundaryOp ) = 0;
 
 	void Initialize(const FDynamicMesh3& DynamicMesh, const ELaplacianWeightScheme Scheme);
 
 protected:
-	bool Linearize(const FSOAPositions& PositionalVector, TArray<FVector3d>& LinearArray) const;
+
+	// Copy the 
+	bool CopyInternalPositions(const FSOAPositions& PositionalVector, TArray<FVector3d>& LinearArray) const;
+
+	bool CopyBoundaryPositions(TArray<FVector3d>& LinearArray) const;
 
 	// Cache the vertex count.
-	int32 VertexCount;
+	int32                 VertexCount;
 
-	// Used to map vertex ID to entry in linear vector..
+	// Cache the number of internal vertices
+	int32                 InternalVertexCount;
+
+	// Used to map between VtxId and vertex Index in linear vector..
 	FVertexLinearization  VtxLinearization;
 
+
 	// I don't know if we want to keep this after the constructor
 	bool                                   bIsSymmetric;
-	TUniquePtr<FSparseMatrixD>             DiffusionOperator;
+	FSparseMatrixD                         DiffusionOperator;
+	FSparseMatrixD			               BoundaryOperator;
 	TArray<int32>                          EdgeVerts;
 	FSparseMatrixD::Scalar                 MinDiagonalValue;
+	
+	FSOAPositions                       BoundaryPositions;
 
 	FSOAPositions                       Tmp[2];
 	int32                               Id; // double buffer id
@@ -220,13 +241,13 @@ public:
 	}
 	
 protected:
-	TUniquePtr<FSparseMatrixD> ConstructDiffusionOperator( const ELaplacianWeightScheme Scheme,
-														   const FDynamicMesh3& Mesh,
-		                                                   bool& bIsOperatorSymmetric,
-		                                                   FVertexLinearization& Linearization,
-		                                                   TArray<int32>* EdgeVtx) override ;
 	
-
+	void ConstructOperators(const ELaplacianWeightScheme Scheme,
+		                    const FDynamicMesh3& Mesh,
+		                    bool& bIsOperatorSymmetric,
+		                    FVertexLinearization& Linearization,
+		                    FSparseMatrixD& DiffusionOp,
+		                    FSparseMatrixD& BoundaryOp) override;
 };
 
 class  FBiHarmonicDiffusionMeshSmoother : public FDiffusionIntegrator
@@ -240,10 +261,12 @@ public:
 	}
 
 protected:
-	TUniquePtr<FSparseMatrixD> ConstructDiffusionOperator( const ELaplacianWeightScheme Scheme,
-														   const FDynamicMesh3& Mesh,
-														   bool& bIsOperatorSymmetric,
-														   FVertexLinearization& Linearization,
-														   TArray<int32>* EdgeVtx) override;
+
+	void ConstructOperators( const ELaplacianWeightScheme Scheme,
+		                     const FDynamicMesh3& Mesh,
+		                     bool& bIsOperatorSymmetric,
+		                     FVertexLinearization& Linearization,
+		                     FSparseMatrixD& DiffusionOp,
+		                     FSparseMatrixD& BoundaryOp) override;
 };
 
