Move FMeshProjectionTarget method impls to cpp file, and handle garbage return from FDistPoint3Triangle3d
#rnx

diff --git a/Engine/Plugins/Experimental/GeometryProcessing/Source/DynamicMesh/Private/ProjectionTargets.cpp b/Engine/Plugins/Experimental/GeometryProcessing/Source/DynamicMesh/Private/ProjectionTargets.cpp
new file mode 100644
--- /dev/null
+++ b/Engine/Plugins/Experimental/GeometryProcessing/Source/DynamicMesh/Private/ProjectionTargets.cpp
@@ -0,0 +1,55 @@
+// Copyright 1998-2019 Epic Games, Inc. All Rights Reserved.
+
+#include "ProjectionTargets.h"
+#include "Distance/DistPoint3Triangle3.h"
+
+FVector3d FMeshProjectionTarget::Project(const FVector3d& Point, int Identifier)
+{
+	double fDistSqr;
+	int tNearestID = Spatial->FindNearestTriangle(Point, fDistSqr);
+	if (tNearestID <= 0)
+	{
+		return Point;
+	}
+	FTriangle3d Triangle;
+	Mesh->GetTriVertices(tNearestID, Triangle.V[0], Triangle.V[1], Triangle.V[2]);
+
+	FDistPoint3Triangle3d DistanceQuery(Point, Triangle);
+	DistanceQuery.GetSquared();
+	if (VectorUtil::IsFinite(DistanceQuery.ClosestTrianglePoint))
+	{
+		return DistanceQuery.ClosestTrianglePoint;
+	}
+	else
+	{
+		return Point;
+	}
+}
+
+/**
+ * @return Projection of Point onto this target, and set ProjectNormalOut to the triangle normal at the returned point (*not* interpolated vertex normal)
+ */
+FVector3d FMeshProjectionTarget::Project(const FVector3d& Point, FVector3d& ProjectNormalOut, int Identifier)
+{
+	double fDistSqr;
+	int tNearestID = Spatial->FindNearestTriangle(Point, fDistSqr);
+	if (tNearestID <= 0)
+	{
+		return Point;
+	}
+	FTriangle3d Triangle;
+	Mesh->GetTriVertices(tNearestID, Triangle.V[0], Triangle.V[1], Triangle.V[2]);
+
+	ProjectNormalOut = Triangle.Normal();
+
+	FDistPoint3Triangle3d DistanceQuery(Point, Triangle);
+	DistanceQuery.GetSquared();
+	if (VectorUtil::IsFinite(DistanceQuery.ClosestTrianglePoint))
+	{
+		return DistanceQuery.ClosestTrianglePoint;
+	}
+	else
+	{
+		return Point;
+	}
+}
\ No newline at end of file
diff --git a/Engine/Plugins/Experimental/GeometryProcessing/Source/DynamicMesh/Public/ProjectionTargets.h b/Engine/Plugins/Experimental/GeometryProcessing/Source/DynamicMesh/Public/ProjectionTargets.h
--- a/Engine/Plugins/Experimental/GeometryProcessing/Source/DynamicMesh/Public/ProjectionTargets.h
+++ b/Engine/Plugins/Experimental/GeometryProcessing/Source/DynamicMesh/Public/ProjectionTargets.h
@@ -5,14 +5,13 @@
 // Port of geometry3cpp ProjectionTargets
 
 #include "DynamicMeshAABBTree3.h"
-#include "Distance/DistPoint3Triangle3.h"
 
 
 /**
  * FMeshProjectionTarget provides an IProjectionTarget interface to a FDynamicMesh + FDynamicMeshAABBTree3
  * Use to project points to mesh surface.
  */
-class FMeshProjectionTarget : public IOrientedProjectionTarget
+class DYNAMICMESH_API FMeshProjectionTarget : public IOrientedProjectionTarget
 {
 public:
 	/** The mesh to project onto */
@@ -36,42 +35,12 @@ public:
 	/**
 	 * @return Projection of Point onto this target
 	 */
-	virtual FVector3d Project(const FVector3d& Point, int Identifier = -1) override
-	{
-		double fDistSqr;
-		int tNearestID = Spatial->FindNearestTriangle(Point, fDistSqr);
-		if (tNearestID <= 0)
-		{
-			return Point;
-		}
-		FTriangle3d Triangle;
-		Mesh->GetTriVertices(tNearestID, Triangle.V[0], Triangle.V[1], Triangle.V[2]);
-
-		FDistPoint3Triangle3d DistanceQuery(Point, Triangle);
-		DistanceQuery.GetSquared();
-		return DistanceQuery.ClosestTrianglePoint;
-	}
+	virtual FVector3d Project(const FVector3d& Point, int Identifier = -1) override;
 
 	/**
 	 * @return Projection of Point onto this target, and set ProjectNormalOut to the triangle normal at the returned point (*not* interpolated vertex normal)
 	 */
-	virtual FVector3d Project(const FVector3d& Point, FVector3d& ProjectNormalOut, int Identifier = -1) override
-	{
-		double fDistSqr;
-		int tNearestID = Spatial->FindNearestTriangle(Point, fDistSqr);
-		if (tNearestID <= 0)
-		{
-			return Point;
-		}
-		FTriangle3d Triangle;
-		Mesh->GetTriVertices(tNearestID, Triangle.V[0], Triangle.V[1], Triangle.V[2]);
-
-		ProjectNormalOut = Triangle.Normal();
-
-		FDistPoint3Triangle3d DistanceQuery(Point, Triangle);
-		DistanceQuery.GetSquared();
-		return DistanceQuery.ClosestTrianglePoint;
-	}
+	virtual FVector3d Project(const FVector3d& Point, FVector3d& ProjectNormalOut, int Identifier = -1) override;
 
 
 };
\ No newline at end of file
