GeometryProcessing: Need to pass copies to .InsertAt(), in case it resizes internal array, because internal blocks are currently copied instead of moved (this will be fixed)
#rnx

diff --git a/Engine/Plugins/Experimental/GeometryProcessing/Source/DynamicMesh/Private/DynamicMesh3_Edits.cpp b/Engine/Plugins/Experimental/GeometryProcessing/Source/DynamicMesh/Private/DynamicMesh3_Edits.cpp
--- a/Engine/Plugins/Experimental/GeometryProcessing/Source/DynamicMesh/Private/DynamicMesh3_Edits.cpp
+++ b/Engine/Plugins/Experimental/GeometryProcessing/Source/DynamicMesh/Private/DynamicMesh3_Edits.cpp
@@ -902,7 +902,8 @@ EMeshResult FDynamicMesh3::SplitEdge(int eab, FEdgeSplitInfo& SplitInfo, double
 		int t2 = AddTriangleInternal(f, b, c, InvalidID, InvalidID, InvalidID);
 		if (HasTriangleGroups())
 		{
-			TriangleGroups->InsertAt((*TriangleGroups)[t0], t2);
+			int group0 = (*TriangleGroups)[t0];
+			TriangleGroups->InsertAt(group0, t2);
 		}
 
 		// rewrite edge bc, create edge af
@@ -983,8 +984,10 @@ EMeshResult FDynamicMesh3::SplitEdge(int eab, FEdgeSplitInfo& SplitInfo, double
 		int t3 = AddTriangleInternal(f, d, b, InvalidID, InvalidID, InvalidID);
 		if (HasTriangleGroups()) 
 		{
-			TriangleGroups->InsertAt((*TriangleGroups)[t0], t2);
-			TriangleGroups->InsertAt((*TriangleGroups)[t1], t3);
+			int group0 = (*TriangleGroups)[t0];
+			TriangleGroups->InsertAt(group0, t2);
+			int group1 = (*TriangleGroups)[t1];
+			TriangleGroups->InsertAt(group1, t3);
 		}
 
 		// update the edges we found above, to point to triangles
diff --git a/Engine/Plugins/Experimental/GeometryProcessing/Source/GeometricObjects/Public/Spatial/MeshAABBTree3.h b/Engine/Plugins/Experimental/GeometryProcessing/Source/GeometricObjects/Public/Spatial/MeshAABBTree3.h
--- a/Engine/Plugins/Experimental/GeometryProcessing/Source/GeometricObjects/Public/Spatial/MeshAABBTree3.h
+++ b/Engine/Plugins/Experimental/GeometryProcessing/Source/GeometricObjects/Public/Spatial/MeshAABBTree3.h
@@ -531,10 +531,13 @@ protected:
 		// ok now append internal node boxes & index ptrs
 		for (i = 0; i < Nodes.IBoxCur; ++i)
 		{
-			BoxCenters.InsertAt(Nodes.BoxCenters[i], iBoxShift + i);
-			BoxExtents.InsertAt(Nodes.BoxExtents[i], iBoxShift + i);
+			FVector3d NodeBoxCenter = Nodes.BoxCenters[i];		// cannot pass as argument in case a resize happens
+			BoxCenters.InsertAt(NodeBoxCenter, iBoxShift + i);
+			FVector3d NodeBoxExtents = Nodes.BoxExtents[i];
+			BoxExtents.InsertAt(NodeBoxExtents, iBoxShift + i);
 			// internal node indices are shifted
-			BoxToIndex.InsertAt(iIndexShift + Nodes.BoxToIndex[i], iBoxShift + i);
+			int NodeBoxIndex = Nodes.BoxToIndex[i];
+			BoxToIndex.InsertAt(iIndexShift + NodeBoxIndex, iBoxShift + i);
 		}
 
 		// now append index list
