Replace DistPoint3Triangle3 implementation with a direct port of the gte version (was previous an indirect port via geometry3sharp), add test to verify edge case is fixed by doing so

#rnx

diff --git a/Engine/Plugins/Experimental/GeometryProcessing/Source/GeometricObjects/Public/Distance/DistPoint3Triangle3.h b/Engine/Plugins/Experimental/GeometryProcessing/Source/GeometricObjects/Public/Distance/DistPoint3Triangle3.h
--- a/Engine/Plugins/Experimental/GeometryProcessing/Source/GeometricObjects/Public/Distance/DistPoint3Triangle3.h
+++ b/Engine/Plugins/Experimental/GeometryProcessing/Source/GeometricObjects/Public/Distance/DistPoint3Triangle3.h
@@ -1,6 +1,6 @@
 // Copyright 1998-2019 Epic Games, Inc. All Rights Reserved.
 
-// Port of geometry3Sharp DistPoint3Triangle3
+// Port of gte's GteDistPointTriangle to use GeometryProcessing data types
 
 #pragma once
 
@@ -9,7 +9,7 @@
 
 
 /**
- * Compute unsigned distance between 3D point and 3D triangle
+ * Compute unsigned distance between 3D Point and 3D Triangle
  */
 template <typename Real>
 class TDistPoint3Triangle3
@@ -41,231 +41,201 @@ public:
 
 	Real ComputeResult()
 	{
-		FVector3<Real> diff = Triangle.V[0] - Point;
+		FVector3<Real> diff = Point - Triangle.V[0];
 		FVector3<Real> edge0 = Triangle.V[1] - Triangle.V[0];
 		FVector3<Real> edge1 = Triangle.V[2] - Triangle.V[0];
 		Real a00 = edge0.SquaredLength();
 		Real a01 = edge0.Dot(edge1);
 		Real a11 = edge1.SquaredLength();
-		Real b0 = diff.Dot(edge0);
-		Real b1 = diff.Dot(edge1);
-		Real c = diff.SquaredLength();
-		Real det = FMath::Abs(a00*a11 - a01 * a01);
-		Real s = a01 * b1 - a11 * b0;
-		Real t = a01 * b0 - a00 * b1;
-		Real sqrDistance;
-
-		if (s + t <= det)
+		Real b0 = -diff.Dot(edge0);
+		Real b1 = -diff.Dot(edge1);
+
+		Real f00 = b0;
+		Real f10 = b0 + a00;
+		Real f01 = b0 + a01;
+
+		FVector2<Real> p0, p1, p;
+		Real dt1, h0, h1;
+
+		// Compute the endpoints p0 and p1 of the segment.  The segment is
+		// parameterized by L(z) = (1-z)*p0 + z*p1 for z in [0,1] and the
+		// directional derivative of half the quadratic on the segment is
+		// H(z) = Dot(p1-p0,gradient[Q](L(z))/2), where gradient[Q]/2 = (F,G).
+		// By design, F(L(z)) = 0 for cases (2), (4), (5), and (6).  Cases (1) and
+		// (3) can correspond to no-intersection or intersection of F = 0 with the
+		// Triangle.
+		if (f00 >= (Real)0)
 		{
-			if (s < (Real)0)
+			if (f01 >= (Real)0)
+			{
+				// (1) p0 = (0,0), p1 = (0,1), H(z) = G(L(z))
+				GetMinEdge02(a11, b1, p);
+			}
+			else
 			{
-				if (t < (Real)0)  // region 4
+				// (2) p0 = (0,t10), p1 = (t01,1-t01), H(z) = (t11 - t10)*G(L(z))
+				p0[0] = (Real)0;
+				p0[1] = f00 / (f00 - f01);
+				p1[0] = f01 / (f01 - f10);
+				p1[1] = (Real)1 - p1[0];
+				dt1 = p1[1] - p0[1];
+				h0 = dt1 * (a11 * p0[1] + b1);
+				if (h0 >= (Real)0)
 				{
-					if (b0 < (Real)0)
-					{
-						t = (Real)0;
-						if (-b0 >= a00)
-						{
-							s = (Real)1;
-							sqrDistance = a00 + ((Real)2)*b0 + c;
-						}
-						else
-						{
-							s = -b0 / a00;
-							sqrDistance = b0 * s + c;
-						}
-					}
-					else
-					{
-						s = (Real)0;
-						if (b1 >= (Real)0)
-						{
-							t = (Real)0;
-							sqrDistance = c;
-						}
-						else if (-b1 >= a11)
-						{
-							t = (Real)1;
-							sqrDistance = a11 + ((Real)2)*b1 + c;
-						}
-						else
-						{
-							t = -b1 / a11;
-							sqrDistance = b1 * t + c;
-						}
-					}
+					GetMinEdge02(a11, b1, p);
 				}
-				else  // region 3
+				else
 				{
-					s = (Real)0;
-					if (b1 >= (Real)0)
-					{
-						t = (Real)0;
-						sqrDistance = c;
-					}
-					else if (-b1 >= a11)
+					h1 = dt1 * (a01 * p1[0] + a11 * p1[1] + b1);
+					if (h1 <= (Real)0)
 					{
-						t = (Real)1;
-						sqrDistance = a11 + ((Real)2)*b1 + c;
+						GetMinEdge12(a01, a11, b1, f10, f01, p);
 					}
 					else
 					{
-						t = -b1 / a11;
-						sqrDistance = b1 * t + c;
+						GetMinInterior(p0, h0, p1, h1, p);
 					}
 				}
 			}
-			else if (t < (Real)0)  // region 5
-			{
-				t = (Real)0;
-				if (b0 >= (Real)0)
-				{
-					s = (Real)0;
-					sqrDistance = c;
-				}
-				else if (-b0 >= a00)
-				{
-					s = (Real)1;
-					sqrDistance = a00 + ((Real)2)*b0 + c;
-				}
-				else
-				{
-					s = -b0 / a00;
-					sqrDistance = b0 * s + c;
-				}
-			}
-			else  // region 0
-			{
-				// minimum at interior point
-				Real invDet = ((Real)1) / det;
-				s *= invDet;
-				t *= invDet;
-				sqrDistance = s * (a00*s + a01 * t + ((Real)2)*b0) + t * (a01*s + a11 * t + ((Real)2)*b1) + c;
-			}
 		}
-		else
+		else if (f01 <= (Real)0)
 		{
-			Real tmp0, tmp1, numer, denom;
-
-			if (s < (Real)0)  // region 2
+			if (f10 <= (Real)0)
+			{
+				// (3) p0 = (1,0), p1 = (0,1), H(z) = G(L(z)) - F(L(z))
+				GetMinEdge12(a01, a11, b1, f10, f01, p);
+			}
+			else
 			{
-				tmp0 = a01 + b0;
-				tmp1 = a11 + b1;
-				if (tmp1 > tmp0)
+				// (4) p0 = (t00,0), p1 = (t01,1-t01), H(z) = t11*G(L(z))
+				p0[0] = f00 / (f00 - f10);
+				p0[1] = (Real)0;
+				p1[0] = f01 / (f01 - f10);
+				p1[1] = (Real)1 - p1[0];
+				h0 = p1[1] * (a01 * p0[0] + b1);
+				if (h0 >= (Real)0)
 				{
-					numer = tmp1 - tmp0;
-					denom = a00 - ((Real)2)*a01 + a11;
-					if (numer >= denom)
-					{
-						s = (Real)1;
-						t = (Real)0;
-						sqrDistance = a00 + ((Real)2)*b0 + c;
-					}
-					else
-					{
-						s = numer / denom;
-						t = (Real)1 - s;
-						sqrDistance = s * (a00*s + a01 * t + ((Real)2)*b0) + t * (a01*s + a11 * t + ((Real)2)*b1) + c;
-					}
+					p = p0;  // GetMinEdge01
 				}
 				else
 				{
-					s = (Real)0;
-					if (tmp1 <= (Real)0)
+					h1 = p1[1] * (a01 * p1[0] + a11 * p1[1] + b1);
+					if (h1 <= (Real)0)
 					{
-						t = (Real)1;
-						sqrDistance = a11 + ((Real)2)*b1 + c;
-					}
-					else if (b1 >= (Real)0)
-					{
-						t = (Real)0;
-						sqrDistance = c;
+						GetMinEdge12(a01, a11, b1, f10, f01, p);
 					}
 					else
 					{
-						t = -b1 / a11;
-						sqrDistance = b1 * t + c;
+						GetMinInterior(p0, h0, p1, h1, p);
 					}
 				}
 			}
-			else if (t < (Real)0)  // region 6
+		}
+		else if (f10 <= (Real)0)
+		{
+			// (5) p0 = (0,t10), p1 = (t01,1-t01), H(z) = (t11 - t10)*G(L(z))
+			p0[0] = (Real)0;
+			p0[1] = f00 / (f00 - f01);
+			p1[0] = f01 / (f01 - f10);
+			p1[1] = (Real)1 - p1[0];
+			dt1 = p1[1] - p0[1];
+			h0 = dt1 * (a11 * p0[1] + b1);
+			if (h0 >= (Real)0)
+			{
+				GetMinEdge02(a11, b1, p);
+			}
+			else
 			{
-				tmp0 = a01 + b1;
-				tmp1 = a00 + b0;
-				if (tmp1 > tmp0)
+				h1 = dt1 * (a01 * p1[0] + a11 * p1[1] + b1);
+				if (h1 <= (Real)0)
 				{
-					numer = tmp1 - tmp0;
-					denom = a00 - ((Real)2)*a01 + a11;
-					if (numer >= denom)
-					{
-						t = (Real)1;
-						s = (Real)0;
-						sqrDistance = a11 + ((Real)2)*b1 + c;
-					}
-					else
-					{
-						t = numer / denom;
-						s = (Real)1 - t;
-						sqrDistance = s * (a00*s + a01 * t + ((Real)2)*b0) + t * (a01*s + a11 * t + ((Real)2)*b1) + c;
-					}
+					GetMinEdge12(a01, a11, b1, f10, f01, p);
 				}
 				else
 				{
-					t = (Real)0;
-					if (tmp1 <= (Real)0)
-					{
-						s = (Real)1;
-						sqrDistance = a00 + ((Real)2)*b0 + c;
-					}
-					else if (b0 >= (Real)0)
-					{
-						s = (Real)0;
-						sqrDistance = c;
-					}
-					else
-					{
-						s = -b0 / a00;
-						sqrDistance = b0 * s + c;
-					}
+					GetMinInterior(p0, h0, p1, h1, p);
 				}
 			}
-			else  // region 1
+		}
+		else
+		{
+			// (6) p0 = (t00,0), p1 = (0,t11), H(z) = t11*G(L(z))
+			p0[0] = f00 / (f00 - f10);
+			p0[1] = (Real)0;
+			p1[0] = (Real)0;
+			p1[1] = f00 / (f00 - f01);
+			h0 = p1[1] * (a01 * p0[0] + b1);
+			if (h0 >= (Real)0)
+			{
+				p = p0;  // GetMinEdge01
+			}
+			else
 			{
-				numer = a11 + b1 - a01 - b0;
-				if (numer <= (Real)0)
+				h1 = p1[1] * (a11 * p1[1] + b1);
+				if (h1 <= (Real)0)
 				{
-					s = (Real)0;
-					t = (Real)1;
-					sqrDistance = a11 + ((Real)2)*b1 + c;
+					GetMinEdge02(a11, b1, p);
 				}
 				else
 				{
-					denom = a00 - ((Real)2)*a01 + a11;
-					if (numer >= denom)
-					{
-						s = (Real)1;
-						t = (Real)0;
-						sqrDistance = a00 + ((Real)2)*b0 + c;
-					}
-					else
-					{
-						s = numer / denom;
-						t = (Real)1 - s;
-						sqrDistance = s * (a00*s + a01 * t + ((Real)2)*b0) + t * (a01*s + a11 * t + ((Real)2)*b1) + c;
-					}
+					GetMinInterior(p0, h0, p1, h1, p);
 				}
 			}
 		}
 
-		// Account for numerical round-off error.
-		if (sqrDistance < (Real)0)
+		TriangleBaryCoords = FVector3<Real>((Real)1 - p[0] - p[1], p[0], p[1]);
+		ClosestTrianglePoint = Triangle.V[0] + p[0] * edge0 + p[1] * edge1;
+		return Point.DistanceSquared(ClosestTrianglePoint);
+	}
+
+private:
+	void GetMinEdge02(Real const& a11, Real const& b1, FVector2<Real>& p) const
+	{
+		p[0] = (Real)0;
+		if (b1 >= (Real)0)
+		{
+			p[1] = (Real)0;
+		}
+		else if (a11 + b1 <= (Real)0)
+		{
+			p[1] = (Real)1;
+		}
+		else
+		{
+			p[1] = -b1 / a11;
+		}
+	}
+
+	void GetMinEdge12(
+			Real const& a01, Real const& a11, Real const& b1, Real const& f10,
+			Real const& f01, FVector2<Real>& p) const
+	{
+		Real h0 = a01 + b1 - f10;
+		if (h0 >= (Real)0)
 		{
-			sqrDistance = (Real)0;
+			p[1] = (Real)0;
+		}
+		else
+		{
+			Real h1 = a11 + b1 - f01;
+			if (h1 <= (Real)0)
+			{
+				p[1] = (Real)1;
+			}
+			else
+			{
+				p[1] = h0 / (h0 - h1);
+			}
 		}
+		p[0] = (Real)1 - p[1];
+	}
 
-		ClosestTrianglePoint = Triangle.V[0] + s * edge0 + t * edge1;
-		TriangleBaryCoords = FVector3<Real>((Real)1 - s - t, s, t);
-		return sqrDistance;
+	void GetMinInterior(
+			FVector2<Real> const& p0, Real const& h0, FVector2<Real> const& p1,
+			Real const& h1, FVector2<Real>& p) const
+	{
+		Real z = h0 / (h0 - h1);
+		p = ((Real)1 - z) * p0 + z * p1;
 	}
 
 };
