Use TIndirectArray to back geometryprocessing's chunked array (instead of TArray)

diff --git a/Engine/Plugins/Experimental/GeometryProcessing/Source/GeometricObjects/Public/Util/DynamicVector.h b/Engine/Plugins/Experimental/GeometryProcessing/Source/GeometricObjects/Public/Util/DynamicVector.h
--- a/Engine/Plugins/Experimental/GeometryProcessing/Source/GeometricObjects/Public/Util/DynamicVector.h
+++ b/Engine/Plugins/Experimental/GeometryProcessing/Source/GeometricObjects/Public/Util/DynamicVector.h
@@ -6,6 +6,7 @@
 
 #include <CoreMinimal.h>
 #include "Containers/StaticArray.h"
+#include "Containers/IndirectArray.h"
 #include "VectorTypes.h"
 #include "IndexTypes.h"
 
@@ -23,7 +24,7 @@ public:
 	{
 		CurBlock = 0;
 		CurBlockUsed = 0;
-		Blocks.Emplace();
+		Blocks.Add(new BlockType());
 	}
 
 	TDynamicVector(const TDynamicVector& Copy)
@@ -86,7 +87,7 @@ protected:
 	unsigned int CurBlockUsed;
 
 	using BlockType = TStaticArray<Type, BlockSize>;
-	TArray<BlockType> Blocks;
+	TIndirectArray<BlockType> Blocks;
 
 	friend class FIterator;
 
@@ -355,7 +356,7 @@ void TDynamicVector<Type>::Clear()
 	Blocks.Empty();
 	CurBlock = 0;
 	CurBlockUsed = 0;
-	Blocks.Add(BlockType());
+	Blocks.Add(new BlockType());
 }
 
 template <class Type>
@@ -383,21 +384,20 @@ void TDynamicVector<Type>::Resize(size_t Count)
 	int nNumSegs = 1 + (int)Count / BlockSize;
 
 	// figure out how many are currently allocated...
-	size_t nCurCount = Blocks.Num();
+	int32 nCurCount = Blocks.Num();
 
 	// resize to right number of segments
-	if (nNumSegs >= Blocks.Num())
+	if (nNumSegs >= nCurCount)
 	{
 		// allocate new segments
 		for (int i = (int)nCurCount; i < nNumSegs; ++i)
 		{
-			Blocks.Emplace();
+			Blocks.Add(new BlockType());
 		}
 	}
 	else
 	{
-		//Blocks.RemoveRange(nNumSegs, Blocks.Count - nNumSegs);
-		Blocks.SetNum(nNumSegs);
+		Blocks.RemoveAt(nNumSegs, nCurCount - nNumSegs);
 	}
 
 	// mark last segment
@@ -423,7 +423,7 @@ void TDynamicVector<Type>::Add(const Type& Value)
 	{
 		if (CurBlock == Blocks.Num() - 1)
 		{
-			Blocks.Emplace();
+			Blocks.Add(new BlockType());
 		}
 		CurBlock++;
 		CurBlockUsed = 0;
