GeometryProcessing: add 2d discrete scalar field with bilinear interpolation sampling
#rnx

diff --git a/Engine/Plugins/Experimental/GeometryProcessing/Source/GeometricObjects/Public/Spatial/SampledScalarField2.h b/Engine/Plugins/Experimental/GeometryProcessing/Source/GeometricObjects/Public/Spatial/SampledScalarField2.h
new file mode 100644
--- /dev/null
+++ b/Engine/Plugins/Experimental/GeometryProcessing/Source/GeometricObjects/Public/Spatial/SampledScalarField2.h
@@ -0,0 +1,113 @@
+// Copyright 1998-2019 Epic Games, Inc. All Rights Reserved.
+
+// Port of geometry3Sharp DSparseGrid3
+
+#pragma once
+
+#include "CoreMinimal.h"
+#include "MathUtil.h"
+#include "VectorTypes.h"
+
+
+
+/**
+ * TSampledScalarField2 implements a generic 2D grid of values that can be interpolated in various ways.
+ * The grid is treated as a set of sample points in 2D space, IE a grid origin and x/y point-spacing is part of this class.
+ * 
+ * The class is templated on two types:
+ *    RealType: the real type used for spatial calculations, ie 2D grid origin, cell dimensions, and sample positions
+ *    ValueType: the type of value stored in the grid. Could be real or vector-typed, needs to support multiplication by RealType (for interpolation)
+ */
+template<typename RealType, typename ValueType>
+class TSampledScalarField2
+{
+public:
+	FVector2i GridDimensions;
+	TArray<ValueType> GridValues;
+
+	FVector2<RealType> GridOrigin;
+	FVector2<RealType> CellDimensions;
+
+	/**
+	 * Create empty grid, defaults to 2x2 grid of whatever default value of ValueType is
+	 */
+	TSampledScalarField2()
+	{
+		GridDimensions = FVector2i(2, 2);
+		GridValues.SetNum(4);
+		GridOrigin = FVector2<RealType>::Zero();
+		CellDimensions = FVector2<RealType>::One();
+	}
+
+	/**
+	 * Resize the grid to given Width/Height, and initialize w/ given InitValue
+	 */
+	void Resize(int Width, int Height, const ValueType& InitValue)
+	{
+		GridDimensions = FVector2i(Width, Height);
+		GridValues.Init(InitValue, Width*Height);
+	}
+
+	/**
+	 * Set the 2D origin of the grid
+	 */
+	void SetPosition(const FVector2<RealType>& Origin)
+	{
+		GridOrigin = Origin;
+	}
+	
+	/**
+	 * Set the size of the grid cells to uniform CellSize
+	 */
+	void SetCellSize(RealType CellSize)
+	{
+		CellDimensions.X = CellDimensions.Y = CellSize;
+	}
+
+	/**
+	 * Sample scalar field with bilinear interpolation at given Position
+	 * @param Position sample point relative to grid origin/dimensions
+	 * @return interpolated value at this position
+	 */
+	ValueType BilinearSampleClamped(const FVector2<RealType>& Position) const
+	{
+		// transform Position into grid coordinates
+		FVector2<RealType> GridPoint(
+			((Position.X - GridOrigin.X) / CellDimensions.X),
+			((Position.Y - GridOrigin.Y) / CellDimensions.Y));
+
+		// compute integer grid coordinates
+		int x0 = (int)GridPoint.X;
+		int x1 = x0 + 1;
+		int y0 = (int)GridPoint.Y;
+		int y1 = y0 + 1;
+
+		// clamp to valid range
+		x0 = FMath::Clamp(x0, 0, GridDimensions.X - 1);
+		x1 = FMath::Clamp(x1, 0, GridDimensions.X - 1);
+		y0 = FMath::Clamp(y0, 0, GridDimensions.Y - 1);
+		y1 = FMath::Clamp(y1, 0, GridDimensions.Y - 1);
+
+		// convert real-valued grid coords to [0,1] range
+		RealType fAx = FMath::Clamp(GridPoint.X - (RealType)x0, (RealType)0, (RealType)1);
+		RealType fAy = FMath::Clamp(GridPoint.Y - (RealType)y0, (RealType)0, (RealType)1);
+		RealType OneMinusfAx = (RealType)1 - fAx;
+		RealType OneMinusfAy = (RealType)1 - fAy;
+
+		// fV## is grid cell corner index
+		const ValueType& fV00 = GridValues[y0*GridDimensions.X + x0];
+		const ValueType& fV01 = GridValues[y0*GridDimensions.X + x1];
+		const ValueType& fV10 = GridValues[y1*GridDimensions.X + x0];
+		const ValueType& fV11 = GridValues[y1*GridDimensions.X + x1];
+
+		return
+			(OneMinusfAx * OneMinusfAy) * fV00 +
+			(OneMinusfAx * fAy)         * fV01 +
+			(fAx         * OneMinusfAy) * fV10 +
+			(fAx         * fAy)         * fV11;
+	}
+
+};
+
+typedef TSampledScalarField2<double, double> FSampledScalarField2d;
+typedef TSampledScalarField2<float, float> FSampledScalarField2f;
\ No newline at end of file
