GeometryProcessing: minor additions to MathUtil and SparseGrid3, add DynamicVector const-iterator
#rnx

diff --git a/Engine/Plugins/Experimental/GeometryProcessing/Source/GeometricObjects/Public/MathUtil.h b/Engine/Plugins/Experimental/GeometryProcessing/Source/GeometricObjects/Public/MathUtil.h
--- a/Engine/Plugins/Experimental/GeometryProcessing/Source/GeometricObjects/Public/MathUtil.h
+++ b/Engine/Plugins/Experimental/GeometryProcessing/Source/GeometricObjects/Public/MathUtil.h
@@ -61,6 +61,8 @@ public:
 	static inline RealType Sin(const RealType Value);
 	static inline RealType Cos(const RealType Value);
 	static inline RealType ACos(const RealType Value);
+	static inline RealType Floor(const RealType Value);
+	static inline RealType Ceil(const RealType Value);
 
 
 
@@ -170,6 +172,18 @@ RealType TMathUtil<RealType>::ACos(const RealType Value)
 	return acos(Value);
 }
 
+template<typename RealType>
+RealType TMathUtil<RealType>::Floor(const RealType Value)
+{
+	return floor(Value);
+}
+
+template<typename RealType>
+RealType TMathUtil<RealType>::Ceil(const RealType Value)
+{
+	return ceil(Value);
+}
+
 
 template<typename RealType>
 RealType TMathUtil<RealType>::Atan2Positive(const RealType Y, const RealType X)
diff --git a/Engine/Plugins/Experimental/GeometryProcessing/Source/GeometricObjects/Public/Spatial/SparseGrid3.h b/Engine/Plugins/Experimental/GeometryProcessing/Source/GeometricObjects/Public/Spatial/SparseGrid3.h
--- a/Engine/Plugins/Experimental/GeometryProcessing/Source/GeometricObjects/Public/Spatial/SparseGrid3.h
+++ b/Engine/Plugins/Experimental/GeometryProcessing/Source/GeometricObjects/Public/Spatial/SparseGrid3.h
@@ -54,6 +54,17 @@ public:
 	}
 
 
+	/**
+	 * Get the grid element at this index
+	 * @param Index integer grid index
+	 * @return pointer to ElemType instance, or nullptr if element doesn't exist
+	 */
+	const ElemType* Get(const FVector3i& Index) const
+	{
+		return Elements.FindRef(Index);
+	}
+
+
 	/**
 	 * Get the grid element at this index, and optionally allocate it if it doesn't exist
 	 * @param Index integer grid index
@@ -135,6 +146,19 @@ public:
 	}
 
 
+	/**
+	 * Iterate over existing elements and apply ElementFunc to each of them
+	 */
+	template<typename Func>
+	void AllocatedIteration(Func ElementFunc) const
+	{
+		for (auto Pair : Elements)
+		{
+			ElementFunc(Pair.Value);
+		}
+	}
+
+
 protected:
 
 	ElemType* Allocate(const FVector3i& index)
diff --git a/Engine/Plugins/Experimental/GeometryProcessing/Source/GeometricObjects/Public/Util/DynamicVector.h b/Engine/Plugins/Experimental/GeometryProcessing/Source/GeometricObjects/Public/Util/DynamicVector.h
--- a/Engine/Plugins/Experimental/GeometryProcessing/Source/GeometricObjects/Public/Util/DynamicVector.h
+++ b/Engine/Plugins/Experimental/GeometryProcessing/Source/GeometricObjects/Public/Util/DynamicVector.h
@@ -148,6 +148,65 @@ public:
 	{
 		return FIterator(this, (int)GetLength());
 	}
+
+
+
+
+
+	/*
+	 * FConstIterator class iterates over values of vector
+	 */
+	class FConstIterator
+	{
+	public:
+		inline const Type& operator*() const
+		{
+			return (*DVector)[Idx];
+		}
+		inline FConstIterator& operator++()   // prefix
+		{
+			Idx++;
+			return *this;
+		}
+		inline FConstIterator operator++(int) // postfix
+		{
+			FConstIterator Copy(*this);
+			Idx++;
+			return Copy;
+		}
+		inline bool operator==(const FConstIterator& Itr2)
+		{
+			return DVector == Itr2.DVector && Idx == Itr2.Idx;
+		}
+		inline bool operator!=(const FConstIterator& Itr2)
+		{
+			return DVector != Itr2.DVector || Idx != Itr2.Idx;
+		}
+
+	protected:
+		const TDynamicVector<Type>* DVector;
+		int Idx;
+		inline FConstIterator(const TDynamicVector<Type>* Parent, int ICur)
+		{
+			DVector = Parent;
+			Idx = ICur;
+		}
+		friend class TDynamicVector<Type>;
+	};
+
+	/** @return iterator at beginning of vector */
+	FConstIterator begin() const
+	{
+		return IsEmpty() ? end() : FConstIterator(this, 0);
+	}
+	/** @return iterator at end of vector */
+	FConstIterator end() const
+	{
+		return FConstIterator(this, (int)GetLength());
+	}
+
+
+
 };
 
 
