initial code to walk paths on mesh surfaces and embed them in the mesh, to be used for cutting shapes into meshes

#rnx

diff --git a/Engine/Plugins/Experimental/GeometryProcessing/Source/DynamicMesh/Private/Operations/EmbedSurfacePath.cpp b/Engine/Plugins/Experimental/GeometryProcessing/Source/DynamicMesh/Private/Operations/EmbedSurfacePath.cpp
new file mode 100644
--- /dev/null
+++ b/Engine/Plugins/Experimental/GeometryProcessing/Source/DynamicMesh/Private/Operations/EmbedSurfacePath.cpp
@@ -0,0 +1,725 @@
+// Copyright 1998-2019 Epic Games, Inc. All Rights Reserved.
+
+#include "Operations/EmbedSurfacePath.h"
+#include "MathUtil.h"
+#include "VectorUtil.h"
+#include "TriangleTypes.h"
+#include "SegmentTypes.h"
+#include "Distance/DistPoint3Triangle3.h"
+
+#include "DynamicMesh3.h"
+
+FVector3d FMeshSurfacePoint::Pos(const FDynamicMesh3 *Mesh) const
+{
+	if (PointType == ESurfacePointType::Vertex)
+	{
+		return Mesh->GetVertex(ElementID);
+	}
+	else if (PointType == ESurfacePointType::Edge)
+	{
+		FVector3d EA, EB;
+		Mesh->GetEdgeV(ElementID, EA, EB);
+		return BaryCoord[0] * EA + BaryCoord[1] * EB;
+	}
+	else // PointType == ESurfacePointType::Triangle
+	{
+		FVector3d TA, TB, TC;
+		Mesh->GetTriVertices(ElementID, TA, TB, TC);
+		return BaryCoord[0] * TA + BaryCoord[1] * TB + BaryCoord[2] * TC;
+	}
+}
+
+/**
+ * Helper function to snap a triangle surface point to the triangle vertices or edges if it's close enough.  Input SurfacePt must be a triangle.
+ */
+void RefineSurfacePtFromTriangleToSubElement(const FDynamicMesh3* Mesh, FVector3d Pos, FMeshSurfacePoint& SurfacePt, double SnapElementThresholdSq)
+{
+	// expect this to only be called on SurfacePoints with PointType == Triangle; otherwise indicative of incorrect usage
+	if (!ensure(SurfacePt.PointType == ESurfacePointType::Triangle))
+	{
+		return;
+	}
+	int TriID = SurfacePt.ElementID;
+
+	FIndex3i TriVertIDs = Mesh->GetTriangle(TriID);
+	int BestSubIdx = -1;
+	double BestElementDistSq = 0;
+	for (int VertSubIdx = 0; VertSubIdx < 3; VertSubIdx++)
+	{
+		double DistSq = Pos.DistanceSquared(Mesh->GetVertex(TriVertIDs[VertSubIdx]));
+		if (DistSq <= SnapElementThresholdSq && (BestSubIdx == -1 || DistSq < BestElementDistSq))
+		{
+			BestSubIdx = VertSubIdx;
+			BestElementDistSq = DistSq;
+		}
+	}
+
+	if (BestSubIdx > -1)
+	{
+		SurfacePt.ElementID = TriVertIDs[BestSubIdx];
+		SurfacePt.PointType = ESurfacePointType::Vertex;
+		return;
+	}
+
+	// failed to snap to vertex, try snapping to edge
+	FIndex3i TriEdgeIDs = Mesh->GetTriEdges(TriID);
+	
+	check(BestSubIdx == -1); // otherwise would have returned the within-threshold vertex!
+	double BestEdgeParam = 0;
+	for (int EdgeSubIdx = 0; EdgeSubIdx < 3; EdgeSubIdx++)
+	{
+		int EdgeID = TriEdgeIDs[EdgeSubIdx];
+		FVector3d EPosA, EPosB;
+		Mesh->GetEdgeV(EdgeID, EPosA, EPosB);
+		FSegment3d EdgeSeg(EPosA, EPosB);
+		double DistSq = EdgeSeg.DistanceSquared(Pos);
+		if (DistSq <= SnapElementThresholdSq && (BestSubIdx == -1 || DistSq < BestElementDistSq))
+		{
+			BestSubIdx = EdgeSubIdx;
+			BestElementDistSq = DistSq;
+			BestEdgeParam = EdgeSeg.ProjectUnitRange(Pos);
+		}
+	}
+
+	if (BestSubIdx > -1)
+	{
+		SurfacePt.ElementID = TriEdgeIDs[BestSubIdx];
+		SurfacePt.PointType = ESurfacePointType::Edge;
+		SurfacePt.BaryCoord = FVector3d(BestEdgeParam, 1 - BestEdgeParam, 0);
+		return;
+	}
+
+	// no snapping to be done, leave surfacept on the triangle
+}
+
+// For when a triangle is replaced by multiple triangles, create a new surface point for the point's new location among the smaller triangles.
+//  Note input position is in the coordinate space of the mesh vertices, not e.g. a transformed space (unlike the positions passed to WalkMeshPlanar below!)
+FMeshSurfacePoint RelocateTrianglePointAfterRefinement(const FDynamicMesh3* Mesh, const FVector3d& PosInVertexCoordSpace, TArray<int> TriIDs, double SnapElementThresholdSq)
+{
+	double BestTriDistSq = 0;
+	FVector3d BestBaryCoords;
+	int BestTriID = -1;
+	for (int TriID : TriIDs)
+	{
+		check(Mesh->IsTriangle(TriID));
+		FIndex3i TriVertIDs = Mesh->GetTriangle(TriID);
+		FTriangle3d Tri(Mesh->GetVertex(TriVertIDs.A), Mesh->GetVertex(TriVertIDs.B), Mesh->GetVertex(TriVertIDs.C));
+		FDistPoint3Triangle3d TriDist(PosInVertexCoordSpace, Tri); // heavy duty way to get barycentric coordinates and check if on triangle; should be robust to degenerate triangles unlike VectorUtil's barycentric coordinate function
+		double DistSq = TriDist.GetSquared();
+		if (BestTriID == -1 || DistSq < BestTriDistSq)
+		{
+			BestTriID = TriID;
+			BestTriDistSq = DistSq;
+			BestBaryCoords = TriDist.TriangleBaryCoords;
+		}
+	}
+
+	ensure(Mesh->IsTriangle(BestTriID));
+	FMeshSurfacePoint SurfacePt(BestTriID, BestBaryCoords);
+	RefineSurfacePtFromTriangleToSubElement(Mesh, PosInVertexCoordSpace, SurfacePt, SnapElementThresholdSq);
+	return SurfacePt;
+}
+
+bool WalkMeshPlanar(const FDynamicMesh3* Mesh, int StartTri, FVector3d StartPt, int EndTri, FVector3d EndPt, FVector3d WalkPlaneNormal, TFunction<FVector3d(const FDynamicMesh3*, int)> VertexToPosnFn, bool bAllowBackwardsSearch, double AcceptEndPtOutsideDist, double PtOnPlaneThresholdSq, TArray<TPair<FMeshSurfacePoint, int>>& WalkedPath)
+{
+	auto SetTriVertPositions = [&VertexToPosnFn, &Mesh](FIndex3i TriVertIDs, FTriangle3d& Tri)
+	{
+		Tri.V[0] = VertexToPosnFn(Mesh, TriVertIDs.A);
+		Tri.V[1] = VertexToPosnFn(Mesh, TriVertIDs.B);
+		Tri.V[2] = VertexToPosnFn(Mesh, TriVertIDs.C);
+	};
+
+	auto PtTriPlaneSignedDist = [](FVector3d Pt, const FTriangle3d& Tri)
+	{
+		FVector3d Normal = Tri.Normal();
+		Normal.Normalize(FMathd::Epsilon);
+		return (Pt - Tri.V[0]).Dot(Normal);
+	};
+
+	auto PtInsideTri = [](const FVector3d& BaryCoord, double BaryThreshold = FMathd::ZeroTolerance)
+	{
+		return BaryCoord[0] >= -BaryThreshold && BaryCoord[1] >= -BaryThreshold && BaryCoord[2] >= -BaryThreshold;
+	};
+
+	// track where you came from and where you are going
+	struct FWalkIndices
+	{
+		FVector3d Position;		// Position in the coordinate space used for the walk (note: may not be the same space as the mesh vertices, e.g. if we walk on UV positions)
+		int WalkedFromPt,		// index into ComputedPointsAndSources (or -1)
+			WalkingOnTri;		// ID of triangle in mesh (or -1)
+
+		FWalkIndices() : WalkedFromPt(-1), WalkingOnTri(-1)
+		{}
+
+		FWalkIndices(FVector3d Position, int FromPt, int OnTri) : Position(Position), WalkedFromPt(FromPt), WalkingOnTri(OnTri)
+		{}
+	};
+
+	// TODO: vertex/edge snapping?
+	TArray<TPair<FMeshSurfacePoint, FWalkIndices>> ComputedPointsAndSources;
+	// TODO: switch this to a priority queue where distance to end is stored alongside, and we always pick the closest to goal ...
+	TArray<int> UnexploredEnds;
+	int BestKnownEnd = -1;
+	double BestKnownEndDistSq = FMathd::MaxReal;
+	TSet<int> ExploredTriangles;
+
+	bool bHasArrived = false;
+
+	FTriangle3d CurrentTri;
+	FIndex3i StartTriVertIDs = Mesh->GetTriangle(StartTri);
+	SetTriVertPositions(StartTriVertIDs, CurrentTri);
+	FDistPoint3Triangle3d CurrentTriDist(StartPt, CurrentTri); // heavy duty way to get barycentric coordinates and check if on triangle; should be robust to degenerate triangles unlike VectorUtil's barycentric coordinate function
+	// TODO: use TrianglePosToSurfacePoint to snap to edge or vertex as needed (OR do this as a post-process and delete the point if doing so leads to a duplicate point!)
+	CurrentTriDist.ComputeResult();
+	// TODO: replace barycoords result with edge or vertex surface point data if within distance threshold of vertex or edge!
+	ComputedPointsAndSources.Emplace(FMeshSurfacePoint(StartTri, CurrentTriDist.TriangleBaryCoords), FWalkIndices(StartPt, -1, StartTri));
+
+	double InitialDistSq = EndPt.DistanceSquared(StartPt);
+
+	int CurrentEnd = 0;
+	int IterCountSafety = 0;
+	int NumTriangles = Mesh->TriangleCount();
+	while (true)
+	{
+		if (!ensure(IterCountSafety++ < NumTriangles * 2)) // safety check to protect against infinite loop
+		{
+			return false;
+		}
+		const FMeshSurfacePoint& FromPt = ComputedPointsAndSources[CurrentEnd].Key;
+		const FWalkIndices& CurrentWalk = ComputedPointsAndSources[CurrentEnd].Value;
+		int TriID = CurrentWalk.WalkingOnTri;
+		check(Mesh->IsTriangle(TriID));
+		FIndex3i TriVertIDs = Mesh->GetTriangle(TriID);
+		SetTriVertPositions(TriVertIDs, CurrentTri);
+
+		bool OnEndTri = EndTri == TriID;
+		bool ComputedEndPtOnTri = false;
+		if (EndTri == -1)
+		{
+			CurrentTriDist.Triangle = CurrentTri;
+			CurrentTriDist.Point = EndPt;
+			ComputedEndPtOnTri = true;
+			double DistSq = CurrentTriDist.GetSquared();
+			if (DistSq < FMathd::Epsilon/* && PtInsideTri(CurrentTriDist.TriangleBaryCoords, FMathd::Epsilon)*/)  // TODO: we don't really need to check the barycentric coordinates for being 'inside' the triangle if the distance is within epsilon?  especially since the barycoords test also has n epsilon threshold?
+			{
+				OnEndTri = true;
+			}
+		}
+
+		// if we're on the final triangle, end the search!
+		if (OnEndTri)
+		{
+			if (!ComputedEndPtOnTri)
+			{
+				CurrentTriDist.Triangle = CurrentTri;
+				CurrentTriDist.Point = EndPt;
+				ComputedEndPtOnTri = true;
+				CurrentTriDist.GetSquared();
+			}
+			CurrentEnd = ComputedPointsAndSources.Emplace(FMeshSurfacePoint(TriID, CurrentTriDist.TriangleBaryCoords), FWalkIndices(EndPt, CurrentEnd, TriID));
+			bHasArrived = true;
+			BestKnownEnd = CurrentEnd;
+			break;
+		}
+
+		if (ExploredTriangles.Contains(TriID))
+		{
+			// note we only add explored triangles to the search at all to handle the specific case where we go 'the long way' around and back to the start triangle ... otherwise we could have just not added them to the search at all
+			// currently that case is not even possible, but if it were, it would have been handled in the above `if (OnEndTri)` so we should be able to safely kill this branch of search here
+			// TODO: this code is a copy of the code at the end of the while(true) block!  consolidate?!
+			if (UnexploredEnds.Num())
+			{
+				// TODO: consider storing scores for all pts and using heappop to get the closest unexplored
+				CurrentEnd = UnexploredEnds.Pop();
+				continue;
+			}
+			else
+			{
+				return false; // failed to find
+			}
+		}
+		ExploredTriangles.Add(TriID);
+
+		// not on a terminal triangle, cross the triangle and continue the search
+		double SignDist[3];
+		int Side[3];
+		int32 InitialComputedPointsNum = ComputedPointsAndSources.Num();
+		for (int TriSubIdx = 0; TriSubIdx < 3; TriSubIdx++)
+		{
+			double SD = (CurrentTri.V[TriSubIdx] - StartPt).Dot(WalkPlaneNormal);
+			SignDist[TriSubIdx] = SD;
+			if (FMathd::Abs(SD) <= PtOnPlaneThresholdSq)
+			{
+				Side[TriSubIdx] = 0;
+				int CandidateVertID = TriVertIDs[TriSubIdx];
+				if (FromPt.PointType != ESurfacePointType::Vertex || CandidateVertID != FromPt.ElementID)
+				{
+					FMeshSurfacePoint SurfPt(CandidateVertID);
+					FWalkIndices WalkInds(CurrentTri.V[TriSubIdx], CurrentEnd, -1);
+
+					// TODO: extract this "next triangle candidate" logic to be used in more places??
+					// walking over a vertex is gross because we have to search the whole one ring for candidate next triangles and there might be multiple of them
+					// note that currently this means I compute signs for all vertices of neighboring triangles here, and do not re-use those signs when I actually process the triangle later; TODO reconsider if/when optimizing this fn
+					for (int32 NbrTriID : Mesh->VtxTrianglesItr(CandidateVertID))
+					{
+						if (NbrTriID != TriID)
+						{
+							FIndex3i NbrTriVertIDs = Mesh->GetTriangle(NbrTriID);
+							FTriangle3d NbrTri;
+							SetTriVertPositions(NbrTriVertIDs, NbrTri);
+							int SignsMultiplied = 1;
+							for (int NbrTriSubIdx = 0; NbrTriSubIdx < 3; NbrTriSubIdx++)
+							{
+								if (NbrTriVertIDs[NbrTriSubIdx] == CandidateVertID)
+								{
+									continue;
+								}
+								double NbrSD = (NbrTri.V[NbrTriSubIdx] - StartPt).Dot(WalkPlaneNormal);
+								int NbrSign = FMathd::Abs(NbrSD) <= PtOnPlaneThresholdSq ? 0 : NbrSD > 0 ? 1 : -1;
+								SignsMultiplied *= NbrSign;
+							}
+							if (SignsMultiplied < 1) // plane will cross this triangle, so try walking it
+							{
+								WalkInds.WalkingOnTri = NbrTriID;
+								ComputedPointsAndSources.Emplace(SurfPt, WalkInds);
+							}
+						}
+					}
+				}
+			}
+			else
+			{
+				Side[TriSubIdx] = SD > 0 ? 1 : -1;
+			}
+		}
+		FIndex3i TriEdgeIDs = Mesh->GetTriEdges(TriID);
+		for (int TriSubIdx = 0; TriSubIdx < 3; TriSubIdx++)
+		{
+			int NextSubIdx = (TriSubIdx + 1) % 3;
+			if (SignDist[TriSubIdx] * SignDist[NextSubIdx] < 0)
+			{
+				// edge crossing
+				int CandidateEdgeID = TriEdgeIDs[TriSubIdx];
+				if (FromPt.PointType != ESurfacePointType::Edge || CandidateEdgeID != FromPt.ElementID)
+				{
+					double CrossingT = SignDist[TriSubIdx] / (SignDist[TriSubIdx] - SignDist[NextSubIdx]);
+					FVector3d CrossingP = (1 - CrossingT) * CurrentTri.V[TriSubIdx] + CrossingT * CurrentTri.V[NextSubIdx];
+					FIndex4i EdgeInfo = Mesh->GetEdge(CandidateEdgeID);
+					if (EdgeInfo.A != TriVertIDs[TriSubIdx]) // edge verts are stored backwards from the order in the local triangle, reverse the crossing accordingly
+					{
+						CrossingT = 1 - CrossingT;
+					}
+					int CrossToTriID = EdgeInfo.C;
+					if (CrossToTriID == TriID)
+					{
+						CrossToTriID = EdgeInfo.D;
+					}
+					if (CrossToTriID == -1)
+					{
+						// We've walked off the border of the mesh
+						// TODO: check if this is close enough to the EndPt, and if so just stop the walk here
+						continue;
+					}
+					double DSq = EndPt.DistanceSquared(CrossingP);
+					if (!bAllowBackwardsSearch && DSq > InitialDistSq + 100 * FMathd::Epsilon)
+					{
+						// not allowed to go in a direction that gets us further from the destination than our initial point if backwards search not allowed
+						continue;
+					}
+					ComputedPointsAndSources.Emplace(FMeshSurfacePoint(CandidateEdgeID, CrossingT), FWalkIndices(CrossingP, CurrentEnd, CrossToTriID));
+				}
+			}
+		}
+
+		int BestCandidate = -1;
+		double BestCandidateDistSq = FMathd::MaxReal;
+		for (int32 NewComputedPtIdx = InitialComputedPointsNum; NewComputedPtIdx < ComputedPointsAndSources.Num(); NewComputedPtIdx++)
+		{
+			double DistSq = EndPt.DistanceSquared(ComputedPointsAndSources[NewComputedPtIdx].Value.Position);
+			// TODO: reject cases that move us backwards if !bAllowBackwardsSearch
+			if (BestCandidate == -1 || DistSq < BestCandidateDistSq)
+			{
+				BestCandidateDistSq = DistSq;
+				BestCandidate = NewComputedPtIdx;
+			}
+		}
+		if (BestCandidate == -1)
+		{
+			// TODO: this code is a copy of the code that handles terminating the search if we hit a triangle we've already seen, above; consolidate!?
+			if (UnexploredEnds.Num())
+			{
+				// TODO: consider storing scores for all pts and using heappop to get the closest unexplored
+				CurrentEnd = UnexploredEnds.Pop();
+				continue;
+			}
+			else
+			{
+				return false; // failed to find
+			}
+		}
+		CurrentEnd = BestCandidate;
+		for (int32 NewComputedPtIdx = InitialComputedPointsNum; NewComputedPtIdx < ComputedPointsAndSources.Num(); NewComputedPtIdx++)
+		{
+			if (NewComputedPtIdx != BestCandidate)
+			{
+				UnexploredEnds.Add(NewComputedPtIdx);
+			}
+		}
+	}
+
+
+	int TrackedPtIdx = BestKnownEnd;
+	int SafetyIdxBacktrack = 0;
+	TArray<int> AcceptedIndices;
+	while (TrackedPtIdx > -1)
+	{
+		if (!ensure(SafetyIdxBacktrack++ < 2*ComputedPointsAndSources.Num())) // infinite loop guard
+		{
+			return false;
+		}
+		AcceptedIndices.Add(TrackedPtIdx);
+		TrackedPtIdx = ComputedPointsAndSources[TrackedPtIdx].Value.WalkedFromPt;
+	}
+	WalkedPath.Reset();
+	for (int32 IdxIdx = AcceptedIndices.Num() - 1; IdxIdx >= 0; IdxIdx--)
+	{
+		WalkedPath.Emplace(ComputedPointsAndSources[AcceptedIndices[IdxIdx]].Key, ComputedPointsAndSources[AcceptedIndices[IdxIdx]].Value.WalkingOnTri);
+	}
+
+	// try refining start and end points if they were on triangles, and remove them if they turn out to be duplicates after refinement
+	//  (note we could instead do the refinement up front and avoid the possibility of duplicate points, but that would slightly complicate the traversal logic ...)
+	if (WalkedPath.Num() && WalkedPath[0].Key.PointType == ESurfacePointType::Triangle)
+	{
+		RefineSurfacePtFromTriangleToSubElement(Mesh, WalkedPath[0].Key.Pos(Mesh), WalkedPath[0].Key, PtOnPlaneThresholdSq);
+		if (WalkedPath.Num() > 1 &&
+			WalkedPath[0].Key.PointType == WalkedPath[1].Key.PointType &&
+			WalkedPath[0].Key.ElementID == WalkedPath[1].Key.ElementID &&
+			(WalkedPath[0].Key.PointType == ESurfacePointType::Vertex ||
+				WalkedPath[0].Key.BaryCoord.DistanceSquared(WalkedPath[1].Key.BaryCoord) < FMathd::Epsilon))
+		{
+			WalkedPath.RemoveAt(0);
+		}
+	}
+	if (WalkedPath.Num() && WalkedPath.Last().Key.PointType == ESurfacePointType::Triangle)
+	{
+		RefineSurfacePtFromTriangleToSubElement(Mesh, WalkedPath.Last().Key.Pos(Mesh), WalkedPath.Last().Key, PtOnPlaneThresholdSq);
+		if (WalkedPath.Num() > 1 &&
+			WalkedPath.Last().Key.PointType == WalkedPath.Last(1).Key.PointType &&
+			WalkedPath.Last().Key.ElementID == WalkedPath.Last(1).Key.ElementID &&
+			(WalkedPath.Last().Key.PointType == ESurfacePointType::Vertex ||
+				WalkedPath.Last().Key.BaryCoord.DistanceSquared(WalkedPath.Last(1).Key.BaryCoord) < FMathd::Epsilon))
+		{
+			WalkedPath.Pop();
+		}
+	}
+
+	return true;
+}
+
+
+//static int FMeshSurfacePath::FindSharedTriangle(const FDynamicMesh3* Mesh, const FMeshSurfacePoint& A, const FMeshSurfacePoint& B)
+//{
+//	if (A.PointType == B.PointType && A.PointType == ESurfacePointType::Vertex)
+//	{
+//		int Edge = Mesh->FindEdge(A.ElementID, B.ElementID);
+//		if (Edge == FDynamicMesh3::InvalidID)
+//		{
+//			return FDynamicMesh3::InvalidID;
+//		}
+//		return Mesh->GetEdgeT(Edge).A;
+//	}
+//	else if (A.PointType == ESurfacePointType::Vertex || B.PointType == ESurfacePointType::Vertex)
+//	{
+//		int VertexID = A.ElementID, OtherID = B.ElementID;
+//		ESurfacePointType OtherType = B.PointType;
+//		if (A.PointType != ESurfacePointType::Vertex)
+//		{
+//			VertexID = B.ElementID;
+//			OtherID = A.ElementID;
+//			OtherType = A.PointType;
+//		}
+//		if (OtherType == ESurfacePointType::Triangle)
+//		{
+//			if (!Mesh->GetTriangle(OtherID).Contains(VertexID))
+//			{
+//				return FDynamicMesh3::InvalidID;
+//			}
+//			return OtherID;
+//		}
+//		else // OtherType == ESurfacePointType::Edge
+//		{
+//			FIndex2i TriIDs = Mesh->GetEdgeT(OtherID);
+//			if (Mesh->GetTriangle(TriIDs.A).Contains(VertexID))
+//			{
+//				return TriIDs.A;
+//			}
+//			else if (TriIDs.B != FDynamicMesh3::InvalidID && Mesh->GetTriangle(TriIDs.B).Contains(VertexID))
+//			{
+//				return TriIDs.B;
+//			}
+//			else
+//			{
+//				return FDynamicMesh3::InvalidID;
+//			}
+//		}
+//	}
+//	else if (A.PointType == ESurfacePointType::Edge || B.FMeshSurfacePoint == ESurfacePointType::Edge)
+//	{
+//		int EdgeID = A.ElementID, OtherID = B.ElementID;
+//		ESurfacePointType OtherType = B.PointType;
+//		if (A.PointType != ESurfacePointType::Edge)
+//		{
+//			EdgeID = B.ElementID;
+//			OtherID = A.ElementID;
+//			OtherType = A.PointType;
+//		}
+//		if (OtherType == ESurfacePointType::Triangle)
+//		{
+//			FIndex2i 
+//		}
+//	}
+//	else // both ESurfacePointType::Triangle
+//	{
+//		if (A.ElementID == B.ElementID)
+//		{
+//			return A.ElementID;
+//		}
+//		else
+//		{
+//			return FDynamicMesh3::InvalidID;
+//		}
+//	}
+//}
+
+bool FMeshSurfacePath::IsConnected() const
+{
+	int Idx = 1, LastIdx = 0;
+	if (bIsClosed)
+	{
+		LastIdx = Path.Num() - 1;
+		Idx = 0;
+	}
+	for (; Idx < Path.Num(); LastIdx = Idx++)
+	{
+		int WalkingOnTri = Path[LastIdx].Value;
+		if (!Mesh->IsTriangle(WalkingOnTri))
+		{
+			return false;
+		}
+		int Inds[2] = { LastIdx, Idx };
+		for (int IndIdx = 0; IndIdx < 2; IndIdx++)
+		{
+			const FMeshSurfacePoint& P = Path[Inds[IndIdx]].Key;
+			switch (P.PointType)
+			{
+			case ESurfacePointType::Triangle:
+				if (P.ElementID != WalkingOnTri)
+				{
+					return false;
+				}
+				break;
+			case ESurfacePointType::Edge:
+				if (!Mesh->GetEdgeT(P.ElementID).Contains(WalkingOnTri))
+				{
+					return false;
+				}
+				break;
+			case ESurfacePointType::Vertex:
+				if (!Mesh->GetTriangle(WalkingOnTri).Contains(P.ElementID))
+				{
+					return false;
+				}
+				break;
+			}
+		}
+	}
+	return true;
+}
+
+
+bool FMeshSurfacePath::AddViaPlanarWalk(int StartTri, FVector3d StartPt, int EndTri, FVector3d EndPt, FVector3d WalkPlaneNormal, TFunction<FVector3d(const FDynamicMesh3*, int)> VertexToPosnFn, bool bAllowBackwardsSearch, double AcceptEndPtOutsideDist, double PtOnPlaneThresholdSq)
+{
+	if (!VertexToPosnFn)
+	{
+		VertexToPosnFn = [](const FDynamicMesh3* Mesh, int VertexID)
+		{
+			return Mesh->GetVertex(VertexID);
+		};
+	}
+	return WalkMeshPlanar(Mesh, StartTri, StartPt, EndTri, EndPt, WalkPlaneNormal, VertexToPosnFn, bAllowBackwardsSearch, AcceptEndPtOutsideDist, PtOnPlaneThresholdSq, Path);
+}
+
+// TODO: general path embedding becomes an arbitrary 2D remeshing problem per triangle; requires support from e.g. GeometryAlgorithms CDT. Not implemented yet; this is a vague sketch of what might go there.
+//bool FMeshSurfacePath::EmbedPath(bool bUpdatePath, TArray<int>& PathVertices, TFunction<bool(const TArray<FVector2d>& Vertices, const TArray<FIndex3i>& LabelledEdges, TArray<FVector2d>& OutVertices, TArray<int32>& OutVertexMap, TArray<FIndex3i>& OutTriangles)> MeshGraphFn)
+//{
+//	// Array mapping surface pts in Path to IDs of new vertices in the mesh
+//	TArray<int> SurfacePtToNewVertexID;
+//
+//	struct FRemeshTriInfo
+//	{
+//		TArray<int> AddSurfacePtIdx;
+//		TArray<FVector2d> LocalVertexPosns;
+//		TArray<FIndex2i> LocalEdges;
+//	};
+//
+//	TMap<int, FRemeshTriInfo> TrianglesToRemesh;
+//
+//	// STEP 1: Collect all triangles that need re-triangulation because of pts added to edges or faces
+//	// STEP 2: Perform re-triangulations in local coordinate space per triangle.  Store in FRemeshTriInfo map.  If any failure detected, stop process here (before altering source mesh)
+//	// STEP 3: Add floating verts for all points that require a new vert, track their correspondence to surface pts via SurfacePtToNewVertexID
+//	// STEP 4: Remove all triangles that needed re-triangulation (ignoring creation of bowties, etc)
+//	// STEP 5: Add re-triangulations into mesh
+//	// TODO: consider how to handle dynamic mesh overlay / attribute stuff
+//	
+//	return false;
+//}
+
+/**
+* Embed a surface path in mesh provided that the path only crosses vertices and edges except at the start and end, so we can add the path easily with local edge splits and possibly two triangle pokes (rather than needing general remeshing machinery)
+* Also assumes triangles are only crossed over once (except possibly  to loop around to the start triangle on the end triangle)
+* Planar walks naturally create simple paths, so this function can be used on any paths created by single planar walks.
+*
+* @param bUpdatePath Updating the Path array with the new vertices (if false, the path will no longer be valid after running this function)\
+* @param PathVertices Indices of the vertices on the path after embedding succeeds; NOTE these will not be 1:1 with the input Path
+* @return true if embedding succeeded.
+*/
+bool FMeshSurfacePath::EmbedSimplePath(bool bUpdatePath, TArray<int>& PathVertices, double SnapElementThresholdSq)
+{
+	int32 InitialPathIdx = PathVertices.Num(); // TODO: rm this variable (used for debugging)
+	if (!Path.Num())
+	{
+		return true;
+	}
+
+	int32 PathNum = Path.Num();
+	const FMeshSurfacePoint& OrigEndPt = Path[PathNum - 1].Key;
+	
+	// If FinalTri is split or poked, we will need to re-locate the last point in the path
+	int StartProcessIdx = 0, EndSimpleProcessIdx = PathNum - 1;
+	bool bEndPointSpecialProcess = false;
+	if (OrigEndPt.PointType == ESurfacePointType::Triangle)
+	{
+		EndSimpleProcessIdx = PathNum - 2;
+		bEndPointSpecialProcess = true;
+	}
+	bool bNeedFinalRelocate = false;
+	FMeshSurfacePoint EndPtUpdated = Path.Last().Key;
+	FVector3d EndPtPos = OrigEndPt.Pos(Mesh);
+	
+	if (Path[0].Key.PointType == ESurfacePointType::Triangle)
+	{
+		// TODO: poke triangle, and place initial vertex
+		FDynamicMesh3::FPokeTriangleInfo PokeInfo;
+		Mesh->PokeTriangle(Path[0].Key.ElementID, Path[0].Key.BaryCoord, PokeInfo);
+		if (EndPtUpdated.PointType == ESurfacePointType::Triangle && Path[0].Key.ElementID == EndPtUpdated.ElementID)
+		{
+			EndPtUpdated = RelocateTrianglePointAfterRefinement(Mesh, EndPtPos, { PokeInfo.NewTriangles.A, PokeInfo.NewTriangles.B, PokeInfo.OriginalTriangle }, SnapElementThresholdSq);
+		}
+		PathVertices.Add(PokeInfo.NewVertex);
+		StartProcessIdx = 1;
+	}
+	for (int32 PathIdx = StartProcessIdx; PathIdx <= EndSimpleProcessIdx; PathIdx++)
+	{
+		if (!ensure(Path[PathIdx].Key.PointType != ESurfacePointType::Triangle))
+		{
+			// Input assumptions violated -- Simple path can only have Triangle points at the very first and/or last points!  Would need a more powerful embed function to handle this case.
+			return false;
+		}
+		const FMeshSurfacePoint& Pt = Path[PathIdx].Key;
+		if (Pt.PointType == ESurfacePointType::Edge)
+		{
+			ensure(Mesh->IsEdge(Pt.ElementID));
+			FDynamicMesh3::FEdgeSplitInfo SplitInfo;
+			Mesh->SplitEdge(Pt.ElementID, SplitInfo, Pt.BaryCoord[0]);
+			PathVertices.Add(SplitInfo.NewVertex);
+			if (EndPtUpdated.PointType == ESurfacePointType::Triangle && SplitInfo.OriginalTriangles.Contains(EndPtUpdated.ElementID))
+			{
+				TArray<int> TriInds = { EndPtUpdated.ElementID };
+				if (SplitInfo.OriginalTriangles.A == EndPtUpdated.ElementID)
+				{
+					TriInds.Add(SplitInfo.NewTriangles.A);
+				}
+				else
+				{
+					TriInds.Add(SplitInfo.NewTriangles.B);
+				}
+				EndPtUpdated = RelocateTrianglePointAfterRefinement(Mesh, EndPtPos, TriInds, SnapElementThresholdSq);
+			}
+		}
+		else
+		{
+			ensure(Pt.PointType == ESurfacePointType::Vertex);
+			ensure(Mesh->IsVertex(Pt.ElementID));
+			PathVertices.Add(Pt.ElementID);
+		}
+	}
+
+	if (bEndPointSpecialProcess)
+	{
+		if (EndPtUpdated.PointType == ESurfacePointType::Triangle)
+		{
+			FDynamicMesh3::FPokeTriangleInfo PokeInfo;
+			Mesh->PokeTriangle(EndPtUpdated.ElementID, EndPtUpdated.BaryCoord, PokeInfo);
+			PathVertices.Add(PokeInfo.NewVertex);
+		}
+		else if (EndPtUpdated.PointType == ESurfacePointType::Edge)
+		{
+			FDynamicMesh3::FEdgeSplitInfo SplitInfo;
+			Mesh->SplitEdge(EndPtUpdated.ElementID, SplitInfo, EndPtUpdated.BaryCoord[0]);
+			PathVertices.Add(SplitInfo.NewVertex);
+		}
+		else
+		{
+			PathVertices.Add(EndPtUpdated.ElementID);
+		}
+	}
+
+
+	// TODO: rm this debugging check
+	for (int PathIdx = InitialPathIdx; PathIdx + 1 < PathVertices.Num(); PathIdx++)
+	{
+		if (!ensure(Mesh->IsEdge(Mesh->FindEdge(PathVertices[PathIdx], PathVertices[PathIdx + 1]))))
+		{
+			return false;
+		}
+	}
+
+	if (bUpdatePath)
+	{
+		ensure(false); // todo implement this
+	}
+
+	return true;
+}
+
+bool EmbedProjectedPath(FDynamicMesh3* Mesh, int StartTriID, FFrame3d Frame, const TArray<FVector2d>& Path2D, TArray<int>& OutPathVertices, bool bClosePath, double PtSnapVertexOrEdgeThresholdSq)
+{
+	int32 EndIdxA = Path2D.Num() - (bClosePath ? 1 : 2);
+	int CurrentSeedTriID = StartTriID;
+	OutPathVertices.Reset();
+
+	TFunction<FVector3d(const FDynamicMesh3*, int)> ProjectToFrame = [&Frame](const FDynamicMesh3* Mesh, int VertexID)
+	{
+		FVector2d ProjPt = Frame.ToPlaneUV(Mesh->GetVertex(VertexID));
+		return FVector3d(ProjPt.X, ProjPt.Y, 0);
+	};
+
+	for (int32 IdxA = 0; IdxA <= EndIdxA; IdxA++)
+	{
+		int32 IdxB = (IdxA + 1) % Path2D.Num();
+		FMeshSurfacePath SurfacePath(Mesh);
+		FVector2d WalkDir = Path2D[IdxB] - Path2D[IdxA];
+		bool WalkSuccess = SurfacePath.AddViaPlanarWalk(CurrentSeedTriID, FVector3d(Path2D[IdxA].X, Path2D[IdxA].Y, 0), -1, FVector3d(Path2D[IdxB].X, Path2D[IdxB].Y, 0), FVector3d(-WalkDir.Y, WalkDir.X, 0), ProjectToFrame, false, 0, PtSnapVertexOrEdgeThresholdSq);
+		if (!WalkSuccess)
+		{
+			return false;
+		}
+		SurfacePath.EmbedSimplePath(false, OutPathVertices, PtSnapVertexOrEdgeThresholdSq);
+		TArray<int> TrianglesOut;
+		Mesh->GetVertexOneRingTriangles(OutPathVertices.Last(), TrianglesOut);
+		check(TrianglesOut.Num());
+		CurrentSeedTriID = TrianglesOut[0];
+	}
+
+	return true;
+}
diff --git a/Engine/Plugins/Experimental/GeometryProcessing/Source/DynamicMesh/Public/Operations/EmbedSurfacePath.h b/Engine/Plugins/Experimental/GeometryProcessing/Source/DynamicMesh/Public/Operations/EmbedSurfacePath.h
new file mode 100644
--- /dev/null
+++ b/Engine/Plugins/Experimental/GeometryProcessing/Source/DynamicMesh/Public/Operations/EmbedSurfacePath.h
@@ -0,0 +1,141 @@
+// Copyright 1998-2019 Epic Games, Inc. All Rights Reserved.
+
+// Port of geometry3Sharp MeshPlaneCut
+
+#pragma once
+#include "CoreMinimal.h"
+
+#include "MathUtil.h"
+#include "VectorTypes.h"
+#include "FrameTypes.h"
+#include "GeometryTypes.h"
+#include "Curve/GeneralPolygon2.h"
+
+
+class FDynamicMesh3;
+
+enum class ESurfacePointType
+{
+	Vertex = 0,
+	Edge = 1,
+	Triangle = 2
+};
+
+struct DYNAMICMESH_API FMeshSurfacePoint
+{
+	int ElementID;
+	FVector3d BaryCoord;
+	ESurfacePointType PointType;
+
+	FMeshSurfacePoint() : ElementID(-1)
+	{
+	}
+	FMeshSurfacePoint(int TriangleID, const FVector3d& BaryCoord) : ElementID(TriangleID), BaryCoord(BaryCoord), PointType(ESurfacePointType::Triangle)
+	{
+	}
+	FMeshSurfacePoint(int EdgeID, double FirstCoordWt) : ElementID(EdgeID), BaryCoord(FirstCoordWt, 1-FirstCoordWt, 0), PointType(ESurfacePointType::Edge)
+	{
+	}
+	FMeshSurfacePoint(int VertexID) : ElementID(VertexID), BaryCoord(1,0,0), PointType(ESurfacePointType::Vertex)
+	{
+	}
+
+	FVector3d Pos(const FDynamicMesh3 *Mesh) const;
+};
+
+/**
+ * Walk the surface of an FDynamicMesh to try find a planar path connecting two points.  Paths include every vertex and edge they need to cross.  Greedy algorithm will only return one path if there are multiple.
+ */
+bool DYNAMICMESH_API WalkMeshPlanar(const FDynamicMesh3* Mesh, int StartTri, FVector3d StartPt, int EndTri, FVector3d EndPt, FVector3d WalkPlaneNormal, TFunction<FVector3d(const FDynamicMesh3*, int)> VertexToPosnFn, bool bAllowBackwardsSearch, double AcceptEndPtOutsideDist, double PtOnPlaneThreshold, TArray<TPair<FMeshSurfacePoint, int>>& WalkedPath);
+
+
+/**
+ * Represent a path on the surface of a mesh via barycentric coordinates and triangle references
+ */
+class DYNAMICMESH_API FMeshSurfacePath
+{
+public:
+	FDynamicMesh3 *Mesh;
+	TArray<TPair<FMeshSurfacePoint, int>> Path; // Surface points paired with triangle to walk to get to next surface point
+	bool bIsClosed;
+
+public:
+
+	/**
+	 * Cut mesh with plane. Assumption is that plane normal is Z value.
+	 */
+	FMeshSurfacePath(FDynamicMesh3* Mesh) : Mesh(Mesh), bIsClosed(false)
+	{
+	}
+	virtual ~FMeshSurfacePath() {}
+
+	// TODO: not clear if we need this -- helper to find the shared triangle (if it exists) that we could traverse to connect two surface points.  this gets a bit hairy and it's easier to keep track of the traversed triangles during the mesh walk that creates any path in the first place
+	//static int FindSharedTriangle(const FDynamicMesh3* Mesh, const FMeshSurfacePoint& A, const FMeshSurfacePoint& B);
+	
+	/**
+	 * @return True if the Path exactly sticks to the mesh surface, and never jumps to disconnected elements
+	 */
+	bool IsConnected() const;
+
+	bool IsClosed() const
+	{
+		return bIsClosed;
+	}
+
+	void Reset()
+	{
+		Path.Reset();
+		bIsClosed = false;
+	}
+
+	bool AddViaPlanarWalk(int StartTri, FVector3d StartPt, int EndTri, FVector3d EndPt, FVector3d WalkPlaneNormal, TFunction<FVector3d(const FDynamicMesh3*, int)> VertexToPosnFn = nullptr, bool bAllowBackwardsSearch = true, double AcceptEndPtOutsideDist = FMathd::Epsilon, double PtOnPlaneThresholdSq = FMathd::Epsilon*100);
+	// TODO: Also support geodesic walks, other alternatives?
+
+	/**
+	 * Make the embedded path into a closed loop.  Only succeeds if path start and end already share a triangle!
+	 *
+	 * @return true if succeeded
+	 */
+	bool ClosePath();
+
+	/**
+	* @return EOperationValidationResult::Ok if we can apply operation, or error code if we cannot
+	*/
+	virtual EOperationValidationResult Validate()
+	{
+		if (!IsConnected())
+		{
+			return EOperationValidationResult::Failed_UnknownReason;
+		}
+
+		return EOperationValidationResult::Ok;
+	}
+
+	// TODO: support for embedding arbitrary paths in mesh using a graph remesher
+	///**
+	// * Embed surface path in mesh, such that each point in path is a mesh vertex with a connecting edge.  Note that IsConnected must be true for this to succeed.
+	// *
+	// * @param bUpdatePath Updating the Path array with the new vertices (if false, the path will no longer be valid after running this function)\
+	// * @param PathVertices Indices of the vertices on the path after embedding succeeds; NOTE these will not be 1:1 with the input Path
+	// * @param MeshGraphFn Function to create a mesh from an input 2D graph of labeled edges.
+	// * @return true if embedding succeeded.
+	// */
+	//bool EmbedPath(bool bUpdatePath, TArray<int>& PathVertices, TFunction<bool(const TArray<FVector2d>& Vertices, const TArray<FIndex3i>& LabelledEdges, TArray<FVector2d>& OutVertices, TArray<int32>& OutVertexMap, TArray<FIndex3i>& OutTriangles)> MeshGraphFn);
+
+	/**
+	 * Embed a surface path in mesh provided that the path only crosses vertices and edges except at the start and end, so we can add the path easily with local edge splits and possibly two triangle pokes (rather than needing general remeshing machinery)
+	 *
+	 * @param bUpdatePath Updating the Path array with the new vertices (if false, the path will no longer be valid after running this function)\
+	 * @param PathVertices Indices of the vertices on the path after embedding succeeds; NOTE these will not be 1:1 with the input Path
+	 * @return true if embedding succeeded.
+	 */
+	bool EmbedSimplePath(bool bUpdatePath, TArray<int>& PathVertices, double SnapElementThresholdSq = FMathd::Epsilon*100);
+
+	// TODO: add functionality to delete 'inside' of path -- but how do we determine what is inside?  ref MeshFacesFromLoop.cs in geometry3sharp
+};
+
+/**
+ * Embed a 2D path into a mesh by projection, starting the walk from a given triangle.
+ */
+bool DYNAMICMESH_API EmbedProjectedPath(FDynamicMesh3* Mesh, int StartTriID, FFrame3d Frame, const TArray<FVector2d>& Path2D, TArray<int>& OutPathVertices, bool bClosePath, double PtSnapVertexOrEdgeThresholdSq = FMathd::Epsilon*100);
+
diff --git a/Engine/Plugins/Experimental/GeometryProcessing/Source/GeometricObjects/Public/FrameTypes.h b/Engine/Plugins/Experimental/GeometryProcessing/Source/GeometricObjects/Public/FrameTypes.h
--- a/Engine/Plugins/Experimental/GeometryProcessing/Source/GeometricObjects/Public/FrameTypes.h
+++ b/Engine/Plugins/Experimental/GeometryProcessing/Source/GeometricObjects/Public/FrameTypes.h
@@ -218,7 +218,7 @@ struct TFrame3
 	 * @param PlaneNormalAxis which plane to project onto, identified by perpendicular normal. Default is 2, ie normal is Z, plane is (X,Y)
 	 * @return 2D coordinates in UV plane, relative to origin
 	 */
-	FVector2<RealType> ToPlaneUV(const FVector3<RealType>& Pos, int PlaneNormalAxis) const
+	FVector2<RealType> ToPlaneUV(const FVector3<RealType>& Pos, int PlaneNormalAxis = 2) const
 	{
 		int Axis0 = 0, Axis1 = 1;
 		if (PlaneNormalAxis == 0)
@@ -243,7 +243,7 @@ struct TFrame3
 	 * @param PlaneNormalAxis which plane to map to, identified by perpendicular normal. Default is 2, ie normal is Z, plane is (X,Y)
 	 * @return 3D coordinates in frame's plane (including Origin translation)
 	 */
-	FVector3<RealType> FromPlaneUV(const FVector2<RealType>& PosUV, int PlaneNormalAxis) const
+	FVector3<RealType> FromPlaneUV(const FVector2<RealType>& PosUV, int PlaneNormalAxis = 2) const
 	{
 		FVector3<RealType> PlanePos(PosUV[0], PosUV[1], 0);
 		if (PlaneNormalAxis == 0)
@@ -265,7 +265,7 @@ struct TFrame3
 	 * @param PlaneNormalAxis which plane to project onto, identified by perpendicular normal. Default is 2, ie normal is Z, plane is (X,Y)
 	 * @return 3D coordinate in the plane
 	 */
-	FVector3<RealType> ToPlane(const FVector3<RealType>& Pos, int PlaneNormalAxis) const
+	FVector3<RealType> ToPlane(const FVector3<RealType>& Pos, int PlaneNormalAxis = 2) const
 	{
 		FVector3<RealType> Normal = GetAxis(PlaneNormalAxis);
 		FVector3<RealType> LocalVec = Pos - Origin;
diff --git a/Engine/Plugins/Experimental/GeometryProcessing/Source/GeometricObjects/Public/IndexTypes.h b/Engine/Plugins/Experimental/GeometryProcessing/Source/GeometricObjects/Public/IndexTypes.h
--- a/Engine/Plugins/Experimental/GeometryProcessing/Source/GeometricObjects/Public/IndexTypes.h
+++ b/Engine/Plugins/Experimental/GeometryProcessing/Source/GeometricObjects/Public/IndexTypes.h
@@ -61,6 +61,16 @@ struct FIndex2i
 		return A != Other.A || B != Other.B;
 	}
 
+	int IndexOf(int Value) const
+	{
+		return (A == Value) ? 0 : ((B == Value) ? 1 : -1);
+	}
+
+	bool Contains(int Value) const
+	{
+		return (A == Value) || (B == Value);
+	}
+
 	inline void Swap()
 	{
 		int tmp = A;
@@ -136,7 +146,7 @@ struct FIndex3i
 		return (A == Value) ? 0 : ((B == Value) ? 1 : (C == Value ? 2 : -1));
 	}
 
-	int Contains(int Value) const
+	bool Contains(int Value) const
 	{
 		return (A == Value) || (B == Value) || (C == Value);
 	}
diff --git a/Engine/Plugins/Experimental/GeometryProcessing/Source/GeometricObjects/Public/VectorUtil.h b/Engine/Plugins/Experimental/GeometryProcessing/Source/GeometricObjects/Public/VectorUtil.h
--- a/Engine/Plugins/Experimental/GeometryProcessing/Source/GeometricObjects/Public/VectorUtil.h
+++ b/Engine/Plugins/Experimental/GeometryProcessing/Source/GeometricObjects/Public/VectorUtil.h
@@ -284,6 +284,56 @@ namespace VectorUtil
 		}
 	}
 
+	/**
+	 * Compute barycentric coordinates/weights of vPoint inside 3D triangle (V0,V1,V2). 
+	 * If point is in triangle plane and inside triangle, coords will be positive and sum to 1.
+	 * ie if result is a, then vPoint = a.x*V0 + a.y*V1 + a.z*V2.
+	 * TODO: make robust to degenerate triangles?
+	 */
+	template <typename RealType>
+	FVector3<RealType> BarycentricCoords(const FVector3<RealType>& Point, const FVector3<RealType>& V0, const FVector3<RealType>& V1, const FVector3<RealType>& V2)
+	{
+		FVector3<RealType> kV02 = V0 - V2;
+		FVector3<RealType> kV12 = V1 - V2;
+		FVector3<RealType> kPV2 = Point - V2;
+		RealType fM00 = kV02.Dot(kV02);
+		RealType fM01 = kV02.Dot(kV12);
+		RealType fM11 = kV12.Dot(kV12);
+		RealType fR0 = kV02.Dot(kPV2);
+		RealType fR1 = kV12.Dot(kPV2);
+		RealType fDet = fM00 * fM11 - fM01 * fM01;
+		RealType fInvDet = 1.0 / fDet;
+		RealType fBary1 = (fM11 * fR0 - fM01 * fR1) * fInvDet;
+		RealType fBary2 = (fM00 * fR1 - fM01 * fR0) * fInvDet;
+		RealType fBary3 = 1.0 - fBary1 - fBary2;
+		return new FVector3<RealType>(fBary1, fBary2, fBary3);
+	}
+
+	/**
+	* Compute barycentric coordinates/weights of vPoint inside 2D triangle (V0,V1,V2). 
+	* If point is inside triangle, coords will be positive and sum to 1.
+	* ie if result is a, then vPoint = a.x*V0 + a.y*V1 + a.z*V2.
+	* TODO: make robust to degenerate triangles?
+	*/
+	template <typename RealType>
+	FVector3<RealType> BarycentricCoords(const FVector3<RealType>& Point, const FVector2<RealType>& V0, const FVector2<RealType>& V1, const FVector2<RealType>& V2)
+	{
+		FVector2<RealType> kV02 = V0 - V2;
+		FVector2<RealType> kV12 = V1 - V2;
+		FVector2<RealType> kPV2 = Point - V2;
+		RealType fM00 = kV02.Dot(kV02);
+		RealType fM01 = kV02.Dot(kV12);
+		RealType fM11 = kV12.Dot(kV12);
+		RealType fR0 = kV02.Dot(kPV2);
+		RealType fR1 = kV12.Dot(kPV2);
+		RealType fDet = fM00 * fM11 - fM01 * fM01;
+		RealType fInvDet = 1.0 / fDet;
+		RealType fBary1 = (fM11 * fR0 - fM01 * fR1) * fInvDet;
+		RealType fBary2 = (fM00 * fR1 - fM01 * fR0) * fInvDet;
+		RealType fBary3 = 1.0 - fBary1 - fBary2;
+		return new FVector3<RealType>(fBary1, fBary2, fBary3);
+	}
+
 	/**
 	 * @return solid angle at point P for triangle A,B,C
 	 */
