Fix shadowed variable warning #rb none #rnx

diff --git a/Engine/Plugins/Experimental/GeometryProcessing/Source/DynamicMesh/Private/Operations/EmbedSurfacePath.cpp b/Engine/Plugins/Experimental/GeometryProcessing/Source/DynamicMesh/Private/Operations/EmbedSurfacePath.cpp
--- a/Engine/Plugins/Experimental/GeometryProcessing/Source/DynamicMesh/Private/Operations/EmbedSurfacePath.cpp
+++ b/Engine/Plugins/Experimental/GeometryProcessing/Source/DynamicMesh/Private/Operations/EmbedSurfacePath.cpp
@@ -541,9 +541,9 @@ bool FMeshSurfacePath::AddViaPlanarWalk(int StartTri, FVector3d StartPt, int End
 {
 	if (!VertexToPosnFn)
 	{
-		VertexToPosnFn = [](const FDynamicMesh3* Mesh, int VertexID)
+		VertexToPosnFn = [](const FDynamicMesh3* MeshArg, int VertexID)
 		{
-			return Mesh->GetVertex(VertexID);
+			return MeshArg->GetVertex(VertexID);
 		};
 	}
 	return WalkMeshPlanar(Mesh, StartTri, StartPt, EndTri, EndPt, WalkPlaneNormal, VertexToPosnFn, bAllowBackwardsSearch, AcceptEndPtOutsideDist, PtOnPlaneThresholdSq, Path);
