fix variable shadowing warning
#rnx

diff --git a/Engine/Plugins/Experimental/GeometryProcessing/Source/DynamicMesh/Private/Operations/EmbedSurfacePath.cpp b/Engine/Plugins/Experimental/GeometryProcessing/Source/DynamicMesh/Private/Operations/EmbedSurfacePath.cpp
--- a/Engine/Plugins/Experimental/GeometryProcessing/Source/DynamicMesh/Private/Operations/EmbedSurfacePath.cpp
+++ b/Engine/Plugins/Experimental/GeometryProcessing/Source/DynamicMesh/Private/Operations/EmbedSurfacePath.cpp
@@ -698,9 +698,9 @@ bool EmbedProjectedPath(FDynamicMesh3* Mesh, int StartTriID, FFrame3d Frame, con
 	int CurrentSeedTriID = StartTriID;
 	OutPathVertices.Reset();
 
-	TFunction<FVector3d(const FDynamicMesh3*, int)> ProjectToFrame = [&Frame](const FDynamicMesh3* Mesh, int VertexID)
+	TFunction<FVector3d(const FDynamicMesh3*, int)> ProjectToFrame = [&Frame](const FDynamicMesh3* MeshArg, int VertexID)
 	{
-		FVector2d ProjPt = Frame.ToPlaneUV(Mesh->GetVertex(VertexID));
+		FVector2d ProjPt = Frame.ToPlaneUV(MeshArg->GetVertex(VertexID));
 		return FVector3d(ProjPt.X, ProjPt.Y, 0);
 	};
 
