MeshProcessing: add FDynamicMeshChangeTracker::GetSavedTriangleList utility function
#rnx

diff --git a/Engine/Plugins/Experimental/GeometryProcessing/Source/DynamicMesh/Private/DynamicMeshChangeTracker.cpp b/Engine/Plugins/Experimental/GeometryProcessing/Source/DynamicMesh/Private/DynamicMeshChangeTracker.cpp
--- a/Engine/Plugins/Experimental/GeometryProcessing/Source/DynamicMesh/Private/DynamicMeshChangeTracker.cpp
+++ b/Engine/Plugins/Experimental/GeometryProcessing/Source/DynamicMesh/Private/DynamicMeshChangeTracker.cpp
@@ -667,4 +667,17 @@ TUniquePtr<FDynamicMeshChange> FDynamicMeshChangeTracker::EndChange()
 void FDynamicMeshChangeTracker::VerifySaveState()
 {
 	Change->VerifySaveState();
+}
+
+
+
+void FDynamicMeshChange::GetSavedTriangleList(TArray<int>& TrianglesOut, bool bInitial) const
+{
+	const TArray<FChangeTriangle>& UseList = (bInitial) ? OldTriangles : NewTriangles;
+	int32 N = UseList.Num();
+	TrianglesOut.Reserve(TrianglesOut.Num() + N);
+	for (int32 i = 0; i < N; ++i)
+	{
+		TrianglesOut.Add(UseList[i].TriangleID);
+	}
 }
\ No newline at end of file
diff --git a/Engine/Plugins/Experimental/GeometryProcessing/Source/DynamicMesh/Public/DynamicMeshChangeTracker.h b/Engine/Plugins/Experimental/GeometryProcessing/Source/DynamicMesh/Public/DynamicMeshChangeTracker.h
--- a/Engine/Plugins/Experimental/GeometryProcessing/Source/DynamicMesh/Public/DynamicMeshChangeTracker.h
+++ b/Engine/Plugins/Experimental/GeometryProcessing/Source/DynamicMesh/Public/DynamicMeshChangeTracker.h
@@ -138,6 +138,9 @@ public:
 	/** @return true if this vertex was saved. Uses linear search. */
 	bool HasSavedVertex(int VertexID);
 
+	/** store IDs of saved triangles in TrianglesOut. if bInitial=true, old triangles are stored, otherwise new triangles */
+	void GetSavedTriangleList(TArray<int>& TrianglesOut, bool bInitial) const;
+
 protected:
 
 	struct FChangeVertex
