MeshProcessing: add Remesher::OnEdgeFlip to notify subclasses/etc on edge flips
#rnx

diff --git a/Engine/Plugins/Experimental/GeometryProcessing/Source/DynamicMesh/Private/Remesher.cpp b/Engine/Plugins/Experimental/GeometryProcessing/Source/DynamicMesh/Private/Remesher.cpp
--- a/Engine/Plugins/Experimental/GeometryProcessing/Source/DynamicMesh/Private/Remesher.cpp
+++ b/Engine/Plugins/Experimental/GeometryProcessing/Source/DynamicMesh/Private/Remesher.cpp
@@ -271,6 +271,7 @@ abort_collapse:
 			EMeshResult result = Mesh->FlipEdge(edgeID, flipInfo);
 			if (result == EMeshResult::Ok) 
 			{
+				OnEdgeFlip(edgeID, flipInfo);
 				DoDebugChecks();
 				return EProcessResult::Ok_Flipped;
 			}
diff --git a/Engine/Plugins/Experimental/GeometryProcessing/Source/DynamicMesh/Public/Remesher.h b/Engine/Plugins/Experimental/GeometryProcessing/Source/DynamicMesh/Public/Remesher.h
--- a/Engine/Plugins/Experimental/GeometryProcessing/Source/DynamicMesh/Public/Remesher.h
+++ b/Engine/Plugins/Experimental/GeometryProcessing/Source/DynamicMesh/Public/Remesher.h
@@ -146,6 +146,13 @@ public:
 		// this is for subclasses...
 	}
 
+	/** Callback for subclasses to override to implement custom behavior */
+	virtual void OnEdgeFlip(int EdgeID, const FDynamicMesh3::FEdgeFlipInfo& FlipInfo)
+	{
+		// this is for subclasses...
+	}
+
+
 	/**
 	 * Number of edges that were modified in previous Remesh pass.
 	 * If this number gets small relative to edge count, you have probably converged (ish)
