MeshProcessing: add FMeshNormals::ComputeOverlayNormal() to compute single element normal (static func) and TDynamicMeshOverlay::TriangleHasElement()/GetElementTriangles() util functions to interrogate triangles connected to elements
#rnx

diff --git a/Engine/Plugins/Experimental/GeometryProcessing/Source/DynamicMesh/Private/DynamicMeshOverlay.cpp b/Engine/Plugins/Experimental/GeometryProcessing/Source/DynamicMesh/Private/DynamicMeshOverlay.cpp
--- a/Engine/Plugins/Experimental/GeometryProcessing/Source/DynamicMesh/Private/DynamicMeshOverlay.cpp
+++ b/Engine/Plugins/Experimental/GeometryProcessing/Source/DynamicMesh/Private/DynamicMeshOverlay.cpp
@@ -270,6 +270,25 @@ int TDynamicMeshOverlay<RealType, ElementSize>::CountVertexElements(int vid, boo
 
 
 
+template<typename RealType, int ElementSize>
+void TDynamicMeshOverlay<RealType, ElementSize>::GetElementTriangles(int ElementID, TArray<int>& OutTriangles) const
+{
+	check(ElementsRefCounts.IsValid(ElementID));
+	int VertexID = ParentVertices[ElementID];
+
+	for (int TriangleID : ParentMesh->VtxTrianglesItr(VertexID))
+	{
+		int i = 3 * TriangleID;
+		if (ElementTriangles[i] == ElementID || ElementTriangles[i+1] == ElementID || ElementTriangles[i+2] == ElementID)
+		{
+			OutTriangles.Add(TriangleID);
+		}
+	}
+}
+
+
+
+
 template<typename RealType, int ElementSize>
 void TDynamicMeshOverlay<RealType, ElementSize>::OnRemoveTriangle(int TriangleID, bool bRemoveIsolatedVertices)
 {
diff --git a/Engine/Plugins/Experimental/GeometryProcessing/Source/DynamicMesh/Private/MeshNormals.cpp b/Engine/Plugins/Experimental/GeometryProcessing/Source/DynamicMesh/Private/MeshNormals.cpp
--- a/Engine/Plugins/Experimental/GeometryProcessing/Source/DynamicMesh/Private/MeshNormals.cpp
+++ b/Engine/Plugins/Experimental/GeometryProcessing/Source/DynamicMesh/Private/MeshNormals.cpp
@@ -135,4 +135,24 @@ FVector3d FMeshNormals::ComputeVertexNormal(const FDynamicMesh3& Mesh, int VertI
 		SumNormal += Area * Normal;
 	}
 	return SumNormal.Normalized();
+}
+
+
+FVector3d FMeshNormals::ComputeOverlayNormal(const FDynamicMesh3& Mesh, FDynamicMeshNormalOverlay* NormalOverlay, int ElemIdx)
+{
+	int ParentVertexID = NormalOverlay->GetParentVertex(ElemIdx);
+	FVector3d SumNormal = FVector3d::Zero();
+	int Count = 0;
+	for (int TriIdx : Mesh.VtxTrianglesItr(ParentVertexID))
+	{
+		if (NormalOverlay->TriangleHasElement(TriIdx, ElemIdx))
+		{
+			FVector3d Normal, Centroid; double Area;
+			Mesh.GetTriInfo(TriIdx, Normal, Area, Centroid);
+			SumNormal += Area * Normal;
+			Count++;
+		}
+	}
+
+	return (Count > 0) ? SumNormal.Normalized() : FVector3d::Zero();
 }
\ No newline at end of file
diff --git a/Engine/Plugins/Experimental/GeometryProcessing/Source/DynamicMesh/Public/DynamicMeshOverlay.h b/Engine/Plugins/Experimental/GeometryProcessing/Source/DynamicMesh/Public/DynamicMeshOverlay.h
--- a/Engine/Plugins/Experimental/GeometryProcessing/Source/DynamicMesh/Public/DynamicMeshOverlay.h
+++ b/Engine/Plugins/Experimental/GeometryProcessing/Source/DynamicMesh/Public/DynamicMeshOverlay.h
@@ -195,6 +195,12 @@ public:
 		}
 	}
 
+	/** @return true if triangle contains element */
+	inline bool TriangleHasElement(int TriangleID, int ElementID) const
+	{
+		int i = 3 * TriangleID;
+		return (ElementTriangles[i] == ElementID || ElementTriangles[i+1] == ElementID || ElementTriangles[i+2] == ElementID);
+	}
 
 
 	/** Returns true if the parent-mesh edge is a "Seam" in this overlay */
@@ -207,6 +213,9 @@ public:
 	/** Count the number of unique elements for a given parent-mesh vertex */
 	int CountVertexElements(int VertexID, bool bBruteForce = false) const;
 
+	/** find the triangles connected to an element */
+	void GetElementTriangles(int ElementID, TArray<int>& OutTriangles) const;
+
 	/**
 	 * Checks that the overlay mesh is well-formed, ie all internal data structures are consistent
 	 */
diff --git a/Engine/Plugins/Experimental/GeometryProcessing/Source/DynamicMesh/Public/MeshNormals.h b/Engine/Plugins/Experimental/GeometryProcessing/Source/DynamicMesh/Public/MeshNormals.h
--- a/Engine/Plugins/Experimental/GeometryProcessing/Source/DynamicMesh/Public/MeshNormals.h
+++ b/Engine/Plugins/Experimental/GeometryProcessing/Source/DynamicMesh/Public/MeshNormals.h
@@ -98,11 +98,18 @@ public:
 
 
 	/**
-	 * @return the vertex normal at vertex VertIDx of Mesh
+	 * @return the vertex normal at vertex VertIdx of Mesh
 	 */
 	static FVector3d ComputeVertexNormal(const FDynamicMesh3& Mesh, int VertIdx);
 
 
+	/**
+	 * @return the computed overlay normal at an element of the overlay (ie based on normals of triangles connected to this element)
+	 */
+	static FVector3d ComputeOverlayNormal(const FDynamicMesh3& Mesh, FDynamicMeshNormalOverlay* NormalOverlay, int ElemIdx);
+
+
+
 protected:
 	/** Compute per-vertex normals using area-weighted averaging of one-ring triangles */
 	void Compute_FaceAvg_AreaWeighted();
