MeshProcessing: add FDynamicMeshChange::CheckValidity() to do internal data structure sanity checks
#rnx

diff --git a/Engine/Plugins/Experimental/GeometryProcessing/Source/DynamicMesh/Private/DynamicMeshChangeTracker.cpp b/Engine/Plugins/Experimental/GeometryProcessing/Source/DynamicMesh/Private/DynamicMeshChangeTracker.cpp
--- a/Engine/Plugins/Experimental/GeometryProcessing/Source/DynamicMesh/Private/DynamicMeshChangeTracker.cpp
+++ b/Engine/Plugins/Experimental/GeometryProcessing/Source/DynamicMesh/Private/DynamicMeshChangeTracker.cpp
@@ -168,7 +168,6 @@ template class DYNAMICMESH_API TDynamicMeshAttributeChange<int, 3>;
 
 FDynamicMeshChange::~FDynamicMeshChange()
 {
-	check(true);
 }
 
 
@@ -680,4 +679,55 @@ void FDynamicMeshChange::GetSavedTriangleList(TArray<int>& TrianglesOut, bool bI
 	{
 		TrianglesOut.Add(UseList[i].TriangleID);
 	}
+}
+
+
+
+void FDynamicMeshChange::CheckValidity(EValidityCheckFailMode FailMode) const
+{
+	bool is_ok = true;
+	TFunction<void(bool)> CheckOrFailF = [&](bool b)
+	{
+		is_ok = is_ok && b;
+	};
+	if (FailMode == EValidityCheckFailMode::Check)
+	{
+		CheckOrFailF = [&](bool b)
+		{
+			checkf(b, TEXT("FDynamicMeshChange::CheckValidity failed!"));
+			is_ok = is_ok && b;
+		};
+	}
+	else if (FailMode == EValidityCheckFailMode::Ensure)
+	{
+		CheckOrFailF = [&](bool b)
+		{
+			ensureMsgf(b, TEXT("FDynamicMeshChange::CheckValidity failed!"));
+			is_ok = is_ok && b;
+		};
+	}
+
+
+	TSet<int> SavedOldV, SavedNewV;
+	for (const FChangeVertex& ChangeVert : OldVertices)
+	{
+		SavedOldV.Add(ChangeVert.VertexID);
+	}
+	for (const FChangeVertex& ChangeVert : NewVertices)
+	{
+		SavedNewV.Add(ChangeVert.VertexID);
+	}
+
+	for (const FChangeTriangle& ChangedTri : OldTriangles)
+	{
+		CheckOrFailF(SavedOldV.Contains(ChangedTri.Vertices.A));
+		CheckOrFailF(SavedOldV.Contains(ChangedTri.Vertices.B));
+		CheckOrFailF(SavedOldV.Contains(ChangedTri.Vertices.C));
+	}
+	for (const FChangeTriangle& ChangedTri : NewTriangles)
+	{
+		CheckOrFailF(SavedNewV.Contains(ChangedTri.Vertices.A));
+		CheckOrFailF(SavedNewV.Contains(ChangedTri.Vertices.B));
+		CheckOrFailF(SavedNewV.Contains(ChangedTri.Vertices.C));
+	}
 }
\ No newline at end of file
diff --git a/Engine/Plugins/Experimental/GeometryProcessing/Source/DynamicMesh/Public/DynamicMeshChangeTracker.h b/Engine/Plugins/Experimental/GeometryProcessing/Source/DynamicMesh/Public/DynamicMeshChangeTracker.h
--- a/Engine/Plugins/Experimental/GeometryProcessing/Source/DynamicMesh/Public/DynamicMeshChangeTracker.h
+++ b/Engine/Plugins/Experimental/GeometryProcessing/Source/DynamicMesh/Public/DynamicMeshChangeTracker.h
@@ -141,6 +141,9 @@ public:
 	/** store IDs of saved triangles in TrianglesOut. if bInitial=true, old triangles are stored, otherwise new triangles */
 	void GetSavedTriangleList(TArray<int>& TrianglesOut, bool bInitial) const;
 
+	/** run self-validity checks on internal data structures to test if change is well-formed */
+	void CheckValidity(EValidityCheckFailMode FailMode = EValidityCheckFailMode::Check) const;
+
 protected:
 
 	struct FChangeVertex
