port mesh selection classes from geometry3sharp and hook up floodfill select to the projected path embedder, to optionally select the faces inside the cut path

#rnx

diff --git a/Engine/Plugins/Experimental/GeometryProcessing/Source/DynamicMesh/Private/Operations/EmbedSurfacePath.cpp b/Engine/Plugins/Experimental/GeometryProcessing/Source/DynamicMesh/Private/Operations/EmbedSurfacePath.cpp
--- a/Engine/Plugins/Experimental/GeometryProcessing/Source/DynamicMesh/Private/Operations/EmbedSurfacePath.cpp
+++ b/Engine/Plugins/Experimental/GeometryProcessing/Source/DynamicMesh/Private/Operations/EmbedSurfacePath.cpp
@@ -692,7 +692,7 @@ bool FMeshSurfacePath::EmbedSimplePath(bool bUpdatePath, TArray<int>& PathVertic
 	return true;
 }
 
-bool EmbedProjectedPath(FDynamicMesh3* Mesh, int StartTriID, FFrame3d Frame, const TArray<FVector2d>& Path2D, TArray<int>& OutPathVertices, bool bClosePath, double PtSnapVertexOrEdgeThresholdSq)
+bool EmbedProjectedPath(FDynamicMesh3* Mesh, int StartTriID, FFrame3d Frame, const TArray<FVector2d>& Path2D, TArray<int>& OutPathVertices, bool bClosePath, FMeshFaceSelection* EnclosedFaces, double PtSnapVertexOrEdgeThresholdSq)
 {
 	int32 EndIdxA = Path2D.Num() - (bClosePath ? 1 : 2);
 	int CurrentSeedTriID = StartTriID;
@@ -721,5 +721,49 @@ bool EmbedProjectedPath(FDynamicMesh3* Mesh, int StartTriID, FFrame3d Frame, con
 		CurrentSeedTriID = TrianglesOut[0];
 	}
 
+	if (EnclosedFaces)
+	{
+		TSet<int> Edges;
+		int32 NumEdges = bClosePath ? OutPathVertices.Num() : OutPathVertices.Num() - 1;
+		int32 SeedTriID = -1;
+		for (int32 IdxA = 0; IdxA < NumEdges; IdxA++)
+		{
+			int32 IdxB = (IdxA + 1) % OutPathVertices.Num();
+			int32 IDA = OutPathVertices[IdxA];
+			int32 IDB = OutPathVertices[IdxB];
+			int EID = Mesh->FindEdge(IDA, IDB);
+			Edges.Add(EID);
+
+			if (SeedTriID == -1)
+			{
+				FVector2d PA = Frame.ToPlaneUV(Mesh->GetVertex(IDA));
+				FVector2d PB = Frame.ToPlaneUV(Mesh->GetVertex(IDB));
+
+				FIndex2i OppVIDs = Mesh->GetEdgeOpposingV(EID);
+				double SignedAreaA = FTriangle2d::SignedArea(PA, PB, Frame.ToPlaneUV(Mesh->GetVertex(OppVIDs.A)));
+				if (SignedAreaA > FMathd::Epsilon)
+				{
+					SeedTriID = Mesh->GetEdgeT(EID).A;
+				}
+				else if (OppVIDs.B != FDynamicMesh3::InvalidID)
+				{
+					double SignedAreaB = FTriangle2d::SignedArea(PA, PB, Frame.ToPlaneUV(Mesh->GetVertex(OppVIDs.B)));
+					if (SignedAreaB > FMathd::Epsilon)
+					{
+						SeedTriID = Mesh->GetEdgeT(EID).B;
+					}
+				}
+			}
+		}
+		if (SeedTriID)
+		{
+			EnclosedFaces->FloodFill(SeedTriID, nullptr, [&Edges](int ID)
+			{
+				return !Edges.Contains(ID);
+			});
+		}
+		
+	}
+
 	return true;
 }
diff --git a/Engine/Plugins/Experimental/GeometryProcessing/Source/DynamicMesh/Private/Selections/MeshEdgeSelection.cpp b/Engine/Plugins/Experimental/GeometryProcessing/Source/DynamicMesh/Private/Selections/MeshEdgeSelection.cpp
new file mode 100644
--- /dev/null
+++ b/Engine/Plugins/Experimental/GeometryProcessing/Source/DynamicMesh/Private/Selections/MeshEdgeSelection.cpp
@@ -0,0 +1,67 @@
+// Copyright 1998-2019 Epic Games, Inc. All Rights Reserved.
+
+#include "Selections/MeshEdgeSelection.h"
+
+#include "DynamicMesh3.h"
+#include "Containers/Set.h"
+#include "Selections/MeshFaceSelection.h"
+#include "Selections/MeshVertexSelection.h"
+
+
+// convert vertex selection to edge selection. Require at least minCount verts of edge to be selected
+FMeshEdgeSelection::FMeshEdgeSelection(const FDynamicMesh3* mesh, const FMeshVertexSelection& convertV, int minCount) : Mesh(mesh)
+{
+	for (int eid : Mesh->EdgeIndicesItr())
+	{
+		FIndex2i ev = Mesh->GetEdgeV(eid);
+		int n = (convertV.IsSelected(ev.A) ? 1 : 0) +
+			(convertV.IsSelected(ev.B) ? 1 : 0);
+		if (n >= minCount)
+		{
+			add(eid);
+		}
+	}
+}
+
+// convert face selection to edge selection. Require at least minCount tris of edge to be selected
+FMeshEdgeSelection::FMeshEdgeSelection(const FDynamicMesh3* mesh, const FMeshFaceSelection& convertT, int minCount) : Mesh(mesh)
+{
+	minCount = FMathd::Clamp(minCount, 1, 2);
+
+	if (minCount == 1)
+	{
+		for (int tid : convertT)
+		{
+			FIndex3i te = Mesh->GetTriEdges(tid);
+			add(te.A); add(te.B); add(te.C);
+		}
+	}
+	else
+	{
+		for (int eid : Mesh->EdgeIndicesItr())
+		{
+			FIndex2i et = Mesh->GetEdgeT(eid);
+			if (convertT.IsSelected(et.A) && convertT.IsSelected(et.B))
+			{
+				add(eid);
+			}
+		}
+	}
+}
+
+void FMeshEdgeSelection::SelectBoundaryTriEdges(const FMeshFaceSelection& Triangles)
+{
+	for (int tid : Triangles)
+	{
+		FIndex3i te = Mesh->GetTriEdges(tid);
+		for ( int j = 0; j < 3; ++j )
+		{
+			FIndex2i et = Mesh->GetEdgeT(te[j]);
+			int other_tid = (et.A == tid) ? et.B : et.A;
+			if (Triangles.IsSelected(other_tid) == false)
+			{
+				add(te[j]);
+			}
+		}
+	}
+}
\ No newline at end of file
diff --git a/Engine/Plugins/Experimental/GeometryProcessing/Source/DynamicMesh/Private/Selections/MeshFaceSelection.cpp b/Engine/Plugins/Experimental/GeometryProcessing/Source/DynamicMesh/Private/Selections/MeshFaceSelection.cpp
new file mode 100644
--- /dev/null
+++ b/Engine/Plugins/Experimental/GeometryProcessing/Source/DynamicMesh/Private/Selections/MeshFaceSelection.cpp
@@ -0,0 +1,45 @@
+// Copyright 1998-2019 Epic Games, Inc. All Rights Reserved.
+
+#include "Selections/MeshFaceSelection.h"
+#include "Selections/MeshVertexSelection.h"
+#include "DynamicMesh3.h"
+
+
+// convert vertex selection to face selection. Require at least minCount verts of
+// tri to be selected (valid values are 1,2,3)
+FMeshFaceSelection::FMeshFaceSelection(const FDynamicMesh3* mesh, const FMeshVertexSelection& convertV, int minCount) : Mesh(mesh)
+{
+	minCount = FMathd::Clamp(minCount, 1, 3);
+
+	if (minCount == 1)
+	{
+		for ( int vid : convertV )
+		{
+			for (int tid : Mesh->VtxTrianglesItr(vid))
+			{
+				add(tid);
+			}
+		}
+	} else {
+		for (int tid : Mesh->TriangleIndicesItr()) {
+			FIndex3i tri = Mesh->GetTriangle(tid);
+			if (minCount == 3)
+			{
+				if (convertV.IsSelected(tri.A) && convertV.IsSelected(tri.B) && convertV.IsSelected(tri.C))
+				{
+					add(tid);
+				}
+			}
+			else
+			{
+				int n = (convertV.IsSelected(tri.A) ? 1 : 0) +
+					(convertV.IsSelected(tri.B) ? 1 : 0) +
+					(convertV.IsSelected(tri.C) ? 1 : 0);
+				if (n >= minCount)
+				{
+					add(tid);
+				}
+			}
+		}
+	}
+}
\ No newline at end of file
diff --git a/Engine/Plugins/Experimental/GeometryProcessing/Source/DynamicMesh/Private/Selections/MeshVertexSelection.cpp b/Engine/Plugins/Experimental/GeometryProcessing/Source/DynamicMesh/Private/Selections/MeshVertexSelection.cpp
new file mode 100644
--- /dev/null
+++ b/Engine/Plugins/Experimental/GeometryProcessing/Source/DynamicMesh/Private/Selections/MeshVertexSelection.cpp
@@ -0,0 +1,67 @@
+// Copyright 1998-2019 Epic Games, Inc. All Rights Reserved.
+
+#include "Selections/MeshVertexSelection.h"
+#include "Selections/MeshFaceSelection.h"
+#include "Selections/MeshEdgeSelection.h"
+
+#include "DynamicMesh3.h"
+
+// convert face selection to vertex selection. 
+FMeshVertexSelection::FMeshVertexSelection(const FDynamicMesh3* mesh, const FMeshFaceSelection& convertT) : Mesh(mesh)
+{
+	for (int tid : convertT) {
+		FIndex3i tv = Mesh->GetTriangle(tid);
+		add(tv.A); add(tv.B); add(tv.C);
+	}
+}
+
+// convert edge selection to vertex selection. 
+FMeshVertexSelection::FMeshVertexSelection(const FDynamicMesh3* mesh, const FMeshEdgeSelection& convertE) : Mesh(mesh)
+{
+	for (int eid : convertE) {
+		FIndex2i ev = Mesh->GetEdgeV(eid);
+		add(ev.A); add(ev.B);
+	}
+}
+
+void FMeshVertexSelection::SelectTriangleVertices(const FMeshFaceSelection& Triangles)
+{
+	for (int tid : Triangles)
+	{
+		FIndex3i tri = Mesh->GetTriangle(tid);
+		add(tri.A); add(tri.B); add(tri.C);
+	}
+}
+
+void FMeshVertexSelection::SelectInteriorVertices(const FMeshFaceSelection& Triangles)
+{
+	TSet<int> borderv;
+	for (int tid : Triangles) {
+		FIndex3i tv = Mesh->GetTriangle(tid);
+		for ( int j = 0; j < 3; ++j ) {
+			int vid = tv[j];
+			if (Selected.Contains(vid) || borderv.Contains(vid))
+			{
+				continue;
+			}
+			bool full_ring = true;
+			for (int ring_tid : Mesh->VtxTrianglesItr(vid))
+			{
+				if (Triangles.IsSelected(ring_tid) == false)
+				{
+					full_ring = false;
+					break;
+				}
+			}
+			if (full_ring)
+			{
+				add(vid);
+			}
+			else
+			{
+				borderv.Add(vid);
+			}
+		}
+	}
+}
+
diff --git a/Engine/Plugins/Experimental/GeometryProcessing/Source/DynamicMesh/Public/Operations/EmbedSurfacePath.h b/Engine/Plugins/Experimental/GeometryProcessing/Source/DynamicMesh/Public/Operations/EmbedSurfacePath.h
--- a/Engine/Plugins/Experimental/GeometryProcessing/Source/DynamicMesh/Public/Operations/EmbedSurfacePath.h
+++ b/Engine/Plugins/Experimental/GeometryProcessing/Source/DynamicMesh/Public/Operations/EmbedSurfacePath.h
@@ -10,7 +10,7 @@
 #include "FrameTypes.h"
 #include "GeometryTypes.h"
 #include "Curve/GeneralPolygon2.h"
-
+#include "Selections/MeshFaceSelection.h"
 
 class FDynamicMesh3;
 
@@ -115,7 +115,7 @@ public:
 	///**
 	// * Embed surface path in mesh, such that each point in path is a mesh vertex with a connecting edge.  Note that IsConnected must be true for this to succeed.
 	// *
-	// * @param bUpdatePath Updating the Path array with the new vertices (if false, the path will no longer be valid after running this function)\
+	// * @param bUpdatePath Updating the Path array with the new vertices (if false, the path will no longer be valid after running this function)
 	// * @param PathVertices Indices of the vertices on the path after embedding succeeds; NOTE these will not be 1:1 with the input Path
 	// * @param MeshGraphFn Function to create a mesh from an input 2D graph of labeled edges.
 	// * @return true if embedding succeeded.
@@ -137,5 +137,5 @@ public:
 /**
  * Embed a 2D path into a mesh by projection, starting the walk from a given triangle.
  */
-bool DYNAMICMESH_API EmbedProjectedPath(FDynamicMesh3* Mesh, int StartTriID, FFrame3d Frame, const TArray<FVector2d>& Path2D, TArray<int>& OutPathVertices, bool bClosePath, double PtSnapVertexOrEdgeThresholdSq = FMathd::Epsilon*100);
+bool DYNAMICMESH_API EmbedProjectedPath(FDynamicMesh3* Mesh, int StartTriID, FFrame3d Frame, const TArray<FVector2d>& Path2D, TArray<int>& OutPathVertices, bool bClosePath, FMeshFaceSelection *EnclosedFaces = nullptr, double PtSnapVertexOrEdgeThresholdSq = FMathd::Epsilon*100);
 
diff --git a/Engine/Plugins/Experimental/GeometryProcessing/Source/DynamicMesh/Public/Selections/MeshEdgeSelection.h b/Engine/Plugins/Experimental/GeometryProcessing/Source/DynamicMesh/Public/Selections/MeshEdgeSelection.h
new file mode 100644
--- /dev/null
+++ b/Engine/Plugins/Experimental/GeometryProcessing/Source/DynamicMesh/Public/Selections/MeshEdgeSelection.h
@@ -0,0 +1,173 @@
+// Copyright 1998-2019 Epic Games, Inc. All Rights Reserved.
+
+// Port of geometry3sharp MeshEdgeSelection
+
+#pragma once
+
+#include "DynamicMesh3.h"
+#include "CoreMinimal.h"
+#include "Containers/Set.h"
+#include "Selections/MeshVertexSelection.h"
+
+/**
+ * Currently a thin wrapper of a TSet<int> of Edge IDs paired with a Mesh; the backing storage will likely change as we need to optimize in the future
+ */
+class DYNAMICMESH_API FMeshEdgeSelection
+{
+private:
+	const FDynamicMesh3* Mesh;
+
+	TSet<int> Selected;
+
+public:
+
+	FMeshEdgeSelection(const FDynamicMesh3* mesh)
+	{
+		Mesh = mesh;
+	}
+	FMeshEdgeSelection(const FMeshEdgeSelection& copy)
+	{
+		Mesh = copy.Mesh;
+		Selected = copy.Selected;
+	}
+
+
+	// convert vertex selection to edge selection. Require at least minCount verts of edge to be selected
+	FMeshEdgeSelection(const FDynamicMesh3* mesh, const FMeshVertexSelection& convertV, int minCount = 2);
+
+	// convert face selection to edge selection. Require at least minCount tris of edge to be selected
+	FMeshEdgeSelection(const FDynamicMesh3* mesh, const FMeshFaceSelection& convertT, int minCount = 1);
+
+
+
+	TSet<int> AsSet() const
+	{
+		return Selected;
+	}
+	TArray<int> AsArray() const
+	{
+		return Selected.Array();
+	}
+	TBitArray<FDefaultBitArrayAllocator> AsBitArray() const
+	{
+		TBitArray<FDefaultBitArrayAllocator> Bitmap(false, Mesh->MaxEdgeID());
+		for (int tid : Selected)
+		{
+			Bitmap[tid] = true;
+		}
+		return Bitmap;
+	}
+
+public:
+	/**
+	* DO NOT USE DIRECTLY
+	* STL-like iterators to enable range-based for loop support.
+	*/
+	TSet<int>::TRangedForIterator      begin() { return Selected.begin(); }
+	TSet<int>::TRangedForConstIterator begin() const { return Selected.begin(); }
+	TSet<int>::TRangedForIterator      end() { return Selected.end(); }
+	TSet<int>::TRangedForConstIterator end() const { return Selected.end(); }
+
+private:
+	void add(int eid)
+	{
+		Selected.Add(eid);
+	}
+	void remove(int eid)
+	{
+		Selected.Remove(eid);
+	}
+
+
+public:
+	int Num()
+	{
+		return Selected.Num();
+	}
+
+
+
+	bool IsSelected(int eid)
+	{
+		return Selected.Contains(eid);
+	}
+
+
+	void Select(int eid)
+	{
+		ensure(Mesh->IsEdge(eid));
+		if (Mesh->IsEdge(eid))
+		{
+			add(eid);
+		}
+	}
+	
+	void Select(const TArray<int>& edges)
+	{
+		for (int eid : edges)
+		{
+			if (Mesh->IsEdge(eid))
+			{
+				add(eid);
+			}
+		}
+	}
+	void Select(TArrayView<const int> edges)
+	{
+		for (int eid : edges)
+		{
+			if (Mesh->IsEdge(eid))
+			{
+				add(eid);
+			}
+		}
+	}
+	void Select(TFunction<bool(int)> SelectF)
+	{
+		int NT = Mesh->MaxEdgeID();
+		for (int eid = 0; eid < NT; ++eid)
+		{
+			if (Mesh->IsEdge(eid) && SelectF(eid))
+			{
+				add(eid);
+			}
+		}
+	}
+
+	void SelectVertexEdges(TArrayView<const int> vertices)
+	{
+		for (int vid : vertices) {
+			for (int eid : Mesh->VtxEdgesItr(vid))
+			{
+				add(eid);
+			}
+		}
+	}
+
+	void SelectTriangleEdges(TArrayView<const int> Triangles)
+	{
+		for (int tid : Triangles)
+		{
+			FIndex3i et = Mesh->GetTriEdges(tid);
+			add(et.A); add(et.B); add(et.C);
+		}
+	}
+
+
+	void SelectBoundaryTriEdges(const FMeshFaceSelection& Triangles);
+
+	void Deselect(int tid) {
+		remove(tid);
+	}
+	void Deselect(TArrayView<const int> edges) {
+		for (int tid : edges)
+		{
+			remove(tid);
+		}
+	}
+	void DeselectAll()
+	{
+		Selected.Empty();
+	}
+
+};
\ No newline at end of file
diff --git a/Engine/Plugins/Experimental/GeometryProcessing/Source/DynamicMesh/Public/Selections/MeshFaceSelection.h b/Engine/Plugins/Experimental/GeometryProcessing/Source/DynamicMesh/Public/Selections/MeshFaceSelection.h
new file mode 100644
--- /dev/null
+++ b/Engine/Plugins/Experimental/GeometryProcessing/Source/DynamicMesh/Public/Selections/MeshFaceSelection.h
@@ -0,0 +1,681 @@
+// Copyright 1998-2019 Epic Games, Inc. All Rights Reserved.
+
+// Port of geometry3cpp MeshBoundaryLoops
+
+#pragma once
+
+#include "DynamicMesh3.h"
+#include "Util/DynamicVector.h"
+#include "EdgeLoop.h"
+
+class FMeshVertexSelection;
+class FMeshEdgeSelection;
+
+
+class DYNAMICMESH_API FMeshFaceSelection
+{
+private:
+	const FDynamicMesh3* Mesh;
+
+	TSet<int> Selected;
+
+public:
+
+	FMeshFaceSelection(const FDynamicMesh3* mesh)
+	{
+		Mesh = mesh;
+	}
+	FMeshFaceSelection(const FMeshFaceSelection& copy)
+	{
+		Mesh = copy.Mesh;
+		Selected = copy.Selected;
+	}
+
+	// convert vertex selection to face selection. Require at least minCount verts of
+	// tri to be selected (valid values are 1,2,3)
+	FMeshFaceSelection(const FDynamicMesh3* mesh, const FMeshVertexSelection& convertV, int minCount = 3);
+
+	// select a group
+	FMeshFaceSelection(const FDynamicMesh3* mesh, int group_id) : Mesh(mesh)
+	{
+		SelectGroup(group_id);
+	}
+
+	/**
+	* DO NOT USE DIRECTLY
+	* STL-like iterators to enable range-based for loop support.
+	*/
+	TSet<int>::TRangedForIterator      begin() { return Selected.begin(); }
+	TSet<int>::TRangedForConstIterator begin() const { return Selected.begin(); }
+	TSet<int>::TRangedForIterator      end() { return Selected.end(); }
+	TSet<int>::TRangedForConstIterator end() const { return Selected.end(); }
+
+
+private:
+
+	void add(int tid)
+	{
+		Selected.Add(tid);
+	}
+	void remove(int tid)
+	{
+		Selected.Remove(tid);
+	}
+
+public:
+
+
+	int Num() const
+	{
+		return Selected.Num();
+	}
+
+
+
+	bool IsSelected(int tid) const
+	{
+		return Selected.Contains(tid);
+	}
+
+
+	void Select(int tid)
+	{
+		ensure(Mesh->IsTriangle(tid));
+		if (Mesh->IsTriangle(tid))
+		{
+			add(tid);
+		}
+	}
+	void Select(TArrayView<const int> Triangles)
+	{
+		for (int tID : Triangles)
+		{
+			if (Mesh->IsTriangle(tID))
+			{
+				add(tID);
+			}
+		}
+	}
+	void Select(TFunction<bool(int)> SelectF)
+	{
+		int NT = Mesh->MaxTriangleID();
+		for (int tID = 0; tID < NT; ++tID)
+		{
+			if (Mesh->IsTriangle(tID) && SelectF(tID))
+			{
+				add(tID);
+			}
+		}
+	}
+
+
+	void SelectVertexOneRing(int vid) {
+		for (int tid : Mesh->VtxTrianglesItr(vid))
+		{
+			add(tid);
+		}
+	}
+	void SelectVertexOneRings(TArrayView<const int> Vertices)
+	{
+		for (int vid : Vertices)
+		{
+			for (int tid : Mesh->VtxTrianglesItr(vid))
+			{
+				add(tid);
+			}
+		}
+	}
+
+	void SelectEdgeTris(int eid)
+	{
+		FIndex2i et = Mesh->GetEdgeT(eid);
+		add(et.A);
+		if (et.B != FDynamicMesh3::InvalidID)
+		{
+			add(et.B);
+		}
+	}
+
+
+	void Deselect(int tid)
+	{
+		remove(tid);
+	}
+	void Deselect(TArrayView<const int> Triangles)
+	{
+		for (int TID : Triangles)
+		{
+			remove(TID);
+		}
+	}
+
+	void DeselectAll()
+	{
+		Selected.Empty();
+	}
+
+
+	void SelectGroup(int gid)
+	{
+		int NT = Mesh->MaxTriangleID();
+		for (int tid = 0; tid < NT; ++tid)
+		{
+			if (Mesh->IsTriangle(tid) && Mesh->GetTriangleGroup(tid) == gid)
+			{
+				add(tid);
+			}
+		}
+	}
+	void SelectGroupInverse(int gid)
+	{
+		int NT = Mesh->MaxTriangleID();
+		for (int tid = 0; tid < NT; ++tid)
+		{
+			if (Mesh->IsTriangle(tid) && Mesh->GetTriangleGroup(tid) != gid)
+			{
+				add(tid);
+			}
+		}
+	}
+	void DeselectGroup(int gid)
+	{
+		// cannot just iterate over selected tris because remove() will change them...
+		int NT = Mesh->MaxTriangleID();
+		for (int tid = 0; tid < NT; ++tid)
+		{
+			if (Mesh->IsTriangle(tid) && Mesh->GetTriangleGroup(tid) == gid)
+			{
+				remove(tid);
+			}
+		}
+	}
+
+	TSet<int> AsSet() const
+	{
+		return Selected;
+	}
+	TArray<int> AsArray() const
+	{
+		return Selected.Array();
+	}
+	TBitArray<FDefaultBitArrayAllocator> AsBitArray() const
+	{
+		TBitArray<FDefaultBitArrayAllocator> Bitmap(false, Mesh->MaxTriangleID());
+		for (int tid : Selected)
+		{
+			Bitmap[tid] = true;
+		}
+		return Bitmap;
+	}
+
+
+
+	/**
+	 *  find set of tris just outside border of selection
+	 */
+	TArray<int> FindNeighbourTris() const
+	{
+		TArray<int> result;
+		for (int tid : Selected)
+		{
+			FIndex3i nbr_tris = Mesh->GetTriNeighbourTris(tid);
+			for (int j = 0; j < 3; ++j)
+			{
+				if (nbr_tris[j] != FDynamicMesh3::InvalidID && IsSelected(nbr_tris[j]) == false)
+				{
+					result.Add(nbr_tris[j]);
+				}
+			}
+		}
+		return result;
+	}
+
+
+	/**
+	 *  find set of tris just inside border of selection
+	 */
+	TArray<int> FindBorderTris() const
+	{
+		TArray<int> result;
+		for (int tid : Selected)
+		{
+			FIndex3i nbr_tris = Mesh->GetTriNeighbourTris(tid);
+			if (IsSelected(nbr_tris.A) == false || IsSelected(nbr_tris.B) == false || IsSelected(nbr_tris.C) == false)
+			{
+				result.Add(tid);
+			}
+		}
+		return result;
+	}
+
+
+	void ExpandToFaceNeighbours(TFunction<bool(int)> FilterF = nullptr)
+	{
+		TArray<int> ToAdd;
+
+		for (int tid : Selected) { 
+			FIndex3i nbr_tris = Mesh->GetTriNeighbourTris(tid);
+			for (int j = 0; j < 3; ++j)
+			{
+				if (FilterF != nullptr && FilterF(nbr_tris[j]) == false)
+				{
+					continue;
+				}
+				if (nbr_tris[j] != FDynamicMesh3::InvalidID && !IsSelected(nbr_tris[j]))
+				{
+					ToAdd.Add(nbr_tris[j]);
+				}
+			}
+		}
+
+		for (int ID : ToAdd)
+		{
+			add(ID);
+		}
+	}
+	void ExpandToFaceNeighbours(int rounds, TFunction<bool(int)> FilterF = nullptr)
+	{
+		for (int k = 0; k < rounds; ++k)
+		{
+			ExpandToFaceNeighbours(FilterF);
+		}
+	}
+
+
+	/**
+	 *  Add all triangles in vertex one-rings of current selection to set.
+	 *  On a large mesh this is quite expensive as we don't know the boundary,
+	 *  so we have to iterate over all triangles.
+	 *  
+	 *  Return false from FilterF to prevent triangles from being included.
+	 */
+	void ExpandToOneRingNeighbours(TFunction<bool(int)> FilterF = nullptr)
+	{
+		TArray<int> ToAdd;
+
+		for (int tid : Selected) { 
+			FIndex3i tri_v = Mesh->GetTriangle(tid);
+			for (int j = 0; j < 3; ++j)
+			{
+				int vid = tri_v[j];
+				for (int nbr_t : Mesh->VtxTrianglesItr(vid))
+				{
+					if (FilterF != nullptr && FilterF(nbr_t) == false)
+					{
+						continue;
+					}
+					if (!IsSelected(nbr_t))
+					{
+						ToAdd.Add(nbr_t);
+					}
+				}
+			}
+		}
+
+		for (int ID : ToAdd)
+		{
+			add(ID);
+		}
+	}
+
+
+	/**
+	 *  Expand selection by N vertex one-rings. This is *significantly* faster
+	 *  than calling ExpandToOnering() multiple times, because we can track
+	 *  the growing front and only check the new triangles.
+	 *  
+	 *  Return false from FilterF to prevent triangles from being included.
+	 */
+	void ExpandToOneRingNeighbours(int nRings, TFunction<bool(int)> FilterF = nullptr)
+	{
+		if (nRings == 1)
+		{
+			ExpandToOneRingNeighbours(FilterF);
+			return;
+		}
+
+		TArray<int> triArrays[2];
+		int addIdx = 0, checkIdx = 1;
+		triArrays[checkIdx] = Selected.Array();
+
+		TBitArray<FDefaultBitArrayAllocator> Bitmap = AsBitArray();
+
+		for (int ri = 0; ri < nRings; ++ri)
+		{
+			TArray<int>& addTris = triArrays[addIdx];
+			TArray<int>& checkTris = triArrays[checkIdx];
+			addTris.Empty();
+
+			for (int tid : checkTris)
+			{
+				FIndex3i tri_v = Mesh->GetTriangle(tid);
+				for (int j = 0; j < 3; ++j)
+				{
+					int vid = tri_v[j];
+					for (int nbr_t : Mesh->VtxTrianglesItr(vid))
+					{
+						if (FilterF != nullptr && FilterF(nbr_t) == false)
+						{
+							continue;
+						}
+						if (Bitmap[nbr_t] == false)
+						{
+							addTris.Add(nbr_t);
+							Bitmap[nbr_t] = true;
+						}
+					}
+				}
+			}
+
+			for (int TID : addTris)
+			{
+				add(TID);
+			}
+
+			Swap(addIdx, checkIdx); // check in the next iter what we added in this iter
+		}
+	}
+
+
+
+
+
+	/**
+	 *  remove all triangles in vertex one-rings of current selection to set.
+	 *  On a large mesh this is quite expensive as we don't know the boundary,
+	 *  so we have to iterate over all triangles.
+	 *  
+	 *  Return false from FilterF to prevent triangles from being deselected.
+	 */
+	void ContractBorderByOneRingNeighbours()
+	{
+		TArray<int> BorderVIDs;   // border vertices
+
+		// [TODO] border vertices are pushed onto the temp list multiple times.
+		// minor inefficiency, but maybe we could improve it?
+
+		// find set of vertices on border
+		for (int tid : Selected)
+		{
+			FIndex3i tri_v = Mesh->GetTriangle(tid);
+			for (int j = 0; j < 3; ++j)
+			{
+				int vid = tri_v[j];
+				for (int nbr_t : Mesh->VtxTrianglesItr(vid))
+				{
+					if (!IsSelected(nbr_t))
+					{
+						BorderVIDs.Add(vid);
+						break;
+					}
+				}
+			}
+		}
+
+		for (int vid : BorderVIDs)
+		{
+			for (int nbr_t : Mesh->VtxTrianglesItr(vid))
+			{
+				Deselect(nbr_t);
+			}
+		}
+
+	}
+
+
+
+
+	/**
+	 *  Grow selection outwards from seed triangle, until it hits boundaries defined by triangle and edge filters.
+	 *  Edge filter is not effective unless it (possibly combined w/ triangle filter) defines closed loops.
+	 */
+	void FloodFill(int tSeed, TFunction<bool(int)> TriFilterF = nullptr, TFunction<bool(int)> EdgeFilterF = nullptr)
+	{
+		TArray<int> Seeds = { tSeed };
+		FloodFill(Seeds, TriFilterF, EdgeFilterF);
+	}
+	/**
+	 *  Grow selection outwards from seed triangles, until it hits boundaries defined by triangle and edge filters.
+	 *  Edge filter is not effective unless it (possibly combined w/ triangle filter) defines closed loops.
+	 */
+	void FloodFill(const TArray<int>& Seeds, TFunction<bool(int)> TriFilterF = nullptr, TFunction<bool(int)> EdgeFilterF = nullptr)
+	{
+		TDynamicVector<int> stack(Seeds);
+		for (int Seed : Seeds)
+		{
+			add(Seed);
+		}
+		while (stack.Num() > 0)
+		{
+			int TID = stack.Back();
+			stack.PopBack();
+
+			FIndex3i nbrs = Mesh->GetTriNeighbourTris(TID);
+			for (int j = 0; j < 3; ++j)
+			{
+				int nbr_tid = nbrs[j];
+				if (nbr_tid == FDynamicMesh3::InvalidID || IsSelected(nbr_tid))
+				{
+					continue;
+				}
+				if (TriFilterF != nullptr && TriFilterF(nbr_tid) == false)
+				{
+					continue;
+				}
+				int EID = Mesh->GetTriEdge(TID, j);
+				if (EdgeFilterF != nullptr && EdgeFilterF(EID) == false)
+				{
+					continue;
+				}
+				add(nbr_tid);
+
+				stack.Add(nbr_tid);
+			}
+		}
+	}
+
+
+
+	// return true if we clipped something
+	bool ClipFins(bool bClipLoners)
+	{
+		TArray<int> ToRemove; // temp array so we don't remove as we iterate the set
+		for (int tid : Selected)
+		{
+			if (is_fin(tid, bClipLoners))
+			{
+				ToRemove.Add(tid);
+			}
+		}
+		for (int tid : ToRemove)
+		{
+			remove(tid);
+		}
+		return ToRemove.Num() > 0;
+	}
+
+
+	// return true if we filled any ears.
+	bool FillEars(bool bFillTinyHoles)
+	{
+		// [TODO] not efficient! checks each nbr 3 times !! ugh!!
+		TArray<int> temp;
+		for (int tid : Selected)
+		{
+			FIndex3i nbr_tris = Mesh->GetTriNeighbourTris(tid);
+			for (int j = 0; j < 3; ++j)
+			{
+				int nbr_t = nbr_tris[j];
+				if (IsSelected(nbr_t))
+				{
+					continue;
+				}
+				if (is_ear(nbr_t, bFillTinyHoles))
+				{
+					temp.Add(nbr_t);
+				}
+			}
+		}
+		for (int tid : temp)
+		{
+			add(tid);
+		}
+		return temp.Num() > 0;
+	}
+
+	// returns true if selection was modified
+	bool LocalOptimize(bool bClipFins, bool bFillEars, bool bFillTinyHoles = true, bool bClipLoners = true, bool bRemoveBowties = false)
+	{
+		bool bModified = false;
+		bool done = false;
+		int count = 0;
+		do
+		{
+			bool bDidClip = bClipFins && ClipFins(bClipLoners);
+			bool bDidEars = bFillEars && FillEars(bFillTinyHoles);
+			bool bDidBows = bRemoveBowties && remove_bowties();
+			done = !(bDidClip || bDidEars || bDidBows); // only done if did nothing this pass
+			bModified = bModified || !done;
+		} while (!done && count++ < 25);
+		if (bRemoveBowties)
+		{
+			remove_bowties();        // do a final pass of this because it is usually the most problematic...
+		}
+		return bModified;
+	}
+
+	bool LocalOptimize(bool bRemoveBowties = true)
+	{
+		return LocalOptimize(true, true, true, true, bRemoveBowties);
+	}
+
+
+
+
+	/**
+	 *  Find any "bowtie" vertices - ie vertex v such taht there is multiple spans of triangles
+	 *  selected in v's triangle one-ring - and deselect those one-rings.
+	 *  Returns true if selection was modified.
+	 */
+	bool RemoveBowties() {
+		return remove_bowties();
+	}
+	bool remove_bowties()
+	{
+		TSet<int> tempHash;
+		bool bModified = false;
+		bool done = false;
+		TSet<int> vertices;
+		while (!done)
+		{
+			done = true;
+			vertices.Empty();
+			for (int tid : Selected)
+			{
+				FIndex3i tv = Mesh->GetTriangle(tid);
+				vertices.Add(tv.A); vertices.Add(tv.B); vertices.Add(tv.C);
+			}
+
+			for (int vid : vertices) {
+				if (is_bowtie_vtx(vid)) {
+					for (int TID : Mesh->VtxTrianglesItr(vid))
+					{
+						Deselect(TID);
+					}
+					done = false;
+				}
+			}
+			if (done == false)
+			{
+				bModified = true;
+			}
+		}
+		return bModified;
+	}
+
+private:
+	bool is_bowtie_vtx(int vid) const
+	{
+		int border_edges = 0;
+		for (int eid : Mesh->VtxEdgesItr(vid)) {
+			FIndex2i et = Mesh->GetEdgeT(eid);
+			if (et.B != FDynamicMesh3::InvalidID)
+			{
+				bool in_a = IsSelected(et.A);
+				bool in_b = IsSelected(et.B);
+				if (in_a != in_b)
+				{
+					border_edges++;
+				}
+			}
+			else
+			{
+				if (IsSelected(et.A))
+				{
+					border_edges++;
+				}
+			}
+		}
+		return border_edges > 2;
+	}
+
+	void count_nbrs(int tid, int& nbr_in, int& nbr_out, int& bdry_e) const
+	{
+		FIndex3i nbr_tris = Mesh->GetTriNeighbourTris(tid);
+		nbr_in = 0; nbr_out = 0; bdry_e = 0;
+		for ( int j = 0; j < 3; ++j )
+		{
+			int nbr_t = nbr_tris[j];
+			if (nbr_t == FDynamicMesh3::InvalidID)
+			{
+				bdry_e++;
+			}
+			else if (IsSelected(nbr_t) == true)
+			{
+				nbr_in++;
+			}
+			else
+			{
+				nbr_out++;
+			}
+		}
+	}
+	bool is_ear(int tid, bool include_tiny_holes) const
+	{
+		if (IsSelected(tid) == true)
+		{
+			return false;
+		}
+		int nbr_in, nbr_out, bdry_e;
+		count_nbrs(tid, nbr_in, nbr_out, bdry_e);
+		if (bdry_e == 2 && nbr_in == 1)
+		{
+			return true;        // unselected w/ 2 boundary edges, nbr is  in
+		}
+		else if (nbr_in == 2)
+		{
+			if (bdry_e == 1 || nbr_out == 1)
+			{
+				return true;        // unselected w/ 2 selected nbrs
+			}
+		}
+		else if (include_tiny_holes && nbr_in == 3)
+		{
+			return true;
+		}
+		return false;
+	}
+	bool is_fin(int tid, bool include_loners) const
+	{
+		if (IsSelected(tid) == false)
+		{
+			return false;
+		}
+		int nbr_in, nbr_out, bdry_e;
+		count_nbrs(tid, nbr_in, nbr_out, bdry_e);
+		return (nbr_in == 1 && nbr_out == 2) ||
+			   (include_loners == true && nbr_in == 0 && nbr_out == 3);
+	}
+
+
+};
diff --git a/Engine/Plugins/Experimental/GeometryProcessing/Source/DynamicMesh/Public/Selections/MeshVertexSelection.h b/Engine/Plugins/Experimental/GeometryProcessing/Source/DynamicMesh/Public/Selections/MeshVertexSelection.h
new file mode 100644
--- /dev/null
+++ b/Engine/Plugins/Experimental/GeometryProcessing/Source/DynamicMesh/Public/Selections/MeshVertexSelection.h
@@ -0,0 +1,263 @@
+// Copyright 1998-2019 Epic Games, Inc. All Rights Reserved.
+
+// Port of geometry3sharp FMeshVertexSelection
+
+#pragma once
+
+#include "DynamicMesh3.h"
+
+class FMeshFaceSelection;
+class FMeshEdgeSelection;
+
+class DYNAMICMESH_API FMeshVertexSelection
+{
+private:
+	const FDynamicMesh3* Mesh;
+
+	TSet<int> Selected;
+
+public:
+	FMeshVertexSelection(const FDynamicMesh3* mesh)
+	{
+		Mesh = mesh;
+	}
+
+	// convert face selection to vertex selection. 
+	FMeshVertexSelection(const FDynamicMesh3* mesh, const FMeshFaceSelection& convertT);
+
+	// convert edge selection to vertex selection. 
+	FMeshVertexSelection(const FDynamicMesh3* mesh, const FMeshEdgeSelection& convertE);
+
+
+	TSet<int> AsSet() const
+	{
+		return Selected;
+	}
+	TArray<int> AsArray() const
+	{
+		return Selected.Array();
+	}
+	TBitArray<FDefaultBitArrayAllocator> AsBitArray() const
+	{
+		TBitArray<FDefaultBitArrayAllocator> Bitmap(false, Mesh->MaxVertexID());
+		for (int tid : Selected)
+		{
+			Bitmap[tid] = true;
+		}
+		return Bitmap;
+	}
+
+
+
+public:
+	/**
+	* DO NOT USE DIRECTLY
+	* STL-like iterators to enable range-based for loop support.
+	*/
+	TSet<int>::TRangedForIterator      begin() { return Selected.begin(); }
+	TSet<int>::TRangedForConstIterator begin() const { return Selected.begin(); }
+	TSet<int>::TRangedForIterator      end() { return Selected.end(); }
+	TSet<int>::TRangedForConstIterator end() const { return Selected.end(); }
+
+private:
+	void add(int vID)
+	{
+		Selected.Add(vID);
+	}
+	void remove(int vID)
+	{
+		Selected.Remove(vID);
+	}
+
+public:
+	int Num() const
+	{
+		return Selected.Num();
+	}
+
+	bool IsSelected(int vID) const
+	{
+		return Selected.Contains(vID);
+	}
+
+	void Select(int vID)
+	{
+		ensure(Mesh->IsVertex(vID));
+		if (Mesh->IsVertex(vID))
+		{
+			add(vID);
+		}
+	}
+	void Select(TArrayView<const int> Vertices)
+	{
+		for (int VID : Vertices)
+		{
+			if (Mesh->IsVertex(VID))
+			{
+				add(VID);
+			}
+		}
+	}
+
+	void SelectTriangleVertices(TArrayView<const int> Triangles)
+	{
+		for (int TID : Triangles)
+		{
+			FIndex3i tri = Mesh->GetTriangle(TID);
+			add(tri.A); add(tri.B); add(tri.C);
+		}
+	}
+	void SelectTriangleVertices(const FMeshFaceSelection& Triangles);
+
+
+	/**
+	 *  for each vertex of input triangle set, select vertex if all
+	 *  one-ring triangles are contained in triangle set (ie vertex is not on boundary of triangle set).
+	 */
+	void SelectInteriorVertices(const FMeshFaceSelection& triangles);
+
+
+	/**
+	 *  Select set of boundary vertices connected to vSeed.
+	 */
+	void SelectConnectedBoundaryV(int vSeed)
+	{
+		if (!ensureMsgf(Mesh->IsBoundaryVertex(vSeed), TEXT("MeshConnectedComponents.FindConnectedBoundaryV: vSeed is not a boundary vertex")))
+		{
+			return;
+		}
+
+		TSet<int> &found = Selected;
+		found.Add(vSeed);
+		TArray<int> queue;
+		queue.Add(vSeed);
+		while (queue.Num() > 0) {
+			int vid = queue.Pop(false);
+			for (int nbrid : Mesh->VtxVerticesItr(vid)) {
+				if (Mesh->IsBoundaryVertex(nbrid) && found.Contains(nbrid) == false)
+				{
+					found.Add(nbrid);
+					queue.Add(nbrid);
+				}
+			}
+		}
+	}
+
+
+	void SelectEdgeVertices(TArrayView<const int> Edges)
+	{
+		for (int EID : Edges)
+		{
+			FIndex2i ev = Mesh->GetEdgeV(EID);
+			add(ev.A); add(ev.B);
+		}
+	}
+
+
+	void Deselect(int vID)
+	{
+		remove(vID);
+	}
+	void Deselect(TArrayView<const int> Vertices)
+	{
+		for (int VID : Vertices)
+		{
+			remove(VID);
+		}
+	}
+	void DeselectEdge(int eid)
+	{
+		FIndex2i ev = Mesh->GetEdgeV(eid);
+		remove(ev.A); remove(ev.B);
+	}
+
+	void DeselectEdges(TArrayView<const int> Edges) {
+		for (int EID : Edges)
+		{
+			FIndex2i ev = Mesh->GetEdgeV(EID);
+			remove(ev.A); remove(ev.B);
+		}
+	}
+
+
+	/**
+	 *  Add all one-ring neighbours of current selection to set.
+	 *  On a large mesh this is quite expensive as we don't know the boundary,
+	 *  so we have to iterate over all triangles.
+	 *  
+	 *  Return false from FilterF to prevent vertices from being included.
+	 */
+	void ExpandToOneRingNeighbours(TFunction<bool(int)> FilterF = nullptr)
+	{
+		TArray<int> temp;
+
+		for (int vid : Selected) {
+			for (int nbr_vid : Mesh->VtxVerticesItr(vid))
+			{
+				if (FilterF != nullptr && FilterF(nbr_vid) == false)
+				{
+					continue;
+				}
+				if (IsSelected(nbr_vid) == false)
+				{
+					temp.Add(nbr_vid);
+				}
+			}
+		}
+
+		for (int ID : temp)
+		{
+			add(ID);
+		}
+	}
+
+
+	// [TODO] should do this more efficiently, like FMeshFaceSelection
+	void ExpandToOneRingNeighbours(int nRings, TFunction<bool(int)> FilterF = nullptr)
+	{
+		for (int k = 0; k < nRings; ++k)
+		{
+			ExpandToOneRingNeighbours(FilterF);
+		}
+	}
+
+
+
+
+
+	/**
+	 *  Grow selection outwards from seed vertex, until it hits boundaries defined by vertex filter.
+	 */
+	void FloodFill(int vSeed, TFunction<bool(int)> VertIncludedF = nullptr)
+	{
+		FloodFill({ vSeed }, VertIncludedF);
+	}
+	/**
+	 *  Grow selection outwards from seed vertex, until it hits boundaries defined by vertex filter.
+	 */
+	void FloodFill(TArray<int> Seeds, TFunction<bool(int)> VertIncludedF = nullptr)
+	{
+		TDynamicVector<int> stack(Seeds);
+		for (int Seed : Seeds)
+		{
+			add(Seed);
+		}
+		while (stack.Num() > 0)
+		{
+			int vID = stack.Back();
+			stack.PopBack();
+
+			for (int nbr_vid : Mesh->VtxVerticesItr(vID))
+			{
+				if (IsSelected(nbr_vid) == true || VertIncludedF(nbr_vid) == false)
+				{
+					continue;
+				}
+				add(nbr_vid);
+				stack.Add(nbr_vid);
+			}
+		}
+	}
+
+
+};
\ No newline at end of file
diff --git a/Engine/Plugins/Experimental/GeometryProcessing/Source/GeometricObjects/Public/Util/DynamicVector.h b/Engine/Plugins/Experimental/GeometryProcessing/Source/GeometricObjects/Public/Util/DynamicVector.h
--- a/Engine/Plugins/Experimental/GeometryProcessing/Source/GeometricObjects/Public/Util/DynamicVector.h
+++ b/Engine/Plugins/Experimental/GeometryProcessing/Source/GeometricObjects/Public/Util/DynamicVector.h
@@ -37,6 +37,23 @@ public:
 		*this = MoveTemp(Moved);
 	}
 
+	TDynamicVector(const TArray<Type>& Array)
+	{
+		SetNum(Array.Num());
+		for (int Idx = 0; Idx < Array.Num(); Idx++)
+		{
+			(*this)[Idx] = Array[Idx];
+		}
+	}
+	TDynamicVector(TArrayView<const Type> Array)
+	{
+		SetNum(Array.Num());
+		for (int Idx = 0; Idx < Array.Num(); Idx++)
+		{
+			(*this)[Idx] = Array[Idx];
+		}
+	}
+
 	~TDynamicVector() {}
 
 	const TDynamicVector& operator=(const TDynamicVector& Copy);
