Fix infinite recursion in vertex select flood fill

#rnx

diff --git a/Engine/Plugins/Experimental/GeometryProcessing/Source/DynamicMesh/Public/Selections/MeshVertexSelection.h b/Engine/Plugins/Experimental/GeometryProcessing/Source/DynamicMesh/Public/Selections/MeshVertexSelection.h
--- a/Engine/Plugins/Experimental/GeometryProcessing/Source/DynamicMesh/Public/Selections/MeshVertexSelection.h
+++ b/Engine/Plugins/Experimental/GeometryProcessing/Source/DynamicMesh/Public/Selections/MeshVertexSelection.h
@@ -230,12 +230,13 @@ public:
 	 */
 	void FloodFill(int vSeed, TFunction<bool(int)> VertIncludedF = nullptr)
 	{
-		FloodFill({ vSeed }, VertIncludedF);
+		TArray<int> Seeds = { vSeed };
+		FloodFill(Seeds, VertIncludedF);
 	}
 	/**
 	 *  Grow selection outwards from seed vertex, until it hits boundaries defined by vertex filter.
 	 */
-	void FloodFill(TArray<int> Seeds, TFunction<bool(int)> VertIncludedF = nullptr)
+	void FloodFill(const TArray<int>& Seeds, TFunction<bool(int)> VertIncludedF = nullptr)
 	{
 		TDynamicVector<int> stack(Seeds);
 		for (int Seed : Seeds)
