MeshProcessing: don't call ExpandIterator expand function if we are already at end (iterator value is invalid)
#rnx

diff --git a/Engine/Plugins/Experimental/GeometryProcessing/Source/GeometricObjects/Public/Util/IteratorUtil.h b/Engine/Plugins/Experimental/GeometryProcessing/Source/GeometricObjects/Public/Util/IteratorUtil.h
--- a/Engine/Plugins/Experimental/GeometryProcessing/Source/GeometricObjects/Public/Util/IteratorUtil.h
+++ b/Engine/Plugins/Experimental/GeometryProcessing/Source/GeometricObjects/Public/Util/IteratorUtil.h
@@ -101,7 +101,7 @@ public:
 		Cur = CurItr;
 		End = EndItr;
 		this->FilterFunc = FilterFuncIn;
-		if (FilterFunc(*Cur) == false)
+		if (Cur != End && FilterFunc(*Cur) == false)
 		{
 			GotoNextElement();
 		}
