GeometryProcessing: FMeshConstraintUtil::ConstrainAllSeams now done in parallel, added FMeshConstraintUtil::ConstrainEdgeROISeams. Refactored FRemesher internals for smoothing to allow for localized smoothing in subclasses.
#rnx

diff --git a/Engine/Plugins/Experimental/GeometryProcessing/Source/DynamicMesh/Private/MeshConstraintsUtil.cpp b/Engine/Plugins/Experimental/GeometryProcessing/Source/DynamicMesh/Private/MeshConstraintsUtil.cpp
--- a/Engine/Plugins/Experimental/GeometryProcessing/Source/DynamicMesh/Private/MeshConstraintsUtil.cpp
+++ b/Engine/Plugins/Experimental/GeometryProcessing/Source/DynamicMesh/Private/MeshConstraintsUtil.cpp
@@ -1,9 +1,10 @@
 // Copyright 1998-2019 Epic Games, Inc. All Rights Reserved.
 
 #include "MeshConstraintsUtil.h"
+#include "Async/ParallelFor.h"
 
 
-void FMeshConstraintsUtil::ConstrainAllSeams(FMeshConstraints& Constraints, const FDynamicMesh3& Mesh, bool bAllowSplits, bool bAllowSmoothing)
+void FMeshConstraintsUtil::ConstrainAllSeams(FMeshConstraints& Constraints, const FDynamicMesh3& Mesh, bool bAllowSplits, bool bAllowSmoothing, bool bParallel)
 {
 	if (Mesh.HasAttributes() == false)
 	{
@@ -14,14 +15,57 @@ void FMeshConstraintsUtil::ConstrainAllSeams(FMeshConstraints& Constraints, cons
 	FEdgeConstraint EdgeConstraint = (bAllowSplits) ? FEdgeConstraint::SplitsOnly() : FEdgeConstraint::FullyConstrained();
 	FVertexConstraint VtxConstraint = (bAllowSmoothing) ? FVertexConstraint::PinnedMovable() : FVertexConstraint::Pinned();
 
-	for (int EdgeID : Mesh.EdgeIndicesItr())
+	FCriticalSection ConstraintSetLock;
+
+	int32 NumEdges = Mesh.MaxEdgeID();
+	ParallelFor(NumEdges, [&](int EdgeID)
 	{
+		if (Mesh.IsEdge(EdgeID))
+		{
+			if (Attributes->IsSeamEdge(EdgeID))
+			{
+				FIndex2i EdgeVerts = Mesh.GetEdgeV(EdgeID);
+
+				ConstraintSetLock.Lock();
+				Constraints.SetOrUpdateEdgeConstraint(EdgeID, EdgeConstraint);
+				Constraints.SetOrUpdateVertexConstraint(EdgeVerts.A, VtxConstraint);
+				Constraints.SetOrUpdateVertexConstraint(EdgeVerts.B, VtxConstraint);
+				ConstraintSetLock.Unlock();
+			}
+		}
+	}, (bParallel==false) );
+}
+
+
+
+void FMeshConstraintsUtil::ConstrainEdgeROISeams(FMeshConstraints& Constraints, const FDynamicMesh3& Mesh, const TArray<int>& EdgeROI, bool bAllowSplits, bool bAllowSmoothing, bool bParallel)
+{
+	if (Mesh.HasAttributes() == false)
+	{
+		return;
+	}
+	const FDynamicMeshAttributeSet* Attributes = Mesh.Attributes();
+
+	FEdgeConstraint EdgeConstraint = (bAllowSplits) ? FEdgeConstraint::SplitsOnly() : FEdgeConstraint::FullyConstrained();
+	FVertexConstraint VtxConstraint = (bAllowSmoothing) ? FVertexConstraint::PinnedMovable() : FVertexConstraint::Pinned();
+
+	FCriticalSection ConstraintSetLock;
+
+	int32 NumEdges = EdgeROI.Num();
+	ParallelFor(NumEdges, [&](int k)
+	{
+		int EdgeID = EdgeROI[k];
+
 		if (Attributes->IsSeamEdge(EdgeID))
 		{
-			Constraints.SetOrUpdateEdgeConstraint(EdgeID, EdgeConstraint);
 			FIndex2i EdgeVerts = Mesh.GetEdgeV(EdgeID);
+
+			ConstraintSetLock.Lock();
+			Constraints.SetOrUpdateEdgeConstraint(EdgeID, EdgeConstraint);
 			Constraints.SetOrUpdateVertexConstraint(EdgeVerts.A, VtxConstraint);
 			Constraints.SetOrUpdateVertexConstraint(EdgeVerts.B, VtxConstraint);
+			ConstraintSetLock.Unlock();
 		}
-	}
-}
\ No newline at end of file
+	}, (bParallel==false) );
+}
+
diff --git a/Engine/Plugins/Experimental/GeometryProcessing/Source/DynamicMesh/Private/Remesher.cpp b/Engine/Plugins/Experimental/GeometryProcessing/Source/DynamicMesh/Private/Remesher.cpp
--- a/Engine/Plugins/Experimental/GeometryProcessing/Source/DynamicMesh/Private/Remesher.cpp
+++ b/Engine/Plugins/Experimental/GeometryProcessing/Source/DynamicMesh/Private/Remesher.cpp
@@ -452,27 +452,29 @@ static FVector3d CotanSmooth(const FDynamicMesh3& mesh, int vID, double t)
 	return (1.0 - t)*v + (t)*c;
 }
 
-void FRemesher::FullSmoothPass_Buffer(bool bParallel)
+TFunction<FVector3d(const FDynamicMesh3&, int, double)> FRemesher::GetSmoothFunction()
 {
-	InitializeVertexBufferForPass();
-
-	TFunction<FVector3d(const FDynamicMesh3&, int, double)> UseSmoothFunc = UniformSmooth;
-	//Func<FDynamicMesh3, int, double, FVector3d> smoothFunc = MeshUtil.UniformSmooth;
-	if (CustomSmoothF != nullptr) 
+	if (CustomSmoothF != nullptr)
+	{
+		return CustomSmoothF;
+	} 
+	else if (SmoothType == ESmoothTypes::MeanValue)
 	{
-		UseSmoothFunc = CustomSmoothF;
+		return MeanValueSmooth;
 	}
-	else 
+	else if (SmoothType == ESmoothTypes::Cotan)
 	{
-		if (SmoothType == ESmoothTypes::MeanValue)
-		{
-			UseSmoothFunc = MeanValueSmooth;
-		}
-		else if (SmoothType == ESmoothTypes::Cotan)
-		{
-			UseSmoothFunc = CotanSmooth;
-		}
+		return CotanSmooth;
 	}
+	return UniformSmooth;
+}
+
+
+void FRemesher::FullSmoothPass_Buffer(bool bParallel)
+{
+	InitializeVertexBufferForPass();
+
+	TFunction<FVector3d(const FDynamicMesh3&, int, double)> UseSmoothFunc = GetSmoothFunction();
 
 	auto SmoothAndUpdateFunc = [&](int vID) 
 	{
diff --git a/Engine/Plugins/Experimental/GeometryProcessing/Source/DynamicMesh/Public/MeshConstraintsUtil.h b/Engine/Plugins/Experimental/GeometryProcessing/Source/DynamicMesh/Public/MeshConstraintsUtil.h
--- a/Engine/Plugins/Experimental/GeometryProcessing/Source/DynamicMesh/Public/MeshConstraintsUtil.h
+++ b/Engine/Plugins/Experimental/GeometryProcessing/Source/DynamicMesh/Public/MeshConstraintsUtil.h
@@ -24,7 +24,18 @@ public:
 	 * @param bAllowSplits should we allow constrained edges to be split
 	 * @param bAllowSmoothing should we allow constrained vertices to be smoothed
 	 */
-	static void ConstrainAllSeams(FMeshConstraints& Constraints, const FDynamicMesh3& Mesh, bool bAllowSplits, bool bAllowSmoothing);
+	static void ConstrainAllSeams(FMeshConstraints& Constraints, const FDynamicMesh3& Mesh, bool bAllowSplits, bool bAllowSmoothing, bool bParallel = true);
+
+
+	/**
+	 * Constrain all attribute seams for all overlays of a mesh, for edges in the edge array
+	 * @param Constraints the set of constraints to add to
+	 * @param Mesh the mesh to constrain
+	 * @param EdgeROI list of edges to try to constrain
+	 * @param bAllowSplits should we allow constrained edges to be split
+	 * @param bAllowSmoothing should we allow constrained vertices to be smoothed
+	 */
+	static void ConstrainEdgeROISeams(FMeshConstraints& Constraints, const FDynamicMesh3& Mesh, const TArray<int>& EdgeROI, bool bAllowSplits, bool bAllowSmoothing, bool bParallel = true);
 
 
 	/**
diff --git a/Engine/Plugins/Experimental/GeometryProcessing/Source/DynamicMesh/Public/Remesher.h b/Engine/Plugins/Experimental/GeometryProcessing/Source/DynamicMesh/Public/Remesher.h
--- a/Engine/Plugins/Experimental/GeometryProcessing/Source/DynamicMesh/Public/Remesher.h
+++ b/Engine/Plugins/Experimental/GeometryProcessing/Source/DynamicMesh/Public/Remesher.h
@@ -234,15 +234,18 @@ protected:
 
 
 
-
+	// applies a smoothing pass to the mesh, storing intermediate positions in a buffer and then writing them at the end (so, no order effect)
 	virtual void FullSmoothPass_Buffer(bool bParallel);
 
-	/** computes smoothed vertex position w/ proper constraints/etc. Does not modify mesh. */
+	// computes smoothed vertex position w/ proper constraints/etc. Does not modify mesh. 
 	virtual FVector3d ComputeSmoothedVertexPos(int VertexID,
 		TFunction<FVector3d(const FDynamicMesh3 &, int, double)> SmoothFunc, bool& bModified);
 
+	// FullSmoothPass_Buffer this to calls VertexSmoothFunc for each vertex of the mesh
 	virtual void ApplyToSmoothVertices(const TFunction<void(int)>& VertexSmoothFunc);
 
+	// returns the function we want to use to compute a smoothed vertex position - will be CustomSmoothF if set, otherwise one of cotan/meanvalue/uniform
+	virtual TFunction<FVector3d(const FDynamicMesh3&, int, double)> GetSmoothFunction();
 
 
 
