better choices of function param types for mesh element selection filters

#rnx

diff --git a/Engine/Plugins/Experimental/GeometryProcessing/Source/DynamicMesh/Public/Selections/MeshEdgeSelection.h b/Engine/Plugins/Experimental/GeometryProcessing/Source/DynamicMesh/Public/Selections/MeshEdgeSelection.h
--- a/Engine/Plugins/Experimental/GeometryProcessing/Source/DynamicMesh/Public/Selections/MeshEdgeSelection.h
+++ b/Engine/Plugins/Experimental/GeometryProcessing/Source/DynamicMesh/Public/Selections/MeshEdgeSelection.h
@@ -122,7 +122,7 @@ public:
 			}
 		}
 	}
-	void Select(TFunction<bool(int)> SelectF)
+	void Select(TFunctionRef<bool(int)> SelectF)
 	{
 		int NT = Mesh->MaxEdgeID();
 		for (int eid = 0; eid < NT; ++eid)
diff --git a/Engine/Plugins/Experimental/GeometryProcessing/Source/DynamicMesh/Public/Selections/MeshFaceSelection.h b/Engine/Plugins/Experimental/GeometryProcessing/Source/DynamicMesh/Public/Selections/MeshFaceSelection.h
--- a/Engine/Plugins/Experimental/GeometryProcessing/Source/DynamicMesh/Public/Selections/MeshFaceSelection.h
+++ b/Engine/Plugins/Experimental/GeometryProcessing/Source/DynamicMesh/Public/Selections/MeshFaceSelection.h
@@ -96,7 +96,7 @@ public:
 			}
 		}
 	}
-	void Select(TFunction<bool(int)> SelectF)
+	void Select(TFunctionRef<bool(int)> SelectF)
 	{
 		int NT = Mesh->MaxTriangleID();
 		for (int tID = 0; tID < NT; ++tID)
@@ -249,7 +249,7 @@ public:
 	}
 
 
-	void ExpandToFaceNeighbours(TFunction<bool(int)> FilterF = nullptr)
+	void ExpandToFaceNeighbours(const TUniqueFunction<bool(int)>& FilterF = nullptr)
 	{
 		TArray<int> ToAdd;
 
@@ -257,7 +257,7 @@ public:
 			FIndex3i nbr_tris = Mesh->GetTriNeighbourTris(tid);
 			for (int j = 0; j < 3; ++j)
 			{
-				if (FilterF != nullptr && FilterF(nbr_tris[j]) == false)
+				if (FilterF && FilterF(nbr_tris[j]) == false)
 				{
 					continue;
 				}
@@ -273,7 +273,7 @@ public:
 			add(ID);
 		}
 	}
-	void ExpandToFaceNeighbours(int rounds, TFunction<bool(int)> FilterF = nullptr)
+	void ExpandToFaceNeighbours(int rounds, const TUniqueFunction<bool(int)>& FilterF = nullptr)
 	{
 		for (int k = 0; k < rounds; ++k)
 		{
@@ -289,7 +289,7 @@ public:
 	 *  
 	 *  Return false from FilterF to prevent triangles from being included.
 	 */
-	void ExpandToOneRingNeighbours(TFunction<bool(int)> FilterF = nullptr)
+	void ExpandToOneRingNeighbours(const TUniqueFunction<bool(int)>& FilterF = nullptr)
 	{
 		TArray<int> ToAdd;
 
@@ -300,7 +300,7 @@ public:
 				int vid = tri_v[j];
 				for (int nbr_t : Mesh->VtxTrianglesItr(vid))
 				{
-					if (FilterF != nullptr && FilterF(nbr_t) == false)
+					if (FilterF && FilterF(nbr_t) == false)
 					{
 						continue;
 					}
@@ -326,7 +326,7 @@ public:
 	 *  
 	 *  Return false from FilterF to prevent triangles from being included.
 	 */
-	void ExpandToOneRingNeighbours(int nRings, TFunction<bool(int)> FilterF = nullptr)
+	void ExpandToOneRingNeighbours(int nRings, const TUniqueFunction<bool(int)>& FilterF = nullptr)
 	{
 		if (nRings == 1)
 		{
@@ -354,7 +354,7 @@ public:
 					int vid = tri_v[j];
 					for (int nbr_t : Mesh->VtxTrianglesItr(vid))
 					{
-						if (FilterF != nullptr && FilterF(nbr_t) == false)
+						if (FilterF && FilterF(nbr_t) == false)
 						{
 							continue;
 						}
@@ -429,7 +429,7 @@ public:
 	 *  Grow selection outwards from seed triangle, until it hits boundaries defined by triangle and edge filters.
 	 *  Edge filter is not effective unless it (possibly combined w/ triangle filter) defines closed loops.
 	 */
-	void FloodFill(int tSeed, TFunction<bool(int)> TriFilterF = nullptr, TFunction<bool(int)> EdgeFilterF = nullptr)
+	void FloodFill(int tSeed, const TUniqueFunction<bool(int)>& TriFilterF = nullptr, const TUniqueFunction<bool(int)>& EdgeFilterF = nullptr)
 	{
 		TArray<int> Seeds = { tSeed };
 		FloodFill(Seeds, TriFilterF, EdgeFilterF);
@@ -438,7 +438,7 @@ public:
 	 *  Grow selection outwards from seed triangles, until it hits boundaries defined by triangle and edge filters.
 	 *  Edge filter is not effective unless it (possibly combined w/ triangle filter) defines closed loops.
 	 */
-	void FloodFill(const TArray<int>& Seeds, TFunction<bool(int)> TriFilterF = nullptr, TFunction<bool(int)> EdgeFilterF = nullptr)
+	void FloodFill(const TArray<int>& Seeds, const TUniqueFunction<bool(int)>& TriFilterF = nullptr, const TUniqueFunction<bool(int)>& EdgeFilterF = nullptr)
 	{
 		TDynamicVector<int> stack(Seeds);
 		for (int Seed : Seeds)
@@ -458,12 +458,12 @@ public:
 				{
 					continue;
 				}
-				if (TriFilterF != nullptr && TriFilterF(nbr_tid) == false)
+				if (TriFilterF && TriFilterF(nbr_tid) == false)
 				{
 					continue;
 				}
 				int EID = Mesh->GetTriEdge(TID, j);
-				if (EdgeFilterF != nullptr && EdgeFilterF(EID) == false)
+				if (EdgeFilterF && EdgeFilterF(EID) == false)
 				{
 					continue;
 				}
diff --git a/Engine/Plugins/Experimental/GeometryProcessing/Source/DynamicMesh/Public/Selections/MeshVertexSelection.h b/Engine/Plugins/Experimental/GeometryProcessing/Source/DynamicMesh/Public/Selections/MeshVertexSelection.h
--- a/Engine/Plugins/Experimental/GeometryProcessing/Source/DynamicMesh/Public/Selections/MeshVertexSelection.h
+++ b/Engine/Plugins/Experimental/GeometryProcessing/Source/DynamicMesh/Public/Selections/MeshVertexSelection.h
@@ -187,14 +187,14 @@ public:
 	 *  
 	 *  Return false from FilterF to prevent vertices from being included.
 	 */
-	void ExpandToOneRingNeighbours(TFunction<bool(int)> FilterF = nullptr)
+	void ExpandToOneRingNeighbours(const TUniqueFunction<bool(int)>& FilterF = nullptr)
 	{
 		TArray<int> temp;
 
 		for (int vid : Selected) {
 			for (int nbr_vid : Mesh->VtxVerticesItr(vid))
 			{
-				if (FilterF != nullptr && FilterF(nbr_vid) == false)
+				if (FilterF && FilterF(nbr_vid) == false)
 				{
 					continue;
 				}
@@ -213,7 +213,7 @@ public:
 
 
 	// [TODO] should do this more efficiently, like FMeshFaceSelection
-	void ExpandToOneRingNeighbours(int nRings, TFunction<bool(int)> FilterF = nullptr)
+	void ExpandToOneRingNeighbours(int nRings, const TUniqueFunction<bool(int)>& FilterF = nullptr)
 	{
 		for (int k = 0; k < nRings; ++k)
 		{
@@ -228,7 +228,7 @@ public:
 	/**
 	 *  Grow selection outwards from seed vertex, until it hits boundaries defined by vertex filter.
 	 */
-	void FloodFill(int vSeed, TFunction<bool(int)> VertIncludedF = nullptr)
+	void FloodFill(int vSeed, const TUniqueFunction<bool(int)>& VertIncludedF = nullptr)
 	{
 		TArray<int> Seeds = { vSeed };
 		FloodFill(Seeds, VertIncludedF);
@@ -236,7 +236,7 @@ public:
 	/**
 	 *  Grow selection outwards from seed vertex, until it hits boundaries defined by vertex filter.
 	 */
-	void FloodFill(const TArray<int>& Seeds, TFunction<bool(int)> VertIncludedF = nullptr)
+	void FloodFill(const TArray<int>& Seeds, const TUniqueFunction<bool(int)>& VertIncludedF = nullptr)
 	{
 		TDynamicVector<int> stack(Seeds);
 		for (int Seed : Seeds)
@@ -250,7 +250,7 @@ public:
 
 			for (int nbr_vid : Mesh->VtxVerticesItr(vID))
 			{
-				if (IsSelected(nbr_vid) == true || VertIncludedF(nbr_vid) == false)
+				if (IsSelected(nbr_vid) == true || (VertIncludedF && VertIncludedF(nbr_vid) == false))
 				{
 					continue;
 				}
