MeshProcessing:
Add FSparseDynamicOctree3 and FDynamicMeshOctree3. Implements dynamically-updateable octree-set for mesh triangles based on layered uniform grids. FDynamicMeshOctree3 extension supports calculating/updating FTreeCutSet to partition mesh into smaller chunks based on octree branches.
Add FArrayIndexSetsDecomposition, stores set of index arrays with unique keys, used to track partition of mesh based on FTreeCutSet
#rnx

diff --git a/Engine/Plugins/Experimental/GeometryProcessing/Source/DynamicMesh/Public/DynamicMeshOctree3.h b/Engine/Plugins/Experimental/GeometryProcessing/Source/DynamicMesh/Public/DynamicMeshOctree3.h
new file mode 100644
--- /dev/null
+++ b/Engine/Plugins/Experimental/GeometryProcessing/Source/DynamicMesh/Public/DynamicMeshOctree3.h
@@ -0,0 +1,374 @@
+// Copyright 1998-2019 Epic Games, Inc. All Rights Reserved.
+
+#pragma once
+
+#include "Spatial/MeshAABBTree3.h"
+#include "Spatial/SparseDynamicOctree3.h"
+#include "MeshQueries.h"
+#include "DynamicMesh3.h"
+
+
+/**
+ * FDynamicMeshOctree3 is an extension of FSparseDynamicOctree3 for the triangles of a FDynamicMesh3 instance.
+ * This extension does several things:
+ *   1) provides a simplified API based on triangle IDs to various Octree functions
+ *   2) tracks ModifiedBounds box of modified areas
+ *   3) support for computing/updating/querying a "Cut" of the octree, ie a set of cells which are roots of sub-branches
+ *      that partition the tree. This is useful for splitting up mesh processing/rendering into spatially-coherent chunks.
+ *      (This functionality should probably be extracted into a separate class...)
+ */
+class FDynamicMeshOctree3 : public FSparseDynamicOctree3
+{
+	// potential optimizations:
+	//    - keep track of how many triangles are descendents of each cell? root cells?
+
+
+
+public:
+	/** parent mesh */
+	const FDynamicMesh3* Mesh;
+
+	/** bounding box of triangles that have been inserted/removed since last clear */
+	FAxisAlignedBox3d ModifiedBounds;
+
+	/** 
+	 * Add all triangles of MeshIn to the octree 
+	 */
+	void Initialize(const FDynamicMesh3* MeshIn)
+	{
+		this->Mesh = MeshIn;
+
+		for (int tid : Mesh->TriangleIndicesItr())
+		{
+			InsertTriangle(tid);
+		}
+	}
+
+	/**
+	 * Reset the internal ModifiedBounds box that tracks modified triangle bounds
+	 */
+	void ResetModifiedBounds()
+	{
+		ModifiedBounds = FAxisAlignedBox3d::Empty();
+	}
+
+	/**
+	 * Insert a triangle into the tree
+	 */
+	void InsertTriangle(int32 TriangleID)
+	{
+		FAxisAlignedBox3d Bounds = Mesh->GetTriBounds(TriangleID);
+		ModifiedBounds.Contain(Bounds);
+		InsertObject(TriangleID, Bounds);
+	}
+
+	/**
+	 * Insert a list of triangles into the tree
+	 */
+	void InsertTriangles(const TArray<int>& Triangles)
+	{
+		int N = Triangles.Num();
+		for (int i = 0; i < N; ++i)
+		{
+			FAxisAlignedBox3d Bounds = Mesh->GetTriBounds(Triangles[i]);
+			ModifiedBounds.Contain(Bounds);
+			InsertObject(Triangles[i], Bounds);
+		}
+	}
+
+	/**
+	 * Insert a set of triangles into the tree
+	 */
+	void InsertTriangles(const TSet<int>& Triangles)
+	{
+		for (int TriangleID : Triangles)
+		{
+			FAxisAlignedBox3d Bounds = Mesh->GetTriBounds(TriangleID);
+			ModifiedBounds.Contain(Bounds);
+			InsertObject(TriangleID, Bounds);
+		}
+	}
+
+
+	/**
+	 * Remove a triangle from the tree
+	 */
+	bool RemoveTriangle(int32 TriangleID)
+	{
+		FAxisAlignedBox3d Bounds = Mesh->GetTriBounds(TriangleID);
+		ModifiedBounds.Contain(Bounds);
+		RemoveObject(TriangleID);
+	}
+
+	/**
+	 * Remove a list of triangles into the tree
+	 */
+	void RemoveTriangles(const TArray<int>& Triangles)
+	{
+		int N = Triangles.Num();
+		for ( int i = 0; i < N; ++i )
+		{
+			FAxisAlignedBox3d Bounds = Mesh->GetTriBounds(Triangles[i]);
+			ModifiedBounds.Contain(Bounds);
+			RemoveObject(Triangles[i]);
+		}
+	}
+
+	/**
+	 * Remove a set of triangles into the tree
+	 */
+	void RemoveTriangles(const TSet<int>& Triangles)
+	{
+		for (int TriangleID : Triangles)
+		{
+			FAxisAlignedBox3d Bounds = Mesh->GetTriBounds(TriangleID);
+			ModifiedBounds.Contain(Bounds);
+			RemoveObject(TriangleID);
+		}
+	}
+
+
+	/**
+	 * Reinsert a set of triangles into the tree
+	 */
+	void ReinsertTriangles(const TSet<int>& Triangles)
+	{
+		for (int TriangleID : Triangles)
+		{
+			FAxisAlignedBox3d Bounds = Mesh->GetTriBounds(TriangleID);
+			ModifiedBounds.Contain(Bounds);
+			ReinsertObject(TriangleID, Bounds);
+		}
+	}
+
+
+	/**
+	 * Include the current bounds of a triangle in the ModifiedBounds box
+	 */
+	void NotifyPendingModification(int TriangleID)
+	{
+		FAxisAlignedBox3d Bounds = Mesh->GetTriBounds(TriangleID);
+		ModifiedBounds.Contain(Bounds);
+	}
+
+	/**
+	 * Include the current bounds of a set of triangles in the ModifiedBounds box
+	 */
+	void NotifyPendingModification(const TSet<int>& Triangles)
+	{
+		for (int TriangleID : Triangles)
+		{
+			FAxisAlignedBox3d Bounds = Mesh->GetTriBounds(TriangleID);
+			ModifiedBounds.Contain(Bounds);
+		}
+	}
+
+
+	/**
+	 * Find the nearest triangle of the mesh that is hit by the ray
+	 */
+	int32 FindNearestHitObject(const FRay3d& Ray,
+		double MaxDistance = TNumericLimits<double>::Max()) const
+	{
+		return FSparseDynamicOctree3::FindNearestHitObject(Ray,
+			[&](int tid) {return Mesh->GetTriBounds(tid); },
+			[&](int tid, const FRay3d& Ray) {
+			FIntrRay3Triangle3d Intr = TMeshQueries<FDynamicMesh3>::TriangleIntersection(*Mesh, tid, Ray);
+			return (Intr.IntersectionType == EIntersectionType::Point) ?
+				Intr.RayParameter : TNumericLimits<double>::Max();
+
+		}, MaxDistance);
+	}
+
+
+	/**
+	 * Check that the Octree is internally valid
+	 */
+	void CheckValidity(
+		EValidityCheckFailMode FailMode = EValidityCheckFailMode::Check,
+		bool bVerbose = false,
+		bool bFailOnMissingObjects = false) const
+	{
+		FSparseDynamicOctree3::CheckValidity(
+			[&](int tid) {return Mesh->IsTriangle(tid); },
+			[&](int tid) {return Mesh->GetTriBounds(tid); },
+			FailMode, bVerbose, bFailOnMissingObjects);
+	}
+
+
+
+
+
+
+	//
+	// Support for building "cuts" of the octree, which are sets
+	// of internal nodes which can be used to decompose the tree
+	// (eg into spatially-coherent chunks of triangles, for example)
+	// 
+
+	/** FCellReference is a handle to an internal cell of the octree */
+	struct FCellReference
+	{
+		uint32 CellID;
+	};
+
+	/** FTreeCutSet is a cut of the tree, ie a set of internal cells which are "parents" of separate branches */
+	class FTreeCutSet
+	{
+	public:
+		TArray<FCellReference> CutCells;
+
+	protected:
+		TSet<uint32> CutCellIDs;			// just a set version of CutCells for internal use - maybe use set there instead, and add compare ops to FCellReference?
+		int FixedCutLevel;
+
+		friend class FDynamicMeshOctree3;
+	};
+
+
+	/**
+	 * @return a cut of the tree at a fixed level
+	 */
+	FTreeCutSet BuildLevelCutSet(uint32 CutLevel = 5) const
+	{
+		FTreeCutSet CutSet;
+		CutSet.FixedCutLevel = CutLevel;
+
+		for (const FSparseOctreeCell& Cell : Cells)
+		{
+			if (Cell.Level == CutSet.FixedCutLevel)
+			{
+				FCellReference CellRef;
+				CellRef.CellID = Cell.CellID;
+				CutSet.CutCells.Add(CellRef);
+
+				CutSet.CutCellIDs.Add(Cell.CellID);
+			}
+		}
+
+		return CutSet;
+	}
+
+
+	/**
+	 * For a fixed-level cut set created by BuildLevelCutSet, check that all current cells at that level are in the cut set
+	 * (call this after adding/removing to the tree to make sure the CutSet is up to date)
+	 * @param CutSet the tree cut set to check/update
+	 * @param NewCutCellsOut list of new cut cells added by the update
+	 */
+	void UpdateLevelCutSet(FTreeCutSet& CutSet, TArray<FCellReference>& NewCutCellsOut) const
+	{
+		for (const FSparseOctreeCell& Cell : Cells)
+		{
+			if (Cell.Level == CutSet.FixedCutLevel)
+			{
+				if (CutSet.CutCellIDs.Contains(Cell.CellID) == false)
+				{
+					FCellReference CellRef;
+					CellRef.CellID = Cell.CellID;
+					CutSet.CutCells.Add(CellRef);
+					CutSet.CutCellIDs.Add(Cell.CellID);
+					NewCutCellsOut.Add(CellRef);
+				}
+			}
+		}
+	}
+
+	/**
+	 * @return true if the cell identified by CellRef intersects the Bounds box
+	 */
+	bool TestCellIntersection(const FCellReference& CellRef, const FAxisAlignedBox3d& Bounds) const
+	{
+		check(CellRefCounts.IsValid(CellRef.CellID));
+		const FSparseOctreeCell& Cell = Cells[CellRef.CellID];
+		return GetCellBox(Cell, MaxExpandFactor).Intersects(Bounds);
+	}
+
+	/**
+	 * Call TriangleFunc on any triangles in the branch of the tree starting at CellRef
+	 */
+	void CollectTriangles(
+		const FCellReference& CellRef, 
+		TFunctionRef<void(int)> TriangleFunc) const
+	{
+		check(CellRefCounts.IsValid(CellRef.CellID));
+
+		TArray<const FSparseOctreeCell*> Queue;
+		Queue.Add(&Cells[CellRef.CellID]);
+		while (Queue.Num() > 0)
+		{
+			const FSparseOctreeCell* CurCell = Queue.Pop(false);
+
+			// process elements
+			for (int ObjectID : CellObjectLists.Values(CurCell->CellID))
+			{
+				TriangleFunc(ObjectID);
+			}
+
+			for (int k = 0; k < 8; ++k)
+			{
+				if (CurCell->HasChild(k))
+				{
+					const FSparseOctreeCell* ChildCell = &Cells[CurCell->GetChildCellID(k)];
+					Queue.Add(ChildCell);
+				}
+			}
+		}
+	}
+
+
+	/**
+	 * Call TriangleFunc for all triangles in the octree "above" the CutSet
+	 * (ie at tree cells that are not children of any cut cells)
+	 */
+	void CollectRootTriangles(const FTreeCutSet& CutSet,
+		TFunctionRef<void(int)> TriangleFunc) const
+	{
+		TArray<const FSparseOctreeCell*> Queue;
+
+		// start at root cells
+		RootCells.AllocatedIteration([&](const uint32* RootCellID)
+		{
+			const FSparseOctreeCell* RootCell = &Cells[*RootCellID];
+			Queue.Add(&Cells[*RootCellID]);
+		});
+
+		while (Queue.Num() > 0)
+		{
+			const FSparseOctreeCell* CurCell = Queue.Pop(false);
+			if (CutSet.CutCellIDs.Contains(CurCell->CellID))
+			{
+				continue;
+			}
+
+			// process elements
+			for (int TriangleID : CellObjectLists.Values(CurCell->CellID))
+			{
+				TriangleFunc(TriangleID);
+			}
+
+			for (int k = 0; k < 8; ++k)
+			{
+				if (CurCell->HasChild(k))
+				{
+					const FSparseOctreeCell* ChildCell = &Cells[CurCell->GetChildCellID(k)];
+					Queue.Add(ChildCell);
+				}
+			}
+		}
+	}
+
+	/**
+	 * Call TriangleFunc for any triangles in the spill set (ie not contained in any Root cell)
+	 */
+	void CollectSpillTriangles(
+		TFunctionRef<void(int)> TriangleFunc) const
+	{
+		for (int tid : SpillObjectSet)
+		{
+			TriangleFunc(tid);
+		}
+	}
+
+
+};
\ No newline at end of file
diff --git a/Engine/Plugins/Experimental/GeometryProcessing/Source/DynamicMesh/Public/Generators/RectangleMeshGenerator.h b/Engine/Plugins/Experimental/GeometryProcessing/Source/DynamicMesh/Public/Generators/RectangleMeshGenerator.h
--- a/Engine/Plugins/Experimental/GeometryProcessing/Source/DynamicMesh/Public/Generators/RectangleMeshGenerator.h
+++ b/Engine/Plugins/Experimental/GeometryProcessing/Source/DynamicMesh/Public/Generators/RectangleMeshGenerator.h
@@ -8,7 +8,7 @@
  * Generate planar rectangular mesh with variable number of subdivisions along width and height.
  * By default, center of rectangle is centered at (0,0,0) origin
  */
-class DYNAMICMESH_API FRectangleMeshGenerator : FMeshShapeGenerator
+class DYNAMICMESH_API FRectangleMeshGenerator : public FMeshShapeGenerator
 {
 public:
 	/** Rectangle will be translated so that center is at this point */
diff --git a/Engine/Plugins/Experimental/GeometryProcessing/Source/GeometricObjects/Public/Spatial/SparseDynamicOctree3.h b/Engine/Plugins/Experimental/GeometryProcessing/Source/GeometricObjects/Public/Spatial/SparseDynamicOctree3.h
new file mode 100644
--- /dev/null
+++ b/Engine/Plugins/Experimental/GeometryProcessing/Source/GeometricObjects/Public/Spatial/SparseDynamicOctree3.h
@@ -0,0 +1,900 @@
+// Copyright 1998-2019 Epic Games, Inc. All Rights Reserved.
+
+#pragma once
+
+
+#include "CoreMinimal.h"
+#include "Containers/StaticArray.h"
+#include "BoxTypes.h"
+#include "Util/DynamicVector.h"
+#include "Util/SmallListSet.h"
+#include "Spatial/SparseGrid3.h"
+#include "Intersection/IntrRay3AxisAlignedBox3.h"
+#include "ASync/ParallelFor.h"
+
+
+/**
+ * Utility class that allows for get/set of a flag for each integer ID, where
+ * the flag set automatically grows to contain whatever integer ID is passed
+ */
+class FDynamicFlagArray
+{
+public:
+	TBitArray<> BitArray;
+	uint32 MaxIndex = 0;
+	static constexpr int32 GrowChunkSize = 0xFFF;
+
+	FDynamicFlagArray()
+	{
+		BitArray.Init(false, GrowChunkSize);
+		MaxIndex = BitArray.Num();
+	}
+
+	void Set(uint32 BitIndex, bool bValue)
+	{
+		if (bValue && BitIndex >= MaxIndex)
+		{
+			int32 ToAdd = (BitIndex - MaxIndex) | GrowChunkSize;
+			BitArray.Add(false, ToAdd);
+			MaxIndex = BitArray.Num();
+		}
+		BitArray[BitIndex] = bValue;
+	}
+
+	inline bool Get(uint32 BitIndex) const
+	{
+		return (BitIndex < MaxIndex) ? BitArray[BitIndex] : false;
+	}
+};
+
+
+
+
+/**
+ * FSparseOctreeCell is a Node in a SparseDynamicOctree3. 
+ */
+struct FSparseOctreeCell
+{
+	static constexpr uint32 InvalidID = TNumericLimits<uint32>::Max();;
+	static constexpr uint8 InvalidLevel = TNumericLimits<uint8>::Max();
+
+	/** ID of cell (index into cell list) */
+	uint32 CellID;
+
+	/** Level of cell in octree */
+	uint8 Level;
+
+	/** i,j,k index of cell in level, relative to origin */
+	FVector3i Index;
+
+	/** CellID of each child, or InvalidID if that child does not exist */
+	TStaticArray<uint32, 8> Children;
+
+	FSparseOctreeCell()
+		: CellID(InvalidID), Level(0), Index(FVector3i::Zero())
+	{
+		Children[0] = Children[1] = Children[2] = Children[3] = InvalidID;
+		Children[4] = Children[5] = Children[6] = Children[7] = InvalidID;
+	}
+
+	FSparseOctreeCell(uint8 LevelIn, const FVector3i& IndexIn)
+		: CellID(InvalidID), Level(LevelIn), Index(IndexIn)
+	{
+		Children[0] = Children[1] = Children[2] = Children[3] = InvalidID;
+		Children[4] = Children[5] = Children[6] = Children[7] = InvalidID;
+	}
+
+	inline bool IsExistingCell() const
+	{
+		return CellID != InvalidID;
+	}
+
+	inline bool HasChild(int ChildIndex) const
+	{
+		return Children[ChildIndex] != InvalidID;
+	}
+
+	inline uint32 GetChildCellID(int ChildIndex) const
+	{
+		return Children[ChildIndex];
+	}
+
+	inline FSparseOctreeCell MakeChildCell(int ChildIndex)
+	{
+		FVector3i IndexOffset(
+			((ChildIndex & 1) != 0) ? 1 : 0,
+			((ChildIndex & 2) != 0) ? 1 : 0,
+			((ChildIndex & 4) != 0) ? 1 : 0);
+		return FSparseOctreeCell(Level + 1, 2*Index + IndexOffset);
+	}
+
+	inline void SetChild(uint32 ChildIndex, const FSparseOctreeCell& ChildCell)
+	{
+		Children[ChildIndex] = ChildCell.CellID;
+	}
+
+};
+
+
+/**
+ * FSparseDynamicOctree3 sorts objects with axis-aligned bounding boxes into a dynamic
+ * sparse octree of axis-aligned uniform grid cells. At the top level we have an infinite
+ * grid of "root cells" of size RootDimension, which then contain 8 children, and so on.
+ * (So in fact each cell is a separate octree, and we have a uniform grid of octrees)
+ * 
+ * The objects and their bounding-boxes are not stored in the tree. You must have an
+ * integer identifier (ObjectID) for each object, and call Insert(ObjectID, BoundingBox).
+ * Some query functions will require you to provide a lambda/etc that can be called to
+ * retrieve the bounding box for a given ObjectID.
+ * 
+ * Objects are currently inserted at the maximum possible depth, ie smallest cell that
+ * will contain them, or MaxTreeDepth. The tree boxes are expanded by MaxExpandFactor
+ * to allow for deeper insertion. If MaxExpandFactor > 0 then the tree does not strictly
+ * partition space, IE adjacent cells overlap.
+ * 
+ * The octree is dynamic. Objects can be removed and re-inserted.
+ * 
+ */
+class FSparseDynamicOctree3
+{
+	// potential optimizations/improvements
+	//    - Cell sizes are known at each level...can keep a lookup table? will this improve performance?
+	//    - Store cells for each level in separate TDynamicVectors. CellID is then [Level:8 | Index:24].
+	//      This would allow level-grids to be processed separately / in-parallel (for example a cut at given level would be much faster)
+	//    - Currently insertion is max-depth but we do not dynamically expand more than once. So early
+	//      insertions end up in very large buckets. When a child expands we should check if any of its parents would fit.
+	//    - Currently insertion is max-depth so we end up with a huge number of single-object cells. Should only go down a level
+	//      if enough objects exist in current cell. Can do this in a greedy fashion, less optimal but still acceptable...
+	//    - Store an expand-factor for each cell? or an actual AABB for each cell? this would allow for tighter bounds but
+	//      requires accumulating expansion "up" the tree...
+	//    - get rid of ValidObjectIDs, I don't think we need it?
+
+public:
+
+	//
+	// Tree configuration parameters. It is not safe to change these after tree initialization!
+	// 
+
+	/**
+	 * Size of the Root cells of the octree. Objects that don't fit in a Root cell are added to a "Spill set"
+	 */
+	double RootDimension = 1000.0;
+
+	/**
+	 * Fraction we expand the dimension of any cell, to allow extra space to fit objects.
+	 */
+	double MaxExpandFactor = 0.25;
+
+	/**
+	 * Objects will not be inserted more than this many levels deep from a Root cell
+	 */
+	int MaxTreeDepth = 10;
+
+
+public:
+
+	/**
+	 * Test if an object is stored in the tree
+	 * @param ObjectID ID of the object
+	 * @return true if ObjectID is stored in this octree
+	 */
+	inline bool ContainsObject(int32 ObjectID) const;
+
+	/**
+	 * Insert ObjectID into the Octree
+	 * @param ObjectID ID of the object to insert
+	 * @param Bounds bounding box of the object
+	 *
+	 */
+	inline void InsertObject(int32 ObjectID, const FAxisAlignedBox3d& Bounds);
+
+	/**
+	 * Remove an object from the octree
+	 * @param ObjectID ID of the object
+	 * @return true if the object was in the tree and removed
+	 */
+	inline bool RemoveObject(int32 ObjectID);
+
+	/**
+	 * Update the position of an object in the octree. This is more efficient than doing a remove+insert
+	 * @param ObjectID ID of the object
+	 * @param NewBounds enw bounding box of the object
+	 */
+	inline void ReinsertObject(int32 ObjectID, const FAxisAlignedBox3d& NewBounds);
+
+	/**
+	 * Find nearest ray-hit point with objects in tree
+	 * @param Ray the ray 
+	 * @param GetObjectBoundsFunc function that returns bounding box of object identified by ObjectID
+	 * @param HitObjectDistFunc function that returns distance along ray to hit-point on object identified by ObjectID (or TNumericLimits<double>::Max() on miss)
+	 * @param MaxDistance maximum hit distance
+	 * @return ObjectID of hit object, or -1 on miss
+	 */
+	inline int32 FindNearestHitObject(const FRay3d& Ray,
+		TFunctionRef<FAxisAlignedBox3d(int)> GetObjectBoundsFunc,
+		TFunctionRef<double(int, const FRay3d&)> HitObjectDistFunc,
+		double MaxDistance = TNumericLimits<double>::Max()) const;
+
+	/**
+	 * Process ObjectIDs from all the cells with bounding boxes that intersect Bounds
+	 * @param Bounds query box
+	 * @param ObjectIDFunc this function is called 
+	 */
+	inline void RangeQuery(const FAxisAlignedBox3d& Bounds,
+		TFunctionRef<void(int)> ObjectIDFunc) const;
+
+	/**
+	 * Collect ObjectIDs from all the cells with bounding boxes that intersect Bounds
+	 * @param Bounds query box
+	 * @param ObjectIDsOut collected ObjectIDs are stored here 
+	 */
+	inline void RangeQuery(const FAxisAlignedBox3d& Bounds,
+		TArray<int>& ObjectIDsOut ) const;
+
+
+	/**
+	 * Check that the octree is internally valid
+	 * @param IsValidObjectIDFunc function that returns true if given ObjectID is valid
+	 * @param GetObjectBoundSFunc function that returns bounding box of object identified by ObjectID
+	 * @param FailMode how should validity checks fail
+	 * @param bVerbose if true, print some debug info via UE_LOG
+	 * @param bFailOnMissingObjects if true, assume ObjectIDs are dense and that all ObjectIDs must be in the tree
+	 */
+	inline void CheckValidity(
+		TFunctionRef<bool(int)> IsValidObjectIDFunc,
+		TFunctionRef<FAxisAlignedBox3d(int)> GetObjectBoundsFunc,
+		EValidityCheckFailMode FailMode = EValidityCheckFailMode::Check,
+		bool bVerbose = false,
+		bool bFailOnMissingObjects = false) const;
+
+	/**
+	 * statistics about internal structure of the octree
+	 */
+	struct FStatistics
+	{
+		int32 Levels;
+		TArray<int32> LevelBoxCounts;
+		TArray<int32> LevelObjCounts;
+		int32 SpillObjCount;
+		inline FString ToString() const;
+	};
+
+	/**
+	 * Populate given FStatistics with info about the octree
+	 */
+	inline void ComputeStatistics(FStatistics& StatsOut) const;
+
+protected:
+	// this identifier is used for unknown cells
+	static constexpr uint32 InvalidCellID = FSparseOctreeCell::InvalidID;
+
+	// if an object is in the spill cell, that means it didn't fit in the tree
+	static constexpr uint32 SpillCellID = InvalidCellID - 1;
+
+	// reference counts for Cells list. We don't actually need reference counts here, but we need a free
+	// list and iterators, and FRefCountVector provides this
+	FRefCountVector CellRefCounts;
+
+	// list of cells. Note that some cells may be unused, depending on CellRefCounts
+	TDynamicVector<FSparseOctreeCell> Cells;
+
+	FSmallListSet CellObjectLists;			// per-cell object ID lists
+	TSet<int32> SpillObjectSet;				// list of object IDs for objects that didn't fit in a root cell
+
+	TDynamicVector<uint32> ObjectIDToCellMap;	// map from external Object IDs to which cell the object is in (or spill cell, or invalid)
+	FDynamicFlagArray ValidObjectIDs;			// set of ObjectIDs in the tree. This is perhaps not necessary...couldn't we rely on ObjectIDToCellMap?
+
+	// RootCells are the top-level cells of the octree, of size RootDimension. 
+	// So the elements of this sparse grid are CellIDs
+	TSparseGrid3<uint32> RootCells;
+
+
+	// calculate the base width of a cell at a given level
+	inline double GetCellWidth(uint32 Level) const
+	{
+		double Divisor = (double)(1 << Level);
+		double CellWidth = RootDimension / Divisor;
+		return CellWidth;
+	}
+
+
+	FAxisAlignedBox3d GetBox(uint32 Level, const FVector3i& Index, double ExpandFactor) const
+	{
+		double CellWidth = GetCellWidth(Level);
+		double ExpandDelta = CellWidth * ExpandFactor;
+		double MinX = (CellWidth * (double)Index.X) - ExpandDelta;
+		double MinY = (CellWidth * (double)Index.Y) - ExpandDelta;
+		double MinZ = (CellWidth * (double)Index.Z) - ExpandDelta;
+		CellWidth += 2.0 * ExpandDelta;
+		return FAxisAlignedBox3d(
+			FVector3d(MinX, MinY, MinZ),
+			FVector3d(MinX + CellWidth, MinY + CellWidth, MinZ + CellWidth));
+	}
+	inline FAxisAlignedBox3d GetCellBox(const FSparseOctreeCell& Cell, double ExpandFactor = 0) const
+	{
+		return GetBox(Cell.Level, Cell.Index, ExpandFactor);
+	}
+	FVector3d GetCellCenter(const FSparseOctreeCell& Cell) const
+	{
+		double CellWidth = GetCellWidth(Cell.Level);
+		double MinX = CellWidth * (double)Cell.Index.X;
+		double MinY = CellWidth * (double)Cell.Index.Y;
+		double MinZ = CellWidth * (double)Cell.Index.Z;
+		CellWidth *= 0.5;
+		return FVector3d(MinX + CellWidth, MinY + CellWidth, MinZ + CellWidth);
+	}
+
+
+	FVector3i PointToIndex(uint32 Level, const FVector3d& Position) const
+	{
+		double CellWidth = GetCellWidth(Level);
+		int32 i = (int32)FMathd::Floor(Position.X / CellWidth);
+		int32 j = (int32)FMathd::Floor(Position.Y / CellWidth);
+		int32 k = (int32)FMathd::Floor(Position.Z / CellWidth);
+		return FVector3i(i, j, k);
+	}
+
+
+	int ToChildCellIndex(const FSparseOctreeCell& Cell, const FVector3d& Position) const
+	{
+		FVector3d Center = GetCellCenter(Cell);
+		int ChildIndex =
+			((Position.X < Center.X) ? 0 : 1) +
+			((Position.Y < Center.Y) ? 0 : 2) +
+			((Position.Z < Center.Z) ? 0 : 4);
+		return ChildIndex;
+	}
+
+	bool CanFit(const FSparseOctreeCell& Cell, const FAxisAlignedBox3d& Bounds) const
+	{
+		FAxisAlignedBox3d CellBox = GetCellBox(Cell, MaxExpandFactor);
+		return CellBox.Contains(Bounds);
+	}
+
+	uint32 GetCellForObject(int32 ObjectID) const
+	{
+		if (ObjectID >= 0 && ObjectID < ObjectIDToCellMap.Num())
+		{
+			return ObjectIDToCellMap[ObjectID];
+		}
+		return InvalidCellID;
+	}
+
+	inline FSparseOctreeCell FindCurrentContainingCell(const FAxisAlignedBox3d& Bounds) const;
+
+
+	inline void Insert_Spill(int32 ObjectID, const FAxisAlignedBox3d& Bounds);
+	inline void Insert_NewRoot(int32 ObjectID, const FAxisAlignedBox3d& Bounds, FSparseOctreeCell NewRootCell);
+	inline void Insert_ToCell(int32 ObjectID, const FAxisAlignedBox3d& Bounds, const FSparseOctreeCell& ExistingCell);
+	inline void Insert_NewChildCell(int32 ObjectID, const FAxisAlignedBox3d& Bounds, int ParentCellID, FSparseOctreeCell NewChildCell, int ChildIdx);
+
+
+	inline double FindNearestRayCellIntersection(const FSparseOctreeCell& Cell, const FRay3d& Ray) const;
+
+};
+
+
+bool FSparseDynamicOctree3::ContainsObject(int32 ObjectID) const
+{
+	return ValidObjectIDs.Get(ObjectID);
+}
+
+
+void FSparseDynamicOctree3::InsertObject(int32 ObjectID, const FAxisAlignedBox3d& Bounds)
+{
+	check(ContainsObject(ObjectID) == false);
+
+	FSparseOctreeCell CurrentCell = FindCurrentContainingCell(Bounds);
+
+	// if we could not find a containing root cell, we spill
+	if (CurrentCell.Level == FSparseOctreeCell::InvalidLevel)
+	{
+		Insert_Spill(ObjectID, Bounds);
+		return;
+	}
+
+	// if we found a containing root cell but it doesn't exist, create it and insert
+	if (CurrentCell.Level == 0 && CurrentCell.IsExistingCell() == false)
+	{
+		Insert_NewRoot(ObjectID, Bounds, CurrentCell);
+		return;
+	}
+
+	// YIKES this currently does max-depth insertion...
+	//   desired behavior is that parent cell accumulates and then splits later!
+
+	int PotentialChildIdx = ToChildCellIndex(CurrentCell, Bounds.Center());
+	// if current cell does not have this child we might fit there so try it
+	if (CurrentCell.HasChild(PotentialChildIdx) == false)
+	{
+		// todo can we do a fast check based on level and dimensions??
+		FSparseOctreeCell NewChildCell = CurrentCell.MakeChildCell(PotentialChildIdx);
+		if (NewChildCell.Level <= MaxTreeDepth && CanFit(NewChildCell, Bounds))
+		{
+			Insert_NewChildCell(ObjectID, Bounds, CurrentCell.CellID, NewChildCell, PotentialChildIdx);
+			return;
+		}
+	}
+
+	// insert into current cell if
+	//   1) child cell exists (in which case we didn't fit or FindCurrentContainingCell would have returned)
+	//   2) we tried to fit in child cell and failed
+	Insert_ToCell(ObjectID, Bounds, CurrentCell);
+}
+
+void FSparseDynamicOctree3::Insert_NewChildCell(int32 ObjectID, const FAxisAlignedBox3d& Bounds,
+	int ParentCellID, FSparseOctreeCell NewChildCell, int ChildIdx)
+{
+	FSparseOctreeCell& OrigParentCell = Cells[ParentCellID];
+	check(OrigParentCell.HasChild(ChildIdx) == false);
+
+	NewChildCell.CellID = CellRefCounts.Allocate();
+	Cells.InsertAt(NewChildCell, NewChildCell.CellID);
+
+	ObjectIDToCellMap.InsertAt(NewChildCell.CellID, ObjectID);
+	ValidObjectIDs.Set(ObjectID, true);
+
+	CellObjectLists.AllocateAt(NewChildCell.CellID);
+	CellObjectLists.Insert(NewChildCell.CellID, ObjectID);
+
+	OrigParentCell.SetChild(ChildIdx, NewChildCell);
+
+	check(CanFit(NewChildCell, Bounds));
+	check(PointToIndex(NewChildCell.Level, Bounds.Center()) == NewChildCell.Index);
+}
+
+
+void FSparseDynamicOctree3::Insert_ToCell(int32 ObjectID, const FAxisAlignedBox3d& Bounds, const FSparseOctreeCell& ExistingCell)
+{
+	check(CellRefCounts.IsValid(ExistingCell.CellID));
+
+	ObjectIDToCellMap.InsertAt(ExistingCell.CellID, ObjectID);
+	ValidObjectIDs.Set(ObjectID, true);
+
+	CellObjectLists.Insert(ExistingCell.CellID, ObjectID);
+
+	check(CanFit(ExistingCell, Bounds));
+	check(PointToIndex(ExistingCell.Level, Bounds.Center()) == ExistingCell.Index);
+}
+
+
+void FSparseDynamicOctree3::Insert_NewRoot(int32 ObjectID, const FAxisAlignedBox3d& Bounds, FSparseOctreeCell NewRootCell)
+{
+	check(RootCells.Has(NewRootCell.Index) == false);
+
+	NewRootCell.CellID = CellRefCounts.Allocate();
+	Cells.InsertAt(NewRootCell, NewRootCell.CellID);
+
+	ObjectIDToCellMap.InsertAt(NewRootCell.CellID, ObjectID);
+	ValidObjectIDs.Set(ObjectID, true);
+
+	uint32* RootCellElem = RootCells.Get(NewRootCell.Index, true);
+	*RootCellElem = NewRootCell.CellID;
+
+	CellObjectLists.AllocateAt(NewRootCell.CellID);
+	CellObjectLists.Insert(NewRootCell.CellID, ObjectID);
+}
+
+void FSparseDynamicOctree3::Insert_Spill(int32 ObjectID, const FAxisAlignedBox3d& Bounds)
+{
+	SpillObjectSet.Add(ObjectID);
+	ObjectIDToCellMap.InsertAt(SpillCellID, ObjectID);
+	ValidObjectIDs.Set(ObjectID, true);
+}
+
+
+
+
+bool FSparseDynamicOctree3::RemoveObject(int32 ObjectID)
+{
+	if (ContainsObject(ObjectID) == false)
+	{
+		return false;
+	}
+
+	uint32 CellID = GetCellForObject(ObjectID);
+	if (CellID == SpillCellID)
+	{
+		int32 RemovedCount = SpillObjectSet.Remove(ObjectID);
+		check(RemovedCount > 0);
+		ValidObjectIDs.Set(ObjectID, false);
+		return (RemovedCount > 0);
+	}
+	if (CellID == InvalidCellID)
+	{
+		return false;
+	}
+
+	ObjectIDToCellMap[ObjectID] = InvalidCellID;
+	ValidObjectIDs.Set(ObjectID, false);
+
+	bool bInList = CellObjectLists.Remove(CellID, ObjectID);
+	check(bInList);
+	return true;
+}
+
+
+
+
+
+void FSparseDynamicOctree3::ReinsertObject(int32 ObjectID, const FAxisAlignedBox3d& NewBounds)
+{
+	if (ContainsObject(ObjectID))
+	{
+		uint32 CellID = GetCellForObject(ObjectID);
+		if (CellID != SpillCellID && CellID != InvalidCellID)
+		{
+			FSparseOctreeCell& CurrentCell = Cells[CellID];
+			if (CanFit(CurrentCell, NewBounds))
+			{
+				return;		// everything is fine
+			}
+
+		}
+	}
+
+	RemoveObject(ObjectID);
+	InsertObject(ObjectID, NewBounds);
+}
+
+
+
+
+
+double FSparseDynamicOctree3::FindNearestRayCellIntersection(const FSparseOctreeCell& Cell, const FRay3d& Ray) const
+{
+	FAxisAlignedBox3d Box = GetCellBox(Cell, MaxExpandFactor);
+	double ray_t = TNumericLimits<double>::Max();
+	if (FIntrRay3AxisAlignedBox3d::FindIntersection(Ray, Box, ray_t))
+	{
+		return ray_t;
+	}
+	else
+	{
+		return TNumericLimits<double>::Max();
+	}
+}
+
+
+
+int32 FSparseDynamicOctree3::FindNearestHitObject(const FRay3d& Ray,
+	TFunctionRef<FAxisAlignedBox3d(int)> GetObjectBoundsFunc,
+	TFunctionRef<double(int, const FRay3d&)> HitObjectDistFunc,
+	double MaxDistance) const
+{
+	// this should take advantage of raster!
+
+	int32 HitObjectID = -1;
+	TArray<const FSparseOctreeCell*> Queue;
+	Queue.Reserve(32);
+
+	// start at root cells
+	RootCells.AllocatedIteration([&](const uint32* RootCellID)
+	{
+		const FSparseOctreeCell* RootCell = &Cells[*RootCellID];
+		double RayHitParam = FindNearestRayCellIntersection(*RootCell, Ray);
+		if (RayHitParam < MaxDistance)
+		{
+			Queue.Add(&Cells[*RootCellID]);
+		}
+	});
+
+
+	while (Queue.Num() > 0)
+	{
+		const FSparseOctreeCell* CurCell = Queue.Pop(false);
+		
+		// process elements
+		for (int ObjectID : CellObjectLists.Values(CurCell->CellID))
+		{
+			double HitDist = HitObjectDistFunc(ObjectID, Ray);
+			if (HitDist < MaxDistance)
+			{
+				MaxDistance = HitDist;
+				HitObjectID = ObjectID;
+			}
+		}
+
+		// descend to child cells
+		// sort by distance? use DDA?
+		for (int k = 0; k < 8; ++k) 
+		{
+			if (CurCell->HasChild(k))
+			{
+				const FSparseOctreeCell* ChildCell = &Cells[CurCell->GetChildCellID(k)];
+				double RayHitParam = FindNearestRayCellIntersection(*ChildCell, Ray);
+				if (RayHitParam < MaxDistance)
+				{
+					Queue.Add(ChildCell);
+				}
+			}
+		}
+	}
+
+	return HitObjectID;
+}
+
+
+
+
+
+void FSparseDynamicOctree3::RangeQuery(
+	const FAxisAlignedBox3d& Bounds,
+	TFunctionRef<void(int)> ObjectIDFunc) const
+{
+	// todo: this should take advantage of raster!
+
+	TArray<const FSparseOctreeCell*> Queue;
+
+	// start at root cells
+	RootCells.AllocatedIteration([&](const uint32* RootCellID)
+	{
+		const FSparseOctreeCell* RootCell = &Cells[*RootCellID];
+		if ( GetCellBox(*RootCell, MaxExpandFactor).Intersects(Bounds) )
+		{
+			Queue.Add(&Cells[*RootCellID]);
+		}
+	});
+
+
+	while (Queue.Num() > 0)
+	{
+		const FSparseOctreeCell* CurCell = Queue.Pop(false);
+
+		// process elements
+		for (int ObjectID : CellObjectLists.Values(CurCell->CellID))
+		{
+			ObjectIDFunc(ObjectID);
+		}
+
+		for (int k = 0; k < 8; ++k)
+		{
+			if (CurCell->HasChild(k))
+			{
+				const FSparseOctreeCell* ChildCell = &Cells[CurCell->GetChildCellID(k)];
+				if (GetCellBox(*ChildCell, MaxExpandFactor).Intersects(Bounds))
+				{
+					Queue.Add(ChildCell);
+				}
+			}
+		}
+	}
+}
+
+
+
+
+
+
+
+void FSparseDynamicOctree3::RangeQuery(
+	const FAxisAlignedBox3d& Bounds,
+	TArray<int>& ObjectIDs) const
+{
+	// todo: this should take advantage of raster!
+
+	TArray<const FSparseOctreeCell*> Queue;
+
+	// start at root cells
+	RootCells.AllocatedIteration([&](const uint32* RootCellID)
+	{
+		const FSparseOctreeCell* RootCell = &Cells[*RootCellID];
+		if (GetCellBox(*RootCell, MaxExpandFactor).Intersects(Bounds))
+		{
+			Queue.Add(&Cells[*RootCellID]);
+		}
+	});
+
+
+	while (Queue.Num() > 0)
+	{
+		const FSparseOctreeCell* CurCell = Queue.Pop(false);
+
+		// process elements
+		for (int ObjectID : CellObjectLists.Values(CurCell->CellID))
+		{
+			ObjectIDs.Add(ObjectID);
+		}
+
+		for (int k = 0; k < 8; ++k)
+		{
+			if (CurCell->HasChild(k))
+			{
+				const FSparseOctreeCell* ChildCell = &Cells[CurCell->GetChildCellID(k)];
+				if (GetCellBox(*ChildCell, MaxExpandFactor).Intersects(Bounds))
+				{
+					Queue.Add(ChildCell);
+				}
+			}
+		}
+	}
+}
+
+
+
+
+
+
+FSparseOctreeCell FSparseDynamicOctree3::FindCurrentContainingCell(const FAxisAlignedBox3d& Bounds) const
+{
+	double BoxWidth = Bounds.MaxDim();
+	FVector3d BoxCenter = Bounds.Center();
+
+	// look up root cell, which may not exist
+	FVector3i RootIndex = PointToIndex(0, BoxCenter);
+	const uint32* RootCellID = RootCells.Get(RootIndex);
+	if (RootCellID == nullptr)
+	{
+		return FSparseOctreeCell(0, RootIndex);
+	}
+	check(CellRefCounts.IsValid(*RootCellID));
+
+	// check if box is contained in root cell, if not we have to spill
+	// (should we do this before checking for existence? we can...)
+	const FSparseOctreeCell* RootCell = &Cells[*RootCellID];
+	if (CanFit(*RootCell, Bounds) == false)
+	{
+		return FSparseOctreeCell(FSparseOctreeCell::InvalidLevel, FVector3i::Zero());
+	}
+
+	const FSparseOctreeCell* CurrentCell = RootCell;
+	do
+	{
+		int ChildIdx = ToChildCellIndex(*CurrentCell, BoxCenter);
+		if (CurrentCell->HasChild(ChildIdx))
+		{
+			int32 ChildCellID = CurrentCell->GetChildCellID(ChildIdx);
+			check(CellRefCounts.IsValid(ChildCellID));
+			const FSparseOctreeCell* ChildCell = &Cells[ChildCellID];
+			if (CanFit(*ChildCell, Bounds))
+			{
+				CurrentCell = ChildCell;
+				continue;
+			}
+		}
+
+		return *CurrentCell;
+
+	} while (true);		// loop will always terminate
+
+	return FSparseOctreeCell(FSparseOctreeCell::InvalidLevel, FVector3i::Zero());
+}
+
+
+void FSparseDynamicOctree3::CheckValidity(
+	TFunctionRef<bool(int)> IsValidObjectIDFunc,
+	TFunctionRef<FAxisAlignedBox3d(int)> GetObjectBoundsFunc,
+	EValidityCheckFailMode FailMode,
+	bool bVerbose,
+	bool bFailOnMissingObjects) const
+{
+	bool is_ok = true;
+	TFunction<void(bool)> CheckOrFailF = [&](bool b)
+	{
+		is_ok = is_ok && b;
+	};
+	if (FailMode == EValidityCheckFailMode::Check)
+	{
+		CheckOrFailF = [&](bool b)
+		{
+			checkf(b, TEXT("FSparseDynamicOctree3::CheckValidity failed!"));
+			is_ok = is_ok && b;
+		};
+	}
+	else if (FailMode == EValidityCheckFailMode::Ensure)
+	{
+		CheckOrFailF = [&](bool b)
+		{
+			ensureMsgf(b, TEXT("FSparseDynamicOctree3::CheckValidity failed!"));
+			is_ok = is_ok && b;
+		};
+	}
+
+	TArray<int> CellsAtLevels, ObjectsAtLevel;
+	CellsAtLevels.Init(0, 32);
+	ObjectsAtLevel.Init(0, 32);
+	uint32 SpillObjectCount = 0;
+	uint32 MissingObjectCount = 0;
+	uint8 MaxLevel = 0;
+
+	// check that all object IDs in per-cell object lists is valid
+	for (int32 CellID : CellRefCounts.Indices())
+	{
+		for (int32 ObjectID : CellObjectLists.Values(CellID))
+		{
+			CheckOrFailF(IsValidObjectIDFunc(ObjectID));
+		}
+	}
+
+	uint32 NumObjectIDs = ObjectIDToCellMap.Num();
+	for (uint32 ObjectID = 0; ObjectID < NumObjectIDs; ObjectID++)
+	{
+		bool IsValidObjectID = IsValidObjectIDFunc(ObjectID);
+		if (IsValidObjectID)
+		{
+			FAxisAlignedBox3d ObjectBounds = GetObjectBoundsFunc(ObjectID);
+
+			CheckOrFailF(ObjectID < ObjectIDToCellMap.Num());
+			uint32 CellID = ObjectIDToCellMap[ObjectID];
+
+			if (bFailOnMissingObjects)
+			{
+				CheckOrFailF(CellID != InvalidCellID);
+			}
+
+			if (CellID == SpillCellID)
+			{
+				// this is a spill node...
+				SpillObjectCount++;
+				check(SpillObjectSet.Contains(ObjectID));
+			}
+			else if (CellID == InvalidCellID)
+			{
+				MissingObjectCount++;
+				check(SpillObjectSet.Contains(ObjectID) == false);
+			}
+			else
+			{
+				CheckOrFailF(CellRefCounts.IsValid(CellID));
+				FSparseOctreeCell Cell = Cells[CellID];
+				FAxisAlignedBox3d CellBounds = GetCellBox(Cell, MaxExpandFactor);
+				CheckOrFailF(CellBounds.Contains(ObjectBounds));
+				CheckOrFailF(CellObjectLists.Contains(CellID, ObjectID));
+
+				ObjectsAtLevel[Cell.Level]++;
+			}
+		}
+	}
+
+
+	for (int32 CellID : CellRefCounts.Indices())
+	{
+		const FSparseOctreeCell& Cell = Cells[CellID];
+		CellsAtLevels[Cell.Level]++;
+		MaxLevel = FMath::Max(MaxLevel, Cell.Level);
+	}
+
+	if (bVerbose)
+	{
+		UE_LOG(LogTemp, Warning, TEXT("FSparseDynamicOctree3::CheckValidity: MaxLevel %d  SpillCount %d  MissingCount %d"), MaxLevel, SpillObjectCount, MissingObjectCount);
+		for (uint32 k = 0; k <= MaxLevel; ++k)
+		{
+			UE_LOG(LogTemp, Warning, TEXT("    Level %4d  Cells %4d  Objects %4d"), k, CellsAtLevels[k], ObjectsAtLevel[k]);
+		}
+	}
+}
+
+
+
+void FSparseDynamicOctree3::ComputeStatistics(FStatistics& StatsOut) const
+{
+	StatsOut.SpillObjCount = SpillObjectSet.Num();
+
+	StatsOut.Levels = 0;
+	for (int32 CellID : CellRefCounts.Indices())
+	{
+		const FSparseOctreeCell& Cell = Cells[CellID];
+		StatsOut.Levels = FMath::Max(StatsOut.Levels, (int32)Cell.Level);
+	}
+	StatsOut.Levels++;
+	StatsOut.LevelObjCounts.Init(0, StatsOut.Levels);
+	StatsOut.LevelBoxCounts.Init(0, StatsOut.Levels);
+	for (int32 CellID : CellRefCounts.Indices())
+	{
+		const FSparseOctreeCell& Cell = Cells[CellID];
+		StatsOut.LevelBoxCounts[Cell.Level]++;
+		StatsOut.LevelObjCounts[Cell.Level] += CellObjectLists.GetCount(CellID);
+	}
+}
+
+FString FSparseDynamicOctree3::FStatistics::ToString() const
+{
+	FString Result = FString::Printf(
+		TEXT("Levels %2d   SpillCount %5d \r\n"), Levels, SpillObjCount);
+	for (int k = 0; k < Levels; ++k)
+	{
+		Result += FString::Printf(TEXT("  Level %2d:  Cells %8d  Tris %8d  Avg %5.3f\r\n"), 
+			k, LevelBoxCounts[k], LevelObjCounts[k], ((float)LevelObjCounts[k] / (float)LevelBoxCounts[k])  );
+	}
+	return Result;
+}
\ No newline at end of file
diff --git a/Engine/Plugins/Experimental/GeometryProcessing/Source/GeometricObjects/Public/Util/IndexSetDecompositions.h b/Engine/Plugins/Experimental/GeometryProcessing/Source/GeometricObjects/Public/Util/IndexSetDecompositions.h
new file mode 100644
--- /dev/null
+++ b/Engine/Plugins/Experimental/GeometryProcessing/Source/GeometricObjects/Public/Util/IndexSetDecompositions.h
@@ -0,0 +1,82 @@
+// Copyright 1998-2019 Epic Games, Inc. All Rights Reserved.
+
+#pragma once
+
+#include "VectorTypes.h"
+
+
+//
+// @todo: base interface that allows for iterating over
+// different kinds of index decompositions in a generic way
+// 
+
+
+/**
+ * FArrayIndexSetsDecomposition represents a decomposition of an integer index set
+ * into subsets, which are stored as arrays (ie this is just an array of arrays,
+ * where each sub-array has a unique integer identifier)
+ */
+class FArrayIndexSetsDecomposition
+{
+protected:
+
+	struct FIndexSet
+	{
+		int32 Identifier;
+		TArray<int32> IndexArray;
+		// todo alternate storage: set, arbitrary function
+	};
+
+	TIndirectArray<FIndexSet> Sets;
+	TMap<int32, FIndexSet*> SetIdentifierMap;
+	int32 IDCounter = 1;
+
+	TArray<int32> AllSetIDs;
+
+public:
+
+	/*
+	 * @return list of ID integers which identify the IndexSets of this Decomposition
+	 */
+	const TArray<int32>& GetIndexSetIDs() const
+	{
+		return AllSetIDs;
+	}
+
+	/**
+	 * Create a new IndexSet for this Decomposition
+	 * @return ID of the new IndexSet
+	 */
+	int32 CreateNewIndexSet()
+	{
+		FIndexSet* Set = new FIndexSet();
+		Set->Identifier = IDCounter++;
+		Sets.Add(Set);
+		SetIdentifierMap.Add(Set->Identifier, Set);
+		AllSetIDs.Add(Set->Identifier);
+		return Set->Identifier;
+	}
+
+	/**
+	 * @return IndexSet for the given SetID
+	 */
+	TArray<int32>& GetIndexSetArray(int32 SetID)
+	{
+		FIndexSet** Found = SetIdentifierMap.Find(SetID);
+		check(Found != nullptr);
+		return (*Found)->IndexArray;
+	}
+
+	/**
+	 * @return IndexSet for the given SetID
+	 */
+	const TArray<int32>& GetIndexSetArray(int32 SetID) const
+	{
+ 		FIndexSet*const* Found = SetIdentifierMap.Find(SetID);
+		check(Found != nullptr);
+		return (*Found)->IndexArray;
+	}
+
+
+
+};
\ No newline at end of file
