MeshProcessing: FSubRegionRemesher now tracks Triangle ROI, and incrementally updates edge/vertex ROI, instead of requiring full pass to rebuild in ::UpdateROI(). Added localized smoothing pass. Can now track removed triangles and provide to external code to do localized updates of other triangle data structures (eg octree).
#rnx

diff --git a/Engine/Plugins/Experimental/GeometryProcessing/Source/DynamicMesh/Public/SubRegionRemesher.h b/Engine/Plugins/Experimental/GeometryProcessing/Source/DynamicMesh/Public/SubRegionRemesher.h
--- a/Engine/Plugins/Experimental/GeometryProcessing/Source/DynamicMesh/Public/SubRegionRemesher.h
+++ b/Engine/Plugins/Experimental/GeometryProcessing/Source/DynamicMesh/Public/SubRegionRemesher.h
@@ -5,21 +5,49 @@
 #include "Remesher.h"
 #include "Util/BufferUtil.h"
 #include "DynamicMeshChangeTracker.h"
+#include "MeshWeights.h"
 
 /**
- * FSubRegionRemesher is an extension
+ * FSubRegionRemesher is an extension of FRemesher that allows for constraining remeshing to
+ * a localized region of a mesh. Currently this is initialized from a Vertex ROI.
+ * 
+ * @warning Currently "boundary" edges of the ROI that are split will result in the ROI
+ * growing to include both new edges created by the split. 
+ * 
  */
 class FSubRegionRemesher : public FRemesher
 {
 protected:
-	// temp buffer of edges connected to VertexROI
+	// active set of edges we will consider for remeshing. This set is updated on 
+	// each edge flip/split/collapse, but is not used during the pass.
 	TSet<int> EdgeROI;
-	// list of edges to consider during current remesh pass (not updated during pass)
+
+	// Active set of triangles. This is conservative, ideally contains one-rings of all
+	// edge-vertices of EdgeROI set, but may include additional triangles accumulated over time
+	// (should *not* contain any triangles that no longer exist)
+	TSet<int> TriangleROI;
+
+	// static list of edges to consider during a pass (set in UpdateROI which must be called each pass)
 	TArray<int> Edges;
-	// current edge we are at in StartEdges/GetNextEdge iteration
+
+	// index of current edge in .Edges we are at in StartEdges/GetNextEdge modulo-iteration
 	int CurEdge;
 
+	// set of triangles removed in last pass. Enable this by calling BeginTrackRemovedTrisInPass()
+	TSet<int> RemovedLastPassTris;
+
+	// controls whether RemovedLastPass will be populated
+	bool bTrackRemoved = false;
+
+
+	// counters for making sure that UpdateROI has been called
+	uint32 LastUpdateROICounter = 0;
+	uint32 LastRemeshPassCounter = 0;
+
+
 public:
+
+
 	FSubRegionRemesher(FDynamicMesh3* Mesh) : FRemesher(Mesh)
 	{
 		VertexControlF = [this](int vid) {
@@ -35,11 +63,14 @@ public:
 
 
 	/**
-	 * Must call UpdateROI before each remesh pass.
+	 * Initialize edge-subregion ROI from the VertexROI member that has been externally initialized
 	 */
-	void UpdateROI()
+	void InitializeFromVertexROI()
 	{
 		EdgeROI.Reset();
+		TriangleROI.Reset();
+
+		// to get active edge set
 		for (int VertIdx : VertexROI)
 		{
 			for (int eid : GetMesh()->VtxEdgesItr(VertIdx))
@@ -47,18 +78,8 @@ public:
 				EdgeROI.Add(eid);
 			}
 		}
-		Edges.Reset();
-		BufferUtil::AppendElements(Edges, EdgeROI);
-	}
+		UpdateROI();
 
-
-	/**
-	 * Tell a MeshChangeTracker about the set of triangles that we might modify in the next remesh pass.
-	 * This could include one-rings of either side of an edge in the ROI, if we collapse.
-	 */
-	TSet<int> Vertices;
-	void SaveActiveROI(FDynamicMeshChangeTracker* Change)
-	{
 		// there is quite a bit of overhead here...perhaps remesher could just save triangles
 		// itself before it touches them?
 
@@ -66,7 +87,7 @@ public:
 
 		// figuring out unique verts means we don't do each vertex N~=valence times, 
 		// which saves a lot of one-ring iterations that are somewhat expensive...
-		Vertices.Reset();
+		TSet<int> Vertices;
 		const FDynamicMesh3* UseMesh = GetMesh();
 		for (int eid : EdgeROI)
 		{
@@ -74,16 +95,103 @@ public:
 			Vertices.Add(EdgeVerts.A);
 			Vertices.Add(EdgeVerts.B);
 		}
-		for ( int vid : Vertices)
+		for (int vid : Vertices)
 		{
 			OneRingTris.Reset();
 			UseMesh->GetVertexOneRingTriangles(vid, OneRingTris);
 			for (int tid : OneRingTris)
 			{
-				Change->SaveTriangle(tid, true);
+				TriangleROI.Add(tid);
 			}
 		}
+	}
+
+
+	/**
+	 * Update the internal data structures in preparation for a call to FRemesher::BasicRemeshPass.
+	 * This must be called before each remesh pass!
+	 */
+	void UpdateROI()
+	{
+		Edges.Reset();
+		for (int eid : EdgeROI)
+		{
+			//check(GetMesh()->IsEdge(eid));
+			Edges.Add(eid);
+		}
+		LastUpdateROICounter++;
+	}
+
+	/**
+	 * Call before BasicRemeshPass() to enable tracking of removed triangles
+	 */
+	void BeginTrackRemovedTrisInPass()
+	{
+		RemovedLastPassTris.Reset();
+		bTrackRemoved = true;
+	}
+
+	/**
+	 * Call after BasicRemeshPass() to disable and return tracking of removed triangles
+	 * @return set of removed triangles. This set will be cleared on next call to BeginTrackRemovedTrisInPass()
+	 */
+	const TSet<int32>& EndTrackRemovedTrisInPass()
+	{
+		bTrackRemoved = false;
+		return RemovedLastPassTris;
+	}
+
+
+
+	/**
+	 * forwards to FRemesher::BasicRemeshPass
+	 */
+	virtual void BasicRemeshPass() override
+	{
+		check(LastRemeshPassCounter != LastUpdateROICounter);
+		LastRemeshPassCounter = LastUpdateROICounter;
+
+		FRemesher::BasicRemeshPass();
+	}
+
+
+
+	/**
+	 * Tell a MeshChangeTracker about the set of triangles that we might modify in the next remesh pass.
+	 * This could include one-rings of either side of an edge in the ROI, if we collapse.
+	 */
+
+	void SaveActiveROI(FDynamicMeshChangeTracker* Change)
+	{
+		for (int tid : TriangleROI)
+		{
+			Change->SaveTriangle(tid, true);
+		}
+	}
+
+	/**
+	 * @return set of triangles that contains edge ROI (note: may also contain additional triangles)
+	 */
+	const TSet<int>& GetCurrentTriangleROI() const
+	{
+		return TriangleROI;
+	}
+
+	/**
+	 * @return set of edges in current edge ROI
+	 */
+	const TSet<int>& GetCurrentEdgeROI() const
+	{
+		return EdgeROI;
+	}
 
+	/**
+	 * @return current edge array. 
+	 * @warning This is only valid after calling UpdateROI() and before calling BasicRemeshPass().
+	 */
+	const TArray<int>& GetCurrentEdgeArray() const
+	{
+		return Edges;
 	}
 
 
@@ -112,19 +220,144 @@ protected:
 		}
 	}
 
+
+
+
 	virtual void OnEdgeSplit(int EdgeID, int VertexA, int VertexB, const FDynamicMesh3::FEdgeSplitInfo& SplitInfo) override
 	{
 		VertexROI.Add(SplitInfo.NewVertex);
+		EdgeROI.Add(SplitInfo.NewEdges.A);
+
+		// By always adding new edges to ROI, we are potentially 'growing' the ROI here.
+		// Could filter out these edges by checking if other vtx is in VertexROI?
+
+		EdgeROI.Add(SplitInfo.NewEdges.B);
+		AddEdgeToTriangleROI(SplitInfo.NewEdges.B);
+		if (SplitInfo.NewEdges.C != FDynamicMesh3::InvalidID)
+		{
+			EdgeROI.Add(SplitInfo.NewEdges.C);
+			AddEdgeToTriangleROI(SplitInfo.NewEdges.C);
+		}
+
+		// these two triangles should be already added by AddEdgeToTriangleROI() calls above...
+		TriangleROI.Add(SplitInfo.NewTriangles.A);
+		if (SplitInfo.NewTriangles.B != FDynamicMesh3::InvalidID)
+		{
+			TriangleROI.Add(SplitInfo.NewTriangles.B);
+		}
+		//AddVertexToTriangleROI(SplitInfo.NewVertex);
+		//for (int tid : GetMesh()->VtxTrianglesItr(SplitInfo.NewVertex))
+		//	check(TriangleROI.Contains(tid));
 	}
 
+
 	virtual void OnEdgeCollapse(int EdgeID, int VertexA, int VertexB, const FDynamicMesh3::FEdgeCollapseInfo& CollapseInfo) override
 	{
+		// remove triangles from ROI
+		TriangleROI.Remove(CollapseInfo.RemovedTris.A);
+		if (bTrackRemoved)
+		{
+			RemovedLastPassTris.Add(CollapseInfo.RemovedTris.A);
+		}
+		if (CollapseInfo.RemovedTris.B != FDynamicMesh3::InvalidID)
+		{
+			TriangleROI.Remove(CollapseInfo.RemovedTris.B);
+			if (bTrackRemoved)
+			{
+				RemovedLastPassTris.Add(CollapseInfo.RemovedTris.B);
+			}
+		}
+
+		// remove vtx
 		VertexROI.Remove(CollapseInfo.RemovedVertex);
+
+		// remove edges
+		EdgeROI.Remove(CollapseInfo.CollapsedEdge);
+		EdgeROI.Remove(CollapseInfo.RemovedEdges.A);
+		if (CollapseInfo.RemovedEdges.B != FDynamicMesh3::InvalidID)
+		{
+			EdgeROI.Remove(CollapseInfo.RemovedEdges.B);
+		}
+
+		// one-ring tris of remaining vtx should already be in TriangleROI!!
+		//AddVertexToTriangleROI(CollapseInfo.KeptVertex);
+		//for (int tid : GetMesh()->VtxTrianglesItr(CollapseInfo.KeptVertex))
+		//	check(TriangleROI.Contains(tid));
+	}
+
+
+	virtual void OnEdgeFlip(int EdgeID, const FDynamicMesh3::FEdgeFlipInfo& FlipInfo)
+	{
+		// flipping an edge potentially connects new verts to the ROI
+
+		FIndex2i EdgeV = GetMesh()->GetEdgeV(EdgeID);
+		VertexROI.Add(EdgeV.A);
+		VertexROI.Add(EdgeV.B);
+
+		AddVertexToTriangleROI(EdgeV.A);
+		AddVertexToTriangleROI(EdgeV.B);
 	}
 
+
+
+	void AddVertexToTriangleROI(int VertexID)
+	{
+		for (int tid : GetMesh()->VtxTrianglesItr(VertexID))
+		{
+			TriangleROI.Add(tid);
+		}
+	}
+
+	void AddEdgeToTriangleROI(int EdgeID)
+	{
+		FIndex2i EdgeV = GetMesh()->GetEdgeV(EdgeID);
+		AddVertexToTriangleROI(EdgeV.A);
+		AddVertexToTriangleROI(EdgeV.B);
+	}
+
+
+
 	EVertexControl VertexFilter(int VertexID)
 	{
 		return (VertexROI.Contains(VertexID) == false) ? EVertexControl::NoMovement : EVertexControl::AllowAll;
 	}
 
+
+
+
+
+	//
+	// localized smoothing
+	// 
+	
+	TMap<int, FVector3d> SmoothedPositions;
+
+	virtual void FullSmoothPass_Buffer(bool bParallel) override
+	{
+		SmoothedPositions.Reset();
+
+		TFunction<FVector3d(const FDynamicMesh3&, int, double)> UseSmoothFunc = GetSmoothFunction();
+		auto SmoothAndUpdateFunc = [&](int vid)
+		{
+			bool bModified = false;
+			FVector3d SmoothedPosition = ComputeSmoothedVertexPos(vid, UseSmoothFunc, bModified);
+			if (bModified)
+			{
+				SmoothedPositions.Add(vid, SmoothedPosition);
+			}
+		};
+
+		for (int vid : VertexROI)
+		{
+			SmoothAndUpdateFunc(vid);
+		}
+
+		for (auto Pair : SmoothedPositions)
+		{
+			Mesh->SetVertex(Pair.Key, Pair.Value);
+		}
+	}
+
+
+
 };
