MeshProcessing: forward return resolve to fix CIS warning
#rnx

diff --git a/Engine/Plugins/Experimental/GeometryProcessing/Source/DynamicMesh/Public/DynamicMeshOctree3.h b/Engine/Plugins/Experimental/GeometryProcessing/Source/DynamicMesh/Public/DynamicMeshOctree3.h
--- a/Engine/Plugins/Experimental/GeometryProcessing/Source/DynamicMesh/Public/DynamicMeshOctree3.h
+++ b/Engine/Plugins/Experimental/GeometryProcessing/Source/DynamicMesh/Public/DynamicMeshOctree3.h
@@ -97,7 +97,7 @@ public:
 	{
 		FAxisAlignedBox3d Bounds = Mesh->GetTriBounds(TriangleID);
 		ModifiedBounds.Contain(Bounds);
-		RemoveObject(TriangleID);
+		return RemoveObject(TriangleID);
 	}
 
 	/**
