Fancy loop stitcher for sparsely corresponded loops, fixes + test coverage for machinery to cut holes through meshes

#rnx

diff --git a/Engine/Plugins/Experimental/GeometryProcessing/Source/DynamicMesh/Private/DynamicMeshEditor.cpp b/Engine/Plugins/Experimental/GeometryProcessing/Source/DynamicMesh/Private/DynamicMeshEditor.cpp
--- a/Engine/Plugins/Experimental/GeometryProcessing/Source/DynamicMesh/Private/DynamicMeshEditor.cpp
+++ b/Engine/Plugins/Experimental/GeometryProcessing/Source/DynamicMesh/Private/DynamicMeshEditor.cpp
@@ -103,6 +103,98 @@ operation_failed:
 
 
 
+bool FDynamicMeshEditor::StitchSparselyCorrespondedVertexLoops(const TArray<int>& VertexIDs1, const TArray<int>& MatchedIndices1, const TArray<int>& VertexIDs2, const TArray<int>& MatchedIndices2, FDynamicMeshEditResult& ResultOut)
+{
+	int CorrespondN = MatchedIndices1.Num();
+	if (!ensureMsgf(CorrespondN == MatchedIndices2.Num(), TEXT("FDynamicMeshEditor::StitchSparselyCorrespondedVertices: correspondence arrays are not the same length!")))
+	{
+		return false;
+	}
+	// TODO: support case of only one corresponded vertex & a connecting a full loop around?
+	// this requires allowing start==end to not immediately stop the walk ...
+	if (!ensureMsgf(CorrespondN >= 2, TEXT("Must have at least two corresponded vertices")))
+	{
+		return false;
+	}
+	ResultOut.NewGroups.Reserve(CorrespondN);
+
+	int i = 0;
+	for (; i < CorrespondN; ++i) 
+	{
+		int Starts[2] { MatchedIndices1[i], MatchedIndices2[i] };
+		int Ends[2] { MatchedIndices1[(i + 1) % CorrespondN], MatchedIndices2[(i + 1) % CorrespondN] };
+
+		auto GetWrappedSpanLen = [](const FDynamicMesh3* Mesh, const TArray<int>& VertexIDs, int StartInd, int EndInd)
+		{
+			float LenTotal = 0;
+			FVector3d V = Mesh->GetVertex(VertexIDs[StartInd]);
+			for (int Ind = StartInd, IndNext; Ind != EndInd;)
+			{
+				IndNext = (Ind + 1) % VertexIDs.Num();
+				FVector3d VNext = Mesh->GetVertex(VertexIDs[IndNext]);
+				LenTotal += V.Distance(VNext);
+				Ind = IndNext;
+				V = VNext;
+			}
+			return LenTotal;
+		};
+		float LenTotal[2] { GetWrappedSpanLen(Mesh, VertexIDs1, Starts[0], Ends[0]), GetWrappedSpanLen(Mesh, VertexIDs2, Starts[1], Ends[1]) };
+		float LenAlong[2] { FMathf::Epsilon, FMathf::Epsilon };
+		LenTotal[0] += FMathf::Epsilon;
+		LenTotal[1] += FMathf::Epsilon;
+
+
+		int NewGroupID = Mesh->AllocateTriangleGroup();
+		ResultOut.NewGroups.Add(NewGroupID);
+		
+		int Walks[2]{ Starts[0], Starts[1] };
+		FVector3d Vertex[2]{ Mesh->GetVertex(VertexIDs1[Starts[0]]), Mesh->GetVertex(VertexIDs2[Starts[1]]) };
+		while (Walks[0] != Ends[0] || Walks[1] != Ends[1])
+		{
+			float PctAlong[2]{ LenAlong[0] / LenTotal[0], LenAlong[1] / LenTotal[1] };
+			bool bAdvanceSecond = (Walks[0] == Ends[0] || (Walks[1] != Ends[1] && PctAlong[0] > PctAlong[1]));
+			FIndex3i Tri(VertexIDs1[Walks[0]], VertexIDs2[Walks[1]], -1);
+			if (!bAdvanceSecond)
+			{
+				Walks[0] = (Walks[0] + 1) % VertexIDs1.Num();
+				
+				Tri.C = VertexIDs1[Walks[0]];
+				FVector3d NextV = Mesh->GetVertex(Tri.C);
+				LenAlong[0] += NextV.Distance(Vertex[0]);
+				Vertex[0] = NextV;
+			}
+			else
+			{
+				Walks[1] = (Walks[1] + 1) % VertexIDs2.Num();
+				Tri.C = VertexIDs2[Walks[1]];
+				FVector3d NextV = Mesh->GetVertex(Tri.C);
+				LenAlong[1] += NextV.Distance(Vertex[1]);
+				Vertex[1] = NextV;
+			}
+			int Tid = Mesh->AppendTriangle(Tri, NewGroupID);
+			ResultOut.NewTriangles.Add(Tid);
+
+			if (Tid < 0)
+			{
+				goto operation_failed;
+			}
+		}
+	}
+
+	return true;
+
+operation_failed:
+	// remove what we added so far
+	if (ResultOut.NewTriangles.Num()) 
+	{
+		if (!RemoveTriangles(ResultOut.NewTriangles, false))
+		{
+			checkf(false, TEXT("FDynamicMeshEditor::StitchSparselyCorrespondedVertexLoops: failed to add all triangles, and also failed to back out changes."));
+		}
+	}
+	return false;
+}
+
 bool FDynamicMeshEditor::AddTriangleFan_OrderedVertexLoop(int CenterVertex, const TArray<int>& VertexLoop, int GroupID, FDynamicMeshEditResult& ResultOut)
 {
 	if (GroupID == -1)
diff --git a/Engine/Plugins/Experimental/GeometryProcessing/Source/DynamicMesh/Private/Operations/EmbedSurfacePath.cpp b/Engine/Plugins/Experimental/GeometryProcessing/Source/DynamicMesh/Private/Operations/EmbedSurfacePath.cpp
--- a/Engine/Plugins/Experimental/GeometryProcessing/Source/DynamicMesh/Private/Operations/EmbedSurfacePath.cpp
+++ b/Engine/Plugins/Experimental/GeometryProcessing/Source/DynamicMesh/Private/Operations/EmbedSurfacePath.cpp
@@ -583,9 +583,10 @@ bool FMeshSurfacePath::AddViaPlanarWalk(int StartTri, FVector3d StartPt, int End
 * @param PathVertices Indices of the vertices on the path after embedding succeeds; NOTE these will not be 1:1 with the input Path
 * @return true if embedding succeeded.
 */
-bool FMeshSurfacePath::EmbedSimplePath(bool bUpdatePath, TArray<int>& PathVertices, double SnapElementThresholdSq)
+bool FMeshSurfacePath::EmbedSimplePath(bool bUpdatePath, TArray<int>& PathVertices, bool bDoNotDuplicateFirstVertexID, double SnapElementThresholdSq)
 {
-	int32 InitialPathIdx = PathVertices.Num(); // TODO: rm this variable (used for debugging)
+	// used to track where the new vertices for *this* path start; used for bDoNotDuplicateFirstVertexID
+	int32 InitialPathIdx = PathVertices.Num();
 	if (!Path.Num())
 	{
 		return true;
@@ -650,7 +651,11 @@ bool FMeshSurfacePath::EmbedSimplePath(bool bUpdatePath, TArray<int>& PathVertic
 		{
 			ensure(Pt.PointType == ESurfacePointType::Vertex);
 			ensure(Mesh->IsVertex(Pt.ElementID));
-			PathVertices.Add(Pt.ElementID);
+			// make sure we don't add a duplicate vertex for the very first vertex (occurs when appending paths sequentially)
+			if (!bDoNotDuplicateFirstVertexID || PathVertices.Num() != InitialPathIdx || 0 == PathVertices.Num() || PathVertices.Last() != Pt.ElementID)
+			{
+				PathVertices.Add(Pt.ElementID);
+			}
 		}
 	}
 
@@ -692,7 +697,7 @@ bool FMeshSurfacePath::EmbedSimplePath(bool bUpdatePath, TArray<int>& PathVertic
 	return true;
 }
 
-bool EmbedProjectedPath(FDynamicMesh3* Mesh, int StartTriID, FFrame3d Frame, const TArray<FVector2d>& Path2D, TArray<int>& OutPathVertices, bool bClosePath, FMeshFaceSelection* EnclosedFaces, double PtSnapVertexOrEdgeThresholdSq)
+bool EmbedProjectedPath(FDynamicMesh3* Mesh, int StartTriID, FFrame3d Frame, const TArray<FVector2d>& Path2D, TArray<int>& OutPathVertices, TArray<int>& OutVertexCorrespondence, bool bClosePath, FMeshFaceSelection* EnclosedFaces, double PtSnapVertexOrEdgeThresholdSq)
 {
 	int32 EndIdxA = Path2D.Num() - (bClosePath ? 1 : 2);
 	int CurrentSeedTriID = StartTriID;
@@ -704,23 +709,47 @@ bool EmbedProjectedPath(FDynamicMesh3* Mesh, int StartTriID, FFrame3d Frame, con
 		return FVector3d(ProjPt.X, ProjPt.Y, 0);
 	};
 
+	OutVertexCorrespondence.Add(0);
 	for (int32 IdxA = 0; IdxA <= EndIdxA; IdxA++)
 	{
+		
 		int32 IdxB = (IdxA + 1) % Path2D.Num();
 		FMeshSurfacePath SurfacePath(Mesh);
 		FVector2d WalkDir = Path2D[IdxB] - Path2D[IdxA];
+		int LastVert = -1;
+		// TODO: for final vertex & closed path, tell the planar walk to connect back to that vertex.  NOTE: cannot simply connect back to the original start triangle, as that will have been split up by the path embedding.
 		bool WalkSuccess = SurfacePath.AddViaPlanarWalk(CurrentSeedTriID, FVector3d(Path2D[IdxA].X, Path2D[IdxA].Y, 0), -1, FVector3d(Path2D[IdxB].X, Path2D[IdxB].Y, 0), FVector3d(-WalkDir.Y, WalkDir.X, 0), ProjectToFrame, false, 0, PtSnapVertexOrEdgeThresholdSq);
 		if (!WalkSuccess)
 		{
 			return false;
 		}
-		SurfacePath.EmbedSimplePath(false, OutPathVertices, PtSnapVertexOrEdgeThresholdSq);
+		SurfacePath.EmbedSimplePath(false, OutPathVertices, true, PtSnapVertexOrEdgeThresholdSq);
+		OutVertexCorrespondence.Add(OutPathVertices.Num() - 1);
 		TArray<int> TrianglesOut;
 		Mesh->GetVertexOneRingTriangles(OutPathVertices.Last(), TrianglesOut);
 		check(TrianglesOut.Num());
 		CurrentSeedTriID = TrianglesOut[0];
 	}
+	// special handling to remove redundant vertex + correspondence at the start and end of a looping path
+	if (bClosePath)
+	{
+		if (OutPathVertices[0] == OutPathVertices.Last())
+		{
+			OutPathVertices.Pop();
+		}
+		else
+		{
+			// TODO: we may consider worrying about the case where the start and end are 'almost' connected / separated by some degenerate triangles that are easily crossed, which would currently fail
+			// for now we only handle the case where the start and end vertices are on the same triangle, which could happen for a single degenerate triangle case, which is the most likely case ...
+			if (!Mesh->FindEdge(OutPathVertices[0], OutPathVertices.Last()))
+			{
+				return false; // failed to properly close path
+			}
+		}
+		OutVertexCorrespondence.Pop();
+	}
 
+	// If a selection was requested, flood fill to select the faces enclosed by the path
 	if (EnclosedFaces)
 	{
 		TSet<int> Edges;
@@ -731,7 +760,10 @@ bool EmbedProjectedPath(FDynamicMesh3* Mesh, int StartTriID, FFrame3d Frame, con
 			int32 IdxB = (IdxA + 1) % OutPathVertices.Num();
 			int32 IDA = OutPathVertices[IdxA];
 			int32 IDB = OutPathVertices[IdxB];
+
+			ensure(IDA != IDB);
 			int EID = Mesh->FindEdge(IDA, IDB);
+			ensure(EID != FDynamicMesh3::InvalidID);
 			Edges.Add(EID);
 
 			if (SeedTriID == -1)
diff --git a/Engine/Plugins/Experimental/GeometryProcessing/Source/DynamicMesh/Public/DynamicMeshEditor.h b/Engine/Plugins/Experimental/GeometryProcessing/Source/DynamicMesh/Public/DynamicMeshEditor.h
--- a/Engine/Plugins/Experimental/GeometryProcessing/Source/DynamicMesh/Public/DynamicMeshEditor.h
+++ b/Engine/Plugins/Experimental/GeometryProcessing/Source/DynamicMesh/Public/DynamicMeshEditor.h
@@ -136,11 +136,24 @@ public:
 	 * @param Loop1 first loop of sequential vertices
 	 * @param Loop2 second loop of sequential vertices
 	 * @param ResultOut lists of newly created triangles/vertices/etc
-	 * @return true if operation suceeded. If a failure occurs, any added triangles are removed via RemoveTriangles
+	 * @return true if operation succeeded. If a failure occurs, any added triangles are removed via RemoveTriangles
 	 */
 	bool StitchVertexLoopsMinimal(const TArray<int>& VertexLoop1, const TArray<int>& VertexLoop2, FDynamicMeshEditResult& ResultOut);
 
 
+
+	/**
+	 * Stitch together two loops of vertices where vertices are only sparsely corresponded
+	 * @param VertexIDs1 first array of sequential vertices
+	 * @param MatchedIndices1 indices into the VertexIDs1 array of vertices that have a corresponding match in the VertexIDs2 array; Must be ordered
+	 * @param VertexIDs2 second array of sequential vertices
+	 * @param MatchedIndices2 indices into the VertexIDs2 array of vertices that have a corresponding match in the VertexIDs1 array; Must be ordered
+	 * @param ResultOut lists of newly created triangles/vertices/etc
+	 * @return true if operation succeeded.  If a failure occurs, any added triangles are removed via RemoveTriangles
+	 */
+	bool StitchSparselyCorrespondedVertexLoops(const TArray<int>& VertexIDs1, const TArray<int>& MatchedIndices1, const TArray<int>& VertexIDs2, const TArray<int>& MatchedIndices2, FDynamicMeshEditResult& ResultOut);
+
+
 	/**
 	 * Fill hole with a triangle fan given an existing (unconnected) center vertex and
 	 * an ordered loop of boundary vertices on the hole border.
diff --git a/Engine/Plugins/Experimental/GeometryProcessing/Source/DynamicMesh/Public/Operations/EmbedSurfacePath.h b/Engine/Plugins/Experimental/GeometryProcessing/Source/DynamicMesh/Public/Operations/EmbedSurfacePath.h
--- a/Engine/Plugins/Experimental/GeometryProcessing/Source/DynamicMesh/Public/Operations/EmbedSurfacePath.h
+++ b/Engine/Plugins/Experimental/GeometryProcessing/Source/DynamicMesh/Public/Operations/EmbedSurfacePath.h
@@ -129,7 +129,7 @@ public:
 	 * @param PathVertices Indices of the vertices on the path after embedding succeeds; NOTE these will not be 1:1 with the input Path
 	 * @return true if embedding succeeded.
 	 */
-	bool EmbedSimplePath(bool bUpdatePath, TArray<int>& PathVertices, double SnapElementThresholdSq = FMathd::Epsilon*100);
+	bool EmbedSimplePath(bool bUpdatePath, TArray<int>& PathVertices, bool bDoNotDuplicateFirstVertexID = true, double SnapElementThresholdSq = FMathd::Epsilon*100);
 
 	// TODO: add functionality to delete 'inside' of path -- but how do we determine what is inside?  ref MeshFacesFromLoop.cs in geometry3sharp
 };
@@ -137,5 +137,5 @@ public:
 /**
  * Embed a 2D path into a mesh by projection, starting the walk from a given triangle.
  */
-bool DYNAMICMESH_API EmbedProjectedPath(FDynamicMesh3* Mesh, int StartTriID, FFrame3d Frame, const TArray<FVector2d>& Path2D, TArray<int>& OutPathVertices, bool bClosePath, FMeshFaceSelection *EnclosedFaces = nullptr, double PtSnapVertexOrEdgeThresholdSq = FMathd::Epsilon*100);
+bool DYNAMICMESH_API EmbedProjectedPath(FDynamicMesh3* Mesh, int StartTriID, FFrame3d Frame, const TArray<FVector2d>& Path2D, TArray<int>& OutPathVertices, TArray<int>& OutVertexCorrespondence, bool bClosePath, FMeshFaceSelection *EnclosedFaces = nullptr, double PtSnapVertexOrEdgeThresholdSq = FMathd::Epsilon*100);
 
