Fix static analysis warning.

diff --git a/Engine/Plugins/Experimental/GeometryProcessing/Source/GeometricObjects/Public/Spatial/SparseDynamicOctree3.h b/Engine/Plugins/Experimental/GeometryProcessing/Source/GeometricObjects/Public/Spatial/SparseDynamicOctree3.h
--- a/Engine/Plugins/Experimental/GeometryProcessing/Source/GeometricObjects/Public/Spatial/SparseDynamicOctree3.h
+++ b/Engine/Plugins/Experimental/GeometryProcessing/Source/GeometricObjects/Public/Spatial/SparseDynamicOctree3.h
@@ -292,7 +292,7 @@ protected:
 	// calculate the base width of a cell at a given level
 	inline double GetCellWidth(uint32 Level) const
 	{
-		double Divisor = (double)(1 << Level);
+		double Divisor = (double)(1U << Level);
 		double CellWidth = RootDimension / Divisor;
 		return CellWidth;
 	}
