fix shadowing in mesh stitch

#rnx

diff --git a/Engine/Plugins/Experimental/GeometryProcessing/Source/DynamicMesh/Private/DynamicMeshEditor.cpp b/Engine/Plugins/Experimental/GeometryProcessing/Source/DynamicMesh/Private/DynamicMeshEditor.cpp
--- a/Engine/Plugins/Experimental/GeometryProcessing/Source/DynamicMesh/Private/DynamicMeshEditor.cpp
+++ b/Engine/Plugins/Experimental/GeometryProcessing/Source/DynamicMesh/Private/DynamicMeshEditor.cpp
@@ -124,14 +124,14 @@ bool FDynamicMeshEditor::StitchSparselyCorrespondedVertexLoops(const TArray<int>
 		int Starts[2] { MatchedIndices1[i], MatchedIndices2[i] };
 		int Ends[2] { MatchedIndices1[(i + 1) % CorrespondN], MatchedIndices2[(i + 1) % CorrespondN] };
 
-		auto GetWrappedSpanLen = [](const FDynamicMesh3* Mesh, const TArray<int>& VertexIDs, int StartInd, int EndInd)
+		auto GetWrappedSpanLen = [](const FDynamicMesh3* M, const TArray<int>& VertexIDs, int StartInd, int EndInd)
 		{
 			float LenTotal = 0;
-			FVector3d V = Mesh->GetVertex(VertexIDs[StartInd]);
+			FVector3d V = M->GetVertex(VertexIDs[StartInd]);
 			for (int Ind = StartInd, IndNext; Ind != EndInd;)
 			{
 				IndNext = (Ind + 1) % VertexIDs.Num();
-				FVector3d VNext = Mesh->GetVertex(VertexIDs[IndNext]);
+				FVector3d VNext = M->GetVertex(VertexIDs[IndNext]);
 				LenTotal += V.Distance(VNext);
 				Ind = IndNext;
 				V = VNext;
