GeometryProcessing: Handle missing points in FPointSetHashTable::FindPointsInBall, remove leftover "new" in VectorUtil::BarycentricCoords
#rnx

diff --git a/Engine/Plugins/Experimental/GeometryProcessing/Source/GeometricObjects/Private/Spatial/PointSetHashTable.cpp b/Engine/Plugins/Experimental/GeometryProcessing/Source/GeometricObjects/Private/Spatial/PointSetHashTable.cpp
--- a/Engine/Plugins/Experimental/GeometryProcessing/Source/GeometricObjects/Private/Spatial/PointSetHashTable.cpp
+++ b/Engine/Plugins/Experimental/GeometryProcessing/Source/GeometricObjects/Private/Spatial/PointSetHashTable.cpp
@@ -44,9 +44,12 @@ bool FPointSetHashtable::FindPointsInBall(const FVector3d& QueryPt, double Query
 	{
 		for (int vid : *CenterCellList)
 		{
-			if (QueryPt.DistanceSquared(Points->GetPoint(vid)) < RadiusSqr)
+			if (Points->IsPoint(vid))		// handle case where points are removed from pointset
 			{
-				ResultOut.Add(vid);
+				if (QueryPt.DistanceSquared(Points->GetPoint(vid)) < RadiusSqr)
+				{
+					ResultOut.Add(vid);
+				}
 			}
 		}
 	}
@@ -75,9 +78,12 @@ bool FPointSetHashtable::FindPointsInBall(const FVector3d& QueryPt, double Query
 			{
 				for (int vid : *NbrCellList)
 				{
-					if (QueryPt.DistanceSquared(Points->GetPoint(vid)) < RadiusSqr)
+					if (Points->IsPoint(vid))		// handle case where points are removed from pointset
 					{
-						ResultOut.Add(vid);
+						if (QueryPt.DistanceSquared(Points->GetPoint(vid)) < RadiusSqr)
+						{
+							ResultOut.Add(vid);
+						}
 					}
 				}
 			}
diff --git a/Engine/Plugins/Experimental/GeometryProcessing/Source/GeometricObjects/Public/VectorUtil.h b/Engine/Plugins/Experimental/GeometryProcessing/Source/GeometricObjects/Public/VectorUtil.h
--- a/Engine/Plugins/Experimental/GeometryProcessing/Source/GeometricObjects/Public/VectorUtil.h
+++ b/Engine/Plugins/Experimental/GeometryProcessing/Source/GeometricObjects/Public/VectorUtil.h
@@ -306,7 +306,7 @@ namespace VectorUtil
 		RealType fBary1 = (fM11 * fR0 - fM01 * fR1) * fInvDet;
 		RealType fBary2 = (fM00 * fR1 - fM01 * fR0) * fInvDet;
 		RealType fBary3 = 1.0 - fBary1 - fBary2;
-		return new FVector3<RealType>(fBary1, fBary2, fBary3);
+		return FVector3<RealType>(fBary1, fBary2, fBary3);
 	}
 
 	/**
@@ -331,7 +331,7 @@ namespace VectorUtil
 		RealType fBary1 = (fM11 * fR0 - fM01 * fR1) * fInvDet;
 		RealType fBary2 = (fM00 * fR1 - fM01 * fR0) * fInvDet;
 		RealType fBary3 = 1.0 - fBary1 - fBary2;
-		return new FVector3<RealType>(fBary1, fBary2, fBary3);
+		return FVector3<RealType>(fBary1, fBary2, fBary3);
 	}
 
 	/**
