Geometry - Constrained Poisson Solvers : Disable timeloggers.

diff --git a/Engine/Plugins/Experimental/GeometryProcessing/Source/MeshSolverUtilities/Private/ConstrainedPoissonSolver.h b/Engine/Plugins/Experimental/GeometryProcessing/Source/MeshSolverUtilities/Private/ConstrainedPoissonSolver.h
--- a/Engine/Plugins/Experimental/GeometryProcessing/Source/MeshSolverUtilities/Private/ConstrainedPoissonSolver.h
+++ b/Engine/Plugins/Experimental/GeometryProcessing/Source/MeshSolverUtilities/Private/ConstrainedPoissonSolver.h
@@ -10,11 +10,11 @@
 #include "MatrixSolver.h"
 
 
+// only needed if you enable logging by changing the typedef for FConstrainedSolverTimeLogger
+//#include "ProfilingDebugging/ScopedTimers.h"
 
 
-
-#include "ProfilingDebugging/ScopedTimers.h"
-#include "Async/ParallelFor.h"
+//#include "Async/ParallelFor.h"
 
 
 
@@ -46,7 +46,14 @@ public:
 	typedef typename FSparseMatrixD::Scalar     ScalarType;
 	typedef typename FSOAPositions::VectorType  VectorType;
 
+	struct FDummyTimeLogger
+	{
+		FDummyTimeLogger(FString InMsg = TEXT("Scoped action"), FOutputDevice* InDevice = GLog) {}
+	};
 
+	// Change this to enable time logging for testing.
+	//typedef FScopedDurationTimeLogger FConstrainedSolverTimeLogger; 
+	typedef FDummyTimeLogger FConstrainedSolverTimeLogger;
 
 	struct FConstraintPosition
 	{
@@ -163,7 +170,7 @@ public:
 	bool Solve(const FSOAPositions& SourceVector, FSOAPositions& SolutionVector) const
 	{
 		checkSlow(bMatrixSolverDirty == false);
-		FScopedDurationTimeLogger Timmer(TEXT("Post-setup solve time"));
+		FConstrainedSolverTimeLogger Timmer(TEXT("Post-setup solve time"));
 
 		FSparseMatrixD& SymmetricMatrix = *SymmetricMatrixPtr;
 		// Set up the source vector
@@ -189,7 +196,7 @@ public:
 	{
 		checkSlow(bMatrixSolverDirty == false);
 		
-		FScopedDurationTimeLogger Timmer(TEXT("Post-setup solve time"));
+		FConstrainedSolverTimeLogger Timmer(TEXT("Post-setup solve time"));
 
 		FSparseMatrixD& SymmetricMatrix = *SymmetricMatrixPtr;
 		FSOAPositions RHSVector(SymmetricMatrix.cols());
@@ -218,7 +225,7 @@ public:
 		{
 			const IIterativeMatrixSolverBase* IterativeSolver = (IIterativeMatrixSolverBase*)MatrixSolver.Get();
 
-			FScopedDurationTimeLogger Timmer(TEXT("Post-setup solve time"));
+			FConstrainedSolverTimeLogger Timmer(TEXT("Post-setup solve time"));
 
 			FSparseMatrixD& SymmetricMatrix = *SymmetricMatrixPtr;
 			FSOAPositions RHSVector(SymmetricMatrix.cols());
@@ -252,7 +259,7 @@ public:
 		{
 			const IIterativeMatrixSolverBase* IterativeSolver = (IIterativeMatrixSolverBase*)MatrixSolver.Get();
 
-			FScopedDurationTimeLogger Timmer(TEXT("Post-setup solve time"));
+			FConstrainedSolverTimeLogger Timmer(TEXT("Post-setup solve time"));
 
 			FSparseMatrixD& SymmetricMatrix = *SymmetricMatrixPtr;
 			// Set up the source vector
@@ -308,7 +315,7 @@ protected:
 	// Note: If constraints haven't been set, or if constraints have been cleared, this must be called prior to solve.
 	void UpdateSolverWithContraints()
 	{
-		FScopedDurationTimeLogger Timmer(TEXT("Matrix setup time"));
+		FConstrainedSolverTimeLogger Timmer(TEXT("Matrix setup time"));
 
 		// The constrained verts are part of the matrix
 
