Geometry - Expose the laplacian deformation solver so it can be used in tools.

diff --git a/Engine/Plugins/Experimental/GeometryProcessing/Source/MeshSolverUtilities/Private/LaplacianMeshSmoother.h b/Engine/Plugins/Experimental/GeometryProcessing/Source/MeshSolverUtilities/Private/LaplacianMeshSmoother.h
--- a/Engine/Plugins/Experimental/GeometryProcessing/Source/MeshSolverUtilities/Private/LaplacianMeshSmoother.h
+++ b/Engine/Plugins/Experimental/GeometryProcessing/Source/MeshSolverUtilities/Private/LaplacianMeshSmoother.h
@@ -12,7 +12,7 @@
 
 
 
-class FConstrainedMeshOperator : public MeshDeformingOperators::IConstrainedMeshOperator
+class MESHSOLVERUTILITIES_API FConstrainedMeshOperator : public MeshDeformingOperators::IConstrainedMeshOperator
 {
 public:
 	typedef FConstrainedSolver::FConstraintPosition   FConstraintPosition;
@@ -96,7 +96,7 @@ protected:
 };
 
 
-class FConstrainedMeshDeformer : public FConstrainedMeshOperator
+class MESHSOLVERUTILITIES_API FConstrainedMeshDeformer : public FConstrainedMeshOperator
 {
 public:
 	FConstrainedMeshDeformer(const FDynamicMesh3& DynamicMesh, const ELaplacianWeightScheme LaplacianType);
