GeometryProcessing - adding a missing const to  FVertexPositionCache

diff --git a/Engine/Plugins/Experimental/GeometryProcessing/Source/DynamicMesh/Public/Operations/GroupTopologyDeformer.h b/Engine/Plugins/Experimental/GeometryProcessing/Source/DynamicMesh/Public/Operations/GroupTopologyDeformer.h
--- a/Engine/Plugins/Experimental/GeometryProcessing/Source/DynamicMesh/Public/Operations/GroupTopologyDeformer.h
+++ b/Engine/Plugins/Experimental/GeometryProcessing/Source/DynamicMesh/Public/Operations/GroupTopologyDeformer.h
@@ -32,9 +32,9 @@ public:
 	}
 
 	/** Apply saved positions to Mesh */
-	void SetPositions(FDynamicMesh3* Mesh)
+	void SetPositions(FDynamicMesh3* Mesh) const
 	{
-		int Num = Vertices.Num();
+		const int Num = Vertices.Num();
 		for (int k = 0; k < Num; ++k)
 		{
 			Mesh->SetVertex(Vertices[k], Positions[k]);
