first pass at polygon on mesh tool (embeds polygons on meshes, to be used to cut through meshes or embed / extrude shapes)

#rnx

diff --git a/Engine/Plugins/Experimental/GeometryProcessing/Source/DynamicMesh/Private/DynamicMesh3.cpp b/Engine/Plugins/Experimental/GeometryProcessing/Source/DynamicMesh/Private/DynamicMesh3.cpp
--- a/Engine/Plugins/Experimental/GeometryProcessing/Source/DynamicMesh/Private/DynamicMesh3.cpp
+++ b/Engine/Plugins/Experimental/GeometryProcessing/Source/DynamicMesh/Private/DynamicMesh3.cpp
@@ -1048,6 +1048,12 @@ int FDynamicMesh3::FindEdge(int vA, int vB) const
 {
 	check(IsVertex(vA));
 	check(IsVertex(vB));
+	if (vA == vB)
+	{
+		// self-edges are not allowed, and if we fall through to the search below on a self edge we will incorrectly
+		// sometimes return an arbitrary edge if queried for a self-edge, due to the optimization of only checking one side of the edge
+		return InvalidID;
+	}
 
 	// [RMS] edge vertices must be sorted (min,max),
 	//   that means we only need one index-check in inner loop.
diff --git a/Engine/Plugins/Experimental/GeometryProcessing/Source/DynamicMesh/Private/DynamicMesh3_Edits.cpp b/Engine/Plugins/Experimental/GeometryProcessing/Source/DynamicMesh/Private/DynamicMesh3_Edits.cpp
--- a/Engine/Plugins/Experimental/GeometryProcessing/Source/DynamicMesh/Private/DynamicMesh3_Edits.cpp
+++ b/Engine/Plugins/Experimental/GeometryProcessing/Source/DynamicMesh/Private/DynamicMesh3_Edits.cpp
@@ -170,9 +170,8 @@ int FDynamicMesh3::AppendTriangle(const FIndex3i& tv, int gid)
 		check(false);
 		return InvalidID;
 	}
-	if (tv[0] == tv[1] || tv[0] == tv[2] || tv[1] == tv[2])
+	if (!ensure(tv[0] != tv[1] && tv[0] != tv[2] && tv[1] != tv[2]))
 	{
-		check(false);
 		return InvalidID;
 	}
 
diff --git a/Engine/Plugins/Experimental/GeometryProcessing/Source/DynamicMesh/Private/DynamicMeshEditor.cpp b/Engine/Plugins/Experimental/GeometryProcessing/Source/DynamicMesh/Private/DynamicMeshEditor.cpp
--- a/Engine/Plugins/Experimental/GeometryProcessing/Source/DynamicMesh/Private/DynamicMeshEditor.cpp
+++ b/Engine/Plugins/Experimental/GeometryProcessing/Source/DynamicMesh/Private/DynamicMeshEditor.cpp
@@ -187,10 +187,7 @@ operation_failed:
 	// remove what we added so far
 	if (ResultOut.NewTriangles.Num()) 
 	{
-		if (!RemoveTriangles(ResultOut.NewTriangles, false))
-		{
-			checkf(false, TEXT("FDynamicMeshEditor::StitchSparselyCorrespondedVertexLoops: failed to add all triangles, and also failed to back out changes."));
-		}
+		ensureMsgf(RemoveTriangles(ResultOut.NewTriangles, false), TEXT("FDynamicMeshEditor::StitchSparselyCorrespondedVertexLoops: failed to add all triangles, and also failed to back out changes."));
 	}
 	return false;
 }
diff --git a/Engine/Plugins/Experimental/GeometryProcessing/Source/DynamicMesh/Private/DynamicMeshOverlay.cpp b/Engine/Plugins/Experimental/GeometryProcessing/Source/DynamicMesh/Private/DynamicMeshOverlay.cpp
--- a/Engine/Plugins/Experimental/GeometryProcessing/Source/DynamicMesh/Private/DynamicMeshOverlay.cpp
+++ b/Engine/Plugins/Experimental/GeometryProcessing/Source/DynamicMesh/Private/DynamicMeshOverlay.cpp
@@ -293,6 +293,12 @@ template<typename RealType, int ElementSize>
 void TDynamicMeshOverlay<RealType, ElementSize>::OnRemoveTriangle(int TriangleID, bool bRemoveIsolatedVertices)
 {
 	FIndex3i Triangle = GetTriangle(TriangleID);
+	if (Triangle.A < 0 && Triangle.B < 0 && Triangle.C < 0)
+	{
+		// if whole triangle has no overlay vertices set, that's OK, just remove nothing
+		// (if only *some* of the triangle vertices were < 0, that would be a bug / invalid overlay triangle)
+		return;
+	}
 	InitializeNewTriangle(TriangleID);
 
 	// decrement element refcounts, and free element if it is now unreferenced
diff --git a/Engine/Plugins/Experimental/GeometryProcessing/Source/DynamicMesh/Private/Operations/EmbedSurfacePath.cpp b/Engine/Plugins/Experimental/GeometryProcessing/Source/DynamicMesh/Private/Operations/EmbedSurfacePath.cpp
--- a/Engine/Plugins/Experimental/GeometryProcessing/Source/DynamicMesh/Private/Operations/EmbedSurfacePath.cpp
+++ b/Engine/Plugins/Experimental/GeometryProcessing/Source/DynamicMesh/Private/Operations/EmbedSurfacePath.cpp
@@ -120,7 +120,7 @@ FMeshSurfacePoint RelocateTrianglePointAfterRefinement(const FDynamicMesh3* Mesh
 	return SurfacePt;
 }
 
-bool WalkMeshPlanar(const FDynamicMesh3* Mesh, int StartTri, FVector3d StartPt, int EndTri, FVector3d EndPt, FVector3d WalkPlaneNormal, TFunction<FVector3d(const FDynamicMesh3*, int)> VertexToPosnFn, bool bAllowBackwardsSearch, double AcceptEndPtOutsideDist, double PtOnPlaneThresholdSq, TArray<TPair<FMeshSurfacePoint, int>>& WalkedPath)
+bool WalkMeshPlanar(const FDynamicMesh3* Mesh, int StartTri, FVector3d StartPt, int EndTri, int EndVertID, FVector3d EndPt, FVector3d WalkPlaneNormal, TFunction<FVector3d(const FDynamicMesh3*, int)> VertexToPosnFn, bool bAllowBackwardsSearch, double AcceptEndPtOutsideDist, double PtOnPlaneThresholdSq, TArray<TPair<FMeshSurfacePoint, int>>& WalkedPath)
 {
 	auto SetTriVertPositions = [&VertexToPosnFn, &Mesh](FIndex3i TriVertIDs, FTriangle3d& Tri)
 	{
@@ -193,14 +193,25 @@ bool WalkMeshPlanar(const FDynamicMesh3* Mesh, int StartTri, FVector3d StartPt,
 		SetTriVertPositions(TriVertIDs, CurrentTri);
 
 		bool OnEndTri = EndTri == TriID;
+
+		// if we're on a triangle that is connected to the known final vertex, end the search!
+		if (EndVertID >= 0 && TriVertIDs.Contains(EndVertID))
+		{
+			OnEndTri = true;
+			CurrentEnd = ComputedPointsAndSources.Emplace(FMeshSurfacePoint(EndVertID), FWalkIndices(EndPt, CurrentEnd, TriID));
+			bHasArrived = true;
+			BestKnownEnd = CurrentEnd;
+			break;
+		}
+
 		bool ComputedEndPtOnTri = false;
-		if (EndTri == -1)
+		if (EndVertID < 0 && EndTri == -1) // if we need to check if this is the end tri, and it could be the end tri
 		{
 			CurrentTriDist.Triangle = CurrentTri;
 			CurrentTriDist.Point = EndPt;
 			ComputedEndPtOnTri = true;
 			double DistSq = CurrentTriDist.GetSquared();
-			if (DistSq < FMathd::Epsilon/* && PtInsideTri(CurrentTriDist.TriangleBaryCoords, FMathd::Epsilon)*/)  // TODO: we don't really need to check the barycentric coordinates for being 'inside' the triangle if the distance is within epsilon?  especially since the barycoords test also has n epsilon threshold?
+			if (DistSq < AcceptEndPtOutsideDist/* && PtInsideTri(CurrentTriDist.TriangleBaryCoords, FMathd::Epsilon)*/)  // TODO: we don't really need to check the barycentric coordinates for being 'inside' the triangle if the distance is within epsilon?  especially since the barycoords test also has n epsilon threshold?
 			{
 				OnEndTri = true;
 			}
@@ -217,6 +228,7 @@ bool WalkMeshPlanar(const FDynamicMesh3* Mesh, int StartTri, FVector3d StartPt,
 				CurrentTriDist.GetSquared();
 			}
 			CurrentEnd = ComputedPointsAndSources.Emplace(FMeshSurfacePoint(TriID, CurrentTriDist.TriangleBaryCoords), FWalkIndices(EndPt, CurrentEnd, TriID));
+
 			bHasArrived = true;
 			BestKnownEnd = CurrentEnd;
 			break;
@@ -250,6 +262,7 @@ bool WalkMeshPlanar(const FDynamicMesh3* Mesh, int StartTri, FVector3d StartPt,
 			SignDist[TriSubIdx] = SD;
 			if (FMathd::Abs(SD) <= PtOnPlaneThresholdSq)
 			{
+				// Vertex crossing
 				Side[TriSubIdx] = 0;
 				int CandidateVertID = TriVertIDs[TriSubIdx];
 				if (FromPt.PointType != ESurfacePointType::Vertex || CandidateVertID != FromPt.ElementID)
@@ -296,7 +309,7 @@ bool WalkMeshPlanar(const FDynamicMesh3* Mesh, int StartTri, FVector3d StartPt,
 		for (int TriSubIdx = 0; TriSubIdx < 3; TriSubIdx++)
 		{
 			int NextSubIdx = (TriSubIdx + 1) % 3;
-			if (SignDist[TriSubIdx] * SignDist[NextSubIdx] < 0)
+			if (Side[TriSubIdx] * Side[NextSubIdx] < 0)
 			{
 				// edge crossing
 				int CandidateEdgeID = TriEdgeIDs[TriSubIdx];
@@ -321,7 +334,7 @@ bool WalkMeshPlanar(const FDynamicMesh3* Mesh, int StartTri, FVector3d StartPt,
 						continue;
 					}
 					double DSq = EndPt.DistanceSquared(CrossingP);
-					if (!bAllowBackwardsSearch && DSq > InitialDistSq + 100 * FMathd::Epsilon)
+					if (!bAllowBackwardsSearch && DSq > InitialDistSq + 100 * FMathd::ZeroTolerance)
 					{
 						// not allowed to go in a direction that gets us further from the destination than our initial point if backwards search not allowed
 						continue;
@@ -388,15 +401,19 @@ bool WalkMeshPlanar(const FDynamicMesh3* Mesh, int StartTri, FVector3d StartPt,
 
 	// try refining start and end points if they were on triangles, and remove them if they turn out to be duplicates after refinement
 	//  (note we could instead do the refinement up front and avoid the possibility of duplicate points, but that would slightly complicate the traversal logic ...)
+	// note we don't even check that the barycoords match in the edge case -- conceptually the path can only cross the edge at one point, so it 'should' be fine to treat them as close enough, and having two points on the same edge would mess up the simple embedding code
 	if (WalkedPath.Num() && WalkedPath[0].Key.PointType == ESurfacePointType::Triangle)
 	{
 		RefineSurfacePtFromTriangleToSubElement(Mesh, WalkedPath[0].Key.Pos(Mesh), WalkedPath[0].Key, PtOnPlaneThresholdSq);
 		if (WalkedPath.Num() > 1 &&
+			WalkedPath[0].Key.PointType != ESurfacePointType::Triangle &&
 			WalkedPath[0].Key.PointType == WalkedPath[1].Key.PointType &&
-			WalkedPath[0].Key.ElementID == WalkedPath[1].Key.ElementID &&
-			(WalkedPath[0].Key.PointType == ESurfacePointType::Vertex ||
-				WalkedPath[0].Key.BaryCoord.DistanceSquared(WalkedPath[1].Key.BaryCoord) < FMathd::Epsilon))
+			WalkedPath[0].Key.ElementID == WalkedPath[1].Key.ElementID)
 		{
+			if (WalkedPath.Last().Key.PointType == ESurfacePointType::Edge) // copy closer barycoord
+			{
+				WalkedPath[1].Key.BaryCoord = WalkedPath[0].Key.BaryCoord;
+			}
 			WalkedPath.RemoveAt(0);
 		}
 	}
@@ -404,11 +421,14 @@ bool WalkMeshPlanar(const FDynamicMesh3* Mesh, int StartTri, FVector3d StartPt,
 	{
 		RefineSurfacePtFromTriangleToSubElement(Mesh, WalkedPath.Last().Key.Pos(Mesh), WalkedPath.Last().Key, PtOnPlaneThresholdSq);
 		if (WalkedPath.Num() > 1 &&
+			WalkedPath.Last().Key.PointType != ESurfacePointType::Triangle &&
 			WalkedPath.Last().Key.PointType == WalkedPath.Last(1).Key.PointType &&
-			WalkedPath.Last().Key.ElementID == WalkedPath.Last(1).Key.ElementID &&
-			(WalkedPath.Last().Key.PointType == ESurfacePointType::Vertex ||
-				WalkedPath.Last().Key.BaryCoord.DistanceSquared(WalkedPath.Last(1).Key.BaryCoord) < FMathd::Epsilon))
+			WalkedPath.Last().Key.ElementID == WalkedPath.Last(1).Key.ElementID)
 		{
+			if (WalkedPath.Last().Key.PointType == ESurfacePointType::Edge) // copy closer barycoord
+			{
+				WalkedPath.Last(1).Key.BaryCoord = WalkedPath.Last().Key.BaryCoord;
+			}
 			WalkedPath.Pop();
 		}
 	}
@@ -537,7 +557,7 @@ bool FMeshSurfacePath::IsConnected() const
 }
 
 
-bool FMeshSurfacePath::AddViaPlanarWalk(int StartTri, FVector3d StartPt, int EndTri, FVector3d EndPt, FVector3d WalkPlaneNormal, TFunction<FVector3d(const FDynamicMesh3*, int)> VertexToPosnFn, bool bAllowBackwardsSearch, double AcceptEndPtOutsideDist, double PtOnPlaneThresholdSq)
+bool FMeshSurfacePath::AddViaPlanarWalk(int StartTri, FVector3d StartPt, int EndTri, int EndVertID, FVector3d EndPt, FVector3d WalkPlaneNormal, TFunction<FVector3d(const FDynamicMesh3*, int)> VertexToPosnFn, bool bAllowBackwardsSearch, double AcceptEndPtOutsideDist, double PtOnPlaneThresholdSq)
 {
 	if (!VertexToPosnFn)
 	{
@@ -546,7 +566,7 @@ bool FMeshSurfacePath::AddViaPlanarWalk(int StartTri, FVector3d StartPt, int End
 			return MeshArg->GetVertex(VertexID);
 		};
 	}
-	return WalkMeshPlanar(Mesh, StartTri, StartPt, EndTri, EndPt, WalkPlaneNormal, VertexToPosnFn, bAllowBackwardsSearch, AcceptEndPtOutsideDist, PtOnPlaneThresholdSq, Path);
+	return WalkMeshPlanar(Mesh, StartTri, StartPt, EndTri, EndVertID, EndPt, WalkPlaneNormal, VertexToPosnFn, bAllowBackwardsSearch, AcceptEndPtOutsideDist, PtOnPlaneThresholdSq, Path);
 }
 
 // TODO: general path embedding becomes an arbitrary 2D remeshing problem per triangle; requires support from e.g. GeometryAlgorithms CDT. Not implemented yet; this is a vague sketch of what might go there.
@@ -598,7 +618,7 @@ bool FMeshSurfacePath::EmbedSimplePath(bool bUpdatePath, TArray<int>& PathVertic
 	// If FinalTri is split or poked, we will need to re-locate the last point in the path
 	int StartProcessIdx = 0, EndSimpleProcessIdx = PathNum - 1;
 	bool bEndPointSpecialProcess = false;
-	if (OrigEndPt.PointType == ESurfacePointType::Triangle)
+	if (PathNum > 1 && OrigEndPt.PointType == ESurfacePointType::Triangle)
 	{
 		EndSimpleProcessIdx = PathNum - 2;
 		bEndPointSpecialProcess = true;
@@ -646,6 +666,11 @@ bool FMeshSurfacePath::EmbedSimplePath(bool bUpdatePath, TArray<int>& PathVertic
 				}
 				EndPtUpdated = RelocateTrianglePointAfterRefinement(Mesh, EndPtPos, TriInds, SnapElementThresholdSq);
 			}
+			else if (PathIdx != PathNum - 1 && EndPtUpdated.PointType == ESurfacePointType::Edge && Pt.ElementID == EndPtUpdated.ElementID)
+			{
+				// TODO: in this case we would need to relocate the endpoint, as its edge is gone ... 
+				ensure(false);
+			}
 		}
 		else
 		{
@@ -699,6 +724,11 @@ bool FMeshSurfacePath::EmbedSimplePath(bool bUpdatePath, TArray<int>& PathVertic
 
 bool EmbedProjectedPath(FDynamicMesh3* Mesh, int StartTriID, FFrame3d Frame, const TArray<FVector2d>& Path2D, TArray<int>& OutPathVertices, TArray<int>& OutVertexCorrespondence, bool bClosePath, FMeshFaceSelection* EnclosedFaces, double PtSnapVertexOrEdgeThresholdSq)
 {
+	if (StartTriID == FDynamicMesh3::InvalidID)
+	{
+		return false;
+	}
+
 	int32 EndIdxA = Path2D.Num() - (bClosePath ? 1 : 2);
 	int CurrentSeedTriID = StartTriID;
 	OutPathVertices.Reset();
@@ -715,23 +745,73 @@ bool EmbedProjectedPath(FDynamicMesh3* Mesh, int StartTriID, FFrame3d Frame, con
 		
 		int32 IdxB = (IdxA + 1) % Path2D.Num();
 		FMeshSurfacePath SurfacePath(Mesh);
-		FVector2d WalkDir = Path2D[IdxB] - Path2D[IdxA];
 		int LastVert = -1;
-		// TODO: for final vertex & closed path, tell the planar walk to connect back to that vertex.  NOTE: cannot simply connect back to the original start triangle, as that will have been split up by the path embedding.
-		bool WalkSuccess = SurfacePath.AddViaPlanarWalk(CurrentSeedTriID, FVector3d(Path2D[IdxA].X, Path2D[IdxA].Y, 0), -1, FVector3d(Path2D[IdxB].X, Path2D[IdxB].Y, 0), FVector3d(-WalkDir.Y, WalkDir.X, 0), ProjectToFrame, false, 0, PtSnapVertexOrEdgeThresholdSq);
-		if (!WalkSuccess)
+		// for closed paths, tell the final segment to connect back to the first vertex
+		if (bClosePath && IdxB == 0 && OutPathVertices.Num() > 0)
+		{
+			LastVert = OutPathVertices[0];
+		}
+		FVector3d StartPos;
+		if (OutPathVertices.Num())  // shift walk start pos to the actual place the last segment ended, to allow for wobble snap
+		{
+			StartPos = ProjectToFrame(Mesh, OutPathVertices.Last());
+		}
+		else
+		{
+			StartPos = FVector3d(Path2D[IdxA].X, Path2D[IdxA].Y, 0);
+		}
+		FVector2d WalkDir = Path2D[IdxB] - FVector2d(StartPos.X, StartPos.Y);
+		double WalkLen = WalkDir.Length();
+		if (WalkLen < PtSnapVertexOrEdgeThresholdSq)
+		{
+			// we're already at the point
+			ensure(false); // TODO: can skip some processing for this case?
+		}
+		WalkDir /= WalkLen;
+		FVector3d WalkNormal(-WalkDir.Y, WalkDir.X, 0);
+		bool bWalkSuccess = SurfacePath.AddViaPlanarWalk(CurrentSeedTriID, StartPos, -1, LastVert, FVector3d(Path2D[IdxB].X, Path2D[IdxB].Y, 0), WalkNormal, ProjectToFrame, false, FMathf::ZeroTolerance, PtSnapVertexOrEdgeThresholdSq);
+		if (!bWalkSuccess)
 		{
 			return false;
 		}
-		SurfacePath.EmbedSimplePath(false, OutPathVertices, true, PtSnapVertexOrEdgeThresholdSq);
+		bool bEmbedSuccess = SurfacePath.EmbedSimplePath(false, OutPathVertices, true, PtSnapVertexOrEdgeThresholdSq);
 		OutVertexCorrespondence.Add(OutPathVertices.Num() - 1);
+		if (!bEmbedSuccess)
+		{
+			return false;
+		}
 		TArray<int> TrianglesOut;
 		Mesh->GetVertexOneRingTriangles(OutPathVertices.Last(), TrianglesOut);
+		
+		//// debugging check that we walked to the right place
+		// useful enough that I'm just commenting out for now; TODO remove entirely
+		//FVector3d ProjFirstVPos = ProjectToFrame(Mesh, OutPathVertices[0]);
+		//double DSqFromGoalPosFirst = ProjFirstVPos.DistanceSquared(FVector3d(Path2D[0].X, Path2D[0].Y, 0));
+		//int LastV = OutPathVertices.Last();
+		//FVector3d LastVPos = Mesh->GetVertex(LastV);
+		//FVector3d ProjLastVPos = ProjectToFrame(Mesh, LastV);
+		//FVector3d ProjSecondLastVPos(0,0,0);
+		//if (OutPathVertices.Num() > 1)
+		//{
+		//	ProjSecondLastVPos = ProjectToFrame(Mesh, OutPathVertices.Last(1));
+		//}
+		//double DSqFromGoalPos = ProjLastVPos.DistanceSquared(FVector3d(Path2D[IdxB].X, Path2D[IdxB].Y, 0));
+		//if (!ensure(DSqFromGoalPos <= PtSnapVertexOrEdgeThresholdSq*100))
+		//{
+		//	int x = 1;
+		//}
+
 		check(TrianglesOut.Num());
 		CurrentSeedTriID = TrianglesOut[0];
 	}
+
+	if (OutPathVertices.Num() == 0) // no path?
+	{
+		return false;
+	}
+
 	// special handling to remove redundant vertex + correspondence at the start and end of a looping path
-	if (bClosePath)
+	if (bClosePath && OutPathVertices.Num() > 1)
 	{
 		if (OutPathVertices[0] == OutPathVertices.Last())
 		{
@@ -741,7 +821,7 @@ bool EmbedProjectedPath(FDynamicMesh3* Mesh, int StartTriID, FFrame3d Frame, con
 		{
 			// TODO: we may consider worrying about the case where the start and end are 'almost' connected / separated by some degenerate triangles that are easily crossed, which would currently fail
 			// for now we only handle the case where the start and end vertices are on the same triangle, which could happen for a single degenerate triangle case, which is the most likely case ...
-			if (!Mesh->FindEdge(OutPathVertices[0], OutPathVertices.Last()))
+			if (Mesh->FindEdge(OutPathVertices[0], OutPathVertices.Last()) == FDynamicMesh3::InvalidID)
 			{
 				return false; // failed to properly close path
 			}
@@ -749,8 +829,10 @@ bool EmbedProjectedPath(FDynamicMesh3* Mesh, int StartTriID, FFrame3d Frame, con
 		OutVertexCorrespondence.Pop();
 	}
 
+	// TODO: check if the walk doesn't repeat any vertices / do something reasonable in cases where it does.
+
 	// If a selection was requested, flood fill to select the faces enclosed by the path
-	if (EnclosedFaces)
+	if (EnclosedFaces && OutPathVertices.Num() > 1)
 	{
 		TSet<int> Edges;
 		int32 NumEdges = bClosePath ? OutPathVertices.Num() : OutPathVertices.Num() - 1;
@@ -763,7 +845,12 @@ bool EmbedProjectedPath(FDynamicMesh3* Mesh, int StartTriID, FFrame3d Frame, con
 
 			ensure(IDA != IDB);
 			int EID = Mesh->FindEdge(IDA, IDB);
-			ensure(EID != FDynamicMesh3::InvalidID);
+			if (!ensure(EID != FDynamicMesh3::InvalidID))
+			{
+				// TODO: some recovery?  This could occur e.g. if you have a self-intersecting path over the mesh surface
+				return false;
+			}
+
 			Edges.Add(EID);
 
 			if (SeedTriID == -1)
@@ -780,7 +867,7 @@ bool EmbedProjectedPath(FDynamicMesh3* Mesh, int StartTriID, FFrame3d Frame, con
 				else if (OppVIDs.B != FDynamicMesh3::InvalidID)
 				{
 					double SignedAreaB = FTriangle2d::SignedArea(PA, PB, Frame.ToPlaneUV(Mesh->GetVertex(OppVIDs.B)));
-					if (SignedAreaB > FMathd::Epsilon)
+					if (SignedAreaB > FMathd::ZeroTolerance)
 					{
 						SeedTriID = Mesh->GetEdgeT(EID).B;
 					}
diff --git a/Engine/Plugins/Experimental/GeometryProcessing/Source/DynamicMesh/Public/Operations/EmbedSurfacePath.h b/Engine/Plugins/Experimental/GeometryProcessing/Source/DynamicMesh/Public/Operations/EmbedSurfacePath.h
--- a/Engine/Plugins/Experimental/GeometryProcessing/Source/DynamicMesh/Public/Operations/EmbedSurfacePath.h
+++ b/Engine/Plugins/Experimental/GeometryProcessing/Source/DynamicMesh/Public/Operations/EmbedSurfacePath.h
@@ -46,7 +46,7 @@ struct DYNAMICMESH_API FMeshSurfacePoint
 /**
  * Walk the surface of an FDynamicMesh to try find a planar path connecting two points.  Paths include every vertex and edge they need to cross.  Greedy algorithm will only return one path if there are multiple.
  */
-bool DYNAMICMESH_API WalkMeshPlanar(const FDynamicMesh3* Mesh, int StartTri, FVector3d StartPt, int EndTri, FVector3d EndPt, FVector3d WalkPlaneNormal, TFunction<FVector3d(const FDynamicMesh3*, int)> VertexToPosnFn, bool bAllowBackwardsSearch, double AcceptEndPtOutsideDist, double PtOnPlaneThreshold, TArray<TPair<FMeshSurfacePoint, int>>& WalkedPath);
+bool DYNAMICMESH_API WalkMeshPlanar(const FDynamicMesh3* Mesh, int StartTri, int EndVertID, FVector3d StartPt, int EndTri, FVector3d EndPt, FVector3d WalkPlaneNormal, TFunction<FVector3d(const FDynamicMesh3*, int)> VertexToPosnFn, bool bAllowBackwardsSearch, double AcceptEndPtOutsideDist, double PtOnPlaneThreshold, TArray<TPair<FMeshSurfacePoint, int>>& WalkedPath);
 
 
 /**
@@ -88,7 +88,7 @@ public:
 		bIsClosed = false;
 	}
 
-	bool AddViaPlanarWalk(int StartTri, FVector3d StartPt, int EndTri, FVector3d EndPt, FVector3d WalkPlaneNormal, TFunction<FVector3d(const FDynamicMesh3*, int)> VertexToPosnFn = nullptr, bool bAllowBackwardsSearch = true, double AcceptEndPtOutsideDist = FMathd::Epsilon, double PtOnPlaneThresholdSq = FMathd::Epsilon*100);
+	bool AddViaPlanarWalk(int StartTri, FVector3d StartPt, int EndTri, int EndVertID, FVector3d EndPt, FVector3d WalkPlaneNormal, TFunction<FVector3d(const FDynamicMesh3*, int)> VertexToPosnFn = nullptr, bool bAllowBackwardsSearch = true, double AcceptEndPtOutsideDist = FMathd::ZeroTolerance, double PtOnPlaneThresholdSq = FMathf::ZeroTolerance*100);
 	// TODO: Also support geodesic walks, other alternatives?
 
 	/**
@@ -129,7 +129,7 @@ public:
 	 * @param PathVertices Indices of the vertices on the path after embedding succeeds; NOTE these will not be 1:1 with the input Path
 	 * @return true if embedding succeeded.
 	 */
-	bool EmbedSimplePath(bool bUpdatePath, TArray<int>& PathVertices, bool bDoNotDuplicateFirstVertexID = true, double SnapElementThresholdSq = FMathd::Epsilon*100);
+	bool EmbedSimplePath(bool bUpdatePath, TArray<int>& PathVertices, bool bDoNotDuplicateFirstVertexID = true, double SnapElementThresholdSq = FMathf::ZeroTolerance*100);
 
 	// TODO: add functionality to delete 'inside' of path -- but how do we determine what is inside?  ref MeshFacesFromLoop.cs in geometry3sharp
 };
@@ -137,5 +137,5 @@ public:
 /**
  * Embed a 2D path into a mesh by projection, starting the walk from a given triangle.
  */
-bool DYNAMICMESH_API EmbedProjectedPath(FDynamicMesh3* Mesh, int StartTriID, FFrame3d Frame, const TArray<FVector2d>& Path2D, TArray<int>& OutPathVertices, TArray<int>& OutVertexCorrespondence, bool bClosePath, FMeshFaceSelection *EnclosedFaces = nullptr, double PtSnapVertexOrEdgeThresholdSq = FMathd::Epsilon*100);
+bool DYNAMICMESH_API EmbedProjectedPath(FDynamicMesh3* Mesh, int StartTriID, FFrame3d Frame, const TArray<FVector2d>& Path2D, TArray<int>& OutPathVertices, TArray<int>& OutVertexCorrespondence, bool bClosePath, FMeshFaceSelection *EnclosedFaces = nullptr, double PtSnapVertexOrEdgeThresholdSq = FMathf::ZeroTolerance*100);
 
diff --git a/Engine/Plugins/Experimental/GeometryProcessing/Source/GeometricObjects/Public/MeshQueries.h b/Engine/Plugins/Experimental/GeometryProcessing/Source/GeometricObjects/Public/MeshQueries.h
--- a/Engine/Plugins/Experimental/GeometryProcessing/Source/GeometricObjects/Public/MeshQueries.h
+++ b/Engine/Plugins/Experimental/GeometryProcessing/Source/GeometricObjects/Public/MeshQueries.h
@@ -139,6 +139,28 @@ public:
 		return tNearestID;
 	}
 
+	// brute force search for all triangle intersections, sorted
+	static void FindHitTriangles_LinearSearch(const TriangleMeshType& Mesh, const FRay3d& Ray, TArray<TPair<float, int>>& SortedHitTriangles)
+	{
+		FTriangle3d Triangle;
+		SortedHitTriangles.Empty();
+
+		for (int TriIdx : Mesh.TriangleIndicesItr())
+		{
+			Mesh.GetTriVertices(TriIdx, Triangle.V[0], Triangle.V[1], Triangle.V[2]);
+			FIntrRay3Triangle3d Query(Ray, Triangle);
+			if (Query.Find())
+			{
+				SortedHitTriangles.Emplace(Query.RayParameter, TriIdx);
+			}
+		}
+
+		SortedHitTriangles.Sort([](const TPair<float, int>& A, const TPair<float, int>& B)
+		{
+			return A.Key < B.Key;
+		});
+	}
+
 	/**
 	 * convenience function to construct a IntrRay3Triangle3 object for a Mesh triangle
 	 */
