initial pass at mesh normals editing tool

diff --git a/Engine/Plugins/Experimental/GeometryProcessing/Source/DynamicMesh/Private/DynamicMesh3_Queries.cpp b/Engine/Plugins/Experimental/GeometryProcessing/Source/DynamicMesh/Private/DynamicMesh3_Queries.cpp
--- a/Engine/Plugins/Experimental/GeometryProcessing/Source/DynamicMesh/Private/DynamicMesh3_Queries.cpp
+++ b/Engine/Plugins/Experimental/GeometryProcessing/Source/DynamicMesh/Private/DynamicMesh3_Queries.cpp
@@ -174,6 +174,116 @@ EMeshResult FDynamicMesh3::GetVtxTriangles(int vID, TArray<int>& TrianglesOut, b
 
 
 
+EMeshResult FDynamicMesh3::GetVtxContiguousTriangles(int VertexID, TArray<int>& TrianglesOut, TArray<int>& SpanLengths, TArray<bool>& IsLoop) const
+{
+	if (!ensure(IsVertex(VertexID)))
+	{
+		return EMeshResult::Failed_NotAVertex;
+	}
+
+	TrianglesOut.Reset();
+	SpanLengths.Reset();
+	IsLoop.Reset();
+
+	int NumEdges = VertexEdgeLists.GetCount(VertexID);
+	if (NumEdges == 0)
+	{
+		return EMeshResult::Ok;
+	}
+
+	TArray<int> StartEdgeIDs;
+	// initial starting edge candidates == boundary edges
+	for (int EID : VertexEdgeLists.Values(VertexID))
+	{
+		if (Edges[EID * 4 + 3] == InvalidID)
+		{
+			StartEdgeIDs.Push(EID);
+		}
+	}
+	bool bHasBoundaries = StartEdgeIDs.Num();
+	
+	if (!bHasBoundaries)
+	{
+		for (int EID : VertexEdgeLists.Values(VertexID))
+		{
+			StartEdgeIDs.Push(EID);
+			break;
+		}
+	}
+
+	int WalkedEdges = 0;
+	bool bHasRemainingBoundaries = bHasBoundaries;
+	while (StartEdgeIDs.Num() || WalkedEdges < NumEdges)
+	{
+		if (!StartEdgeIDs.Num())
+		{
+			bHasRemainingBoundaries = false;
+			// fallback for (hopefully very rare) case of a non-manifold vertex where there are separate one-rings w/ no boundary edges --
+			//  brute force search for an edge that hasn't already been walked
+			for (int EID : VertexEdgeLists.Values(VertexID))
+			{
+				int AttachedTriID = Edges[EID * 4 + 2];
+				bool UsedEdge = TrianglesOut.Contains(AttachedTriID);
+				if (!UsedEdge)
+				{
+					StartEdgeIDs.Push(EID);
+					break;
+				}
+			}
+		}
+
+		// walk starting from this edge, add the found span
+		
+		int StartEID = StartEdgeIDs.Pop();
+		int PrevEID = StartEID;
+		WalkedEdges++;
+		int WalkTri = Edges[StartEID * 4 + 2];
+		int32 SpanStart = TrianglesOut.Num();
+		IsLoop.Add(!bHasRemainingBoundaries);
+		while (true)
+		{
+			TrianglesOut.Add(WalkTri);
+
+			int TriIdx = 3 * WalkTri;
+			FIndex3i TriVIDs(Triangles[TriIdx], Triangles[TriIdx + 1], Triangles[TriIdx + 2]);
+			FIndex3i TriEIDs(TriangleEdges[TriIdx], TriangleEdges[TriIdx + 1], TriangleEdges[TriIdx + 2]);
+			int VertSubIdx = IndexUtil::FindTriIndex(VertexID, TriVIDs);
+			int NextEID = TriEIDs[VertSubIdx];
+			if (NextEID == PrevEID)
+			{
+				NextEID = TriEIDs[(VertSubIdx + 2) % 3];
+			}
+			if (NextEID == StartEID)
+			{
+				break;
+			}
+
+			WalkedEdges++;
+
+			int NextEdgeIdx = NextEID * 4;
+			int NextTriID = Edges[NextEdgeIdx + 2];
+			if (NextTriID == WalkTri)
+			{
+				NextTriID = Edges[NextEdgeIdx + 3];
+			}
+			if (NextTriID == InvalidID)
+			{
+				// remove the corresponding boundary
+				check(StartEdgeIDs.Num() > 0);
+				StartEdgeIDs.RemoveSingleSwap(NextEID);
+				break;
+			}
+			WalkTri = NextTriID;
+			PrevEID = NextEID;
+		}
+		SpanLengths.Add(TrianglesOut.Num() - SpanStart);
+	}
+
+	check(SpanLengths.Num() == IsLoop.Num());
+
+	return EMeshResult::Ok;
+}
+
 
 
 bool FDynamicMesh3::IsBoundaryVertex(int vID) const 
@@ -801,7 +911,7 @@ double FDynamicMesh3::GetTriSolidAngle(int tID, const FVector3d& p) const
 
 
 
-double FDynamicMesh3::GetTriInternalAngleR(int tID, int i)
+double FDynamicMesh3::GetTriInternalAngleR(int tID, int i) const
 {
 	int ti = 3 * tID;
 	int ta = 3 * Triangles[ti];
@@ -825,6 +935,27 @@ double FDynamicMesh3::GetTriInternalAngleR(int tID, int i)
 }
 
 
+FVector3d FDynamicMesh3::GetTriInternalAnglesR(int TID) const
+{
+	int ti = 3 * TID;
+	int ta = 3 * Triangles[ti];
+	FVector3d a(Vertices[ta], Vertices[ta + 1], Vertices[ta + 2]);
+	int tb = 3 * Triangles[ti + 1];
+	FVector3d b(Vertices[tb], Vertices[tb + 1], Vertices[tb + 2]);
+	int tc = 3 * Triangles[ti + 2];
+	FVector3d c(Vertices[tc], Vertices[tc + 1], Vertices[tc + 2]);
+	FVector3d ABhat = (b - a).Normalized();
+	FVector3d BChat = (c - b).Normalized();
+	FVector3d AChat = (c - a).Normalized();
+	return FVector3d(
+		ABhat.AngleR(AChat),	// AB vs AC angle
+		(-ABhat).AngleR(BChat), // BA vs BC angle
+		(AChat).AngleR(BChat) // AC vs BC angle (same as CA vs CB, so don't need to negate both)
+		);
+}
+
+
+
 double FDynamicMesh3::CalculateWindingNumber(const FVector3d& QueryPoint) const
 {
 	double sum = 0;
diff --git a/Engine/Plugins/Experimental/GeometryProcessing/Source/DynamicMesh/Private/DynamicMeshOverlay.cpp b/Engine/Plugins/Experimental/GeometryProcessing/Source/DynamicMesh/Private/DynamicMeshOverlay.cpp
--- a/Engine/Plugins/Experimental/GeometryProcessing/Source/DynamicMesh/Private/DynamicMeshOverlay.cpp
+++ b/Engine/Plugins/Experimental/GeometryProcessing/Source/DynamicMesh/Private/DynamicMeshOverlay.cpp
@@ -82,7 +82,130 @@ EMeshResult TDynamicMeshOverlay<RealType, ElementSize>::InsertElement(int Elemen
 
 
 
+template<typename RealType, int ElementSize>
+void TDynamicMeshOverlay<RealType, ElementSize>::CreateFromPredicate(TFunctionRef<bool(int ParentVertexIdx, int TriIDA, int TriIDB)> TrisCanShareVertexPredicate, RealType InitElementValue, bool bNonTransitiveShare)
+{
+	ClearElements(); // deletes all elements and initializes triangles to be 1:1 w/ parentmesh IDs
+	TArray<int> TrisActiveSubGroup, AppendedElements;
+	TArray<int> TriangleIDs, TriangleContigGroupLens;
+	TArray<bool> GroupIsLoop;
+	for (int VertexID : ParentMesh->VertexIndicesItr())
+	{
+		
+		bool bActiveSubGroupBroken = false;
+		ParentMesh->GetVtxContiguousTriangles(VertexID, TriangleIDs, TriangleContigGroupLens, GroupIsLoop);
+		int GroupStart = 0;
+		for (int GroupIdx = 0; GroupIdx < TriangleContigGroupLens.Num(); GroupIdx++)
+		{
+			bool bIsLoop = GroupIsLoop[GroupIdx];
+			int GroupNum = TriangleContigGroupLens[GroupIdx];
+			if (!ensure(GroupNum > 0)) // sanity check; groups should always have at least one element
+			{
+				continue;
+			}
 
+			TrisActiveSubGroup.Reset();
+			AppendedElements.Reset();
+			TrisActiveSubGroup.SetNumZeroed(GroupNum, false);
+			int CurrentGroupID = 0;
+			int CurrentGroupRefSubIdx = 0;
+			bool bShouldBreakGroup = false; // tracks whether a contiguous group should be broken up (in bNonTransitiveShare mode); final value after for-loop indicates whether final group was broken
+			bool bFirstGroupWasBroken = false;
+			for (int TriSubIdx = 0; TriSubIdx+1 < GroupNum; TriSubIdx++)
+			{
+				int TriIDA = TriangleIDs[GroupStart + TriSubIdx];
+				int TriIDB = TriangleIDs[GroupStart + TriSubIdx + 1];
+				bool bCanShare = TrisCanShareVertexPredicate(VertexID, TriIDA, TriIDB);
+				if (!bShouldBreakGroup && bCanShare && bNonTransitiveShare && TriSubIdx != CurrentGroupRefSubIdx) // require the triangle match the initial triangle in the group as well as the adjacent one
+				{
+					if (!TrisCanShareVertexPredicate(VertexID, TriangleIDs[GroupStart + CurrentGroupRefSubIdx], TriIDB))
+					{
+						bShouldBreakGroup = true;
+						if (CurrentGroupRefSubIdx == 0)
+						{
+							bFirstGroupWasBroken = true;
+						}
+					}
+				}
+				if (!bCanShare)
+				{
+					if (bShouldBreakGroup)
+					{
+						for (int BreakGroupIdx = CurrentGroupRefSubIdx + 1; BreakGroupIdx <= TriSubIdx; BreakGroupIdx++)
+						{
+							TrisActiveSubGroup[BreakGroupIdx] = ++CurrentGroupID;
+						}
+					}
+					CurrentGroupID++;
+					CurrentGroupRefSubIdx = TriSubIdx + 1;
+					bShouldBreakGroup = false;
+				}
+				
+				TrisActiveSubGroup[TriSubIdx + 1] = CurrentGroupID;
+			}
+
+			// break final subgroup
+			if (bShouldBreakGroup)
+			{
+				for (int BreakGroupIdx = CurrentGroupRefSubIdx + 1; BreakGroupIdx < GroupNum; BreakGroupIdx++)
+				{
+					TrisActiveSubGroup[BreakGroupIdx] = ++CurrentGroupID;
+				}
+			}
+
+			// for loops, merge first and last group if needed
+			int NumGroupID = CurrentGroupID + 1;
+			if (bIsLoop && TrisActiveSubGroup[0] != TrisActiveSubGroup.Last())
+			{
+				// if we can merge the first and last triangles AND their groups weren't already both broken, try to merge them 
+				if ((!bShouldBreakGroup || !bFirstGroupWasBroken) && TrisCanShareVertexPredicate(VertexID, TriangleIDs[GroupStart], TriangleIDs[GroupStart + GroupNum - 1]))
+				{
+					int EndGroupID = TrisActiveSubGroup[GroupNum - 1];
+					int StartGroupID = TrisActiveSubGroup[0];
+					int TriID0 = TriangleIDs[GroupStart];
+					bool bStoppedMerge = bShouldBreakGroup || bFirstGroupWasBroken; // if either group was broken, stop the merge before it even starts
+					for (int Idx = GroupNum - 1; !bStoppedMerge && Idx >= 0 && TrisActiveSubGroup[Idx] == EndGroupID; Idx--)
+					{
+						if (Idx < GroupNum - 1 && bNonTransitiveShare && !TrisCanShareVertexPredicate(VertexID, TriID0, TriangleIDs[GroupStart + Idx]))
+						{
+							bStoppedMerge = true;
+							break;
+						}
+						TrisActiveSubGroup[Idx] = StartGroupID;
+					}
+					NumGroupID--;
+					if (bStoppedMerge)
+					{
+						// Merge failed, break apart both groups
+						for (int BreakGroupIdx = 1; BreakGroupIdx < GroupNum; BreakGroupIdx++)
+						{
+							int Group = TrisActiveSubGroup[BreakGroupIdx];
+							if (Group == StartGroupID || Group == EndGroupID)
+							{
+								TrisActiveSubGroup[BreakGroupIdx] = NumGroupID++;
+							}
+						}
+					}
+				}
+			}
+
+			for (int Idx = 0; Idx < NumGroupID; Idx++)
+			{
+				AppendedElements.Add(AppendElement(InitElementValue, VertexID));
+			}
+			for (int TriSubIdx = 0; TriSubIdx < GroupNum; TriSubIdx++)
+			{
+				int TriID = TriangleIDs[GroupStart + TriSubIdx];
+				FIndex3i TriVertIDs = ParentMesh->GetTriangle(TriID);
+				int VertSubIdx = IndexUtil::FindTriIndex(VertexID, TriVertIDs);
+				int i = 3 * TriID;
+				ElementTriangles.InsertAt(AppendedElements[TrisActiveSubGroup[TriSubIdx]], i + VertSubIdx);
+				ElementsRefCounts.Increment(AppendedElements[TrisActiveSubGroup[TriSubIdx]]);
+			}
+			GroupStart += GroupNum;
+		}
+	}
+}
 
 
 
diff --git a/Engine/Plugins/Experimental/GeometryProcessing/Source/DynamicMesh/Private/MeshNormals.cpp b/Engine/Plugins/Experimental/GeometryProcessing/Source/DynamicMesh/Private/MeshNormals.cpp
--- a/Engine/Plugins/Experimental/GeometryProcessing/Source/DynamicMesh/Private/MeshNormals.cpp
+++ b/Engine/Plugins/Experimental/GeometryProcessing/Source/DynamicMesh/Private/MeshNormals.cpp
@@ -74,7 +74,34 @@ void FMeshNormals::Compute_FaceAvg_AreaWeighted()
 	}
 }
 
+void FMeshNormals::Compute_FaceAvg(bool bWeightByArea, bool bWeightByAngle)
+{
+	if (!bWeightByAngle && bWeightByArea)
+	{
+		Compute_FaceAvg_AreaWeighted(); // faster case
+		return;
+	}
+
+	// most general case
+	SetCount(Mesh->MaxVertexID(), true);
+
+	for (int TriIdx : Mesh->TriangleIndicesItr())
+	{
+		FVector3d TriNormal, TriCentroid; double TriArea;
+		Mesh->GetTriInfo(TriIdx, TriNormal, TriArea, TriCentroid);
+		FVector3d TriNormalWeights = GetVertexWeightsOnTriangle(TriIdx, TriArea, bWeightByArea, bWeightByAngle);
 
+		FIndex3i Triangle = Mesh->GetTriangle(TriIdx);
+		Normals[Triangle.A] += TriNormal * TriNormalWeights[0];
+		Normals[Triangle.B] += TriNormal * TriNormalWeights[1];
+		Normals[Triangle.C] += TriNormal * TriNormalWeights[2];
+	}
+
+	for (int VertIdx : Mesh->VertexIndicesItr())
+	{
+		Normals[VertIdx].Normalize();
+	}
+}
 
 void FMeshNormals::Compute_Triangle()
 {
@@ -89,6 +116,39 @@ void FMeshNormals::Compute_Triangle()
 
 
 
+void FMeshNormals::Compute_Overlay_FaceAvg(const FDynamicMeshNormalOverlay* NormalOverlay, bool bWeightByArea, bool bWeightByAngle)
+{
+	if (!bWeightByAngle && bWeightByArea)
+	{
+		Compute_Overlay_FaceAvg_AreaWeighted(NormalOverlay); // faster case
+		return;
+	}
+
+	// most general case
+	SetCount(NormalOverlay->MaxElementID(), true);
+
+	for (int TriIdx : Mesh->TriangleIndicesItr())
+	{
+		FVector3d TriNormal, TriCentroid; double TriArea;
+		Mesh->GetTriInfo(TriIdx, TriNormal, TriArea, TriCentroid);
+		FVector3d TriNormalWeights = GetVertexWeightsOnTriangle(TriIdx, TriArea, bWeightByArea, bWeightByAngle);
+
+		FIndex3i Tri = NormalOverlay->GetTriangle(TriIdx);
+		for (int j = 0; j < 3; ++j) 
+		{
+			if (Tri[j] != FDynamicMesh3::InvalidID)
+			{
+				Normals[Tri[j]] += TriNormal * TriNormalWeights[j];
+			}
+		}
+	}
+
+	for (int ElemIdx : NormalOverlay->ElementIndicesItr())
+	{
+		Normals[ElemIdx].Normalize();
+	}
+}
+
 void FMeshNormals::Compute_Overlay_FaceAvg_AreaWeighted(const FDynamicMeshNormalOverlay* NormalOverlay)
 {
 	SetCount(NormalOverlay->MaxElementID(), true);
@@ -116,7 +176,6 @@ void FMeshNormals::Compute_Overlay_FaceAvg_AreaWeighted(const FDynamicMeshNormal
 }
 
 
-
 void FMeshNormals::QuickComputeVertexNormals(FDynamicMesh3& Mesh, bool bInvert)
 {
 	FMeshNormals normals(&Mesh);
diff --git a/Engine/Plugins/Experimental/GeometryProcessing/Source/DynamicMesh/Public/DynamicMesh3.h b/Engine/Plugins/Experimental/GeometryProcessing/Source/DynamicMesh/Public/DynamicMesh3.h
--- a/Engine/Plugins/Experimental/GeometryProcessing/Source/DynamicMesh/Public/DynamicMesh3.h
+++ b/Engine/Plugins/Experimental/GeometryProcessing/Source/DynamicMesh/Public/DynamicMesh3.h
@@ -827,6 +827,15 @@ public:
 	 */
 	EMeshResult GetVtxTriangles(int VertexID, TArray<int>& TrianglesOut, bool bUseOrientation) const;
 
+	/**
+	* Get triangles connected to vertex in contiguous order, with multiple groups if vertex is a bowtie.
+	* @param VertexID Vertex ID to search around
+	* @param TrianglesOut All triangles connected to the vertex, in contiguous order; if there are multiple contiguous groups they are packed one after another
+	* @param ContiguousGroupLengths Lengths of contiguous groups packed into TrianglesOut (if not a bowtie, this will just be a length-one array w/ {TrianglesOut.Num()})
+	* @param GroupIsLoop Indicates whether each contiguous group is a loop (first triangle connected to last) or not
+	*/
+	EMeshResult GetVtxContiguousTriangles(int VertexID, TArray<int>& TrianglesOut, TArray<int>& ContiguousGroupLengths, TArray<bool>& GroupIsLoop) const;
+
 	/** Returns true if the two triangles connected to edge have different group IDs */
 	bool IsGroupBoundaryEdge(int EdgeID) const;
 
@@ -948,7 +957,10 @@ public:
 	double GetTriSolidAngle(int TriangleID, const FVector3d& p) const;
 
 	/** Compute internal angle at vertex i of triangle (where i is 0,1,2); */
-	double GetTriInternalAngleR(int TriangleID, int i);
+	double GetTriInternalAngleR(int TriangleID, int i) const;
+
+	/** Compute internal angles at all vertices of triangle */
+	FVector3d GetTriInternalAnglesR(int TriangleID) const;
 
 	/** Returns average normal of connected face normals */
 	FVector3d GetEdgeNormal(int EdgeID) const;
diff --git a/Engine/Plugins/Experimental/GeometryProcessing/Source/DynamicMesh/Public/DynamicMeshOverlay.h b/Engine/Plugins/Experimental/GeometryProcessing/Source/DynamicMesh/Public/DynamicMeshOverlay.h
--- a/Engine/Plugins/Experimental/GeometryProcessing/Source/DynamicMesh/Public/DynamicMeshOverlay.h
+++ b/Engine/Plugins/Experimental/GeometryProcessing/Source/DynamicMesh/Public/DynamicMeshOverlay.h
@@ -104,6 +104,20 @@ public:
 
 
 
+	/**
+	 * Build overlay topology from a predicate function, e.g. to build topology for sharp normals
+	 * 
+	 * @param TrisCanShareVertexPredicate Indicator function returns true if the given vertex can be shared for the given pair of triangles
+	 * @param bNonTransitiveShare If false, sharing is transitive so if Tris A and B can share a vertex, and B and C can share the vertex, then A, B and C will all use the same vertex, even if A and C could not share.
+	 *							  Otherwise, for a sequence of tris A, B, C, D, E that would all possibly share the same vertex, we require that BOTH:
+	 *									1. each adjacent pair (A,B), (B,C), etc pass the the TrisCanShareVertexPredicate and
+	 *									2. the first triangle in the sequence, A, passes the predicate with all other grouped triangles (A,C), (A,D), etc
+	 *								For a contiguous section that passes the adjacent pair tests but not the second test, we give each triangle in the section a separate normal for that vertex
+	 *								(doing anything 'smarter' like trying to find contiguous subgroups ends up just looking weird)
+	 */
+	void CreateFromPredicate(TFunctionRef<bool(int ParentVertexIdx, int TriIDA, int TriIDB)> TrisCanShareVertexPredicate, RealType InitElementValue, bool bNonTransitiveShare);
+
+
 	//
 	// Support for inserting element at specific ID. This is a bit tricky
 	// because we likely will need to update the free list in the RefCountVector, which
diff --git a/Engine/Plugins/Experimental/GeometryProcessing/Source/DynamicMesh/Public/MeshNormals.h b/Engine/Plugins/Experimental/GeometryProcessing/Source/DynamicMesh/Public/MeshNormals.h
--- a/Engine/Plugins/Experimental/GeometryProcessing/Source/DynamicMesh/Public/MeshNormals.h
+++ b/Engine/Plugins/Experimental/GeometryProcessing/Source/DynamicMesh/Public/MeshNormals.h
@@ -50,9 +50,9 @@ public:
 	/**
 	 * Compute standard per-vertex normals by averaging one-ring face normals
 	 */
-	void ComputeVertexNormals()
+	void ComputeVertexNormals(bool bWeightByArea = true, bool bWeightByAngle = true)
 	{
-		Compute_FaceAvg_AreaWeighted();
+		Compute_FaceAvg(bWeightByArea, bWeightByAngle);
 	}
 
 	/**
@@ -67,13 +67,12 @@ public:
 	 * Recompute the per-element normals of the given overlay by averaging one-ring face normals
 	 * @warning NormalOverlay must be attached to ParentMesh or an exact copy
 	 */
-	void RecomputeOverlayNormals(FDynamicMeshNormalOverlay* NormalOverlay)
+	void RecomputeOverlayNormals(FDynamicMeshNormalOverlay* NormalOverlay, bool bWeightByArea = true, bool bWeightByAngle = true)
 	{
-		Compute_Overlay_FaceAvg_AreaWeighted(NormalOverlay);
+		Compute_Overlay_FaceAvg(NormalOverlay, bWeightByArea, bWeightByAngle);
 	}
 
 
-
 	/**
 	 * Copy the current set of normals to the vertex normals of SetMesh
 	 * @warning assumes that the computed normals are vertex normals
@@ -111,11 +110,31 @@ public:
 
 
 protected:
-	/** Compute per-vertex normals using area-weighted averaging of one-ring triangles */
+	/** Compute per-vertex normals using area-weighted averaging of one-ring triangle normals */
 	void Compute_FaceAvg_AreaWeighted();
+	/** Compute per-vertex normals using a custom combination of area-weighted and angle-weighted averaging of one-ring triangle normals */
+	void Compute_FaceAvg(bool bWeightByArea, bool bWeightByAngle);
+
 	/** Compute per-triangle normals */
 	void Compute_Triangle();
-	/** Recompute the element Normals of the given attribute overlay using area-weighted averaging of one-ring triangles */
+
+	/** Recompute the element Normals of the given attribute overlay using area-weighted averaging of one-ring triangle normals */
 	void Compute_Overlay_FaceAvg_AreaWeighted(const FDynamicMeshNormalOverlay* NormalOverlay);
+	/** Recompute the element Normals of the given attribute overlay using a custom combination of area-weighted and angle-weighted averaging of one-ring triangle normals */
+	void Compute_Overlay_FaceAvg(const FDynamicMeshNormalOverlay* NormalOverlay, bool bWeightByArea, bool bWeightByAngle);
+
+	inline FVector3d GetVertexWeightsOnTriangle(int TriID, double TriArea, bool bWeightByArea, bool bWeightByAngle) const
+	{
+		FVector3d TriNormalWeights = FVector3d::One();
+		if (bWeightByArea)
+		{
+			TriNormalWeights = Mesh->GetTriInternalAnglesR(TriID); // component-wise multiply by per-vertex internal angles
+		}
+		if (bWeightByArea)
+		{
+			TriNormalWeights *= TriArea;
+		}
+		return TriNormalWeights;
+	}
 
 };
\ No newline at end of file
diff --git a/Engine/Plugins/Experimental/GeometryProcessing/Source/MeshConversion/Private/DynamicMeshToMeshDescription.cpp b/Engine/Plugins/Experimental/GeometryProcessing/Source/MeshConversion/Private/DynamicMeshToMeshDescription.cpp
--- a/Engine/Plugins/Experimental/GeometryProcessing/Source/MeshConversion/Private/DynamicMeshToMeshDescription.cpp
+++ b/Engine/Plugins/Experimental/GeometryProcessing/Source/MeshConversion/Private/DynamicMeshToMeshDescription.cpp
@@ -9,7 +9,7 @@
 
 
 
-void FDynamicMeshToMeshDescription::Update(const FDynamicMesh3* MeshIn, FMeshDescription& MeshOut)
+void FDynamicMeshToMeshDescription::Update(const FDynamicMesh3* MeshIn, FMeshDescription& MeshOut, bool bRecomputeNormals)
 {
 	FMeshDescriptionBuilder Builder;
 	Builder.SetMeshDescription(&MeshOut);
@@ -23,16 +23,108 @@ void FDynamicMeshToMeshDescription::Update(const FDynamicMesh3* MeshIn, FMeshDes
 		Builder.SetPosition(FVertexID(VertID), MeshIn->GetVertex(VertID));
 	}
 
-	// can't trust these yet...
-	//const FDynamicMeshUVOverlay* UVOverlay = MeshIn->HasAttributes() ? MeshIn->Attributes()->PrimaryUV() : nullptr;
-	//const FDynamicMeshNormalOverlay* NormalOverlay = MeshIn->HasAttributes() ? MeshIn->Attributes()->PrimaryNormals() : nullptr;
+	if (bRecomputeNormals)
+	{
+		Builder.RecalculateInstanceNormals();
+	}
+	else
+	{
+		UpdateAttributes(MeshIn, MeshOut, true, false);
+	}
+}
 
-	Builder.RecalculateInstanceNormals();
+
+namespace DynamicMeshToMeshDescriptionConversionHelper
+{
+	// NOTE: assumes the order of triangles in the MeshIn correspond to the ordering you'd get by iterating over polygons, then tris-in-polygons, on MeshOut
+	// This matches conversion currently used in MeshDescriptionToDynamicMesh.cpp, but if that changes we will need to change this function to match!
+	template <typename OutAttributeType, int VecLen, typename InAttributeType>
+	void SetAttributesFromOverlay(
+		const FDynamicMesh3* MeshInArg, const FMeshDescription& MeshOutArg,
+		TVertexInstanceAttributesRef<OutAttributeType>& InstanceAttrib, const TDynamicMeshVectorOverlay<float, VecLen, InAttributeType>* Overlay)
+	{
+		const FPolygonArray& Polygons = MeshOutArg.Polygons();
+		int MeshInTriIdx = 0;
+		for (const FPolygonID PolygonID : Polygons.GetElementIDs())
+		{
+			FPolygonGroupID PolygonGroupID = MeshOutArg.GetPolygonPolygonGroup(PolygonID);
+
+			const TArray<FTriangleID>& TriangleIDs = MeshOutArg.GetPolygonTriangleIDs(PolygonID);
+			int NumTriangles = TriangleIDs.Num();
+			for (int TriIdx = 0; TriIdx < NumTriangles; ++TriIdx, ++MeshInTriIdx)
+			{
+				const FTriangleID TriangleID = TriangleIDs[TriIdx];
+				TArrayView<const FVertexInstanceID> InstanceTri = MeshOutArg.GetTriangleVertexInstances(TriangleID);
+
+				FIndex3i OverlayVertIndices = Overlay->GetTriangle(MeshInTriIdx);
+				InstanceAttrib.Set(InstanceTri[0], OutAttributeType(Overlay->GetElement(OverlayVertIndices.A)));
+				InstanceAttrib.Set(InstanceTri[1], OutAttributeType(Overlay->GetElement(OverlayVertIndices.B)));
+				InstanceAttrib.Set(InstanceTri[2], OutAttributeType(Overlay->GetElement(OverlayVertIndices.C)));
+			}
+		}
+	}
 }
 
 
+void FDynamicMeshToMeshDescription::UpdateAttributes(const FDynamicMesh3* MeshIn, FMeshDescription& MeshOut, bool bUpdateNormals, bool bUpdateUVs)
+{
+	check(MeshIn->IsCompactV());
+	check(MeshIn->VertexCount() == MeshOut.Vertices().Num());
+	check(MeshIn->TriangleCount() == MeshOut.Triangles().Num())
+
+	if (bUpdateNormals)
+	{
+		TVertexInstanceAttributesRef<FVector> InstanceAttrib =
+			MeshOut.VertexInstanceAttributes().GetAttributesRef<FVector>(MeshAttribute::VertexInstance::Normal);
+		ensureMsgf(InstanceAttrib.IsValid(), TEXT("Trying to update normals on a MeshDescription that has no normal attributes"));
+		if (InstanceAttrib.IsValid())
+		{
+			const FDynamicMeshNormalOverlay* Overlay = MeshIn->HasAttributes() ? MeshIn->Attributes()->PrimaryNormals() : nullptr;
+			if (Overlay)
+			{
+				DynamicMeshToMeshDescriptionConversionHelper::SetAttributesFromOverlay(MeshIn, MeshOut, InstanceAttrib, Overlay);
+			}
+			else
+			{
+				for (int VertID : MeshIn->VertexIndicesItr())
+				{
+					FVector Normal = (FVector)MeshIn->GetVertexNormal(VertID);
+					for (FVertexInstanceID InstanceID : MeshOut.GetVertexVertexInstances(FVertexID(VertID)))
+					{
+						InstanceAttrib.Set(InstanceID, Normal);
+					}
+				}
+			}
+		}
+	}
 
+	if (bUpdateUVs)
+	{
+		const FDynamicMeshUVOverlay* Overlay = MeshIn->HasAttributes() ? MeshIn->Attributes()->PrimaryUV() : nullptr;
 
+		TVertexInstanceAttributesRef<FVector2D> InstanceAttrib =
+			MeshOut.VertexInstanceAttributes().GetAttributesRef<FVector2D>(MeshAttribute::VertexInstance::TextureCoordinate);
+		ensureMsgf(InstanceAttrib.IsValid(), TEXT("Trying to update UVs on a MeshDescription that has no texture coordinate attributes"));
+		if (InstanceAttrib.IsValid())
+		{
+			if (Overlay)
+			{
+				DynamicMeshToMeshDescriptionConversionHelper::SetAttributesFromOverlay(MeshIn, MeshOut, InstanceAttrib, Overlay);
+			}
+			else
+			{
+				for (int VertID : MeshIn->VertexIndicesItr())
+				{
+					FVector2D UV = (FVector2D)MeshIn->GetVertexUV(VertID);
+					for (FVertexInstanceID InstanceID : MeshOut.GetVertexVertexInstances(FVertexID(VertID)))
+					{
+						InstanceAttrib.Set(InstanceID, UV);
+					}
+				}
+			}
+		}
+	}
+}
 
 
 
diff --git a/Engine/Plugins/Experimental/GeometryProcessing/Source/MeshConversion/Private/MeshDescriptionToDynamicMesh.cpp b/Engine/Plugins/Experimental/GeometryProcessing/Source/MeshConversion/Private/MeshDescriptionToDynamicMesh.cpp
--- a/Engine/Plugins/Experimental/GeometryProcessing/Source/MeshConversion/Private/MeshDescriptionToDynamicMesh.cpp
+++ b/Engine/Plugins/Experimental/GeometryProcessing/Source/MeshConversion/Private/MeshDescriptionToDynamicMesh.cpp
@@ -167,6 +167,8 @@ void FMeshDescriptionToDynamicMesh::Convert(const FMeshDescription* MeshIn, FDyn
 	FUVWelder UVWelder(UVOverlay);
 	FNormalWelder NormalWelder(NormalOverlay);
 
+	// NOTE: If you change the iteration order here, please update the corresponding iteration in FDynamicMeshToMeshDescription::UpdateAttributes, 
+	//	which assumes the iteration order here is polygons -> triangles, to correspond the triangles when writing updated attributes back!
 	const FPolygonArray& Polygons = MeshIn->Polygons();
 	for (const FPolygonID PolygonID : Polygons.GetElementIDs())
 	{
diff --git a/Engine/Plugins/Experimental/GeometryProcessing/Source/MeshConversion/Public/DynamicMeshToMeshDescription.h b/Engine/Plugins/Experimental/GeometryProcessing/Source/MeshConversion/Public/DynamicMeshToMeshDescription.h
--- a/Engine/Plugins/Experimental/GeometryProcessing/Source/MeshConversion/Public/DynamicMeshToMeshDescription.h
+++ b/Engine/Plugins/Experimental/GeometryProcessing/Source/MeshConversion/Public/DynamicMeshToMeshDescription.h
@@ -27,11 +27,18 @@ public:
 
 	/**
 	 * Update existing MeshDescription based on DynamicMesh. Assumes mesh topology has not changed.
-	 * Copies positions, recalculates MeshDescription normals.
+	 * Copies positions, recalculates MeshDescription normals if bRecomputeNormals is true (copies otherwise).
 	 */
-	void Update(const FDynamicMesh3* MeshIn, FMeshDescription& MeshOut);
+	void Update(const FDynamicMesh3* MeshIn, FMeshDescription& MeshOut, bool bRecomputeNormals = true);
 
 
+	/**
+	 * Update only attributes, assuming the mesh topology has not changed.  Does not touch positions.
+	 *	NOTE: assumes the order of triangles in the MeshIn correspond to the ordering you'd get by iterating over polygons, then tris-in-polygons, on MeshOut
+	 *		  This matches conversion currently used in MeshDescriptionToDynamicMesh.cpp, but if that changes we will need to change this function to match!
+	 */
+	void UpdateAttributes(const FDynamicMesh3* MeshIn, FMeshDescription& MeshOut, bool bUpdateNormals, bool bUpdateUVs);
+
 	//
 	// Internal functions that you can also call directly
 	//
