fix dynamicmesh3_queries.cpp build warning; replace implicit int->bool w/ an explicit comparison vs 0

#rnx

diff --git a/Engine/Plugins/Experimental/GeometryProcessing/Source/DynamicMesh/Private/DynamicMesh3_Queries.cpp b/Engine/Plugins/Experimental/GeometryProcessing/Source/DynamicMesh/Private/DynamicMesh3_Queries.cpp
--- a/Engine/Plugins/Experimental/GeometryProcessing/Source/DynamicMesh/Private/DynamicMesh3_Queries.cpp
+++ b/Engine/Plugins/Experimental/GeometryProcessing/Source/DynamicMesh/Private/DynamicMesh3_Queries.cpp
@@ -200,7 +200,7 @@ EMeshResult FDynamicMesh3::GetVtxContiguousTriangles(int VertexID, TArray<int>&
 			StartEdgeIDs.Push(EID);
 		}
 	}
-	bool bHasBoundaries = StartEdgeIDs.Num();
+	bool bHasBoundaries = StartEdgeIDs.Num() != 0;
 	
 	if (!bHasBoundaries)
 	{
