GeometryProcessing: add FDynamicMeshEditor::RemoveTriangles() and FDynamicMeshEditor::AppendTriangles() utility functions
#rnx

diff --git a/Engine/Plugins/Experimental/GeometryProcessing/Source/DynamicMesh/Private/DynamicMeshEditor.cpp b/Engine/Plugins/Experimental/GeometryProcessing/Source/DynamicMesh/Private/DynamicMeshEditor.cpp
--- a/Engine/Plugins/Experimental/GeometryProcessing/Source/DynamicMesh/Private/DynamicMeshEditor.cpp
+++ b/Engine/Plugins/Experimental/GeometryProcessing/Source/DynamicMesh/Private/DynamicMeshEditor.cpp
@@ -233,16 +233,24 @@ operation_failed:
 
 
 bool FDynamicMeshEditor::RemoveTriangles(const TArray<int>& Triangles, bool bRemoveIsolatedVerts)
+{
+	return RemoveTriangles(Triangles, bRemoveIsolatedVerts, [](int) {});
+}
+
+
+bool FDynamicMeshEditor::RemoveTriangles(const TArray<int>& Triangles, bool bRemoveIsolatedVerts, TFunctionRef<void(int)> OnRemoveTriFunc)
 {
 	bool bAllOK = true;
 	int NumTriangles = Triangles.Num();
-	for (int i = 0; i < NumTriangles; ++i) 
+	for (int i = 0; i < NumTriangles; ++i)
 	{
 		if (Mesh->IsTriangle(Triangles[i]) == false)
 		{
 			continue;
 		}
 
+		OnRemoveTriFunc(Triangles[i]);
+
 		EMeshResult result = Mesh->RemoveTriangle(Triangles[i], bRemoveIsolatedVerts, false);
 		if (result != EMeshResult::Ok)
 		{
@@ -290,6 +298,7 @@ void FDynamicMeshEditor::DuplicateTriangles(const TArray<int>& Triangles, FMeshI
 
 
 
+
 bool FDynamicMeshEditor::DisconnectTriangles(const TArray<int>& Triangles, TArray<FLoopPairSet>& LoopSetOut)
 {
 	check(Mesh->HasAttributes() == false);  // not supported yet
@@ -905,3 +914,167 @@ void FDynamicMeshEditor::AppendUVs(const FDynamicMesh3* AppendMesh,
 		ToUVs->SetTriangle(NewTriID, ElemTri);
 	}
 }
+
+
+
+
+
+
+
+
+
+
+
+
+// can these be replaced w/ template function?
+
+// Utility function for ::AppendTriangles()
+static int AppendTriangleUVAttribute(const FDynamicMesh3* FromMesh, int FromElementID, FDynamicMesh3* ToMesh, int UVLayerIndex, FMeshIndexMappings& IndexMaps)
+{
+	int NewElementID = IndexMaps.GetNewUV(UVLayerIndex, FromElementID);
+	if (NewElementID == IndexMaps.InvalidID())
+	{
+		const FDynamicMeshUVOverlay* FromUVOverlay = FromMesh->Attributes()->GetUVLayer(UVLayerIndex);
+		FDynamicMeshUVOverlay* ToUVOverlay = ToMesh->Attributes()->GetUVLayer(UVLayerIndex);
+
+		// need to determine new parent vertex. It should be in the map already!
+		int ParentVertexID = FromUVOverlay->GetParentVertex(FromElementID);
+		int NewParentVertexID = IndexMaps.GetNewVertex(ParentVertexID);
+		check(NewParentVertexID != IndexMaps.InvalidID());
+
+		NewElementID = ToUVOverlay->AppendElement(
+			FromUVOverlay->GetElement(FromElementID), NewParentVertexID);
+
+		IndexMaps.SetUV(UVLayerIndex, FromElementID, NewElementID);
+	}
+	return NewElementID;
+}
+
+
+// Utility function for ::AppendTriangles()
+static int AppendTriangleNormalAttribute(const FDynamicMesh3* FromMesh, int FromElementID, FDynamicMesh3* ToMesh, int NormalLayerIndex, FMeshIndexMappings& IndexMaps)
+{
+	int NewElementID = IndexMaps.GetNewNormal(NormalLayerIndex, FromElementID);
+	if (NewElementID == IndexMaps.InvalidID())
+	{
+		const FDynamicMeshNormalOverlay* FromNormalOverlay = FromMesh->Attributes()->GetNormalLayer(NormalLayerIndex);
+		FDynamicMeshNormalOverlay* ToNormalOverlay = ToMesh->Attributes()->GetNormalLayer(NormalLayerIndex);
+
+		// need to determine new parent vertex. It should be in the map already!
+		int ParentVertexID = FromNormalOverlay->GetParentVertex(FromElementID);
+		int NewParentVertexID = IndexMaps.GetNewVertex(ParentVertexID);
+		check(NewParentVertexID != IndexMaps.InvalidID());
+
+		NewElementID = ToNormalOverlay->AppendElement(
+			FromNormalOverlay->GetElement(FromElementID), NewParentVertexID);
+
+		IndexMaps.SetNormal(NormalLayerIndex, FromElementID, NewElementID);
+	}
+	return NewElementID;
+}
+
+
+
+
+// Utility function for ::AppendTriangles()
+static void AppendAttributes(const FDynamicMesh3* FromMesh, int FromTriangleID, FDynamicMesh3* ToMesh, int ToTriangleID, FMeshIndexMappings& IndexMaps, FDynamicMeshEditResult& ResultOut)
+{
+	if (FromMesh->HasAttributes() == false || ToMesh->HasAttributes() == false)
+	{
+		return;
+	}
+
+	// todo: copy all attribute layers
+	check(FromMesh->Attributes()->GetNumUVLayers() == 1);
+	check(FromMesh->Attributes()->GetNumNormalLayers() == 1);
+
+	const FDynamicMeshUVOverlay* FromUVOverlay = FromMesh->Attributes()->PrimaryUV();
+	FDynamicMeshUVOverlay* ToUVOverlay = ToMesh->Attributes()->PrimaryUV();
+
+	{
+		FIndex3i FromElemTri = FromUVOverlay->GetTriangle(FromTriangleID);
+		FIndex3i ToElemTri = ToUVOverlay->GetTriangle(ToTriangleID);
+		for (int j = 0; j < 3; ++j)
+		{
+			if (FromElemTri[j] != FDynamicMesh3::InvalidID)
+			{
+				int NewElemID = AppendTriangleUVAttribute(FromMesh, FromElemTri[j], ToMesh, 0, IndexMaps);
+				ToElemTri[j] = NewElemID;
+			}
+		}
+		ToUVOverlay->SetTriangle(ToTriangleID, ToElemTri);
+	}
+
+
+	const FDynamicMeshNormalOverlay* FromNormalOverlay = FromMesh->Attributes()->PrimaryNormals();
+	FDynamicMeshNormalOverlay* ToNormalOverlay = ToMesh->Attributes()->PrimaryNormals();
+
+	{
+		FIndex3i FromElemTri = FromNormalOverlay->GetTriangle(FromTriangleID);
+		FIndex3i ToElemTri = ToNormalOverlay->GetTriangle(ToTriangleID);
+		for (int j = 0; j < 3; ++j)
+		{
+			if (FromElemTri[j] != FDynamicMesh3::InvalidID)
+			{
+				int NewElemID = AppendTriangleNormalAttribute(FromMesh, FromElemTri[j], ToMesh, 0, IndexMaps);
+				ToElemTri[j] = NewElemID;
+			}
+		}
+		ToNormalOverlay->SetTriangle(ToTriangleID, ToElemTri);
+	}
+}
+
+
+
+
+
+
+void FDynamicMeshEditor::AppendTriangles(const FDynamicMesh3* SourceMesh, const TArray<int>& SourceTriangles, FMeshIndexMappings& IndexMaps, FDynamicMeshEditResult& ResultOut)
+{
+	ResultOut.Reset();
+	IndexMaps.Initialize(Mesh);
+
+	for (int SourceTriangleID : SourceTriangles)
+	{
+		check(SourceMesh->IsTriangle(SourceTriangleID));
+		if (SourceMesh->IsTriangle(SourceTriangleID) == false)
+		{
+			continue;	// ignore missing triangles
+		}
+
+		FIndex3i Tri = SourceMesh->GetTriangle(SourceTriangleID);
+
+		// FindOrCreateDuplicateGroup
+		int SourceGroupID = SourceMesh->GetTriangleGroup(SourceTriangleID);
+		int NewGroupID = IndexMaps.GetNewGroup(SourceGroupID);
+		if (NewGroupID == IndexMaps.InvalidID())
+		{
+			NewGroupID = Mesh->AllocateTriangleGroup();
+			IndexMaps.SetGroup(SourceGroupID, NewGroupID);
+			ResultOut.NewGroups.Add(NewGroupID);
+		}
+
+		// FindOrCreateDuplicateVertex
+		FIndex3i NewTri;
+		for (int j = 0; j < 3; ++j)
+		{
+			int SourceVertexID = Tri[j];
+			int NewVertexID = IndexMaps.GetNewVertex(SourceVertexID);
+			if (NewVertexID == IndexMaps.InvalidID())
+			{
+				NewVertexID = Mesh->AppendVertex(*SourceMesh, SourceVertexID);
+				IndexMaps.SetVertex(SourceVertexID, NewVertexID);
+				ResultOut.NewVertices.Add(NewVertexID);
+			}
+			NewTri[j] = NewVertexID;
+		}
+
+		int NewTriangleID = Mesh->AppendTriangle(NewTri, NewGroupID);
+		IndexMaps.SetTriangle(SourceTriangleID, NewTriangleID);
+		ResultOut.NewTriangles.Add(NewTriangleID);
+
+		AppendAttributes(SourceMesh, SourceTriangleID, Mesh, NewTriangleID, IndexMaps, ResultOut);
+
+		//Mesh->CheckValidity(true);
+	}
+}
diff --git a/Engine/Plugins/Experimental/GeometryProcessing/Source/DynamicMesh/Public/DynamicMeshEditor.h b/Engine/Plugins/Experimental/GeometryProcessing/Source/DynamicMesh/Public/DynamicMeshEditor.h
--- a/Engine/Plugins/Experimental/GeometryProcessing/Source/DynamicMesh/Public/DynamicMeshEditor.h
+++ b/Engine/Plugins/Experimental/GeometryProcessing/Source/DynamicMesh/Public/DynamicMeshEditor.h
@@ -201,6 +201,15 @@ public:
 	 */
 	bool RemoveTriangles(const TArray<int>& Triangles, bool bRemoveIsolatedVerts);
 
+	/**
+	 * Remove a list of triangles from the mesh, and optionally any vertices that are now orphaned
+	 * @param Triangles the triangles to remove
+	 * @param bRemoveIsolatedVerts if true, remove vertices that end up with no triangles
+	 * @param OnRemoveTriFunc called for each triangle to be removed
+	 * @return true if all removes succeeded
+	 */
+	bool RemoveTriangles(const TArray<int>& Triangles, bool bRemoveIsolatedVerts, TFunctionRef<void(int)> OnRemoveTriFunc);
+
 
 	//////////////////////////////////////////////////////////////////////////
 	// Normal utility functions
@@ -373,6 +382,20 @@ public:
 		const FDynamicMeshUVOverlay* FromUVs, FDynamicMeshUVOverlay* ToUVs,
 		const FIndexMapi& VertexMap, const FIndexMapi& TriangleMap,
 		FIndexMapi& UVMapOut);
+
+
+
+
+	/**
+	 * Append triangles of an existing mesh. This duplicates the current groups and also any attributes existing on the triangles.
+	 * @param SourceMesh the mesh to copy from
+	 * @param SourceTriangles the triangles to copy
+	 * @param IndexMaps returned mappings from old to new triangles/vertices/etc (you may initialize to optimize memory usage, etc)
+	 * @param ResultOut lists of newly created triangles/vertices/etc
+	 */
+	void AppendTriangles(const FDynamicMesh3* SourceMesh, const TArray<int>& SourceTriangles, FMeshIndexMappings& IndexMaps, FDynamicMeshEditResult& ResultOut);
+
+
 };
 
 
