More generators (sphere, grid box), connected up to MakeMeshTestTool

#rnx

diff --git a/Engine/Plugins/Experimental/GeometryProcessing/Source/DynamicMesh/Public/Generators/GridBoxMeshGenerator.h b/Engine/Plugins/Experimental/GeometryProcessing/Source/DynamicMesh/Public/Generators/GridBoxMeshGenerator.h
new file mode 100644
--- /dev/null
+++ b/Engine/Plugins/Experimental/GeometryProcessing/Source/DynamicMesh/Public/Generators/GridBoxMeshGenerator.h
@@ -0,0 +1,247 @@
+// Copyright 1998-2019 Epic Games, Inc. All Rights Reserved.
+
+#pragma once
+
+#include "MeshShapeGenerator.h"
+#include "OrientedBoxTypes.h"
+#include "Util/IndexUtil.h"
+
+/**
+ *  Generate a mesh of a box that has "gridded" faces, i.e. grid of triangulated quads, 
+ *  with EdgeVertices setting the number of verts along each edge in each dimension
+ */
+class /*DYNAMICMESH_API*/ FGridBoxMeshGenerator : public FMeshShapeGenerator
+{
+public:
+	/** 3D box */
+	FOrientedBox3d Box;
+
+	FIndex3i EdgeVertices { 8,8,8 };
+
+	/** If true (default), UVs are scaled so that there is no stretching. If false, UVs are scaled to fill unit square */
+	bool bScaleUVByAspectRatio = true;
+
+public:
+
+	/** Generate the mesh */
+	virtual FMeshShapeGenerator& Generate() override
+	{
+		FIndex3i N = EdgeVertices;
+		N.A = FMath::Max(2, N.A);
+		N.B = FMath::Max(2, N.B);
+		N.C = FMath::Max(2, N.C);
+
+		FVector3d Nscale(1.0 / (N.A-1), 1.0 / (N.B-1), 1.0 / (N.C-1));
+
+		FIndex3i NInternal = N;
+		NInternal.A -= 2; NInternal.B -= 2; NInternal.C -= 2;
+
+		FIndex3i NTri = N;
+		NTri.A--; NTri.B--; NTri.C--;
+		int NumUVsAndNormals = 2 * (N.A * N.B + N.B * N.C + N.C * N.A);
+		int NumVertices = 8 + (NInternal.A + NInternal.B + NInternal.C) * 4 + (NInternal.A*NInternal.B + NInternal.B*NInternal.C + NInternal.C*NInternal.A) * 2;
+		int NumTriangles = 4 * (NTri.A * NTri.B + NTri.B * NTri.C + NTri.C * NTri.A);
+		SetBufferSizes(NumVertices, NumTriangles, NumUVsAndNormals, NumUVsAndNormals);
+
+		int FaceDimOrder[3]{ 1, 2, 0 }; // ordering from IndexUtil.cpp
+		int D[2][3]{ { 1,2,0 }, { 2,0,1 } }; // helper mapping to go from major dimension to the two canonical sub-dimensions in order; is just D[i][j] -> (j+i+1)%3 (TODO: replace with the modulus version?)
+
+		// allocate a big mapping from faces to vertices, to make life easier when creating triangles later
+		FIndex3i VerticesPerFace(N.B * N.C, N.A * N.C, N.A * N.B);
+		TArray<TArray<int>> FaceVertIndices;
+		FaceVertIndices.SetNum(6);
+		for (int Dim = 0; Dim < 3; Dim++)
+		{
+			int FaceIdxForDim = FaceDimOrder[Dim] * 2;
+			int VertexNum = VerticesPerFace[Dim];
+			FaceVertIndices[FaceIdxForDim].SetNum(VertexNum);
+			FaceVertIndices[FaceIdxForDim+1].SetNum(VertexNum);
+		}
+
+		auto ToFaceV = [&N, &D](int Dim, int D0, int D1)
+		{
+			return D0 + N[D[0][Dim]] * D1;
+		};
+
+		// create the corners and distribute them into the face mapping
+		for (int i = 0; i < 8; ++i) 
+		{
+			Vertices[i] = Box.GetCorner(i);
+			FIndex3i CornerSides = FOrientedBox3d::GetCornerSide(i);
+			for (int Dim = 0; Dim < 3; Dim++)
+			{
+				int FaceIdx = FaceDimOrder[Dim]*2 + CornerSides[Dim];
+				int D0Ind = CornerSides[D[0][Dim]] * (N[D[0][Dim]] - 1);
+				int D1Ind = CornerSides[D[1][Dim]] * (N[D[1][Dim]] - 1);
+				int CornerVInd = ToFaceV(Dim, D0Ind, D1Ind);
+				FaceVertIndices[FaceIdx][CornerVInd] = i;
+			}
+		}
+
+		// create the internal (non-corner) edge vertices and distribute them into the face mapping
+		int CurrentVertIndex = 8;
+		for (int Dim = 0; Dim < 3; Dim++)
+		{
+			int EdgeLen = N[Dim];
+			if (EdgeLen <= 2)
+			{
+				continue; // no internal edge vertices on this dimension
+			}
+			int MajorFaceInd = FaceDimOrder[Dim] * 2;
+			int FaceInds[2]{ -1,-1 };
+			int DSides[2]{ 0,0 };
+			for (DSides[0] = 0; DSides[0] < 2; DSides[0]++)
+			{
+				FaceInds[0] = FaceDimOrder[D[0][Dim]] * 2 + DSides[0];
+				for (DSides[1] = 0; DSides[1] < 2; DSides[1]++)
+				{
+					FaceInds[1] = FaceDimOrder[D[1][Dim]] * 2 + DSides[1];
+					int MajorCornerInd = ToFaceV(Dim, DSides[0] * (N[D[0][Dim]] - 1), DSides[1] * (N[D[1][Dim]] - 1));
+					FVector3d Corners[2]
+					{
+						Vertices[FaceVertIndices[MajorFaceInd  ][MajorCornerInd]],
+						Vertices[FaceVertIndices[MajorFaceInd+1][MajorCornerInd]]
+					};
+					
+					for (int EdgeVert = 1; EdgeVert + 1 < EdgeLen; EdgeVert++)
+					{
+						Vertices[CurrentVertIndex] = FVector3d::Lerp(Corners[0], Corners[1], EdgeVert * Nscale[Dim]);
+						for (int WhichFace = 0; WhichFace < 2; WhichFace++) // each edge is shared by two faces (w/ major axes of subdim 0 and subdim 1 respectively)
+						{
+							int FaceInd = FaceInds[WhichFace];
+
+							int FaceDim = D[WhichFace][Dim];
+							int SubDims[2];
+							SubDims[1 - WhichFace] = EdgeVert;
+							SubDims[WhichFace] = NTri[D[WhichFace][FaceDim]] * DSides[1-WhichFace];
+							int FaceV = ToFaceV(FaceDim, SubDims[0], SubDims[1]);
+							FaceVertIndices[FaceInds[WhichFace]][FaceV] = CurrentVertIndex;
+						}
+						CurrentVertIndex++;
+					}
+				}
+			}
+		}
+
+		// create the internal (non-corner, non-edge) face vertices and distribute them into the face mapping
+		for (int Dim = 0; Dim < 3; Dim++)
+		{
+			int FaceIdxBase = FaceDimOrder[Dim]*2;
+			int FaceInternalVNum = NInternal[D[0][Dim]] * NInternal[D[1][Dim]];
+			if (FaceInternalVNum <= 0)
+			{
+				continue;
+			}
+
+			for (int Side = 0; Side < 2; Side++)
+			{
+				int MajorFaceInd = FaceIdxBase + Side;
+				for (int D0 = 1; D0 + 1 < N[D[0][Dim]]; D0++)
+				{
+					int BotInd = ToFaceV(Dim, D0, 0);
+					int TopInd = ToFaceV(Dim, D0, N[D[1][Dim]] - 1);
+
+					FVector3d Edges[2]
+					{
+						Vertices[FaceVertIndices[MajorFaceInd][BotInd]],
+						Vertices[FaceVertIndices[MajorFaceInd][TopInd]]
+					};
+					for (int D1 = 1; D1 + 1 < N[D[1][Dim]]; D1++)
+					{
+						Vertices[CurrentVertIndex] = FVector3d::Lerp(Edges[0], Edges[1], D1 * Nscale[D[1][Dim]]);
+						FaceVertIndices[MajorFaceInd][ToFaceV(Dim, D0, D1)] = CurrentVertIndex;
+						CurrentVertIndex++;
+					}
+				}
+			}
+		}
+
+		double MaxDimension = 2.0*Box.Extents.MaxAbs();
+		float UVScale = (bScaleUVByAspectRatio) ? (1.0f / (float)MaxDimension) : 1.0f;
+
+		// create the face triangles and UVs+normals
+		int CurrentTriIdx = 0;
+		int CurrentUVIdx = 0;
+		for (int Dim = 0; Dim < 3; Dim++)
+		{
+			int FaceIdxBase = FaceDimOrder[Dim]*2;
+			double FaceWidth = FMathd::Abs(Box.Extents[D[0][Dim]]) * 2.0;
+			double FaceHeight = FMathd::Abs(Box.Extents[D[1][Dim]]) * 2.0;
+			double WidthUVScale = FaceWidth * UVScale;
+			double HeightUVScale = FaceHeight * UVScale;
+
+			
+			for (int Side = 0; Side < 2; Side++)
+			{
+				int SideOpp = 1 - Side;
+
+				FVector3f Normal(0, 0, 0);
+				Normal[Dim] = (2 * Side - 1);
+				int MajorFaceInd = FaceIdxBase + Side;
+
+				int FaceUVStartInd = CurrentUVIdx;
+				// set all the UVs and normals
+				for (int D0 = 0; D0 < N[D[0][Dim]]; D0++)
+				{
+					for (int D1 = 0; D1 < N[D[1][Dim]]; D1++)
+					{
+						UVs[CurrentUVIdx] = FVector2f(D0*WidthUVScale*Nscale[D[0][Dim]], D1*HeightUVScale*Nscale[D[1][Dim]]);
+						Normals[CurrentUVIdx] = Normal;
+						UVParentVertex[CurrentUVIdx] = FaceVertIndices[MajorFaceInd][ToFaceV(Dim, D0, D1)];
+						NormalParentVertex[CurrentUVIdx] = FaceVertIndices[MajorFaceInd][ToFaceV(Dim, D0, D1)];
+						CurrentUVIdx++;
+					}
+				}
+
+				// set all the triangles
+				for (int D0 = 0; D0 + 1 < N[D[0][Dim]]; D0++)
+				{
+					for (int D1 = 0; D1 + 1 < N[D[1][Dim]]; D1++)
+					{
+						SetTriangle(CurrentTriIdx,
+								FaceVertIndices[MajorFaceInd][ToFaceV(Dim, D0, D1)],
+								FaceVertIndices[MajorFaceInd][ToFaceV(Dim, D0+SideOpp, D1+Side)],
+								FaceVertIndices[MajorFaceInd][ToFaceV(Dim, D0+1, D1+1)]
+							);
+
+						SetTriangleUVs(CurrentTriIdx,
+							FaceUVStartInd + D1 + (D0) * N[D[1][Dim]],
+							FaceUVStartInd + D1+Side + (D0+SideOpp) * N[D[1][Dim]],
+							FaceUVStartInd + D1+1 + (D0+1) * N[D[1][Dim]]
+						);
+						SetTriangleNormals(CurrentTriIdx,
+							FaceUVStartInd + D1 + (D0) * N[D[1][Dim]],
+							FaceUVStartInd + D1+Side + (D0+SideOpp) * N[D[1][Dim]],
+							FaceUVStartInd + D1+1 + (D0+1) * N[D[1][Dim]]
+						);
+						SetTrianglePolygon(CurrentTriIdx, MajorFaceInd);
+						CurrentTriIdx++;
+
+						SetTriangle(CurrentTriIdx,
+							FaceVertIndices[MajorFaceInd][ToFaceV(Dim, D0, D1)],
+							FaceVertIndices[MajorFaceInd][ToFaceV(Dim, D0+1, D1+1)],
+							FaceVertIndices[MajorFaceInd][ToFaceV(Dim, D0+Side, D1+SideOpp)]
+						);
+						SetTriangleUVs(CurrentTriIdx,
+							FaceUVStartInd + D1 + (D0) * N[D[1][Dim]],
+							FaceUVStartInd + D1+1 + (D0+1) * N[D[1][Dim]],
+							FaceUVStartInd + D1+SideOpp + (D0+Side) * N[D[1][Dim]]
+						);
+						SetTriangleNormals(CurrentTriIdx,
+							FaceUVStartInd + D1 + (D0) * N[D[1][Dim]],
+							FaceUVStartInd + D1+1 + (D0+1) * N[D[1][Dim]],
+							FaceUVStartInd + D1+SideOpp + (D0+Side) * N[D[1][Dim]]
+						);
+						SetTrianglePolygon(CurrentTriIdx, MajorFaceInd);
+						CurrentTriIdx++;
+					}
+				}
+			}
+		}
+
+
+		return *this;
+	}
+
+};
+
diff --git a/Engine/Plugins/Experimental/GeometryProcessing/Source/DynamicMesh/Public/Generators/SphereGenerator.h b/Engine/Plugins/Experimental/GeometryProcessing/Source/DynamicMesh/Public/Generators/SphereGenerator.h
new file mode 100644
--- /dev/null
+++ b/Engine/Plugins/Experimental/GeometryProcessing/Source/DynamicMesh/Public/Generators/SphereGenerator.h
@@ -0,0 +1,206 @@
+// Copyright 1998-2019 Epic Games, Inc. All Rights Reserved.
+
+#pragma once
+
+#include "MeshShapeGenerator.h"
+#include "OrientedBoxTypes.h"
+#include "Util/IndexUtil.h"
+
+/**
+ * Generate a sphere mesh, with UVs wrapped cylindrically
+ */
+class /*DYNAMICMESH_API*/ FSphereGenerator : public FMeshShapeGenerator
+{
+public:
+	/** Radius */
+	double Radius = 1;
+
+	int NumPhi = 16; // number of vertices along vertical extent from north pole to south pole
+	int NumTheta = 16; // number of vertices around circles
+
+public:
+
+	inline static FVector3d SphericalToCartesian(double r, double theta, double phi)
+	{
+		double Sphi = sin(phi);
+		double Cphi = cos(phi);
+		double Ctheta = cos(theta);
+		double Stheta = sin(theta);
+
+		return FVector3d(r * Ctheta * Sphi, r * Stheta * Sphi, r * Cphi);
+	}
+
+	/** Generate the mesh */
+	virtual FMeshShapeGenerator& Generate() override
+	{
+		// enforce sane values for vertex counts
+		if (NumPhi < 3)
+		{
+			NumPhi = 3;
+		}
+		if (NumTheta < 3)
+		{
+			NumTheta = 3;
+		}
+		int32 NumVertices = (NumPhi - 2) * NumTheta + 2;
+		int32 NumUVs = NumPhi * (NumTheta + 1);
+		int32 NumTris = (NumPhi - 2) * NumTheta * 2;
+		SetBufferSizes(NumVertices, NumTris, NumUVs, NumVertices);
+
+		double Dphi = FMathd::Pi / double(NumPhi - 1);
+		double Dtheta = FMathd::TwoPi / double(NumTheta);
+		float DUVphi = 1.0 / float(NumPhi - 1);
+		float DUVtheta = 1.0 / float(NumTheta);
+
+		TArray<int32> Grid2d; Grid2d.SetNumUninitialized(NumTheta * NumPhi);
+		for (int32 i = 0; i < NumTheta * NumPhi; ++i)
+		{
+			Grid2d[i] = -1;
+		}
+
+
+		auto GridInterface = [&Grid2d, this](int32 itheta, int32 jphi)->int32&
+		{
+			// make periodic in j
+			while (itheta >= NumTheta) { itheta = itheta - NumTheta; }
+
+			int32 offset = itheta + jphi * NumTheta;
+
+			return Grid2d[offset];
+		};
+
+		int32 VtxIdx = 0;
+		for (int32 p = 1; p < NumPhi - 1; ++p) // NB: this skips the poles.
+		{
+			double Phi = p * Dphi;
+
+			for (int32 t = 0; t < NumTheta; ++t)
+			{
+				double Theta = t * Dtheta;
+				FVector3d Pos = SphericalToCartesian(1., Theta, Phi);
+
+				Vertices[VtxIdx] = Pos * Radius;
+				Normals[VtxIdx] = FVector3f(Pos);
+				NormalParentVertex[VtxIdx] = VtxIdx;
+				GridInterface(t, p) = VtxIdx;
+				VtxIdx++;
+			}
+		}
+
+
+		// add a single point at the North Pole
+		{
+			FVector3d Pos = SphericalToCartesian(1., 0., 0.);
+			Vertices[VtxIdx] = Pos * Radius;
+			Normals[VtxIdx] = FVector3f(Pos);
+			NormalParentVertex[VtxIdx] = VtxIdx;
+
+			for (int32 t = 0; t < NumTheta; ++t)
+			{
+				GridInterface(t, 0) = VtxIdx;
+			}
+
+			VtxIdx++;
+		}
+		// add a single point at the south pole
+		{
+			FVector3d Pos = SphericalToCartesian(-1., 0., 0.);
+			Vertices[VtxIdx] = Pos * Radius;
+			Normals[VtxIdx] = FVector3f(Pos);
+			NormalParentVertex[VtxIdx] = VtxIdx;
+
+			for (int32 t = 0; t < NumTheta; ++t)
+			{
+				GridInterface(t, NumPhi - 1) = VtxIdx;
+			}
+
+			VtxIdx++;
+		}
+
+		int32 UVIdx = 0;
+		auto UVIdxLookup = [this](int32 t, int32 p)
+		{
+			check(t >= 0 && t <= NumTheta);
+			return p * (NumTheta+1) + t;
+		};
+		for (int32 p = 0; p < NumPhi; ++p) // NB: this skips the poles.
+		{
+			float Phi = p * DUVphi;
+
+			for (int32 t = 0; t <= NumTheta; ++t)
+			{
+				float Theta = t * DUVtheta;
+				UVs[UVIdx] = FVector2f(Theta, Phi);
+				UVParentVertex[UVIdx] = GridInterface(t, p);
+				ensure(UVIdxLookup(t, p) == UVIdx);
+
+				UVIdx++;
+			}
+		}
+
+		int32 TriIdx = 0, PolyIdx = 0;
+		for (int32 p = 1; p < NumPhi - 2; ++p)
+		{
+			for (int32 t = 0; t < NumTheta; ++t)
+			{
+				// (p,t), (p, t+1), (p+1, t), (p+1, t+1) are the 4 corners of a quad.
+				// counter clockwise.
+				int32 Corners[4] = { GridInterface(t, p + 1), GridInterface(t + 1, p + 1), GridInterface(t + 1, p), GridInterface(t, p) };
+				int32 UVCorners[4] = { UVIdxLookup(t, p + 1),   UVIdxLookup(t + 1, p + 1),   UVIdxLookup(t + 1, p),   UVIdxLookup(t, p) };
+
+				// convert each quad into 2 triangles.
+
+				SetTriangle(TriIdx, Corners[0], Corners[2], Corners[1]);
+				SetTrianglePolygon(TriIdx, PolyIdx);
+				SetTriangleUVs(TriIdx, UVCorners[0], UVCorners[2], UVCorners[1]);
+				SetTriangleNormals(TriIdx, Corners[0], Corners[2], Corners[1]);
+				TriIdx++;
+
+				SetTriangle(TriIdx, Corners[2], Corners[0], Corners[3]);
+				SetTrianglePolygon(TriIdx, PolyIdx);
+				SetTriangleUVs(TriIdx, UVCorners[2], UVCorners[0], UVCorners[3]);
+				SetTriangleNormals(TriIdx, Corners[2], Corners[0], Corners[3]);
+				TriIdx++;
+				PolyIdx++;
+			}
+		}
+
+		// Do Triangles that connect to north pole
+		{
+			int32 p = 0; // North pole 
+			for (int32 t = 0; t < NumTheta; ++t)
+			{
+				int32 Corners[3] = { GridInterface(t, p + 1), GridInterface(t + 1, p + 1), GridInterface(t + 1, p) };
+				int32 UVCorners[3] = { UVIdxLookup(t, p + 1),   UVIdxLookup(t + 1, p + 1),   UVIdxLookup(t + 1, p) };
+
+				SetTriangle(TriIdx, Corners[0], Corners[2], Corners[1]);
+				SetTrianglePolygon(TriIdx, PolyIdx);
+				SetTriangleUVs(TriIdx, UVCorners[0], UVCorners[2], UVCorners[1]);
+				SetTriangleNormals(TriIdx, Corners[0], Corners[2], Corners[1]);
+				TriIdx++;
+				PolyIdx++;
+			}
+		}
+
+		// Do Triangles that connect to south pole
+		{
+			int32 p = NumPhi - 2; // South pole 
+			for (int32 t = 0; t < NumTheta; ++t)
+			{
+				int32 Corners[4] = { GridInterface(t, p + 1), GridInterface(t + 1, p), GridInterface(t, p) };
+				int32 UVCorners[3] = { UVIdxLookup(t, p + 1),   UVIdxLookup(t + 1, p),   UVIdxLookup(t, p) };
+
+				SetTriangle(TriIdx, Corners[0], Corners[2], Corners[1]);
+				SetTrianglePolygon(TriIdx, PolyIdx);
+				SetTriangleUVs(TriIdx, UVCorners[0], UVCorners[2], UVCorners[1]);
+				SetTriangleNormals(TriIdx, Corners[0], Corners[2], Corners[1]);
+				TriIdx++;
+				PolyIdx++;
+			}
+		}
+
+		return *this;
+	}
+
+
+};
\ No newline at end of file
diff --git a/Engine/Plugins/Experimental/GeometryProcessing/Source/DynamicMesh/Public/Generators/SweepGenerator.h b/Engine/Plugins/Experimental/GeometryProcessing/Source/DynamicMesh/Public/Generators/SweepGenerator.h
--- a/Engine/Plugins/Experimental/GeometryProcessing/Source/DynamicMesh/Public/Generators/SweepGenerator.h
+++ b/Engine/Plugins/Experimental/GeometryProcessing/Source/DynamicMesh/Public/Generators/SweepGenerator.h
@@ -255,6 +255,8 @@ public:
 		const TArray<FVector2d>& XVerts = X.GetVertices();
 		ECapType Caps[2] = {ECapType::None, ECapType::None};
 
+		FVector2d NormalSide = (FVector2d(Radius[1], Height) - FVector2d(Radius[0], 0)).Perp().Normalized();
+
 		if (bCapped)
 		{
 			Caps[0] = ECapType::FlatTriangulation;
@@ -268,7 +270,7 @@ public:
 			{
 				Vertices[SubIdx + XIdx * AngleSamples] =
 					FVector3d(XVerts[SubIdx].X * Radius[XIdx], XVerts[SubIdx].Y * Radius[XIdx], XIdx * Height);
-				Normals[SubIdx + XIdx * AngleSamples] = FVector3f(XVerts[SubIdx].X, XVerts[SubIdx].Y, 0);
+				Normals[SubIdx + XIdx * AngleSamples] = FVector3f(XVerts[SubIdx].X*NormalSide.X, XVerts[SubIdx].Y*NormalSide.X, NormalSide.Y);
 
 				if (bCapped)
 				{
@@ -286,6 +288,7 @@ public:
 	}
 };
 
+
 /**
  * Sweep a 2D Profile Polygon along a 3D Path.
  * 
diff --git a/Engine/Plugins/Experimental/GeometryProcessing/Source/GeometricObjects/Public/OrientedBoxTypes.h b/Engine/Plugins/Experimental/GeometryProcessing/Source/GeometricObjects/Public/OrientedBoxTypes.h
--- a/Engine/Plugins/Experimental/GeometryProcessing/Source/GeometricObjects/Public/OrientedBoxTypes.h
+++ b/Engine/Plugins/Experimental/GeometryProcessing/Source/GeometricObjects/Public/OrientedBoxTypes.h
@@ -131,7 +131,20 @@ struct TOrientedBox3
 		return Frame.PointAt(dx, dy, dz);
 	}
 
-
+	/**
+	 * Get whether the corner at Index (see diagram in GetCorner documentation comment) is in the negative or positive direction for each axis
+	 * @param Index corner index in range 0-7
+	 * @return Index3i with 0 or 1 for each axis, 0 if corner is in the negative direction for that axis, 1 if in the positive direction
+	 */
+	static FIndex3i GetCornerSide(int Index)
+	{
+		check(Index >= 0 && Index <= 7);
+		return FIndex3i(
+			(((Index & 1) != 0) ^ ((Index & 2) != 0)) ? 1 : 0,
+			((Index / 2) % 2 == 0) ? 0 : 1,
+			(Index < 4) ? 0 : 1
+		);
+	}
 
 
 
diff --git a/Engine/Plugins/Experimental/GeometryProcessing/Source/GeometricObjects/Public/VectorTypes.h b/Engine/Plugins/Experimental/GeometryProcessing/Source/GeometricObjects/Public/VectorTypes.h
--- a/Engine/Plugins/Experimental/GeometryProcessing/Source/GeometricObjects/Public/VectorTypes.h
+++ b/Engine/Plugins/Experimental/GeometryProcessing/Source/GeometricObjects/Public/VectorTypes.h
@@ -6,35 +6,37 @@
 #include "MathUtil.h"
 
 
+
+
+
 /**
- * Templated 3D Vector. Ported from g3Sharp library, with the intention of 
- * maintaining compatibility with existing g3Sharp code. Has an API
- * similar to WildMagic, GTEngine, Eigen, etc. 
- * 
- * Convenience typedefs for FVector3f/FVector3d/FVector3i are defined, and
- * should be preferentially used over the base template type
- * 
- * Implicit casts to/from FVector are defined, for all types.
- *  
- * @todo Possibly can be replaced/merged with Chaos TVector<T,N>
- */
+* Templated 2D Vector. Ported from g3Sharp library, with the intention of
+* maintaining compatibility with existing g3Sharp code. Has an API
+* similar to WildMagic, GTEngine, Eigen, etc.
+*
+* Convenience typedefs for FVector2f/FVector2d/FVector2i are defined, and
+* should be preferentially used over the base template type
+*
+* Implicit casts to/from FVector2D are defined, for all types.
+*
+* @todo Possibly can be replaced/merged with Chaos TVector<T,N>
+*/
 template <typename T>
-struct FVector3
+struct FVector2
 {
-	T X, Y, Z;
+	T X, Y;
 
-	FVector3()
+	FVector2()
 	{
 	}
-	FVector3(T ValX, T ValY, T ValZ)
-		: X(ValX), Y(ValY), Z(ValZ)
+	FVector2(T ValX, T ValY)
+		: X(ValX), Y(ValY)
 	{
 	}
-	FVector3(const T* Data)
+	FVector2(const T* Data)
 	{
 		X = Data[0];
 		Y = Data[1];
-		Z = Data[2];
 	}
 
 	inline operator const T*() const
@@ -46,68 +48,37 @@
 		return &X;
 	}
 
-	inline operator FVector() const
+	operator FVector2D() const
 	{
-		return FVector((float)X, (float)Y, (float)Z);
+		return FVector2D((float)X, (float)Y);
 	}
-	inline FVector3(const FVector& Vec)
+	FVector2(const FVector2D& Vec)
 	{
 		X = (T)Vec.X;
 		Y = (T)Vec.Y;
-		Z = (T)Vec.Z;
 	}
 
-
-	explicit inline operator FLinearColor() const
+	static FVector2<T> Zero()
 	{
-		return FLinearColor((float)X, (float)Y, (float)Z);
+		return FVector2<T>((T)0, (T)0);
 	}
-	inline FVector3(const FLinearColor& Color)
+	static FVector2<T> One()
 	{
-		X = (T)Color.R;
-		Y = (T)Color.G;
-		Z = (T)Color.B;
+		return FVector2<T>((T)1, (T)1);
+	}
+	static FVector2<T> UnitX()
+	{
+		return FVector2<T>((T)1, (T)0);
+	}
+	static FVector2<T> UnitY()
+	{
+		return FVector2<T>((T)0, (T)1);
 	}
 
-	explicit inline operator FColor() const
+	FVector2<T>& operator=(const FVector2<T>& V2)
 	{
-		return FColor(
-			FMathf::Clamp((int)((float)X*255.0f), 0, 255),
-			FMathf::Clamp((int)((float)Y*255.0f), 0, 255),
-			FMathf::Clamp((int)((float)Z*255.0f), 0, 255));
-	}
-
-
-	static FVector3<T> Zero()
-	{
-		return FVector3<T>((T)0, (T)0, (T)0);
-	}
-	static FVector3<T> One()
-	{
-		return FVector3<T>((T)1, (T)1, (T)1);
-	}
-	static FVector3<T> UnitX()
-	{
-		return FVector3<T>((T)1, (T)0, (T)0);
-	}
-	static FVector3<T> UnitY()
-	{
-		return FVector3<T>((T)0, (T)1, (T)0);
-	}
-	static FVector3<T> UnitZ()
-	{
-		return FVector3<T>((T)0, (T)0, (T)1);
-	}
-	static FVector3<T> Max()
-	{
-		return FVector3<T>(TNumericLimits<T>::Max(), TNumericLimits<T>::Max(), TNumericLimits<T>::Max());
-	}
-
-	FVector3<T>& operator=(const FVector3<T>& V2)
-	{
 		X = V2.X;
 		Y = V2.Y;
-		Z = V2.Z;
 		return *this;
 	}
 
@@ -122,120 +93,126 @@
 
 	T Length() const
 	{
-		return TMathUtil<T>::Sqrt(X * X + Y * Y + Z * Z);
+		return TMathUtil<T>::Sqrt(X * X + Y * Y);
 	}
 	T SquaredLength() const
 	{
-		return X * X + Y * Y + Z * Z;
+		return X * X + Y * Y;
 	}
 
-	inline T Distance(const FVector3<T>& V2) const
+	inline T Distance(const FVector2<T>& V2) const
 	{
 		T dx = V2.X - X;
 		T dy = V2.Y - Y;
-		T dz = V2.Z - Z;
-		return TMathUtil<T>::Sqrt(dx * dx + dy * dy + dz * dz);
+		return TMathUtil<T>::Sqrt(dx * dx + dy * dy);
 	}
-	inline T DistanceSquared(const FVector3<T>& V2) const
+	inline T DistanceSquared(const FVector2<T>& V2) const
 	{
 		T dx = V2.X - X;
 		T dy = V2.Y - Y;
-		T dz = V2.Z - Z;
-		return dx * dx + dy * dy + dz * dz;
+		return dx * dx + dy * dy;
 	}
 
-	inline FVector3<T> operator-() const
+	T Dot(const FVector2<T>& V2) const
 	{
-		return FVector3<T>(-X, -Y, -Z);
+		return X * V2.X + Y * V2.Y;
 	}
 
-	inline FVector3<T> operator+(const FVector3<T>& V2) const
+	// Note: DotPerp and FVector2's version of Cross are the same function
+	inline T DotPerp(const FVector2<T>& V2) const
 	{
-		return FVector3<T>(X + V2.X, Y + V2.Y, Z + V2.Z);
+		return X * V2.Y - Y * V2.X;
+	}
+	inline T Cross(const FVector2<T>& V2) const
+	{
+		return X * V2.Y - Y * V2.X;
 	}
 
-	inline FVector3<T> operator-(const FVector3<T>& V2) const
+	FVector2<T> Perp() const
 	{
-		return FVector3<T>(X - V2.X, Y - V2.Y, Z - V2.Z);
+		return FVector2<T>(Y, -X);
 	}
 
-	inline FVector3<T> operator+(const T& Scalar) const
+	inline FVector2 operator-() const
 	{
-		return FVector3<T>(X + Scalar, Y + Scalar, Z + Scalar);
+		return FVector2(-X, -Y);
 	}
 
-	inline FVector3<T> operator-(const T& Scalar) const
+	inline FVector2 operator+(const FVector2& V2) const
 	{
-		return FVector3<T>(X - Scalar, Y - Scalar, Z - Scalar);
+		return FVector2(X + V2.X, Y + V2.Y);
 	}
 
-	inline FVector3<T> operator*(const T& Scalar) const
+	inline FVector2 operator-(const FVector2& V2) const
 	{
-		return FVector3<T>(X * Scalar, Y * Scalar, Z * Scalar);
+		return FVector2(X - V2.X, Y - V2.Y);
 	}
 
-	inline FVector3<T> operator*(const FVector3<T>& V2) const // component-wise
+	inline FVector2<T> operator+(const T& Scalar) const
 	{
-		return FVector3<T>(X * V2.X, Y * V2.Y, Z * V2.Z);
+		return FVector2<T>(X + Scalar, Y + Scalar);
 	}
 
-	inline FVector3<T> operator/(const T& Scalar) const
+	inline FVector2<T> operator-(const T& Scalar) const
 	{
-		return FVector3<T>(X / Scalar, Y / Scalar, Z / Scalar);
+		return FVector2<T>(X - Scalar, Y - Scalar);
 	}
 
-	inline FVector3<T> operator/(const FVector3<T>& V2) const // component-wise
+	inline FVector2 operator*(const T& Scalar) const
 	{
-		return FVector3<T>(X / V2.X, Y / V2.Y, Z / V2.Z);
+		return FVector2(X * Scalar, Y * Scalar);
 	}
 
-	inline FVector3<T>& operator+=(const FVector3<T>& V2)
+	inline FVector2<T> operator*(const FVector2<T>& V2) const // component-wise
 	{
+		return FVector2<T>(X * V2.X, Y * V2.Y);
+	}
+
+	inline FVector2 operator/(const T& Scalar) const
+	{
+		return FVector2(X / Scalar, Y / Scalar);
+	}
+
+	inline FVector2<T> operator/(const FVector2<T>& V2) const // component-wise
+	{
+		return FVector2<T>(X / V2.X, Y / V2.Y);
+	}
+
+	inline FVector2<T>& operator+=(const FVector2<T>& V2)
+	{
 		X += V2.X;
 		Y += V2.Y;
-		Z += V2.Z;
 		return *this;
 	}
 
-	inline FVector3<T>& operator-=(const FVector3<T>& V2)
+	inline FVector2<T>& operator-=(const FVector2<T>& V2)
 	{
 		X -= V2.X;
 		Y -= V2.Y;
-		Z -= V2.Z;
 		return *this;
 	}
 
-	inline FVector3<T>& operator*=(const T& Scalar)
+	inline FVector2<T>& operator*=(const T& Scalar)
 	{
 		X *= Scalar;
 		Y *= Scalar;
-		Z *= Scalar;
 		return *this;
 	}
 
-	inline FVector3<T>& operator/=(const T& Scalar)
+	inline FVector2<T>& operator/=(const T& Scalar)
 	{
 		X /= Scalar;
 		Y /= Scalar;
-		Z /= Scalar;
 		return *this;
 	}
 
-	T Dot(const FVector3<T>& V2) const
+	inline bool IsNormalized()
 	{
-		return X * V2.X + Y * V2.Y + Z * V2.Z;
+		return TMathUtil<T>::Abs((X * X + Y * Y) - 1) < TMathUtil<T>::ZeroTolerance;
 	}
 
-	FVector3<T> Cross(const FVector3<T>& V2) const
-	{
-		return FVector3(
-			Y * V2.Z - Z * V2.Y,
-			Z * V2.X - X * V2.Z,
-			X * V2.Y - Y * V2.X);
-	}
-
 	// Angle in Degrees
-	T AngleD(const FVector3<T>& V2) const
+	T AngleD(const FVector2<T>& V2) const
 	{
 		T DotVal = Dot(V2);
 		T ClampedDot = (DotVal < (T)-1) ? (T)-1 : ((DotVal > (T)1) ? (T)1 : DotVal);
@@ -243,16 +220,28 @@
 	}
 
 	// Angle in Radians
-	T AngleR(const FVector3<T>& V2) const
+	T AngleR(const FVector2<T>& V2) const
 	{
 		T DotVal = Dot(V2);
 		T ClampedDot = (DotVal < (T)-1) ? (T)-1 : ((DotVal > (T)1) ? (T)1 : DotVal);
 		return (T)acos(ClampedDot);
 	}
 
-	inline bool IsNormalized()
+	// Angle in Radians
+	T SignedAngleR(const FVector2<T>& V2) const
 	{
-		return TMathUtil<T>::Abs((X * X + Y * Y + Z * Z) - 1) < TMathUtil<T>::ZeroTolerance;
+		T DotVal = Dot(V2);
+		T ClampedDot = (DotVal < (T)-1) ? (T)-1 : ((DotVal > (T)1) ? (T)1 : DotVal);
+		T Direction = Cross(V2);
+		if (Direction*Direction < TMathUtil<T>::ZeroTolerance)
+		{
+			return (DotVal < 0) ? TMathUtil<T>::Pi : (T)0;
+		}
+		else
+		{
+			T Sign = Direction < 0 ? (T)-1 : (T)1;
+			return Sign * TMathUtil<T>::ACos(ClampedDot);
+		}
 	}
 
 	T Normalize(const T Epsilon = 0)
@@ -263,99 +252,82 @@
 			T invLength = ((T)1) / length;
 			X *= invLength;
 			Y *= invLength;
-			Z *= invLength;
 			return length;
 		}
-		X = Y = Z = (T)0;
+		X = Y = (T)0;
 		return 0;
 	}
 
-	inline FVector3<T> Normalized(const T Epsilon = 0) const
+	inline FVector2<T> Normalized(const T Epsilon = 0) const
 	{
 		T length = Length();
 		if (length > Epsilon)
 		{
 			T invLength = ((T)1) / length;
-			return FVector3<T>(X * invLength, Y * invLength, Z * invLength);
+			return FVector2<T>(X * invLength, Y * invLength);
 		}
-		return FVector3<T>((T)0, (T)0, (T)0);
+		return FVector2<T>((T)0, (T)0);
 	}
 
-	inline T MaxAbs() const
+
+
+	bool operator==(const FVector2<T>& Other) const
 	{
-		return TMathUtil<T>::Max3(TMathUtil<T>::Abs(X), TMathUtil<T>::Abs(Y), TMathUtil<T>::Abs(Z));
+		return X == Other.X && Y == Other.Y;
 	}
 
-	inline T MinAbs() const
+
+	static FVector2<T> Lerp(const FVector2<T>& A, const FVector2<T>& B, T Alpha)
 	{
-		return TMathUtil<T>::Min3(TMathUtil<T>::Abs(X), TMathUtil<T>::Abs(Y), TMathUtil<T>::Abs(Z));
-	}
-
-	static FVector3<T> Lerp(const FVector3<T>& A, const FVector3<T>& B, T Alpha)
-	{
 		T OneMinusAlpha = (T)1 - Alpha;
-		return FVector3<T>(OneMinusAlpha * A.X + Alpha * B.X,
-						   OneMinusAlpha * A.Y + Alpha * B.Y,
-						   OneMinusAlpha * A.Z + Alpha * B.Z);
+		return FVector2<T>(OneMinusAlpha * A.X + Alpha * B.X,
+			OneMinusAlpha * A.Y + Alpha * B.Y);
 	}
 
-	inline bool operator==(const FVector3<T>& Other) const
+	// Note: This was called "IsLeft" in the Geometry3Sharp code (where it also went through Math.Sign)
+	// Returns >0 if C is to the left of the A->B Line, <0 if to the right, 0 if on the line
+	static T Orient(const FVector2<T>& A, const FVector2<T>& B, const FVector2<T>& C)
 	{
-		return X == Other.X && Y == Other.Y && Z == Other.Z;
+		return (B - A).DotPerp(C - A);
 	}
 };
 
-template <typename RealType>
-inline FVector3<RealType> operator*(RealType Scalar, const FVector3<RealType>& V)
-{
-	return FVector3<RealType>(Scalar * V.X, Scalar * V.Y, Scalar * V.Z);
-}
 
-typedef FVector3<float> FVector3f;
-typedef FVector3<double> FVector3d;
-typedef FVector3<int> FVector3i;
-
-template <typename T>
-FORCEINLINE uint32 GetTypeHash(const FVector3<T>& Vector)
-{
-	// (this is how FIntVector and all the other FVectors do their hash functions)
-	// Note: this assumes there's no padding that could contain uncompared data.
-	return FCrc::MemCrc_DEPRECATED(&Vector, sizeof(FVector3<T>));
-}
-
-
-
-
-
-
 /**
- * Templated 2D Vector. Ported from g3Sharp library, with the intention of
+ * Templated 3D Vector. Ported from g3Sharp library, with the intention of 
  * maintaining compatibility with existing g3Sharp code. Has an API
- * similar to WildMagic, GTEngine, Eigen, etc.
- *
- * Convenience typedefs for FVector2f/FVector2d/FVector2i are defined, and
+ * similar to WildMagic, GTEngine, Eigen, etc. 
+ * 
+ * Convenience typedefs for FVector3f/FVector3d/FVector3i are defined, and
  * should be preferentially used over the base template type
- *
- * Implicit casts to/from FVector2D are defined, for all types.
- *
+ * 
+ * Implicit casts to/from FVector are defined, for all types.
+ *  
  * @todo Possibly can be replaced/merged with Chaos TVector<T,N>
  */
 template <typename T>
-struct FVector2
+struct FVector3
 {
-	T X, Y;
+	T X, Y, Z;
 
-	FVector2()
+	FVector3()
 	{
 	}
-	FVector2(T ValX, T ValY)
-		: X(ValX), Y(ValY)
+	FVector3(T ValX, T ValY, T ValZ)
+		: X(ValX), Y(ValY), Z(ValZ)
 	{
 	}
-	FVector2(const T* Data)
+	FVector3(const T* Data)
 	{
 		X = Data[0];
 		Y = Data[1];
+		Z = Data[2];
+	}
+	FVector3(const FVector2<T>& Data)
+	{
+		X = Data.X;
+		Y = Data.Y;
+		Z = (T)0;
 	}
 
 	inline operator const T*() const
@@ -367,37 +339,68 @@
 		return &X;
 	}
 
-	operator FVector2D() const
+	inline operator FVector() const
 	{
-		return FVector2D((float)X, (float)Y);
+		return FVector((float)X, (float)Y, (float)Z);
 	}
-	FVector2(const FVector2D& Vec)
+	inline FVector3(const FVector& Vec)
 	{
 		X = (T)Vec.X;
 		Y = (T)Vec.Y;
+		Z = (T)Vec.Z;
 	}
 
-	static FVector2<T> Zero()
+
+	explicit inline operator FLinearColor() const
 	{
-		return FVector2<T>((T)0, (T)0);
+		return FLinearColor((float)X, (float)Y, (float)Z);
 	}
-	static FVector2<T> One()
+	inline FVector3(const FLinearColor& Color)
 	{
-		return FVector2<T>((T)1, (T)1);
-	}
-	static FVector2<T> UnitX()
-	{
-		return FVector2<T>((T)1, (T)0);
-	}
-	static FVector2<T> UnitY()
-	{
-		return FVector2<T>((T)0, (T)1);
+		X = (T)Color.R;
+		Y = (T)Color.G;
+		Z = (T)Color.B;
 	}
 
-	FVector2<T>& operator=(const FVector2<T>& V2)
+	explicit inline operator FColor() const
 	{
+		return FColor(
+			FMathf::Clamp((int)((float)X*255.0f), 0, 255),
+			FMathf::Clamp((int)((float)Y*255.0f), 0, 255),
+			FMathf::Clamp((int)((float)Z*255.0f), 0, 255));
+	}
+
+
+	static FVector3<T> Zero()
+	{
+		return FVector3<T>((T)0, (T)0, (T)0);
+	}
+	static FVector3<T> One()
+	{
+		return FVector3<T>((T)1, (T)1, (T)1);
+	}
+	static FVector3<T> UnitX()
+	{
+		return FVector3<T>((T)1, (T)0, (T)0);
+	}
+	static FVector3<T> UnitY()
+	{
+		return FVector3<T>((T)0, (T)1, (T)0);
+	}
+	static FVector3<T> UnitZ()
+	{
+		return FVector3<T>((T)0, (T)0, (T)1);
+	}
+	static FVector3<T> Max()
+	{
+		return FVector3<T>(TNumericLimits<T>::Max(), TNumericLimits<T>::Max(), TNumericLimits<T>::Max());
+	}
+
+	FVector3<T>& operator=(const FVector3<T>& V2)
+	{
 		X = V2.X;
 		Y = V2.Y;
+		Z = V2.Z;
 		return *this;
 	}
 
@@ -412,126 +415,120 @@
 
 	T Length() const
 	{
-		return TMathUtil<T>::Sqrt(X * X + Y * Y);
+		return TMathUtil<T>::Sqrt(X * X + Y * Y + Z * Z);
 	}
 	T SquaredLength() const
 	{
-		return X * X + Y * Y;
+		return X * X + Y * Y + Z * Z;
 	}
 
-	inline T Distance(const FVector2<T>& V2) const
+	inline T Distance(const FVector3<T>& V2) const
 	{
 		T dx = V2.X - X;
 		T dy = V2.Y - Y;
-		return TMathUtil<T>::Sqrt(dx * dx + dy * dy);
+		T dz = V2.Z - Z;
+		return TMathUtil<T>::Sqrt(dx * dx + dy * dy + dz * dz);
 	}
-	inline T DistanceSquared(const FVector2<T>& V2) const
+	inline T DistanceSquared(const FVector3<T>& V2) const
 	{
 		T dx = V2.X - X;
 		T dy = V2.Y - Y;
-		return dx * dx + dy * dy;
+		T dz = V2.Z - Z;
+		return dx * dx + dy * dy + dz * dz;
 	}
 
-	T Dot(const FVector2<T>& V2) const
+	inline FVector3<T> operator-() const
 	{
-		return X * V2.X + Y * V2.Y;
+		return FVector3<T>(-X, -Y, -Z);
 	}
 
-	// Note: DotPerp and FVector2's version of Cross are the same function
-	inline T DotPerp(const FVector2<T>& V2) const
+	inline FVector3<T> operator+(const FVector3<T>& V2) const
 	{
-		return X * V2.Y - Y * V2.X;
-	}
-	inline T Cross(const FVector2<T>& V2) const
-	{
-		return X * V2.Y - Y * V2.X;
+		return FVector3<T>(X + V2.X, Y + V2.Y, Z + V2.Z);
 	}
 
-	FVector2<T> Perp() const
+	inline FVector3<T> operator-(const FVector3<T>& V2) const
 	{
-		return FVector2<T>(Y, -X);
+		return FVector3<T>(X - V2.X, Y - V2.Y, Z - V2.Z);
 	}
 
-	inline FVector2 operator-() const
+	inline FVector3<T> operator+(const T& Scalar) const
 	{
-		return FVector2(-X, -Y);
+		return FVector3<T>(X + Scalar, Y + Scalar, Z + Scalar);
 	}
 
-	inline FVector2 operator+(const FVector2& V2) const
+	inline FVector3<T> operator-(const T& Scalar) const
 	{
-		return FVector2(X + V2.X, Y + V2.Y);
+		return FVector3<T>(X - Scalar, Y - Scalar, Z - Scalar);
 	}
 
-	inline FVector2 operator-(const FVector2& V2) const
+	inline FVector3<T> operator*(const T& Scalar) const
 	{
-		return FVector2(X - V2.X, Y - V2.Y);
+		return FVector3<T>(X * Scalar, Y * Scalar, Z * Scalar);
 	}
 
-	inline FVector2<T> operator+(const T& Scalar) const
+	inline FVector3<T> operator*(const FVector3<T>& V2) const // component-wise
 	{
-		return FVector2<T>(X + Scalar, Y + Scalar);
+		return FVector3<T>(X * V2.X, Y * V2.Y, Z * V2.Z);
 	}
 
-	inline FVector2<T> operator-(const T& Scalar) const
+	inline FVector3<T> operator/(const T& Scalar) const
 	{
-		return FVector2<T>(X - Scalar, Y - Scalar);
+		return FVector3<T>(X / Scalar, Y / Scalar, Z / Scalar);
 	}
 
-	inline FVector2 operator*(const T& Scalar) const
+	inline FVector3<T> operator/(const FVector3<T>& V2) const // component-wise
 	{
-		return FVector2(X * Scalar, Y * Scalar);
+		return FVector3<T>(X / V2.X, Y / V2.Y, Z / V2.Z);
 	}
 
-	inline FVector2<T> operator*(const FVector2<T>& V2) const // component-wise
+	inline FVector3<T>& operator+=(const FVector3<T>& V2)
 	{
-		return FVector2<T>(X * V2.X, Y * V2.Y);
-	}
-
-	inline FVector2 operator/(const T& Scalar) const
-	{
-		return FVector2(X / Scalar, Y / Scalar);
-	}
-
-	inline FVector2<T> operator/(const FVector2<T>& V2) const // component-wise
-	{
-		return FVector2<T>(X / V2.X, Y / V2.Y);
-	}
-
-	inline FVector2<T>& operator+=(const FVector2<T>& V2)
-	{
 		X += V2.X;
 		Y += V2.Y;
+		Z += V2.Z;
 		return *this;
 	}
 
-	inline FVector2<T>& operator-=(const FVector2<T>& V2)
+	inline FVector3<T>& operator-=(const FVector3<T>& V2)
 	{
 		X -= V2.X;
 		Y -= V2.Y;
+		Z -= V2.Z;
 		return *this;
 	}
 
-	inline FVector2<T>& operator*=(const T& Scalar)
+	inline FVector3<T>& operator*=(const T& Scalar)
 	{
 		X *= Scalar;
 		Y *= Scalar;
+		Z *= Scalar;
 		return *this;
 	}
 
-	inline FVector2<T>& operator/=(const T& Scalar)
+	inline FVector3<T>& operator/=(const T& Scalar)
 	{
 		X /= Scalar;
 		Y /= Scalar;
+		Z /= Scalar;
 		return *this;
 	}
 
-	inline bool IsNormalized()
+	T Dot(const FVector3<T>& V2) const
 	{
-		return TMathUtil<T>::Abs((X * X + Y * Y) - 1) < TMathUtil<T>::ZeroTolerance;
+		return X * V2.X + Y * V2.Y + Z * V2.Z;
 	}
 
+	FVector3<T> Cross(const FVector3<T>& V2) const
+	{
+		return FVector3(
+			Y * V2.Z - Z * V2.Y,
+			Z * V2.X - X * V2.Z,
+			X * V2.Y - Y * V2.X);
+	}
+
 	// Angle in Degrees
-	T AngleD(const FVector2<T>& V2) const
+	T AngleD(const FVector3<T>& V2) const
 	{
 		T DotVal = Dot(V2);
 		T ClampedDot = (DotVal < (T)-1) ? (T)-1 : ((DotVal > (T)1) ? (T)1 : DotVal);
@@ -539,28 +536,16 @@
 	}
 
 	// Angle in Radians
-	T AngleR(const FVector2<T>& V2) const
+	T AngleR(const FVector3<T>& V2) const
 	{
 		T DotVal = Dot(V2);
 		T ClampedDot = (DotVal < (T)-1) ? (T)-1 : ((DotVal > (T)1) ? (T)1 : DotVal);
 		return (T)acos(ClampedDot);
 	}
 
-	// Angle in Radians
-	T SignedAngleR(const FVector2<T>& V2) const
+	inline bool IsNormalized()
 	{
-		T DotVal = Dot(V2);
-		T ClampedDot = (DotVal < (T)-1) ? (T)-1 : ((DotVal > (T)1) ? (T)1 : DotVal);
-		T Direction = Cross(V2);
-		if (Direction*Direction < TMathUtil<T>::ZeroTolerance)
-		{
-			return (DotVal < 0) ? TMathUtil<T>::Pi : (T)0;
-		}
-		else
-		{
-			T Sign = Direction < 0 ? (T)-1 : (T)1;
-			return Sign * TMathUtil<T>::ACos(ClampedDot);
-		}
+		return TMathUtil<T>::Abs((X * X + Y * Y + Z * Z) - 1) < TMathUtil<T>::ZeroTolerance;
 	}
 
 	T Normalize(const T Epsilon = 0)
@@ -571,46 +556,70 @@
 			T invLength = ((T)1) / length;
 			X *= invLength;
 			Y *= invLength;
+			Z *= invLength;
 			return length;
 		}
-		X = Y = (T)0;
+		X = Y = Z = (T)0;
 		return 0;
 	}
 
-	inline FVector2<T> Normalized(const T Epsilon = 0) const
+	inline FVector3<T> Normalized(const T Epsilon = 0) const
 	{
 		T length = Length();
 		if (length > Epsilon)
 		{
 			T invLength = ((T)1) / length;
-			return FVector2<T>(X * invLength, Y * invLength);
+			return FVector3<T>(X * invLength, Y * invLength, Z * invLength);
 		}
-		return FVector2<T>((T)0, (T)0);
+		return FVector3<T>((T)0, (T)0, (T)0);
 	}
 
-
-
-	bool operator==(const FVector2<T>& Other) const
+	inline T MaxAbs() const
 	{
-		return X == Other.X && Y == Other.Y;
+		return TMathUtil<T>::Max3(TMathUtil<T>::Abs(X), TMathUtil<T>::Abs(Y), TMathUtil<T>::Abs(Z));
 	}
 
+	inline T MinAbs() const
+	{
+		return TMathUtil<T>::Min3(TMathUtil<T>::Abs(X), TMathUtil<T>::Abs(Y), TMathUtil<T>::Abs(Z));
+	}
 
-	static FVector2<T> Lerp(const FVector2<T>& A, const FVector2<T>& B, T Alpha)
+	static FVector3<T> Lerp(const FVector3<T>& A, const FVector3<T>& B, T Alpha)
 	{
 		T OneMinusAlpha = (T)1 - Alpha;
-		return FVector2<T>(OneMinusAlpha * A.X + Alpha * B.X,
-			OneMinusAlpha * A.Y + Alpha * B.Y);
+		return FVector3<T>(OneMinusAlpha * A.X + Alpha * B.X,
+						   OneMinusAlpha * A.Y + Alpha * B.Y,
+						   OneMinusAlpha * A.Z + Alpha * B.Z);
 	}
 
-	// Note: This was called "IsLeft" in the Geometry3Sharp code (where it also went through Math.Sign)
-	// Returns >0 if C is to the left of the A->B Line, <0 if to the right, 0 if on the line
-	static T Orient(const FVector2<T>& A, const FVector2<T>& B, const FVector2<T>& C)
+	inline bool operator==(const FVector3<T>& Other) const
 	{
-		return (B - A).DotPerp(C - A);
+		return X == Other.X && Y == Other.Y && Z == Other.Z;
 	}
 };
 
+template <typename RealType>
+inline FVector3<RealType> operator*(RealType Scalar, const FVector3<RealType>& V)
+{
+	return FVector3<RealType>(Scalar * V.X, Scalar * V.Y, Scalar * V.Z);
+}
+
+typedef FVector3<float> FVector3f;
+typedef FVector3<double> FVector3d;
+typedef FVector3<int> FVector3i;
+
 template <typename T>
+FORCEINLINE uint32 GetTypeHash(const FVector3<T>& Vector)
+{
+	// (this is how FIntVector and all the other FVectors do their hash functions)
+	// Note: this assumes there's no padding that could contain uncompared data.
+	return FCrc::MemCrc_DEPRECATED(&Vector, sizeof(FVector3<T>));
+}
+
+
+
+
+
+template <typename T>
 FORCEINLINE uint32 GetTypeHash(const FVector2<T>& Vector)
 {
