initial UV projection tool

#rnx

diff --git a/Engine/Plugins/Experimental/GeometryProcessing/Source/DynamicMesh/Public/DynamicMeshOverlay.h b/Engine/Plugins/Experimental/GeometryProcessing/Source/DynamicMesh/Public/DynamicMeshOverlay.h
--- a/Engine/Plugins/Experimental/GeometryProcessing/Source/DynamicMesh/Public/DynamicMeshOverlay.h
+++ b/Engine/Plugins/Experimental/GeometryProcessing/Source/DynamicMesh/Public/DynamicMeshOverlay.h
@@ -102,6 +102,8 @@ public:
 	/** Set the triangle to the given Element index tuple, and increment element reference counts */
 	EMeshResult SetTriangle(int TriangleID, const FIndex3i& TriElements);
 
+	/** @return true if this triangle was set */
+	bool IsSetTriangle(int TID) const { return ElementTriangles[3 * TID] >= 0; }
 
 
 	/**
