Geometry Processing: Removing some timing code from of the Laplacian utilities.
#rnx

diff --git a/Engine/Plugins/Experimental/GeometryProcessing/Source/MeshSolverUtilities/Private/LaplacianMeshSmoother.cpp b/Engine/Plugins/Experimental/GeometryProcessing/Source/MeshSolverUtilities/Private/LaplacianMeshSmoother.cpp
--- a/Engine/Plugins/Experimental/GeometryProcessing/Source/MeshSolverUtilities/Private/LaplacianMeshSmoother.cpp
+++ b/Engine/Plugins/Experimental/GeometryProcessing/Source/MeshSolverUtilities/Private/LaplacianMeshSmoother.cpp
@@ -8,12 +8,12 @@
 #include "MatrixSolver.h"
 #include "MeshSmoothingUtilities.h"
 
-#include "ProfilingDebugging/ScopedTimers.h"
-#include "Async/ParallelFor.h"
-
+#ifdef TIME_LAPLACIAN_SMOOTHERS
 
+#include "ProfilingDebugging/ScopedTimers.h"
 DEFINE_LOG_CATEGORY_STATIC(LogMeshSmoother, Log, All);
 
+#endif
 
 double ComputeDistSqrd(const FSOAPositions& VecA, const FSOAPositions& VecB)
 {
@@ -566,9 +566,9 @@ void FDiffusionIntegrator::Initialize(const FDynamicMesh3& DynamicMesh, const EL
 
 
 
-void FDiffusionIntegrator::Integrate_ForwardEuler(int32 NumSteps, double Alpha, double)
+void FDiffusionIntegrator::Integrate_ForwardEuler(const int32 NumSteps, const double Speed)
 {
-	Alpha = FMath::Clamp(Alpha, 0., 1.);
+	double Alpha = FMath::Clamp(Speed, 0., 1.);
 
 	FSparseMatrixD::Scalar TimeStep = -Alpha / MinDiagonalValue;
 	Id = 0;
@@ -587,7 +587,7 @@ void FDiffusionIntegrator::Integrate_ForwardEuler(int32 NumSteps, double Alpha,
 
 
 
-void FDiffusionIntegrator::Integrate_BackwardEuler(const EMatrixSolverType MatrixSolverType, int32 NumSteps, double Alpha, double Intensity)
+void FDiffusionIntegrator::Integrate_BackwardEuler(const EMatrixSolverType MatrixSolverType, const int32 NumSteps, const double TimeStepSize)
 {
 
 	//typedef typename TMatrixSolverTrait<EMatrixSolverType::LU>::MatrixSolverType   MatrixSolverType;
@@ -612,7 +612,7 @@ void FDiffusionIntegrator::Integrate_BackwardEuler(const EMatrixSolverType Matri
 	FSparseMatrixD Ident(DiffusionOperator.rows(), DiffusionOperator.cols());
 	Ident.setIdentity();
 
-	FSparseMatrixD::Scalar TimeStep = Alpha * FMath::Min(Intensity, 1.e6);
+	FSparseMatrixD::Scalar TimeStep = FMath::Abs(TimeStepSize);// Alpha * FMath::Min(Intensity, 1.e6);
 	
 	FSparseMatrixD SparseMatrix = Ident -TimeStep * DiffusionOperator;
 
@@ -822,14 +822,17 @@ void MeshSmoothingOperators::ComputeSmoothing_BiHarmonic(const ELaplacianWeightS
 #endif
 
 
-
+#ifdef TIME_LAPLACIAN_SMOOTHERS
 	FString DebugLogString = FString::Printf(TEXT("Biharmonic Smoothing of mesh with %d verts "), OriginalMesh.VertexCount()) + LaplacianSchemeName(WeightScheme) + MatrixSolverName(MatrixSolverType);
 
 	FScopedDurationTimeLogger Timer(DebugLogString);
+#endif
+
+	const double TimeStep = Speed * FMath::Min(Intensity, 1.e6);
 
 	FBiHarmonicDiffusionMeshSmoother BiHarmonicDiffusionSmoother(OriginalMesh, WeightScheme);
 
-	BiHarmonicDiffusionSmoother.Integrate_BackwardEuler(MatrixSolverType, NumIterations, Speed, Intensity);
+	BiHarmonicDiffusionSmoother.Integrate_BackwardEuler(MatrixSolverType, NumIterations, TimeStep);
 
 	BiHarmonicDiffusionSmoother.GetPositions(PositionArray);
 
@@ -850,11 +853,11 @@ void MeshSmoothingOperators::ComputeSmoothing_ImplicitBiHarmonicPCG( const ELapl
 	//
 	// re-write as
 	// L^TL[p^{n+1}] + weight * weight p^{n+1} = weight * weight p^{n}
-
+#ifdef TIME_LAPLACIAN_SMOOTHERS
 	FString DebugLogString = FString::Printf(TEXT("PCG Biharmonic Smoothing of mesh with %d verts "), OriginalMesh.VertexCount()) + LaplacianSchemeName(WeightScheme);
 
 	FScopedDurationTimeLogger Timer(DebugLogString);
-
+#endif 
 	if (MaxIterations < 1) return;
 
 	FCGBiHarmonicMeshSmoother Smoother(OriginalMesh, WeightScheme);
@@ -888,6 +891,7 @@ void  MeshSmoothingOperators::ComputeSmoothing_Diffusion( const ELaplacianWeight
 	//const EMatrixSolverType MatrixSolverType = EMatrixSolverType::BICGSTAB;
 #endif
 
+#ifdef TIME_LAPLACIAN_SMOOTHERS
 	FString DebugLogString = FString::Printf(TEXT("Diffusion Smoothing of mesh with %d verts"), OriginalMesh.VertexCount());
 	if (!bForwardEuler)
 	{
@@ -895,18 +899,19 @@ void  MeshSmoothingOperators::ComputeSmoothing_Diffusion( const ELaplacianWeight
 	}
 
 	FScopedDurationTimeLogger Timer(DebugLogString);
-
+#endif
 	if (IterationCount < 1) return;
 
 	FLaplacianDiffusionMeshSmoother Smoother(OriginalMesh, WeightScheme);
 
 	if (bForwardEuler)
 	{
-		Smoother.Integrate_ForwardEuler(IterationCount, Speed, Intensity);
+		Smoother.Integrate_ForwardEuler(IterationCount, Speed);
 	}
 	else
 	{
-		Smoother.Integrate_BackwardEuler(MatrixSolverType, IterationCount, Speed, Intensity);
+		const double TimeStep = Speed * FMath::Min(Intensity, 1.e6);
+		Smoother.Integrate_BackwardEuler(MatrixSolverType, IterationCount, TimeStep);
 	}
 
 	Smoother.GetPositions(PositionArray);
diff --git a/Engine/Plugins/Experimental/GeometryProcessing/Source/MeshSolverUtilities/Private/LaplacianMeshSmoother.h b/Engine/Plugins/Experimental/GeometryProcessing/Source/MeshSolverUtilities/Private/LaplacianMeshSmoother.h
--- a/Engine/Plugins/Experimental/GeometryProcessing/Source/MeshSolverUtilities/Private/LaplacianMeshSmoother.h
+++ b/Engine/Plugins/Experimental/GeometryProcessing/Source/MeshSolverUtilities/Private/LaplacianMeshSmoother.h
@@ -178,10 +178,10 @@ public:
 	FDiffusionIntegrator(const FDynamicMesh3& DynamicMesh, const ELaplacianWeightScheme Scheme);
 	virtual ~FDiffusionIntegrator() {}
 
-	void Integrate_ForwardEuler(int32 NumSteps, double Speed, double);
+	void Integrate_ForwardEuler(const int32 NumSteps, const double Speed);
 
 	// Note: 
-	void Integrate_BackwardEuler(const EMatrixSolverType MatrixSolverType, int32 NumSteps, double Speed, double Intensity);
+	void Integrate_BackwardEuler(const EMatrixSolverType MatrixSolverType, const int32 NumSteps, const double TimeStepSize);
 
 	void GetPositions(TArray<FVector3d>& PositionArray) const;
 
