MeshModeling : Adding automatic simplification pass to the voxel based tools.  Simplifies to a quadric error tol based on the diagonal distance across a voxel.
#rnx

diff --git a/Engine/Plugins/Experimental/GeometryProcessing/Source/DynamicMesh/Private/MeshSimplification.cpp b/Engine/Plugins/Experimental/GeometryProcessing/Source/DynamicMesh/Private/MeshSimplification.cpp
--- a/Engine/Plugins/Experimental/GeometryProcessing/Source/DynamicMesh/Private/MeshSimplification.cpp
+++ b/Engine/Plugins/Experimental/GeometryProcessing/Source/DynamicMesh/Private/MeshSimplification.cpp
@@ -406,6 +406,14 @@ void TMeshSimplification<QuadricErrorType>::DoSimplify()
 				break;
 			}
 		}
+		else if (SimplifyMode == ETargetModes::MaxError)
+		{
+			float qe = EdgeQueue.GetFirstNodePriority();
+			if (FMath::Abs(qe) > MaxErrorAllowed)
+			{
+				break;
+			}
+		}
 		else
 		{
 			if (Mesh->TriangleCount() <= TargetCount)
@@ -413,7 +421,7 @@ void TMeshSimplification<QuadricErrorType>::DoSimplify()
 				break;
 			}
 		}
-
+		
 		COUNT_ITERATIONS++;
 		int eid = EdgeQueue.Dequeue();
 		if (Mesh->IsEdge(eid) == false)
@@ -459,6 +467,7 @@ void TMeshSimplification<QuadricErrorType>::SimplifyToTriangleCount(int nCount)
 	SimplifyMode = ETargetModes::TriangleCount;
 	TargetCount = FMath::Max(1, nCount);
 	MinEdgeLength = FMathd::MaxReal;
+	MaxErrorAllowed = FMathf::MaxReal;
 	DoSimplify();
 }
 
@@ -468,6 +477,7 @@ void TMeshSimplification<QuadricErrorType>::SimplifyToVertexCount(int nCount)
 	SimplifyMode = ETargetModes::VertexCount;
 	TargetCount = FMath::Max(3, nCount);
 	MinEdgeLength = FMathd::MaxReal;
+	MaxErrorAllowed = FMathf::MaxReal;
 	DoSimplify();
 }
 
@@ -477,11 +487,23 @@ void TMeshSimplification<QuadricErrorType>::SimplifyToEdgeLength(double minEdgeL
 	SimplifyMode = ETargetModes::MinEdgeLength;
 	TargetCount = 1;
 	MinEdgeLength = minEdgeLen;
+	MaxErrorAllowed = FMathf::MaxReal;
+	DoSimplify();
+}
+
+template <typename QuadricErrorType>
+void TMeshSimplification<QuadricErrorType>::SimplifyToMaxError(double MaxError)
+{
+	SimplifyMode = ETargetModes::MaxError;
+	TargetCount = 1;
+	MinEdgeLength = FMathd::MaxReal;
+	MaxErrorAllowed = MaxError;
 	DoSimplify();
 }
 
 
 
+
 template <typename QuadricErrorType>
 void TMeshSimplification<QuadricErrorType>::FastCollapsePass(double fMinEdgeLength, int nRounds, bool MeshIsClosedHint)
 {
diff --git a/Engine/Plugins/Experimental/GeometryProcessing/Source/DynamicMesh/Public/MeshSimplification.h b/Engine/Plugins/Experimental/GeometryProcessing/Source/DynamicMesh/Public/MeshSimplification.h
--- a/Engine/Plugins/Experimental/GeometryProcessing/Source/DynamicMesh/Public/MeshSimplification.h
+++ b/Engine/Plugins/Experimental/GeometryProcessing/Source/DynamicMesh/Public/MeshSimplification.h
@@ -74,6 +74,12 @@ public:
 	 */
 	virtual void SimplifyToEdgeLength(double MinEdgeLength);
 
+	/**
+	 * Simplify mesh until the quadric error of an edge collapse exceeds the specified criteria.
+	 * @param MaxError collapse an edge if the corresponding quadric error exceeds this 
+	 */
+	virtual void SimplifyToMaxError(double MaxError);
+
 	/** 
 	 * Does N rounds of collapsing edges longer than fMinEdgeLength. Does not use Quadrics or priority queue.
 	 * This is a quick way to get triangle count down on huge meshes (eg like marching cubes output). 
@@ -94,14 +100,15 @@ protected:
 	// this just lets us write more concise code
 	bool EnableInlineProjection() const { return ProjectionMode == ETargetProjectionMode::Inline; }
 
-
+	float MaxErrorAllowed = FLT_MAX;
 	double MinEdgeLength = FMathd::MaxReal;
 	int TargetCount = INT_MAX;
 	enum class ETargetModes
 	{
 		TriangleCount = 0,
 		VertexCount = 1,
-		MinEdgeLength = 2
+		MinEdgeLength = 2,
+		MaxError  = 3
 	};
 	ETargetModes SimplifyMode = ETargetModes::TriangleCount;
 
diff --git a/Engine/Plugins/Experimental/MeshModelingToolset/Source/MeshModelingToolsEditorOnly/Private/MergeMeshesTool.cpp b/Engine/Plugins/Experimental/MeshModelingToolset/Source/MeshModelingToolsEditorOnly/Private/MergeMeshesTool.cpp
--- a/Engine/Plugins/Experimental/MeshModelingToolset/Source/MeshModelingToolsEditorOnly/Private/MergeMeshesTool.cpp
+++ b/Engine/Plugins/Experimental/MeshModelingToolset/Source/MeshModelingToolsEditorOnly/Private/MergeMeshesTool.cpp
@@ -151,9 +151,10 @@ void UMergeMeshesTool::OnPropertyModified(UObject* PropertySet, UProperty* Prope
 TSharedPtr<FDynamicMeshOperator> UMergeMeshesTool::MakeNewOperator()
 {
 	TSharedPtr<FVoxelMergeMeshesOp> MergeOp = MakeShared<FVoxelMergeMeshesOp>();
-	MergeOp->VoxelCount = MergeProps->VoxelCount;
-	MergeOp->AdaptivityD = MergeProps->Adaptivity;
-	MergeOp->IsoSurfaceD = MergeProps->IsoSurface;
+	MergeOp->VoxelCount     = MergeProps->VoxelCount;
+	MergeOp->AdaptivityD    = MergeProps->Adaptivity;
+	MergeOp->IsoSurfaceD    = MergeProps->IsoSurface;
+	MergeOp->bAutoSimplify  = MergeProps->bAutoSimplify;
 	MergeOp->InputMeshArray = InputMeshes;
 	return MergeOp;
 }
diff --git a/Engine/Plugins/Experimental/MeshModelingToolset/Source/MeshModelingToolsEditorOnly/Private/VoxelCSGMeshesTool.cpp b/Engine/Plugins/Experimental/MeshModelingToolset/Source/MeshModelingToolsEditorOnly/Private/VoxelCSGMeshesTool.cpp
--- a/Engine/Plugins/Experimental/MeshModelingToolset/Source/MeshModelingToolsEditorOnly/Private/VoxelCSGMeshesTool.cpp
+++ b/Engine/Plugins/Experimental/MeshModelingToolset/Source/MeshModelingToolsEditorOnly/Private/VoxelCSGMeshesTool.cpp
@@ -145,10 +145,11 @@ void UVoxelCSGMeshesTool::OnPropertyModified(UObject* PropertySet, UProperty* Pr
 TSharedPtr<FDynamicMeshOperator> UVoxelCSGMeshesTool::MakeNewOperator()
 {
 	TSharedPtr<FVoxelBooleanMeshesOp> CSGOp = MakeShared<FVoxelBooleanMeshesOp>();
-	CSGOp->Operation = (FVoxelBooleanMeshesOp::EBooleanOperation)(int)CSGProps->Operation;
-	CSGOp->VoxelCount = CSGProps->VoxelCount;
-	CSGOp->AdaptivityD = CSGProps->Adaptivity;
-	CSGOp->IsoSurfaceD = CSGProps->IsoSurface;
+	CSGOp->Operation      = (FVoxelBooleanMeshesOp::EBooleanOperation)(int)CSGProps->Operation;
+	CSGOp->VoxelCount     = CSGProps->VoxelCount;
+	CSGOp->AdaptivityD    = CSGProps->Adaptivity;
+	CSGOp->IsoSurfaceD    = CSGProps->IsoSurface;
+	CSGOp->bAutoSimplify  = CSGProps->bAutoSimplify;
 	CSGOp->InputMeshArray = InputMeshes;
 	return CSGOp;
 }
diff --git a/Engine/Plugins/Experimental/MeshModelingToolset/Source/MeshModelingToolsEditorOnly/Public/MergeMeshesTool.h b/Engine/Plugins/Experimental/MeshModelingToolset/Source/MeshModelingToolsEditorOnly/Public/MergeMeshesTool.h
--- a/Engine/Plugins/Experimental/MeshModelingToolset/Source/MeshModelingToolsEditorOnly/Public/MergeMeshesTool.h
+++ b/Engine/Plugins/Experimental/MeshModelingToolset/Source/MeshModelingToolsEditorOnly/Public/MergeMeshesTool.h
@@ -47,11 +47,15 @@ public:
 
 	/** Remeshing adaptivity */
 	UPROPERTY(EditAnywhere, Category = Options, meta = (UIMin = "0", UIMax = "1", ClampMin = "0", ClampMax = "1"))
-	float Adaptivity = 0.01f;
+	float Adaptivity = 0.001f;
 
 	/** Offset when remeshing, measured in voxels units */
 	UPROPERTY(EditAnywhere, Category = Options, meta = (UIMin = "-2", UIMax = "2", ClampMin = "-2", ClampMax = "2"))
 	float IsoSurface = 0;
+
+	/** Automatically simplify the result of voxel-based merge.*/
+	UPROPERTY(EditAnywhere, Category = Options)
+	bool bAutoSimplify = true;
 };
 
 /**
diff --git a/Engine/Plugins/Experimental/MeshModelingToolset/Source/MeshModelingToolsEditorOnly/Public/VoxelCSGMeshesTool.h b/Engine/Plugins/Experimental/MeshModelingToolset/Source/MeshModelingToolsEditorOnly/Public/VoxelCSGMeshesTool.h
--- a/Engine/Plugins/Experimental/MeshModelingToolset/Source/MeshModelingToolsEditorOnly/Public/VoxelCSGMeshesTool.h
+++ b/Engine/Plugins/Experimental/MeshModelingToolset/Source/MeshModelingToolsEditorOnly/Public/VoxelCSGMeshesTool.h
@@ -69,11 +69,15 @@ public:
 
 	/** Remeshing adaptivity */
 	UPROPERTY(EditAnywhere, Category = Options, meta = (UIMin = "0", UIMax = "1", ClampMin = "0", ClampMax = "1"))
-	float Adaptivity = 0.01f;
+	float Adaptivity = 0.001f;
 
 	/** Offset when remeshing, measured in voxels units */
 	UPROPERTY(EditAnywhere, Category = Options, meta = (UIMin = "-2", UIMax = "2", ClampMin = "-2", ClampMax = "2"))
 	float IsoSurface = 0;
+
+	/** Automatically simplify the result of voxel-based merge.*/
+	UPROPERTY(EditAnywhere, Category = Options)
+	bool bAutoSimplify = true;
 };
 
 
diff --git a/Engine/Plugins/Experimental/MeshModelingToolset/Source/ModelingOperatorsEditorOnly/Private/CompositionOps/VoxelBooleanMeshesOp.cpp b/Engine/Plugins/Experimental/MeshModelingToolset/Source/ModelingOperatorsEditorOnly/Private/CompositionOps/VoxelBooleanMeshesOp.cpp
--- a/Engine/Plugins/Experimental/MeshModelingToolset/Source/ModelingOperatorsEditorOnly/Private/CompositionOps/VoxelBooleanMeshesOp.cpp
+++ b/Engine/Plugins/Experimental/MeshModelingToolset/Source/ModelingOperatorsEditorOnly/Private/CompositionOps/VoxelBooleanMeshesOp.cpp
@@ -69,21 +69,18 @@ void FVoxelBooleanMeshesOp::CalculateResult(FProgressCancel* Progress)
 		return;
 	}
 
-	if (AdaptivityD > 0)
+	if (bAutoSimplify)
 	{
 		FQEMSimplification Reducer(ResultMesh.Get());
 		FMeshConstraints constraints;
 		FMeshConstraintsUtil::ConstrainAllSeams(constraints, *ResultMesh, true, false);
 		Reducer.SetExternalConstraints(&constraints);
 		Reducer.Progress = Progress;
-		for (int k = 0; k < NumFastCollapsePasses; ++k)
-		{
-			Reducer.FastCollapsePass(FastCollapseVoxelSizeMultipler * VoxelSizeD);
-			if (Progress->Cancelled())
-			{
-				return;
-			}
-		}
+
+		const double MaxDisplacementSqr = 3. * VoxelSizeD * VoxelSizeD;
+
+		Reducer.SimplifyToMaxError(MaxDisplacementSqr);
+
 	}
 }
 
diff --git a/Engine/Plugins/Experimental/MeshModelingToolset/Source/ModelingOperatorsEditorOnly/Private/CompositionOps/VoxelMergeMeshesOp.cpp b/Engine/Plugins/Experimental/MeshModelingToolset/Source/ModelingOperatorsEditorOnly/Private/CompositionOps/VoxelMergeMeshesOp.cpp
--- a/Engine/Plugins/Experimental/MeshModelingToolset/Source/ModelingOperatorsEditorOnly/Private/CompositionOps/VoxelMergeMeshesOp.cpp
+++ b/Engine/Plugins/Experimental/MeshModelingToolset/Source/ModelingOperatorsEditorOnly/Private/CompositionOps/VoxelMergeMeshesOp.cpp
@@ -43,21 +43,17 @@ void FVoxelMergeMeshesOp::CalculateResult(FProgressCancel* Progress)
 		return;
 	}
 
-	if (AdaptivityD > 0)
+	if (bAutoSimplify)
 	{
 		FQEMSimplification Reducer(ResultMesh.Get());
 		FMeshConstraints constraints;
 		FMeshConstraintsUtil::ConstrainAllSeams(constraints, *ResultMesh, true, false);
 		Reducer.SetExternalConstraints(&constraints);
 		Reducer.Progress = Progress;
-		for (int k = 0; k < NumFastCollapsePasses; ++k)
-		{
-			Reducer.FastCollapsePass(FastCollapseVoxelSizeMultipler * VoxelSizeD);
-			if (Progress->Cancelled())
-			{
-				return;
-			}
-		}
+
+		const double MaxDisplacementSqr = 3. *VoxelSizeD * VoxelSizeD;
+
+		Reducer.SimplifyToMaxError(MaxDisplacementSqr);
 	}
 }
 
diff --git a/Engine/Plugins/Experimental/MeshModelingToolset/Source/ModelingOperatorsEditorOnly/Public/CompositionOps/VoxelBooleanMeshesOp.h b/Engine/Plugins/Experimental/MeshModelingToolset/Source/ModelingOperatorsEditorOnly/Public/CompositionOps/VoxelBooleanMeshesOp.h
--- a/Engine/Plugins/Experimental/MeshModelingToolset/Source/ModelingOperatorsEditorOnly/Public/CompositionOps/VoxelBooleanMeshesOp.h
+++ b/Engine/Plugins/Experimental/MeshModelingToolset/Source/ModelingOperatorsEditorOnly/Public/CompositionOps/VoxelBooleanMeshesOp.h
@@ -29,6 +29,7 @@ public:
 	double VoxelSizeD = 1.0;
 	double AdaptivityD = 0;
 	double IsoSurfaceD = 0;
+	bool   bAutoSimplify = false;
 	EBooleanOperation Operation = EBooleanOperation::DifferenceBA;
 
 	double FastCollapseVoxelSizeMultipler = 1.2;
diff --git a/Engine/Plugins/Experimental/MeshModelingToolset/Source/ModelingOperatorsEditorOnly/Public/CompositionOps/VoxelMergeMeshesOp.h b/Engine/Plugins/Experimental/MeshModelingToolset/Source/ModelingOperatorsEditorOnly/Public/CompositionOps/VoxelMergeMeshesOp.h
--- a/Engine/Plugins/Experimental/MeshModelingToolset/Source/ModelingOperatorsEditorOnly/Public/CompositionOps/VoxelMergeMeshesOp.h
+++ b/Engine/Plugins/Experimental/MeshModelingToolset/Source/ModelingOperatorsEditorOnly/Public/CompositionOps/VoxelMergeMeshesOp.h
@@ -22,6 +22,7 @@ public:
 	double VoxelSizeD = 1.0;
 	double AdaptivityD = 0;
 	double IsoSurfaceD = 0;
+	bool bAutoSimplify = false;
 
 	double FastCollapseVoxelSizeMultipler = 1.2;
 	int NumFastCollapsePasses = 10;
