Fix static analysis error in GteUtil.h: missing header guard
#rnx

diff --git a/Engine/Plugins/Experimental/GeometryProcessing/Source/GeometryAlgorithms/Private/GteUtil.h b/Engine/Plugins/Experimental/GeometryProcessing/Source/GeometryAlgorithms/Private/GteUtil.h
--- a/Engine/Plugins/Experimental/GeometryProcessing/Source/GeometryAlgorithms/Private/GteUtil.h
+++ b/Engine/Plugins/Experimental/GeometryProcessing/Source/GeometryAlgorithms/Private/GteUtil.h
@@ -1,5 +1,7 @@
 // Copyright 1998-2019 Epic Games, Inc. All Rights Reserved.
 
+#pragma once
+
 #include "VectorTypes.h"
 #include "LineTypes.h"
 #include "CircleTypes.h"
