Add support for per-triangle DynamicMesh attributes. Use to implement per-triangle material ID. Copy to/from MeshDescription PolygonGroups (as is convention in existing code)
#rnx

diff --git a/Engine/Plugins/Experimental/GeometryProcessing/Source/DynamicMesh/Private/DynamicMeshAttributeSet.cpp b/Engine/Plugins/Experimental/GeometryProcessing/Source/DynamicMesh/Private/DynamicMeshAttributeSet.cpp
--- a/Engine/Plugins/Experimental/GeometryProcessing/Source/DynamicMesh/Private/DynamicMeshAttributeSet.cpp
+++ b/Engine/Plugins/Experimental/GeometryProcessing/Source/DynamicMesh/Private/DynamicMeshAttributeSet.cpp
@@ -5,6 +5,17 @@
 
 
 
+void FDynamicMeshAttributeSet::EnableMaterialID()
+{
+	if (HasMaterialID() == false)
+	{
+		MaterialIDAttrib = MakeUnique<FDynamicMeshMaterialAttribute>(ParentMesh);
+		MaterialIDAttrib->Initialize((int32)0);
+	}
+}
+
+
+
 bool FDynamicMeshAttributeSet::IsSeamEdge(int eid) const
 {
 	return UV0.IsSeamEdge(eid) || Normals0.IsSeamEdge(eid);
@@ -16,30 +27,50 @@ void FDynamicMeshAttributeSet::OnNewTriangle(int TriangleID, bool bInserted)
 {
 	UV0.InitializeNewTriangle(TriangleID);
 	Normals0.InitializeNewTriangle(TriangleID);
+
+	if (MaterialIDAttrib)
+	{
+		int NewValue = 0;
+		MaterialIDAttrib->SetNewValue(TriangleID, &NewValue);
+	}
 }
 
 void FDynamicMeshAttributeSet::OnRemoveTriangle(int TriangleID, bool bRemoveIsolatedVertices)
 {
 	UV0.OnRemoveTriangle(TriangleID, bRemoveIsolatedVertices);
 	Normals0.OnRemoveTriangle(TriangleID, bRemoveIsolatedVertices);
+
+	// has no effect on MaterialIDAttrib
 }
 
 void FDynamicMeshAttributeSet::OnReverseTriOrientation(int TriangleID)
 {
 	UV0.OnReverseTriOrientation(TriangleID);
 	Normals0.OnReverseTriOrientation(TriangleID);
+
+	// has no effect on MaterialIDAttrib
 }
 
 void FDynamicMeshAttributeSet::OnSplitEdge(const FDynamicMesh3::FEdgeSplitInfo & splitInfo)
 {
 	UV0.OnSplitEdge(splitInfo);
 	Normals0.OnSplitEdge(splitInfo);
+
+	if (MaterialIDAttrib)
+	{
+		MaterialIDAttrib->OnSplitEdge(splitInfo);
+	}
 }
 
 void FDynamicMeshAttributeSet::OnFlipEdge(const FDynamicMesh3::FEdgeFlipInfo & flipInfo)
 {
 	UV0.OnFlipEdge(flipInfo);
 	Normals0.OnFlipEdge(flipInfo);
+
+	if (MaterialIDAttrib)
+	{
+		MaterialIDAttrib->OnFlipEdge(flipInfo);
+	}
 }
 
 
@@ -47,17 +78,32 @@ void FDynamicMeshAttributeSet::OnCollapseEdge(const FDynamicMesh3::FEdgeCollapse
 {
 	UV0.OnCollapseEdge(collapseInfo);
 	Normals0.OnCollapseEdge(collapseInfo);
+
+	if (MaterialIDAttrib)
+	{
+		MaterialIDAttrib->OnCollapseEdge(collapseInfo);
+	}
 }
 
 void FDynamicMeshAttributeSet::OnPokeTriangle(const FDynamicMesh3::FPokeTriangleInfo & pokeInfo)
 {
 	UV0.OnPokeTriangle(pokeInfo);
 	Normals0.OnPokeTriangle(pokeInfo);
+
+	if (MaterialIDAttrib)
+	{
+		MaterialIDAttrib->OnPokeTriangle(pokeInfo);
+	}
 }
 
 void FDynamicMeshAttributeSet::OnMergeEdges(const FDynamicMesh3::FMergeEdgesInfo & mergeInfo)
 {
 	UV0.OnMergeEdges(mergeInfo);
 	Normals0.OnMergeEdges(mergeInfo);
+
+	if (MaterialIDAttrib)
+	{
+		MaterialIDAttrib->OnMergeEdges(mergeInfo);
+	}
 }
 
diff --git a/Engine/Plugins/Experimental/GeometryProcessing/Source/DynamicMesh/Public/DynamicMeshAttributeSet.h b/Engine/Plugins/Experimental/GeometryProcessing/Source/DynamicMesh/Public/DynamicMeshAttributeSet.h
--- a/Engine/Plugins/Experimental/GeometryProcessing/Source/DynamicMesh/Public/DynamicMeshAttributeSet.h
+++ b/Engine/Plugins/Experimental/GeometryProcessing/Source/DynamicMesh/Public/DynamicMeshAttributeSet.h
@@ -4,12 +4,15 @@
 #pragma once
 
 #include "DynamicMeshOverlay.h"
+#include "DynamicMeshTriangleAttribute.h"
 
 /** Standard UV overlay type - 2-element float */
 typedef TDynamicMeshVectorOverlay<float, 2, FVector2f> FDynamicMeshUVOverlay;
 /** Standard Normal overlay type - 3-element float */
 typedef TDynamicMeshVectorOverlay<float, 3, FVector3f> FDynamicMeshNormalOverlay;
 
+/** Standard per-triangle integer material ID */
+typedef TDynamicMeshTriangleAttribute<int32, 1> FDynamicMeshMaterialAttribute;
 
 /**
  * FDynamicMeshAttributeSet manages a set of extended attributes for a FDynamicMesh3.
@@ -38,6 +41,13 @@ public:
 	{
 		UV0.Copy(Copy.UV0);
 		Normals0.Copy(Copy.Normals0);
+
+		if (Copy.MaterialIDAttrib)
+		{
+			EnableMaterialID();
+			MaterialIDAttrib->Copy( *(Copy.MaterialIDAttrib) );
+		}
+
 		// parent mesh is *not* copied!
 	}
 
@@ -143,6 +153,29 @@ public:
 	}
 
 
+
+	//
+	// Per-Triangle Material ID
+	//
+
+	bool HasMaterialID() const
+	{
+		return !!MaterialIDAttrib;
+	}
+
+
+	void EnableMaterialID();
+
+	FDynamicMeshMaterialAttribute* GetMaterialID()
+	{
+		return MaterialIDAttrib.Get();
+	}
+
+	const FDynamicMeshMaterialAttribute* GetMaterialID() const
+	{
+		return MaterialIDAttrib.Get();
+	}
+
 protected:
 	/** Parent mesh of this attribute set */
 	FDynamicMesh3* ParentMesh;
@@ -155,6 +188,8 @@ protected:
 	TArray<FDynamicMeshUVOverlay*> UVLayers;
 	TArray<FDynamicMeshNormalOverlay*> NormalLayers;
 
+	TUniquePtr<FDynamicMeshMaterialAttribute> MaterialIDAttrib;
+	
 
 protected:
 	friend class FDynamicMesh3;
diff --git a/Engine/Plugins/Experimental/GeometryProcessing/Source/DynamicMesh/Public/DynamicMeshTriangleAttribute.h b/Engine/Plugins/Experimental/GeometryProcessing/Source/DynamicMesh/Public/DynamicMeshTriangleAttribute.h
new file mode 100644
--- /dev/null
+++ b/Engine/Plugins/Experimental/GeometryProcessing/Source/DynamicMesh/Public/DynamicMeshTriangleAttribute.h
@@ -0,0 +1,191 @@
+// Copyright 1998-2019 Epic Games, Inc. All Rights Reserved.
+
+#pragma once
+
+#include "DynamicMesh3.h"
+
+/**
+ * TDynamicMeshTriangleAttribute is an add-on to a FDynamicMesh3 that allows for 
+ * per-triangle storage of an attribute value.
+ *
+ * The FDynamicMesh3 mesh topology operations (eg split/flip/collapse edge, poke face, etc)
+ * can be mirrored to the overlay via OnSplitEdge(), etc.
+ */
+template<typename AttribValueType, int AttribDimension>
+class TDynamicMeshTriangleAttribute
+{
+
+protected:
+	/** The parent mesh this overlay belongs to */
+	FDynamicMesh3* ParentMesh;
+
+	/** List of per-triangle attribute values */
+	TDynamicVector<AttribValueType> AttribValues;
+
+	friend class FDynamicMesh3;
+	friend class FDynamicMeshAttributeSet;
+
+public:
+	/** Create an empty overlay */
+	TDynamicMeshTriangleAttribute()
+	{
+		ParentMesh = nullptr;
+	}
+
+	/** Create an overlay for the given parent mesh */
+	TDynamicMeshTriangleAttribute(FDynamicMesh3* ParentMeshIn)
+	{
+		ParentMesh = ParentMeshIn;
+	}
+
+	/** @return the parent mesh for this overlay */
+	const FDynamicMesh3* GetParentMesh() const { return ParentMesh; }
+	/** @return the parent mesh for this overlay */
+	FDynamicMesh3* GetParentMesh() { return ParentMesh; }
+
+	/** Set this overlay to contain the same arrays as the copy overlay */
+	void Copy(const TDynamicMeshTriangleAttribute<AttribValueType, AttribDimension>& Copy)
+	{
+		AttribValues = Copy.AttribValues;
+	}
+
+	/** Initialize the attribute values to the given max triangle ID */
+	void Initialize(AttribValueType InitialValue = (AttribValueType)0)
+	{
+		check(ParentMesh != nullptr);
+		AttribValues.Resize(0);
+		AttribValues.Resize( ParentMesh->MaxTriangleID() * AttribDimension );
+	}
+
+	void SetNewValue(int NewTriangleID, const AttribValueType* Data)
+	{
+		int k = NewTriangleID * AttribDimension;
+		for (int i = 0; i < AttribDimension; ++i)
+		{
+			AttribValues.InsertAt(Data[i], k + i);
+		}
+	}
+
+
+
+	//
+	// Accessors/Queries
+	//  
+
+	/** Get the element at a given index */
+	inline void GetValue(int TriangleID, AttribValueType* Data) const
+	{
+		int k = TriangleID * AttribDimension;
+		for (int i = 0; i < AttribDimension; ++i)
+		{
+			Data[i] = AttribValues[k + i];
+		}
+	}
+
+	/** Get the element at a given index */
+	template<typename AsType>
+	void GetValue(int TriangleID, AsType& Data) const
+	{
+		int k = TriangleID * AttribDimension;
+		for (int i = 0; i < AttribDimension; ++i)
+		{
+			Data[i] = AttribValues[k + i];
+		}
+	}
+
+
+	/** Set the element at a given index */
+	inline void SetValue(int TriangleID, const AttribValueType* Data)
+	{
+		int k = TriangleID * AttribDimension;
+		for (int i = 0; i < AttribDimension; ++i)
+		{
+			AttribValues[k + i] = Data[i];
+		}
+	}
+
+	/** Set the element at a given index */
+	template<typename AsType>
+	void SetValue(int TriangleID, const AsType& Data)
+	{
+		int k = TriangleID * AttribDimension;
+		for (int i = 0; i < AttribDimension; ++i)
+		{
+			AttribValues[k + i] = Data[i];
+		}
+	}
+
+	/**
+	 * Copy the attribute value at FromTriangleID to ToTriangleID
+	 */
+	inline void CopyValue(int FromTriangleID, int ToTriangleID)
+	{
+		int kA = FromTriangleID * AttribDimension;
+		int kB = ToTriangleID * AttribDimension;
+		for (int i = 0; i < AttribDimension; ++i)
+		{
+			AttribValues[kB+i] = AttribValues[kA+i];
+		}
+	}
+
+
+	/** Returns true if the parent-mesh edge is a "Seam" in this overlay */
+	bool IsBorderEdge(int EdgeID, bool bMeshBoundaryIsBorder = true) const
+	{
+		FIndex2i EdgeTris = ParentMesh->GetEdgeT(EdgeID);
+		if (EdgeTris.B == FDynamicMesh3::InvalidID)
+		{
+			return bMeshBoundaryIsBorder;
+		}
+		int kA = EdgeTris.A * AttribDimension;
+		int kB = EdgeTris.B * AttribDimension;
+		for (int i = 0; i < AttribDimension; ++i)
+		{
+			if (AttribValues[kA+i] != AttribValues[kB+i])
+			{
+				return true;
+			}
+		}
+	}
+
+
+
+public:
+
+	/** Update the overlay to reflect an edge split in the parent mesh */
+	void OnSplitEdge(const FDynamicMesh3::FEdgeSplitInfo& SplitInfo)
+	{
+		CopyValue(SplitInfo.OriginalTriangles.A, SplitInfo.NewTriangles.A);
+		if (SplitInfo.OriginalTriangles.B != FDynamicMesh3::InvalidID)
+		{
+			CopyValue(SplitInfo.OriginalTriangles.B, SplitInfo.NewTriangles.B);
+		}
+	}
+
+	/** Update the overlay to reflect an edge flip in the parent mesh */
+	void OnFlipEdge(const FDynamicMesh3::FEdgeFlipInfo& FlipInfo)
+	{
+		// yikes! triangles did not actually change so we will leave attrib unmodified
+	}
+
+	/** Update the overlay to reflect an edge collapse in the parent mesh */
+	void OnCollapseEdge(const FDynamicMesh3::FEdgeCollapseInfo& CollapseInfo)
+	{
+		// nothing to do here, triangles were only deleted
+	}
+
+	/** Update the overlay to reflect a face poke in the parent mesh */
+	void OnPokeTriangle(const FDynamicMesh3::FPokeTriangleInfo& PokeInfo)
+	{
+		CopyValue(PokeInfo.OriginalTriangle, PokeInfo.NewTriangles.A);
+		CopyValue(PokeInfo.OriginalTriangle, PokeInfo.NewTriangles.B);
+	}
+
+	/** Update the overlay to reflect an edge merge in the parent mesh */
+	void OnMergeEdges(const FDynamicMesh3::FMergeEdgesInfo& MergeInfo)
+	{
+		// nothing to do here because triangles did not change
+	}
+
+};
+
diff --git a/Engine/Plugins/Experimental/GeometryProcessing/Source/MeshConversion/Private/DynamicMeshToMeshDescription.cpp b/Engine/Plugins/Experimental/GeometryProcessing/Source/MeshConversion/Private/DynamicMeshToMeshDescription.cpp
--- a/Engine/Plugins/Experimental/GeometryProcessing/Source/MeshConversion/Private/DynamicMeshToMeshDescription.cpp
+++ b/Engine/Plugins/Experimental/GeometryProcessing/Source/MeshConversion/Private/DynamicMeshToMeshDescription.cpp
@@ -230,8 +230,35 @@ void FDynamicMeshToMeshDescription::Convert_SharedInstances(const FDynamicMesh3*
 	}
 
 
-	FPolygonGroupID AllGroupID = Builder.AppendPolygonGroup();
+	FPolygonGroupID ZeroPolygonGroupID = Builder.AppendPolygonGroup();
 
+	// check if we have per-triangle material ID
+	const FDynamicMeshMaterialAttribute* MaterialIDAttrib =
+		(MeshIn->HasAttributes() && MeshIn->Attributes()->HasMaterialID()) ?
+		MeshIn->Attributes()->GetMaterialID() : nullptr;
+
+	// need to know max material index value so we can reserve groups in MeshDescription
+	int32 MaxPolygonGroupID = 0;
+	if (MaterialIDAttrib)
+	{
+		for (int TriID : MeshIn->TriangleIndicesItr())
+		{
+			int32 MaterialID;
+			MaterialIDAttrib->GetValue(TriID, &MaterialID);
+			MaxPolygonGroupID = FMath::Max(MaterialID, MaxPolygonGroupID);
+		}
+		if (MaxPolygonGroupID == 0)
+		{
+			MaterialIDAttrib = nullptr;
+		}
+		else
+		{
+			for (int k = 0; k < MaxPolygonGroupID; ++k)
+			{
+				Builder.AppendPolygonGroup();
+			}
+		}
+	}
 
 	// create new instances when seen
 	TMap<FIndex3i, FVertexInstanceID> InstanceList;
@@ -256,7 +283,16 @@ void FDynamicMeshToMeshDescription::Convert_SharedInstances(const FDynamicMesh3*
 			InstanceTri[j] = InstanceList[InstanceElem];
 		}
 
-		FPolygonID NewPolygonID = Builder.AppendTriangle(InstanceTri[0], InstanceTri[1], InstanceTri[2], AllGroupID);
+		// transfer material index to MeshDescription polygon group (by convention)
+		FPolygonGroupID UsePolygonGroupID = ZeroPolygonGroupID;
+		if (MaterialIDAttrib)
+		{
+			int32 MaterialID;
+			MaterialIDAttrib->GetValue(TriID, &MaterialID);
+			UsePolygonGroupID = FPolygonGroupID(MaterialID);
+		}
+
+		FPolygonID NewPolygonID = Builder.AppendTriangle(InstanceTri[0], InstanceTri[1], InstanceTri[2], UsePolygonGroupID);
 
 		if (bCopyGroupToPolyGroup)
 		{
@@ -292,7 +328,36 @@ void FDynamicMeshToMeshDescription::Convert_NoSharedInstances(const FDynamicMesh
 		MapV[VertID] = Builder.AppendVertex(MeshIn->GetVertex(VertID));
 	}
 
-	FPolygonGroupID AllGroupID = Builder.AppendPolygonGroup();
+	FPolygonGroupID ZeroPolygonGroupID = Builder.AppendPolygonGroup();
+
+	// check if we have per-triangle material ID
+	const FDynamicMeshMaterialAttribute* MaterialIDAttrib =
+		(MeshIn->HasAttributes() && MeshIn->Attributes()->HasMaterialID()) ?
+		MeshIn->Attributes()->GetMaterialID() : nullptr;
+
+	// need to know max material index value so we can reserve groups in MeshDescription
+	int32 MaxPolygonGroupID = 0;
+	if (MaterialIDAttrib)
+	{
+		for (int TriID : MeshIn->TriangleIndicesItr())
+		{
+			int32 MaterialID;
+			MaterialIDAttrib->GetValue(TriID, &MaterialID);
+			MaxPolygonGroupID = FMath::Max(MaterialID, MaxPolygonGroupID);
+		}
+		if (MaxPolygonGroupID == 0)
+		{
+			MaterialIDAttrib = nullptr;
+		}
+		else
+		{
+			for (int k = 0; k < MaxPolygonGroupID; ++k)
+			{
+				Builder.AppendPolygonGroup();
+			}
+		}
+	}
+
 
 	FVertexID TriVertices[3];
 	FVector2D TriUVs[3];
@@ -325,7 +390,16 @@ void FDynamicMeshToMeshDescription::Convert_NoSharedInstances(const FDynamicMesh
 		TriVertices[1] = MapV[Triangle[1]];
 		TriVertices[2] = MapV[Triangle[2]];
 
-		FPolygonID NewPolygonID = Builder.AppendTriangle(TriVertices, AllGroupID, UseUVs, UseNormals);
+		// transfer material index to MeshDescription polygon group (by convention)
+		FPolygonGroupID UsePolygonGroupID = ZeroPolygonGroupID;
+		if (MaterialIDAttrib)
+		{
+			int32 MaterialID;
+			MaterialIDAttrib->GetValue(TriID, &MaterialID);
+			UsePolygonGroupID = FPolygonGroupID(MaterialID);
+		}
+
+		FPolygonID NewPolygonID = Builder.AppendTriangle(TriVertices, UsePolygonGroupID, UseUVs, UseNormals);
 
 		if (bCopyGroupToPolyGroup)
 		{
diff --git a/Engine/Plugins/Experimental/GeometryProcessing/Source/MeshConversion/Private/MeshDescriptionToDynamicMesh.cpp b/Engine/Plugins/Experimental/GeometryProcessing/Source/MeshConversion/Private/MeshDescriptionToDynamicMesh.cpp
--- a/Engine/Plugins/Experimental/GeometryProcessing/Source/MeshConversion/Private/MeshDescriptionToDynamicMesh.cpp
+++ b/Engine/Plugins/Experimental/GeometryProcessing/Source/MeshConversion/Private/MeshDescriptionToDynamicMesh.cpp
@@ -167,12 +167,20 @@ void FMeshDescriptionToDynamicMesh::Convert(const FMeshDescription* MeshIn, FDyn
 	FUVWelder UVWelder(UVOverlay);
 	FNormalWelder NormalWelder(NormalOverlay);
 
+	// only enable material ID if we have more than one material
+	FDynamicMeshMaterialAttribute* MaterialIDAttrib = nullptr;
+
 	// NOTE: If you change the iteration order here, please update the corresponding iteration in FDynamicMeshToMeshDescription::UpdateAttributes, 
 	//	which assumes the iteration order here is polygons -> triangles, to correspond the triangles when writing updated attributes back!
 	const FPolygonArray& Polygons = MeshIn->Polygons();
 	for (const FPolygonID PolygonID : Polygons.GetElementIDs())
 	{
-		FPolygonGroupID PolygonGroupID = MeshIn->GetPolygonPolygonGroup(PolygonID);
+		int32 PolygonGroupID = MeshIn->GetPolygonPolygonGroup(PolygonID).GetValue();
+		if (MaterialIDAttrib == nullptr && PolygonGroupID != 0)
+		{
+			MeshOut.Attributes()->EnableMaterialID();		// initializes all existing values to 0, which is fine since this is first non-zero
+			MaterialIDAttrib = MeshOut.Attributes()->GetMaterialID();
+		}
 
 		const TArray<FTriangleID>& TriangleIDs = MeshIn->GetPolygonTriangleIDs(PolygonID);
 		int NumTriangles = TriangleIDs.Num();
@@ -196,7 +204,7 @@ void FMeshDescriptionToDynamicMesh::Convert(const FMeshDescription* MeshIn, FDyn
 			}
 			else if (GroupMode == EPrimaryGroupMode::SetToPolygonGroupID)
 			{
-				GroupID = PolygonGroupID.GetValue();
+				GroupID = PolygonGroupID;
 			}
 			
 
@@ -285,6 +293,12 @@ void FMeshDescriptionToDynamicMesh::Convert(const FMeshDescription* MeshIn, FDyn
 				NormalOverlay->SetTriangle(NewTriangleID, TriNormals);
 			}
 
+			// use PolygonGroup as MaterialID
+			if (MaterialIDAttrib != nullptr)
+			{
+				MaterialIDAttrib->SetValue(NewTriangleID, &PolygonGroupID);
+			}
+
 		}
 	}
 
