fix normals issues w/ various mesh editing tools by just not recomputing on conversion

#rnx

diff --git a/Engine/Plugins/Experimental/GeometryProcessing/Source/DynamicMesh/Public/MeshNormals.h b/Engine/Plugins/Experimental/GeometryProcessing/Source/DynamicMesh/Public/MeshNormals.h
--- a/Engine/Plugins/Experimental/GeometryProcessing/Source/DynamicMesh/Public/MeshNormals.h
+++ b/Engine/Plugins/Experimental/GeometryProcessing/Source/DynamicMesh/Public/MeshNormals.h
@@ -67,7 +67,7 @@ public:
 	 * Recompute the per-element normals of the given overlay by averaging one-ring face normals
 	 * @warning NormalOverlay must be attached to ParentMesh or an exact copy
 	 */
-	void RecomputeOverlayNormals(FDynamicMeshNormalOverlay* NormalOverlay, bool bWeightByArea = true, bool bWeightByAngle = true)
+	void RecomputeOverlayNormals(const FDynamicMeshNormalOverlay* NormalOverlay, bool bWeightByArea = true, bool bWeightByAngle = true)
 	{
 		Compute_Overlay_FaceAvg(NormalOverlay, bWeightByArea, bWeightByAngle);
 	}
diff --git a/Engine/Plugins/Experimental/GeometryProcessing/Source/MeshConversion/Private/DynamicMeshToMeshDescription.cpp b/Engine/Plugins/Experimental/GeometryProcessing/Source/MeshConversion/Private/DynamicMeshToMeshDescription.cpp
--- a/Engine/Plugins/Experimental/GeometryProcessing/Source/MeshConversion/Private/DynamicMeshToMeshDescription.cpp
+++ b/Engine/Plugins/Experimental/GeometryProcessing/Source/MeshConversion/Private/DynamicMeshToMeshDescription.cpp
@@ -5,33 +5,9 @@
 #include "DynamicMeshAttributeSet.h"
 #include "DynamicMeshOverlay.h"
 #include "MeshDescriptionBuilder.h"
-#include "MeshNormals.h"
 
 
 
-void FDynamicMeshToMeshDescription::Update(const FDynamicMesh3* MeshIn, FMeshDescription& MeshOut, bool bRecomputeNormals)
-{
-	FMeshDescriptionBuilder Builder;
-	Builder.SetMeshDescription(&MeshOut);
-
-	check(MeshIn->IsCompactV());
-	check(MeshIn->VertexCount() == MeshOut.Vertices().Num());
-
-	// update positions
-	for (int VertID : MeshIn->VertexIndicesItr())
-	{
-		Builder.SetPosition(FVertexID(VertID), MeshIn->GetVertex(VertID));
-	}
-
-	if (bRecomputeNormals)
-	{
-		Builder.RecalculateInstanceNormals();
-	}
-	else
-	{
-		UpdateAttributes(MeshIn, MeshOut, true, false);
-	}
-}
 
 
 namespace DynamicMeshToMeshDescriptionConversionHelper
@@ -66,6 +42,25 @@ namespace DynamicMeshToMeshDescriptionConversionHelper
 }
 
 
+void FDynamicMeshToMeshDescription::Update(const FDynamicMesh3* MeshIn, FMeshDescription& MeshOut)
+{
+	FMeshDescriptionBuilder Builder;
+	Builder.SetMeshDescription(&MeshOut);
+
+	check(MeshIn->IsCompactV());
+	check(MeshIn->VertexCount() == MeshOut.Vertices().Num());
+
+	// update positions
+	for (int VertID : MeshIn->VertexIndicesItr())
+	{
+		Builder.SetPosition(FVertexID(VertID), MeshIn->GetVertex(VertID));
+	}
+
+	UpdateAttributes(MeshIn, MeshOut, true, false);
+}
+
+
+
 void FDynamicMeshToMeshDescription::UpdateAttributes(const FDynamicMesh3* MeshIn, FMeshDescription& MeshOut, bool bUpdateNormals, bool bUpdateUVs)
 {
 	check(MeshIn->IsCompactV());
@@ -164,9 +159,6 @@ void FDynamicMeshToMeshDescription::Convert_NoAttributes(const FDynamicMesh3* Me
 		MapV[VertID] = Builder.AppendVertex(MeshIn->GetVertex(VertID));
 	}
 
-	FMeshNormals VertexNormals(MeshIn);
-	VertexNormals.ComputeVertexNormals();
-
 	FPolygonGroupID AllGroupID = Builder.AppendPolygonGroup();
 
 	// create new instances when seen
@@ -184,7 +176,9 @@ void FDynamicMeshToMeshDescription::Convert_NoAttributes(const FDynamicMesh3* Me
 			{
 				FVertexInstanceID NewInstanceID = Builder.AppendInstance(MapV[Triangle[j]]);
 				InstanceList.Add(InstanceElem, NewInstanceID);
-				Builder.SetInstance(NewInstanceID, FVector2f::Zero(), VertexNormals[Triangle[j]]);
+				FVector2D UV = MeshIn->HasVertexUVs() ? MeshIn->GetVertexUV(Triangle[j]) : FVector2D::ZeroVector;
+				FVector Normal = MeshIn->HasVertexNormals() ? MeshIn->GetVertexNormal(Triangle[j]) : FVector::UpVector;
+				Builder.SetInstance(NewInstanceID, UV, Normal);
 			}
 			InstanceTri[j] = InstanceList[InstanceElem];
 		}
@@ -195,8 +189,6 @@ void FDynamicMeshToMeshDescription::Convert_NoAttributes(const FDynamicMesh3* Me
 			Builder.SetPolyGroupID(NewPolygonID, MeshIn->GetTriangleGroup(TriID));
 		}
 	}
-
-	Builder.RecalculateInstanceNormals();
 }
 
 
diff --git a/Engine/Plugins/Experimental/GeometryProcessing/Source/MeshConversion/Private/MeshDescriptionBuilder.cpp b/Engine/Plugins/Experimental/GeometryProcessing/Source/MeshConversion/Private/MeshDescriptionBuilder.cpp
--- a/Engine/Plugins/Experimental/GeometryProcessing/Source/MeshConversion/Private/MeshDescriptionBuilder.cpp
+++ b/Engine/Plugins/Experimental/GeometryProcessing/Source/MeshConversion/Private/MeshDescriptionBuilder.cpp
@@ -8,6 +8,8 @@
 #include "DynamicMesh3.h"
 #include "DynamicMeshAttributeSet.h"
 
+#include "MeshNormals.h"
+
 namespace ExtendedMeshAttribute
 {
 	const FName PolyTriGroups("PolyTriGroups");
@@ -314,52 +316,6 @@ void FMeshDescriptionBuilder::SetAllEdgesHardness(bool bHard)
 
 
 
-void FMeshDescriptionBuilder::RecalculateInstanceNormals()
-{
-	for (int k = 0; k < InstanceNormals.GetNumElements(); ++k)
-	{
-		InstanceNormals.Set(FVertexInstanceID(k), 0, FVector::ZeroVector);
-	}
-
-	const FPolygonArray& Polygons = MeshDescription->Polygons();
-	for (const FPolygonID PolygonID : Polygons.GetElementIDs())
-	{
-		const TArray<FTriangleID>& TriangleIDs = MeshDescription->GetPolygonTriangleIDs(PolygonID);
-		for (const FTriangleID TriangleID : TriangleIDs)
-		{
-			TArrayView<const FVertexInstanceID> TriangleVertexInstanceIDs = MeshDescription->GetTriangleVertexInstances(TriangleID);
-			FVector3d A = GetPosition(TriangleVertexInstanceIDs[0]);
-			FVector3d B = GetPosition(TriangleVertexInstanceIDs[1]);
-			FVector3d C = GetPosition(TriangleVertexInstanceIDs[2]);
-			double Area = 1.0;
-			FVector FaceNormal = VectorUtil::NormalArea(A, B, C, Area);
-			if (Area > FMathf::ZeroTolerance)
-			{
-				FaceNormal *= Area;
-				InstanceNormals.Set(TriangleVertexInstanceIDs[0],
-					InstanceNormals.Get(TriangleVertexInstanceIDs[0], 0) + FaceNormal);
-				InstanceNormals.Set(TriangleVertexInstanceIDs[1],
-					InstanceNormals.Get(TriangleVertexInstanceIDs[1], 0) + FaceNormal);
-				InstanceNormals.Set(TriangleVertexInstanceIDs[2],
-					InstanceNormals.Get(TriangleVertexInstanceIDs[2], 0) + FaceNormal);
-			}
-		}
-	}
-
-
-	for (int k = 0; k < InstanceNormals.GetNumElements(); ++k)
-	{
-		FVector SumNormal = InstanceNormals.Get(FVertexInstanceID(k), 0);
-		SumNormal.Normalize();
-		if ( SumNormal.Size() < 0.99 )
-		{ 
-			SumNormal = FVector(1, 0, 0);
-		}
-		InstanceNormals.Set(FVertexInstanceID(k), 0, SumNormal);
-	}
-}
-
-
 
 FBox FMeshDescriptionBuilder::ComputeBoundingBox() const
 {
diff --git a/Engine/Plugins/Experimental/GeometryProcessing/Source/MeshConversion/Public/DynamicMeshToMeshDescription.h b/Engine/Plugins/Experimental/GeometryProcessing/Source/MeshConversion/Public/DynamicMeshToMeshDescription.h
--- a/Engine/Plugins/Experimental/GeometryProcessing/Source/MeshConversion/Public/DynamicMeshToMeshDescription.h
+++ b/Engine/Plugins/Experimental/GeometryProcessing/Source/MeshConversion/Public/DynamicMeshToMeshDescription.h
@@ -36,9 +36,9 @@ public:
 
 	/**
 	 * Update existing MeshDescription based on DynamicMesh. Assumes mesh topology has not changed.
-	 * Copies positions, recalculates MeshDescription normals if bRecomputeNormals is true (copies otherwise).
+	 * Copies positions and normals; does not update UVs
 	 */
-	void Update(const FDynamicMesh3* MeshIn, FMeshDescription& MeshOut, bool bRecomputeNormals = true);
+	void Update(const FDynamicMesh3* MeshIn, FMeshDescription& MeshOut);
 
 
 	/**
diff --git a/Engine/Plugins/Experimental/GeometryProcessing/Source/MeshConversion/Public/MeshConversionOptions.h b/Engine/Plugins/Experimental/GeometryProcessing/Source/MeshConversion/Public/MeshConversionOptions.h
--- a/Engine/Plugins/Experimental/GeometryProcessing/Source/MeshConversion/Public/MeshConversionOptions.h
+++ b/Engine/Plugins/Experimental/GeometryProcessing/Source/MeshConversion/Public/MeshConversionOptions.h
@@ -10,6 +10,4 @@ struct MESHCONVERSION_API FConversionToMeshDescriptionOptions
 public:
 	/** Should triangle groups be transfered to MeshDescription via custom PolyTriGroups attribute */
 	bool bSetPolyGroups = true;
-
-	// TODO: add additional conversion options (e.g. normal / tangent recompute, etc)
 };
\ No newline at end of file
diff --git a/Engine/Plugins/Experimental/GeometryProcessing/Source/MeshConversion/Public/MeshDescriptionBuilder.h b/Engine/Plugins/Experimental/GeometryProcessing/Source/MeshConversion/Public/MeshDescriptionBuilder.h
--- a/Engine/Plugins/Experimental/GeometryProcessing/Source/MeshConversion/Public/MeshDescriptionBuilder.h
+++ b/Engine/Plugins/Experimental/GeometryProcessing/Source/MeshConversion/Public/MeshDescriptionBuilder.h
@@ -96,8 +96,6 @@ public:
 	/** Translate the MeshDescription vertex positions */
 	void Translate(const FVector& Translation);
 
-	/** Calculate normals on the MeshDescription vertex instances, using area-weighted face average */
-	void RecalculateInstanceNormals();
 
 	/** Return the current bounding box of the mesh */
 	FBox ComputeBoundingBox() const;
