GeometryProcessing: process Octree spill-set objects in ray-intersection and range queries
#rnx

diff --git a/Engine/Plugins/Experimental/GeometryProcessing/Source/GeometricObjects/Public/Spatial/SparseDynamicOctree3.h b/Engine/Plugins/Experimental/GeometryProcessing/Source/GeometricObjects/Public/Spatial/SparseDynamicOctree3.h
--- a/Engine/Plugins/Experimental/GeometryProcessing/Source/GeometricObjects/Public/Spatial/SparseDynamicOctree3.h
+++ b/Engine/Plugins/Experimental/GeometryProcessing/Source/GeometricObjects/Public/Spatial/SparseDynamicOctree3.h
@@ -564,11 +564,23 @@ int32 FSparseDynamicOctree3::FindNearestHitObject(const FRay3d& Ray,
 {
 	// this should take advantage of raster!
 
+	// always test against all spill objects
 	int32 HitObjectID = -1;
+	for (int ObjectID : SpillObjectSet)
+	{
+		double HitDist = HitObjectDistFunc(ObjectID, Ray);
+		if (HitDist < MaxDistance)
+		{
+			MaxDistance = HitDist;
+			HitObjectID = ObjectID;
+		}
+	}
+
+	// we use queue instead of recursion
 	TArray<const FSparseOctreeCell*> Queue;
-	Queue.Reserve(32);
+	Queue.Reserve(64);
 
-	// start at root cells
+	// push all root cells onto queue if they are hit by ray
 	RootCells.AllocatedIteration([&](const uint32* RootCellID)
 	{
 		const FSparseOctreeCell* RootCell = &Cells[*RootCellID];
@@ -579,7 +591,7 @@ int32 FSparseDynamicOctree3::FindNearestHitObject(const FRay3d& Ray,
 		}
 	});
 
-
+	// test cells until the queue is empty
 	while (Queue.Num() > 0)
 	{
 		const FSparseOctreeCell* CurCell = Queue.Pop(false);
@@ -624,6 +636,12 @@ void FSparseDynamicOctree3::RangeQuery(
 {
 	// todo: this should take advantage of raster!
 
+	// always process spill objects
+	for (int ObjectID : SpillObjectSet)
+	{
+		ObjectIDFunc(ObjectID);
+	}
+
 	TArray<const FSparseOctreeCell*> Queue;
 
 	// start at root cells
@@ -673,6 +691,12 @@ void FSparseDynamicOctree3::RangeQuery(
 {
 	// todo: this should take advantage of raster!
 
+	// always collect spill objects
+	for (int ObjectID : SpillObjectSet)
+	{
+		ObjectIDs.Add(ObjectID);
+	}
+
 	TArray<const FSparseOctreeCell*> Queue;
 
 	// start at root cells
