update plugin descriptions
#rnx

diff --git a/Engine/Plugins/Experimental/GeometryProcessing/GeometryProcessing.uplugin b/Engine/Plugins/Experimental/GeometryProcessing/GeometryProcessing.uplugin
--- a/Engine/Plugins/Experimental/GeometryProcessing/GeometryProcessing.uplugin
+++ b/Engine/Plugins/Experimental/GeometryProcessing/GeometryProcessing.uplugin
@@ -2,9 +2,9 @@
 	"FileVersion": 3,
 	"Version": 1,
 	"VersionName": "0.1",
-	"FriendlyName": "GeometryProcessing",
-	"Description": "Classes and Algorithms for Processing 2D and 3D Geometry",
-	"Category": "Other",
+	"FriendlyName": "Geometry Processing",
+	"Description": "Data Structures and Algorithms for Processing 2D and 3D Geometry",
+	"Category": "Geometry",
 	"CreatedBy": "Epic Games, Inc.",
 	"CreatedByURL": "http://epicgames.com",
 	"DocsURL": "",
@@ -38,6 +38,6 @@
 			"Name": "MeshSolverUtilities",
 			"Type": "Runtime",
 			"LoadingPhase": "Default"
-		}			
+		}
 	]
 }
\ No newline at end of file
diff --git a/Engine/Plugins/Experimental/MeshModelingToolset/MeshModelingToolset.uplugin b/Engine/Plugins/Experimental/MeshModelingToolset/MeshModelingToolset.uplugin
--- a/Engine/Plugins/Experimental/MeshModelingToolset/MeshModelingToolset.uplugin
+++ b/Engine/Plugins/Experimental/MeshModelingToolset/MeshModelingToolset.uplugin
@@ -3,7 +3,7 @@
 	"Version": 1,
 	"VersionName": "0.1",
 	"FriendlyName": "Mesh Modeling Toolset",
-	"Description": "Adds Mesh Modeling Modules",
+	"Description": "A set of modules implementing 3D mesh creation and editing based on the Interactive Tools Framework",
 	"Category": "Other",
 	"CreatedBy": "Epic Games, Inc.",
 	"CreatedByURL": "http://epicgames.com",
diff --git a/Engine/Plugins/Experimental/ModelingToolsEditorMode/ModelingToolsEditorMode.uplugin b/Engine/Plugins/Experimental/ModelingToolsEditorMode/ModelingToolsEditorMode.uplugin
--- a/Engine/Plugins/Experimental/ModelingToolsEditorMode/ModelingToolsEditorMode.uplugin
+++ b/Engine/Plugins/Experimental/ModelingToolsEditorMode/ModelingToolsEditorMode.uplugin
@@ -2,9 +2,9 @@
 	"FileVersion": 3,
 	"Version": 1,
 	"VersionName": "0.1",
-	"FriendlyName": "Modeling Tools Mode",
-	"Description": "Adds Modeling Tools Editor Mode Tab",
-	"Category": "Other",
+	"FriendlyName": "Modeling Tools Editor Mode",
+	"Description": "Modeling Tools Mode includes a suite of interactive tools for creating and editing meshes in the Editor",
+	"Category": "Editor",
 	"CreatedBy": "Epic Games, Inc.",
 	"CreatedByURL": "http://epicgames.com",
 	"DocsURL": "",
diff --git a/Engine/Plugins/Experimental/SampleToolsEditorMode/SampleToolsEditorMode.uplugin b/Engine/Plugins/Experimental/SampleToolsEditorMode/SampleToolsEditorMode.uplugin
--- a/Engine/Plugins/Experimental/SampleToolsEditorMode/SampleToolsEditorMode.uplugin
+++ b/Engine/Plugins/Experimental/SampleToolsEditorMode/SampleToolsEditorMode.uplugin
@@ -2,9 +2,9 @@
 	"FileVersion": 3,
 	"Version": 1,
 	"VersionName": "0.1",
-	"FriendlyName": "Sample Tools Mode",
-	"Description": "Adds Sample Tools Editor Mode Tab",
-	"Category": "Other",
+	"FriendlyName": "Sample Tools Editor Mode",
+	"Description": "Sample Tools Mode includes a set of sample Tools demonstrating capabilities of the Interactive Tools Framework",
+	"Category": "Editor",
 	"CreatedBy": "Epic Games, Inc.",
 	"CreatedByURL": "http://epicgames.com",
 	"DocsURL": "",
