fix ambiguous type in conditional expression for dynamic mesh conversion (CIS fix)

#rnx

diff --git a/Engine/Plugins/Experimental/GeometryProcessing/Source/MeshConversion/Private/DynamicMeshToMeshDescription.cpp b/Engine/Plugins/Experimental/GeometryProcessing/Source/MeshConversion/Private/DynamicMeshToMeshDescription.cpp
--- a/Engine/Plugins/Experimental/GeometryProcessing/Source/MeshConversion/Private/DynamicMeshToMeshDescription.cpp
+++ b/Engine/Plugins/Experimental/GeometryProcessing/Source/MeshConversion/Private/DynamicMeshToMeshDescription.cpp
@@ -176,8 +176,8 @@ void FDynamicMeshToMeshDescription::Convert_NoAttributes(const FDynamicMesh3* Me
 			{
 				FVertexInstanceID NewInstanceID = Builder.AppendInstance(MapV[Triangle[j]]);
 				InstanceList.Add(InstanceElem, NewInstanceID);
-				FVector2D UV = MeshIn->HasVertexUVs() ? MeshIn->GetVertexUV(Triangle[j]) : FVector2D::ZeroVector;
-				FVector Normal = MeshIn->HasVertexNormals() ? MeshIn->GetVertexNormal(Triangle[j]) : FVector::UpVector;
+				FVector2D UV = MeshIn->HasVertexUVs() ? FVector2D(MeshIn->GetVertexUV(Triangle[j])) : FVector2D::ZeroVector;
+				FVector Normal = MeshIn->HasVertexNormals() ? FVector(MeshIn->GetVertexNormal(Triangle[j])) : FVector::UpVector;
 				Builder.SetInstance(NewInstanceID, UV, Normal);
 			}
 			InstanceTri[j] = InstanceList[InstanceElem];
