GeometryProcessing: bugfix, need to call insertAt instead of just writing to index in case array has not been resized yet
#rnx

diff --git a/Engine/Plugins/Experimental/GeometryProcessing/Source/DynamicMesh/Public/DynamicMeshTriangleAttribute.h b/Engine/Plugins/Experimental/GeometryProcessing/Source/DynamicMesh/Public/DynamicMeshTriangleAttribute.h
--- a/Engine/Plugins/Experimental/GeometryProcessing/Source/DynamicMesh/Public/DynamicMeshTriangleAttribute.h
+++ b/Engine/Plugins/Experimental/GeometryProcessing/Source/DynamicMesh/Public/DynamicMeshTriangleAttribute.h
@@ -124,7 +124,7 @@ public:
 		int kB = ToTriangleID * AttribDimension;
 		for (int i = 0; i < AttribDimension; ++i)
 		{
-			AttribValues[kB+i] = AttribValues[kA+i];
+			AttribValues.InsertAt(AttribValues[kA+i], kB+i);
 		}
 	}
 
