ModelingTools: improve UVs in DrawPolygonTool. Add UVTranslation to DynamicMeshEditor UV-projection functions. ExtrudeMesh uses this to shift UVs for vertical tiles to be adjacent. Now textures tile along span of extrude.
#rnx

diff --git a/Engine/Plugins/Experimental/GeometryProcessing/Source/DynamicMesh/Private/DynamicMeshEditor.cpp b/Engine/Plugins/Experimental/GeometryProcessing/Source/DynamicMesh/Private/DynamicMeshEditor.cpp
--- a/Engine/Plugins/Experimental/GeometryProcessing/Source/DynamicMesh/Private/DynamicMeshEditor.cpp
+++ b/Engine/Plugins/Experimental/GeometryProcessing/Source/DynamicMesh/Private/DynamicMeshEditor.cpp
@@ -465,7 +465,7 @@ void FDynamicMeshEditor::SetTriangleNormals(const TArray<int>& Triangles, const
 }
 
 
-void FDynamicMeshEditor::SetTriangleUVsFromProjection(const TArray<int>& Triangles, const FFrame3d& ProjectionFrame, float UVScaleFactor, int UVLayerIndex)
+void FDynamicMeshEditor::SetTriangleUVsFromProjection(const TArray<int>& Triangles, const FFrame3d& ProjectionFrame, float UVScaleFactor, const FVector2f& UVTranslation, int UVLayerIndex)
 {
 	if (!Triangles.Num())
 	{
@@ -508,12 +508,13 @@ void FDynamicMeshEditor::SetTriangleUVsFromProjection(const TArray<int>& Triangl
 	{
 		FVector2f UV = UVs->GetElement(UVID);
 		FVector2f TransformedUV = (UV - UVBounds.Min) * UVScaleFactor;
+		TransformedUV += UVTranslation;
 		UVs->SetElement(UVID, TransformedUV);
 	}
 }
 
 
-void FDynamicMeshEditor::SetQuadUVsFromProjection(const FIndex2i& QuadTris, const FFrame3d& ProjectionFrame, float UVScaleFactor, int UVLayerIndex)
+void FDynamicMeshEditor::SetQuadUVsFromProjection(const FIndex2i& QuadTris, const FFrame3d& ProjectionFrame, float UVScaleFactor, const FVector2f& UVTranslation, int UVLayerIndex)
 {
 	check(Mesh->HasAttributes() && Mesh->Attributes()->NumUVLayers() > UVLayerIndex );
 	FDynamicMeshUVOverlay* UVs = Mesh->Attributes()->GetUVLayer(UVLayerIndex);
@@ -568,6 +569,7 @@ void FDynamicMeshEditor::SetQuadUVsFromProjection(const FIndex2i& QuadTris, cons
 		if (AllUVIndices[j] != -1)
 		{
 			FVector2f TransformedUV = (AllUVs[j] - UVBounds.Min) * UVScaleFactor;
+			TransformedUV += UVTranslation;
 			UVs->SetElement(AllUVIndices[j], TransformedUV);
 		}
 	}
diff --git a/Engine/Plugins/Experimental/GeometryProcessing/Source/DynamicMesh/Private/Operations/ExtrudeMesh.cpp b/Engine/Plugins/Experimental/GeometryProcessing/Source/DynamicMesh/Private/Operations/ExtrudeMesh.cpp
--- a/Engine/Plugins/Experimental/GeometryProcessing/Source/DynamicMesh/Private/Operations/ExtrudeMesh.cpp
+++ b/Engine/Plugins/Experimental/GeometryProcessing/Source/DynamicMesh/Private/Operations/ExtrudeMesh.cpp
@@ -101,23 +101,42 @@ bool FExtrudeMesh::Apply()
 		// for each polygon we created in stitch, set UVs and normals
 		if (Mesh->HasAttributes())
 		{
+			float AccumUVTranslation = 0;
+
+			FFrame3d FirstProjectFrame;
+			FVector3d FrameUp;
+
 			int NumNewQuads = StitchResult.NewQuads.Num();
 			for (int k = 0; k < NumNewQuads; k++)
 			{
 				FVector3f Normal = Editor.ComputeAndSetQuadNormal(StitchResult.NewQuads[k], true);
 
-				// @todo is there a simpler way to construct rotation from 3 known axes (third id Normal.Cross(UnitY))
-				//  (converting from matrix might end up being more efficient due to trig in ConstrainedAlignAxis?)
-				FFrame3d ProjectFrame(FVector3d::Zero(), (FVector3d)Normal);
-				if (FMathd::Abs(ProjectFrame.Y().Dot(FVector3d::UnitY())) < 0.01)
+				// align axis 0 of projection frame to first edge, then for further edges,
+				// rotate around 'up' axis to keep normal aligned and frame horizontal
+				FFrame3d ProjectFrame;
+				if (k == 0)
 				{
-					ProjectFrame.ConstrainedAlignAxis(0, FVector3d::UnitX(), ProjectFrame.Z());
+					FVector3d FirstEdge = Mesh->GetVertex(BaseLoop.Vertices[1]) - Mesh->GetVertex(BaseLoop.Vertices[0]);
+					FirstEdge.Normalize();
+					FirstProjectFrame = FFrame3d(FVector3d::Zero(), (FVector3d)Normal);
+					FirstProjectFrame.ConstrainedAlignAxis(0, FirstEdge, (FVector3d)Normal);
+					FrameUp = FirstProjectFrame.GetAxis(1);
+					ProjectFrame = FirstProjectFrame;
 				}
 				else
 				{
-					ProjectFrame.ConstrainedAlignAxis(1, FVector3d::UnitY(), ProjectFrame.Z());
+					ProjectFrame = FirstProjectFrame;
+					ProjectFrame.ConstrainedAlignAxis(2, (FVector3d)Normal, FrameUp);
+				}
+
+				if (k > 0)
+				{
+					AccumUVTranslation += Mesh->GetVertex(BaseLoop.Vertices[k]).Distance(Mesh->GetVertex(BaseLoop.Vertices[k-1]));
 				}
-				Editor.SetQuadUVsFromProjection(StitchResult.NewQuads[k], ProjectFrame, UVScaleFactor);
+
+				// translate horizontally such that vertical spans are adjacent in UV space (so textures tile/wrap properly)
+				float TranslateU = UVScaleFactor*AccumUVTranslation;
+				Editor.SetQuadUVsFromProjection(StitchResult.NewQuads[k], ProjectFrame, UVScaleFactor, FVector2f(TranslateU, 0) );
 			}
 		}
 
diff --git a/Engine/Plugins/Experimental/GeometryProcessing/Source/DynamicMesh/Public/DynamicMeshEditor.h b/Engine/Plugins/Experimental/GeometryProcessing/Source/DynamicMesh/Public/DynamicMeshEditor.h
--- a/Engine/Plugins/Experimental/GeometryProcessing/Source/DynamicMesh/Public/DynamicMeshEditor.h
+++ b/Engine/Plugins/Experimental/GeometryProcessing/Source/DynamicMesh/Public/DynamicMeshEditor.h
@@ -265,20 +265,22 @@ public:
 	 * UVs are translated so that their bbox min-corner is at origin, and scaled by given scale factor
 	 * @param QuadTris pair of triangle IDs. If second ID is invalid, it is ignored
 	 * @param ProjectFrame vertices are projected into XY axes of this frame
-	 * @param UVScaleFactor UVs are scaled
+	 * @param UVScaleFactor UVs are scaled by this uniform scale factor
+	 * @param UVTranslation UVs are translated after scaling
 	 * @param UVLayerIndex which UV layer to operate on (must exist)
 	 */
-	void SetQuadUVsFromProjection(const FIndex2i& QuadTris, const FFrame3d& ProjectionFrame, float UVScaleFactor = 1.0f, int UVLayerIndex = 0);
+	void SetQuadUVsFromProjection(const FIndex2i& QuadTris, const FFrame3d& ProjectionFrame, float UVScaleFactor = 1.0f, const FVector2f& UVTranslation = FVector2f::Zero(), int UVLayerIndex = 0);
 
 	/**
 	* Project triangles onto a plane defined by the ProjectionFrame and use that to create/set new shared per-triangle UVs.
 	* UVs are translated so that their bbox min-corner is at origin, and scaled by given scale factor
 	* @param Triangles TArray of triangle IDs
 	* @param ProjectFrame vertices are projected into XY axes of this frame
-	* @param UVScaleFactor UVs are scaled
+	* @param UVScaleFactor UVs are scaled by this uniform scale factor
+	* @param UVTranslation UVs are translated after scaling
 	* @param UVLayerIndex which UV layer to operate on (must exist)
 	*/
-	void SetTriangleUVsFromProjection(const TArray<int>& Triangles, const FFrame3d& ProjectionFrame, float UVScaleFactor, int UVLayerIndex = 0);
+	void SetTriangleUVsFromProjection(const TArray<int>& Triangles, const FFrame3d& ProjectionFrame, float UVScaleFactor = 1.0f, const FVector2f& UVTranslation = FVector2f::Zero(), int UVLayerIndex = 0);
 
 
 	/**
diff --git a/Engine/Plugins/Experimental/MeshModelingToolset/Source/MeshModelingTools/Private/DrawPolygonTool.cpp b/Engine/Plugins/Experimental/MeshModelingToolset/Source/MeshModelingTools/Private/DrawPolygonTool.cpp
--- a/Engine/Plugins/Experimental/MeshModelingToolset/Source/MeshModelingTools/Private/DrawPolygonTool.cpp
+++ b/Engine/Plugins/Experimental/MeshModelingToolset/Source/MeshModelingTools/Private/DrawPolygonTool.cpp
@@ -131,7 +131,7 @@ void UDrawPolygonTool::Setup()
 	MaterialProperties = NewObject<UNewMeshMaterialProperties>(this);
 
 	// create preview mesh object
-	PreviewMesh = NewObject<UPreviewMesh>(this, TEXT("PreviewMesh"));
+	PreviewMesh = NewObject<UPreviewMesh>(this, TEXT("DrawPolygonPreviewMesh"));
 	PreviewMesh->CreateInWorld(this->TargetWorld, FTransform::Identity);
 	PreviewMesh->SetVisible(false);
 	PreviewMesh->SetMaterial(MaterialProperties->Material);
@@ -875,6 +875,9 @@ void UDrawPolygonTool::UpdateDrawPlaneFromGizmo(const FFrame3d& WorldPosition)
 
 void UDrawPolygonTool::EmitCurrentPolygon()
 {
+	FString BaseName = (PolygonProperties->OutputMode == EDrawPolygonOutputMode::MeshedPolygon) ?
+		TEXT("Polygon") : TEXT("Extrude");
+
 #if WITH_EDITOR
 	// generate new mesh
 	FFrame3d PlaneFrameOut;
@@ -883,12 +886,11 @@ void UDrawPolygonTool::EmitCurrentPolygon()
 		0 : PolygonProperties->ExtrudeHeight;
 	GeneratePolygonMesh(PolygonVertices, &Mesh, PlaneFrameOut, false, ExtrudeDist, false);
 
-
 	GetToolManager()->BeginUndoTransaction(LOCTEXT("CreatePolygon", "Create Polygon"));
 
 	AActor* NewActor = AssetGenerationUtil::GenerateStaticMeshActor(
 		AssetAPI, TargetWorld,
-		&Mesh, PlaneFrameOut.ToTransform(), TEXT("Polygon"),
+		&Mesh, PlaneFrameOut.ToTransform(), BaseName,
 		AssetGenerationUtil::GetDefaultAutoGeneratedAssetPath(),
 		MaterialProperties->Material);
 
