add lengthwise divisions to cylinder (and cone, which is just cylinder with a small top) to support cones with less terrible-looking uvs in addprimitive tool

#rnx

diff --git a/Engine/Plugins/Experimental/GeometryProcessing/Source/DynamicMesh/Public/Generators/SweepGenerator.h b/Engine/Plugins/Experimental/GeometryProcessing/Source/DynamicMesh/Public/Generators/SweepGenerator.h
--- a/Engine/Plugins/Experimental/GeometryProcessing/Source/DynamicMesh/Public/Generators/SweepGenerator.h
+++ b/Engine/Plugins/Experimental/GeometryProcessing/Source/DynamicMesh/Public/Generators/SweepGenerator.h
@@ -246,6 +246,7 @@ public:
 	float Radius[2] = {1.0f, 1.0f};
 	float Height = 1.0f;
 	int AngleSamples = 16;
+	int LengthSamples = 4;
 	bool bCapped = false;
 	bool bUVScaleMatchSidesAndCaps = true;
 
@@ -264,19 +265,31 @@ public:
 			Caps[0] = ECapType::FlatTriangulation;
 			Caps[1] = ECapType::FlatTriangulation;
 		}
-		ConstructMeshTopology(X, {}, {}, 2, Caps, FVector2d(.5, .5), FVector2d(.5, .5));
 
+		int NumX = LengthSamples + 2;
+		ConstructMeshTopology(X, {}, {}, NumX, Caps, FVector2d(.5, .5), FVector2d(.5, .5));
+
+		// set vertex positions and normals for all cross sections along length
+		double LengthFactor = 1.0 / double(NumX-1);
 		for (int SubIdx = 0; SubIdx < X.VertexCount(); SubIdx++)
 		{
-			for (int XIdx = 0; XIdx < 2; ++XIdx)
+			for (int XIdx = 0; XIdx < NumX; ++XIdx)
 			{
+				double Along = XIdx * LengthFactor;
+				double AlongRadius = FMath::Lerp(Radius[0], Radius[1], Along);
 				Vertices[SubIdx + XIdx * AngleSamples] =
-					FVector3d(XVerts[SubIdx].X * Radius[XIdx], XVerts[SubIdx].Y * Radius[XIdx], XIdx * Height);
+					FVector3d(XVerts[SubIdx].X * AlongRadius, XVerts[SubIdx].Y * AlongRadius, Height * Along);
 				Normals[SubIdx + XIdx * AngleSamples] = FVector3f(XVerts[SubIdx].X*NormalSide.X, XVerts[SubIdx].Y*NormalSide.X, NormalSide.Y);
-
-				if (bCapped)
+			}
+		}
+		// if capped, set top/bottom normals
+		if (bCapped)
+		{
+			for (int SubIdx = 0; SubIdx < X.VertexCount(); SubIdx++)
+			{
+				for (int XBotTop = 0; XBotTop < 2; ++XBotTop)
 				{
-					Normals[CapNormalStart[XIdx] + SubIdx] = FVector3f(0, 0, 2 * XIdx - 1);
+					Normals[CapNormalStart[XBotTop] + SubIdx] = FVector3f(0, 0, 2 * XBotTop - 1);
 				}
 			}
 		}
diff --git a/Engine/Plugins/Experimental/MeshModelingToolset/Source/MeshModelingTools/Private/AddPrimitiveTool.cpp b/Engine/Plugins/Experimental/MeshModelingToolset/Source/MeshModelingTools/Private/AddPrimitiveTool.cpp
--- a/Engine/Plugins/Experimental/MeshModelingToolset/Source/MeshModelingTools/Private/AddPrimitiveTool.cpp
+++ b/Engine/Plugins/Experimental/MeshModelingToolset/Source/MeshModelingTools/Private/AddPrimitiveTool.cpp
@@ -67,7 +67,7 @@ namespace
 	  { TEXT("PlaceMode"),     EMakeMeshShapeType::All },
 	  { TEXT("PivotLocation"), EMakeMeshShapeType::All },
 	  { TEXT("Slices"),        EMakeMeshShapeType::Cylinder | EMakeMeshShapeType::Cone | EMakeMeshShapeType::Sphere },
-	  { TEXT("Subdivisions"),  EMakeMeshShapeType::Box | EMakeMeshShapeType::Plane }
+	  { TEXT("Subdivisions"),  EMakeMeshShapeType::Box | EMakeMeshShapeType::Plane | EMakeMeshShapeType::Cylinder | EMakeMeshShapeType::Cone }
 	};
 };
 
@@ -332,6 +332,7 @@ void UAddPrimitiveTool::GenerateCylinder(FDynamicMesh3* OutMesh)
 	CylGen.Radius[1] = CylGen.Radius[0];
 	CylGen.Height = ShapeSettings->Height;
 	CylGen.AngleSamples = ShapeSettings->Slices;
+	CylGen.LengthSamples = ShapeSettings->Subdivisions;
 	CylGen.bCapped = true;
 	CylGen.Generate();
 	OutMesh->Copy(&CylGen);
@@ -346,6 +347,7 @@ void UAddPrimitiveTool::GenerateCone(FDynamicMesh3* OutMesh)
 	CylGen.Radius[1] = .01;
 	CylGen.Height = ShapeSettings->Height;
 	CylGen.AngleSamples = ShapeSettings->Slices;
+	CylGen.LengthSamples = ShapeSettings->Subdivisions;
 	CylGen.bCapped = true;
 	CylGen.Generate();
 	OutMesh->Copy(&CylGen);
