setup multi-uv layer machinery for dmesh attribute set (does not actually hook up conversions yet, so should not have any effect in practice, but is prep for that)

#rnx

diff --git a/Engine/Plugins/Experimental/GeometryProcessing/Source/DynamicMesh/Private/DynamicMeshAttributeSet.cpp b/Engine/Plugins/Experimental/GeometryProcessing/Source/DynamicMesh/Private/DynamicMeshAttributeSet.cpp
--- a/Engine/Plugins/Experimental/GeometryProcessing/Source/DynamicMesh/Private/DynamicMeshAttributeSet.cpp
+++ b/Engine/Plugins/Experimental/GeometryProcessing/Source/DynamicMesh/Private/DynamicMeshAttributeSet.cpp
@@ -18,14 +18,24 @@ void FDynamicMeshAttributeSet::EnableMaterialID()
 
 bool FDynamicMeshAttributeSet::IsSeamEdge(int eid) const
 {
-	return UV0.IsSeamEdge(eid) || Normals0.IsSeamEdge(eid);
+	for (const FDynamicMeshUVOverlay& UVLayer : UVLayers)
+	{
+		if (UVLayer.IsSeamEdge(eid))
+		{
+			return true;
+		}
+	}
+	return Normals0.IsSeamEdge(eid);
 }
 
 
 
 void FDynamicMeshAttributeSet::OnNewTriangle(int TriangleID, bool bInserted)
 {
-	UV0.InitializeNewTriangle(TriangleID);
+	for (FDynamicMeshUVOverlay& UVLayer : UVLayers)
+	{
+		UVLayer.InitializeNewTriangle(TriangleID);
+	}
 	Normals0.InitializeNewTriangle(TriangleID);
 
 	if (MaterialIDAttrib)
@@ -37,7 +47,10 @@ void FDynamicMeshAttributeSet::OnNewTriangle(int TriangleID, bool bInserted)
 
 void FDynamicMeshAttributeSet::OnRemoveTriangle(int TriangleID, bool bRemoveIsolatedVertices)
 {
-	UV0.OnRemoveTriangle(TriangleID, bRemoveIsolatedVertices);
+	for (FDynamicMeshUVOverlay UVLayer : UVLayers)
+	{
+		UVLayer.OnRemoveTriangle(TriangleID, bRemoveIsolatedVertices);
+	}
 	Normals0.OnRemoveTriangle(TriangleID, bRemoveIsolatedVertices);
 
 	// has no effect on MaterialIDAttrib
@@ -45,7 +58,10 @@ void FDynamicMeshAttributeSet::OnRemoveTriangle(int TriangleID, bool bRemoveIsol
 
 void FDynamicMeshAttributeSet::OnReverseTriOrientation(int TriangleID)
 {
-	UV0.OnReverseTriOrientation(TriangleID);
+	for (FDynamicMeshUVOverlay& UVLayer : UVLayers)
+	{
+		UVLayer.OnReverseTriOrientation(TriangleID);
+	}
 	Normals0.OnReverseTriOrientation(TriangleID);
 
 	// has no effect on MaterialIDAttrib
@@ -53,7 +69,10 @@ void FDynamicMeshAttributeSet::OnReverseTriOrientation(int TriangleID)
 
 void FDynamicMeshAttributeSet::OnSplitEdge(const FDynamicMesh3::FEdgeSplitInfo & splitInfo)
 {
-	UV0.OnSplitEdge(splitInfo);
+	for (FDynamicMeshUVOverlay& UVLayer : UVLayers)
+	{
+		UVLayer.OnSplitEdge(splitInfo);
+	}
 	Normals0.OnSplitEdge(splitInfo);
 
 	if (MaterialIDAttrib)
@@ -64,7 +83,10 @@ void FDynamicMeshAttributeSet::OnSplitEdge(const FDynamicMesh3::FEdgeSplitInfo &
 
 void FDynamicMeshAttributeSet::OnFlipEdge(const FDynamicMesh3::FEdgeFlipInfo & flipInfo)
 {
-	UV0.OnFlipEdge(flipInfo);
+	for (FDynamicMeshUVOverlay& UVLayer : UVLayers)
+	{
+		UVLayer.OnFlipEdge(flipInfo);
+	}
 	Normals0.OnFlipEdge(flipInfo);
 
 	if (MaterialIDAttrib)
@@ -76,7 +98,10 @@ void FDynamicMeshAttributeSet::OnFlipEdge(const FDynamicMesh3::FEdgeFlipInfo & f
 
 void FDynamicMeshAttributeSet::OnCollapseEdge(const FDynamicMesh3::FEdgeCollapseInfo & collapseInfo)
 {
-	UV0.OnCollapseEdge(collapseInfo);
+	for (FDynamicMeshUVOverlay& UVLayer : UVLayers)
+	{
+		UVLayer.OnCollapseEdge(collapseInfo);
+	}
 	Normals0.OnCollapseEdge(collapseInfo);
 
 	if (MaterialIDAttrib)
@@ -87,7 +112,10 @@ void FDynamicMeshAttributeSet::OnCollapseEdge(const FDynamicMesh3::FEdgeCollapse
 
 void FDynamicMeshAttributeSet::OnPokeTriangle(const FDynamicMesh3::FPokeTriangleInfo & pokeInfo)
 {
-	UV0.OnPokeTriangle(pokeInfo);
+	for (FDynamicMeshUVOverlay& UVLayer : UVLayers)
+	{
+		UVLayer.OnPokeTriangle(pokeInfo);
+	}
 	Normals0.OnPokeTriangle(pokeInfo);
 
 	if (MaterialIDAttrib)
@@ -98,7 +126,10 @@ void FDynamicMeshAttributeSet::OnPokeTriangle(const FDynamicMesh3::FPokeTriangle
 
 void FDynamicMeshAttributeSet::OnMergeEdges(const FDynamicMesh3::FMergeEdgesInfo & mergeInfo)
 {
-	UV0.OnMergeEdges(mergeInfo);
+	for (FDynamicMeshUVOverlay& UVLayer : UVLayers)
+	{
+		UVLayer.OnMergeEdges(mergeInfo);
+	}
 	Normals0.OnMergeEdges(mergeInfo);
 
 	if (MaterialIDAttrib)
diff --git a/Engine/Plugins/Experimental/GeometryProcessing/Source/DynamicMesh/Private/DynamicMeshChangeTracker.cpp b/Engine/Plugins/Experimental/GeometryProcessing/Source/DynamicMesh/Private/DynamicMeshChangeTracker.cpp
--- a/Engine/Plugins/Experimental/GeometryProcessing/Source/DynamicMesh/Private/DynamicMeshChangeTracker.cpp
+++ b/Engine/Plugins/Experimental/GeometryProcessing/Source/DynamicMesh/Private/DynamicMeshChangeTracker.cpp
@@ -353,9 +353,9 @@ void FDynamicMeshAttributeSetChangeTracker::BeginChange()
 	// initialize new attribute set change
 	check(Change == nullptr);
 	Change = new FDynamicMeshAttributeChangeSet();
-	int NumUVLayers = Attribs->GetNumUVLayers();
+	int NumUVLayers = Attribs->NumUVLayers();
 	Change->UVChanges.SetNum(NumUVLayers);
-	int NumNormalLayers = Attribs->GetNumNormalLayers();
+	int NumNormalLayers = Attribs->NumNormalLayers();
 	Change->NormalChanges.SetNum(NumNormalLayers);
 
 	// initialize UV layer state tracking
@@ -406,7 +406,7 @@ TUniquePtr<FDynamicMeshAttributeChangeSet> FDynamicMeshAttributeSetChangeTracker
 
 void FDynamicMeshAttributeSetChangeTracker::SaveInitialTriangle(int TriangleID)
 {
-	int NumUVLayers = Attribs->GetNumUVLayers();
+	int NumUVLayers = Attribs->NumUVLayers();
 	for (int k = 0; k < NumUVLayers; ++k)
 	{
 		FDynamicMeshUVChange& UVChange = Change->UVChanges[k];
@@ -418,7 +418,7 @@ void FDynamicMeshAttributeSetChangeTracker::SaveInitialTriangle(int TriangleID)
 		UVChange.SaveInitialTriangle(UVLayer, TriangleID);
 	}
 
-	int NumNormalLayers = Attribs->GetNumNormalLayers();
+	int NumNormalLayers = Attribs->NumNormalLayers();
 	for (int k = 0; k < NumNormalLayers; ++k)
 	{
 		FDynamicMeshNormalChange& NormalChange = Change->NormalChanges[k];
@@ -437,7 +437,7 @@ void FDynamicMeshAttributeSetChangeTracker::StoreAllFinalTriangles(const TArray<
 	TSet<int> StoredVertices;	// re-used
 	
 	// store final UV elements for all modified triangles, and final triangles
-	int NumUVLayers = Attribs->GetNumUVLayers();
+	int NumUVLayers = Attribs->NumUVLayers();
 	for (int k = 0; k < NumUVLayers; ++k)
 	{
 		FDynamicMeshUVChange& UVChange = Change->UVChanges[k];
@@ -461,7 +461,7 @@ void FDynamicMeshAttributeSetChangeTracker::StoreAllFinalTriangles(const TArray<
 
 
 	// store final Normal elements for all modified triangles, and final triangles
-	int NumNormalLayers = Attribs->GetNumNormalLayers();
+	int NumNormalLayers = Attribs->NumNormalLayers();
 	for (int k = 0; k < NumNormalLayers; ++k)
 	{
 		FDynamicMeshNormalChange& NormalChange = Change->NormalChanges[k];
@@ -488,7 +488,7 @@ void FDynamicMeshAttributeSetChangeTracker::StoreAllFinalTriangles(const TArray<
 
 bool FDynamicMeshAttributeChangeSet::Apply(FDynamicMeshAttributeSet* Attributes, bool bRevert) const
 {
-	int NumUVLayers = Attributes->GetNumUVLayers();
+	int NumUVLayers = Attributes->NumUVLayers();
 	check(NumUVLayers == UVChanges.Num());
 	for (int k = 0; k < NumUVLayers; ++k)
 	{
@@ -496,7 +496,7 @@ bool FDynamicMeshAttributeChangeSet::Apply(FDynamicMeshAttributeSet* Attributes,
 		UVChanges[k].Apply(UVLayer, bRevert);
 	}
 
-	int NumNormalLayers = Attributes->GetNumNormalLayers();
+	int NumNormalLayers = Attributes->NumNormalLayers();
 	check(NumNormalLayers == NormalChanges.Num());
 	for (int k = 0; k < NumNormalLayers; ++k)
 	{
diff --git a/Engine/Plugins/Experimental/GeometryProcessing/Source/DynamicMesh/Private/DynamicMeshEditor.cpp b/Engine/Plugins/Experimental/GeometryProcessing/Source/DynamicMesh/Private/DynamicMeshEditor.cpp
--- a/Engine/Plugins/Experimental/GeometryProcessing/Source/DynamicMesh/Private/DynamicMeshEditor.cpp
+++ b/Engine/Plugins/Experimental/GeometryProcessing/Source/DynamicMesh/Private/DynamicMeshEditor.cpp
@@ -691,9 +691,10 @@ void FDynamicMeshEditor::CopyAttributes(int FromTriangleID, int ToTriangleID, FM
 		return;
 	}
 
-	int UVLayerIndex = 0;
-	for (FDynamicMeshUVOverlay* UVOverlay : Mesh->Attributes()->GetAllUVLayers())
+	
+	for (int UVLayerIndex = 0; UVLayerIndex < Mesh->Attributes()->NumUVLayers(); UVLayerIndex++)
 	{
+		FDynamicMeshUVOverlay* UVOverlay = Mesh->Attributes()->GetUVLayer(UVLayerIndex);
 		FIndex3i FromElemTri = UVOverlay->GetTriangle(FromTriangleID);
 		FIndex3i ToElemTri = UVOverlay->GetTriangle(ToTriangleID);
 		for (int j = 0; j < 3; ++j)
@@ -705,7 +706,6 @@ void FDynamicMeshEditor::CopyAttributes(int FromTriangleID, int ToTriangleID, FM
 			}
 		}
 		UVOverlay->SetTriangle(ToTriangleID, ToElemTri);
-		UVLayerIndex++;
 	}
 
 
@@ -1030,9 +1030,8 @@ static void AppendAttributes(const FDynamicMesh3* FromMesh, int FromTriangleID,
 		return;
 	}
 
-	// todo: copy all attribute layers
-	check(FromMesh->Attributes()->GetNumUVLayers() == 1);
-	check(FromMesh->Attributes()->GetNumNormalLayers() == 1);
+	// todo: if we ever support multiple normal layers, copy them all
+	check(FromMesh->Attributes()->NumNormalLayers() == 1);
 
 	const FDynamicMeshUVOverlay* FromUVOverlay = FromMesh->Attributes()->PrimaryUV();
 	FDynamicMeshUVOverlay* ToUVOverlay = ToMesh->Attributes()->PrimaryUV();
diff --git a/Engine/Plugins/Experimental/GeometryProcessing/Source/DynamicMesh/Public/DynamicMeshAttributeSet.h b/Engine/Plugins/Experimental/GeometryProcessing/Source/DynamicMesh/Public/DynamicMeshAttributeSet.h
--- a/Engine/Plugins/Experimental/GeometryProcessing/Source/DynamicMesh/Public/DynamicMeshAttributeSet.h
+++ b/Engine/Plugins/Experimental/GeometryProcessing/Source/DynamicMesh/Public/DynamicMeshAttributeSet.h
@@ -27,9 +27,9 @@ class DYNAMICMESH_API FDynamicMeshAttributeSet
 public:
 
 	FDynamicMeshAttributeSet(FDynamicMesh3* Mesh) 
-		: ParentMesh(Mesh), UV0(Mesh), Normals0(Mesh)
+		: ParentMesh(Mesh), Normals0(Mesh)
 	{
-		UVLayers.Add(&UV0);
+		SetNumUVLayers(1);
 		NormalLayers.Add(&Normals0);
 	}
 
@@ -39,7 +39,11 @@ public:
 
 	void Copy(const FDynamicMeshAttributeSet& Copy)
 	{
-		UV0.Copy(Copy.UV0);
+		SetNumUVLayers(Copy.NumUVLayers());
+		for (int UVIdx = 0; UVIdx < NumUVLayers(); UVIdx++)
+		{
+			UVLayers[UVIdx].Copy(Copy.UVLayers[UVIdx]);
+		}
 		Normals0.Copy(Copy.Normals0);
 
 		if (Copy.MaterialIDAttrib)
@@ -60,12 +64,33 @@ public:
 	/** @return number of UV layers */
 	virtual int NumUVLayers() const 
 	{
-		return 1;
+		return UVLayers.Num();
+	}
+
+	virtual void SetNumUVLayers(int Num)
+	{
+		if (UVLayers.Num() == Num)
+		{
+			return;
+		}
+		if (Num >= UVLayers.Num())
+		{
+			for (int i = (int)UVLayers.Num(); i < Num; ++i)
+			{
+				UVLayers.Add(new FDynamicMeshUVOverlay(ParentMesh));
+			}
+		}
+		else
+		{
+			UVLayers.RemoveAt(Num, UVLayers.Num() - Num);
+		}
+		check(UVLayers.Num() == Num);
 	}
 
 	/** @return number of Normals layers */
 	virtual int NumNormalLayers() const
 	{
+		checkSlow(NormalLayers.Num() == 1);
 		return 1;
 	}
 
@@ -77,39 +102,27 @@ public:
 	// UV Layers 
 	//
 
-	/** @return The number of UV layers in this attribute set */
-	int GetNumUVLayers() const
-	{
-		return UVLayers.Num();
-	}
-
 	/** @return the UV layer at the given Index */
 	FDynamicMeshUVOverlay* GetUVLayer(int Index)
 	{
-		return (Index == 0) ? &UV0 : nullptr;
+		return &UVLayers[Index];
 	}
 
 	/** @return the UV layer at the given Index */
 	const FDynamicMeshUVOverlay* GetUVLayer(int Index) const
 	{
-		return (Index == 0) ? &UV0 : nullptr;
-	}
-
-	/** @return list of all UV layers */
-	const TArray<FDynamicMeshUVOverlay*>& GetAllUVLayers() const
-	{
-		return UVLayers;
+		return &UVLayers[Index];
 	}
 
 	/** @return the primary UV layer (layer 0) */
 	FDynamicMeshUVOverlay* PrimaryUV() 
 	{
-		return &UV0;
+		return &UVLayers[0];
 	}
 	/** @return the primary UV layer (layer 0) */
 	const FDynamicMeshUVOverlay* PrimaryUV() const
 	{
-		return &UV0;
+		return &UVLayers[0];
 	}
 
 
@@ -117,12 +130,6 @@ public:
 	// Normal Layers 
 	//
 
-	/** @return The number of normal layers in this attribute set */
-	int GetNumNormalLayers() const
-	{
-		return NormalLayers.Num();
-	}
-
 	/** @return the Normal layer at the given Index */
 	FDynamicMeshNormalOverlay* GetNormalLayer(int Index)
 	{
@@ -180,12 +187,10 @@ protected:
 	/** Parent mesh of this attribute set */
 	FDynamicMesh3* ParentMesh;
 
-	/** Default UV layer */
-	FDynamicMeshUVOverlay UV0;
 	/** Default Normals layer */
 	FDynamicMeshNormalOverlay Normals0;
 
-	TArray<FDynamicMeshUVOverlay*> UVLayers;
+	TIndirectArray<FDynamicMeshUVOverlay> UVLayers;
 	TArray<FDynamicMeshNormalOverlay*> NormalLayers;
 
 	TUniquePtr<FDynamicMeshMaterialAttribute> MaterialIDAttrib;
@@ -199,7 +204,10 @@ protected:
 	 */
 	void Initialize(int MaxVertexID, int MaxTriangleID)
 	{
-		UV0.InitializeTriangles(MaxTriangleID);
+		for (FDynamicMeshUVOverlay& UVLayer : UVLayers)
+		{
+			UVLayer.InitializeTriangles(MaxTriangleID);
+		}
 		Normals0.InitializeTriangles(MaxTriangleID);
 	}
 
