uv layers support in most of the mesh conversion path + associated bugfixes

#rnx

diff --git a/Engine/Plugins/Experimental/GeometryProcessing/Source/DynamicMesh/Private/DynamicMesh3.cpp b/Engine/Plugins/Experimental/GeometryProcessing/Source/DynamicMesh/Private/DynamicMesh3.cpp
--- a/Engine/Plugins/Experimental/GeometryProcessing/Source/DynamicMesh/Private/DynamicMesh3.cpp
+++ b/Engine/Plugins/Experimental/GeometryProcessing/Source/DynamicMesh/Private/DynamicMesh3.cpp
@@ -905,7 +905,10 @@ bool FDynamicMesh3::CheckValidity(bool bAllowNonManifoldVertices, EValidityCheck
 
 	if (HasAttributes())
 	{
-		Attributes()->PrimaryUV()->CheckValidity(true, FailMode);
+		for (int UVLayerIndex = 0; UVLayerIndex < Attributes()->NumUVLayers(); UVLayerIndex++)
+		{
+			Attributes()->GetUVLayer(UVLayerIndex)->CheckValidity(true, FailMode);
+		}
 		Attributes()->PrimaryNormals()->CheckValidity(true, FailMode);
 	}
 
diff --git a/Engine/Plugins/Experimental/GeometryProcessing/Source/DynamicMesh/Private/DynamicMesh3_Edits.cpp b/Engine/Plugins/Experimental/GeometryProcessing/Source/DynamicMesh/Private/DynamicMesh3_Edits.cpp
--- a/Engine/Plugins/Experimental/GeometryProcessing/Source/DynamicMesh/Private/DynamicMesh3_Edits.cpp
+++ b/Engine/Plugins/Experimental/GeometryProcessing/Source/DynamicMesh/Private/DynamicMesh3_Edits.cpp
@@ -698,7 +698,7 @@ EMeshResult FDynamicMesh3::RemoveTriangle(int tID, bool bRemoveIsolatedVertices,
 
 	if (HasAttributes())
 	{
-		Attributes()->OnRemoveTriangle(tID, bRemoveIsolatedVertices);
+		Attributes()->OnRemoveTriangle(tID);
 	}
 
 	UpdateTimeStamp(true, true);
@@ -1177,6 +1177,7 @@ EMeshResult FDynamicMesh3::FlipEdge(int vA, int vB, FEdgeFlipInfo& FlipInfo)
 
 EMeshResult FDynamicMesh3::CollapseEdge(int vKeep, int vRemove, double collapse_t, FEdgeCollapseInfo& CollapseInfo)
 {
+	checkSlow(!HasAttributes() || !Attributes()->IsSeamVertex(vRemove, false));
 	CollapseInfo = FEdgeCollapseInfo();
 
 	if (IsVertex(vKeep) == false || IsVertex(vRemove) == false)
diff --git a/Engine/Plugins/Experimental/GeometryProcessing/Source/DynamicMesh/Private/DynamicMeshAttributeSet.cpp b/Engine/Plugins/Experimental/GeometryProcessing/Source/DynamicMesh/Private/DynamicMeshAttributeSet.cpp
--- a/Engine/Plugins/Experimental/GeometryProcessing/Source/DynamicMesh/Private/DynamicMeshAttributeSet.cpp
+++ b/Engine/Plugins/Experimental/GeometryProcessing/Source/DynamicMesh/Private/DynamicMeshAttributeSet.cpp
@@ -29,6 +29,18 @@ bool FDynamicMeshAttributeSet::IsSeamEdge(int eid) const
 }
 
 
+bool FDynamicMeshAttributeSet::IsSeamVertex(int VID, bool bBoundaryIsSeam) const
+{
+	for (const FDynamicMeshUVOverlay& UVLayer : UVLayers)
+	{
+		if (UVLayer.IsSeamVertex(VID, bBoundaryIsSeam))
+		{
+			return true;
+		}
+	}
+	return Normals0.IsSeamVertex(VID, bBoundaryIsSeam);
+}
+
 
 void FDynamicMeshAttributeSet::OnNewTriangle(int TriangleID, bool bInserted)
 {
@@ -45,13 +57,14 @@ void FDynamicMeshAttributeSet::OnNewTriangle(int TriangleID, bool bInserted)
 	}
 }
 
-void FDynamicMeshAttributeSet::OnRemoveTriangle(int TriangleID, bool bRemoveIsolatedVertices)
+
+void FDynamicMeshAttributeSet::OnRemoveTriangle(int TriangleID)
 {
-	for (FDynamicMeshUVOverlay UVLayer : UVLayers)
+	for (FDynamicMeshUVOverlay& UVLayer : UVLayers)
 	{
-		UVLayer.OnRemoveTriangle(TriangleID, bRemoveIsolatedVertices);
+		UVLayer.OnRemoveTriangle(TriangleID);
 	}
-	Normals0.OnRemoveTriangle(TriangleID, bRemoveIsolatedVertices);
+	Normals0.OnRemoveTriangle(TriangleID);
 
 	// has no effect on MaterialIDAttrib
 }
diff --git a/Engine/Plugins/Experimental/GeometryProcessing/Source/DynamicMesh/Private/DynamicMeshEditor.cpp b/Engine/Plugins/Experimental/GeometryProcessing/Source/DynamicMesh/Private/DynamicMeshEditor.cpp
--- a/Engine/Plugins/Experimental/GeometryProcessing/Source/DynamicMesh/Private/DynamicMeshEditor.cpp
+++ b/Engine/Plugins/Experimental/GeometryProcessing/Source/DynamicMesh/Private/DynamicMeshEditor.cpp
@@ -887,14 +887,19 @@ void FDynamicMeshEditor::AppendMesh(const FDynamicMesh3* AppendMesh,
 				VertexMap, TriangleMap, NormalTransform, NormalMap);
 		}
 
-		const FDynamicMeshUVOverlay* FromUVs = AppendMesh->Attributes()->PrimaryUV();
-		FDynamicMeshUVOverlay* ToUVs = Mesh->Attributes()->PrimaryUV();
-		if (FromUVs != nullptr && ToUVs != nullptr)
+
+		int NumUVLayers = FMath::Min(Mesh->Attributes()->NumUVLayers(), AppendMesh->Attributes()->NumUVLayers());
+		for (int UVLayerIndex = 0; UVLayerIndex < NumUVLayers; UVLayerIndex++)
 		{
-			FIndexMapi& UVMap = IndexMapsOut.GetUVMap(0);
-			UVMap.Reserve(FromUVs->ElementCount());
-			AppendUVs(AppendMesh, FromUVs, ToUVs,
-				VertexMap, TriangleMap, UVMap);
+			const FDynamicMeshUVOverlay* FromUVs = AppendMesh->Attributes()->GetUVLayer(UVLayerIndex);
+			FDynamicMeshUVOverlay* ToUVs = Mesh->Attributes()->GetUVLayer(UVLayerIndex);
+			if (FromUVs != nullptr && ToUVs != nullptr)
+			{
+				FIndexMapi& UVMap = IndexMapsOut.GetUVMap(0);
+				UVMap.Reserve(FromUVs->ElementCount());
+				AppendUVs(AppendMesh, FromUVs, ToUVs,
+					VertexMap, TriangleMap, UVMap);
+			}
 		}
 	}
 }
@@ -1033,10 +1038,10 @@ static void AppendAttributes(const FDynamicMesh3* FromMesh, int FromTriangleID,
 	// todo: if we ever support multiple normal layers, copy them all
 	check(FromMesh->Attributes()->NumNormalLayers() == 1);
 
-	const FDynamicMeshUVOverlay* FromUVOverlay = FromMesh->Attributes()->PrimaryUV();
-	FDynamicMeshUVOverlay* ToUVOverlay = ToMesh->Attributes()->PrimaryUV();
-
+	for (int UVLayerIndex = 0; UVLayerIndex < FMath::Min(FromMesh->Attributes()->NumUVLayers(), ToMesh->Attributes()->NumUVLayers()); UVLayerIndex++)
 	{
+		const FDynamicMeshUVOverlay* FromUVOverlay = FromMesh->Attributes()->GetUVLayer(UVLayerIndex);
+		FDynamicMeshUVOverlay* ToUVOverlay = ToMesh->Attributes()->GetUVLayer(UVLayerIndex);
 		FIndex3i FromElemTri = FromUVOverlay->GetTriangle(FromTriangleID);
 		FIndex3i ToElemTri = ToUVOverlay->GetTriangle(ToTriangleID);
 		for (int j = 0; j < 3; ++j)
diff --git a/Engine/Plugins/Experimental/GeometryProcessing/Source/DynamicMesh/Private/DynamicMeshOverlay.cpp b/Engine/Plugins/Experimental/GeometryProcessing/Source/DynamicMesh/Private/DynamicMeshOverlay.cpp
--- a/Engine/Plugins/Experimental/GeometryProcessing/Source/DynamicMesh/Private/DynamicMeshOverlay.cpp
+++ b/Engine/Plugins/Experimental/GeometryProcessing/Source/DynamicMesh/Private/DynamicMeshOverlay.cpp
@@ -321,11 +321,15 @@ bool TDynamicMeshOverlay<RealType,ElementSize>::IsSeamEdge(int eid) const
 
 
 template<typename RealType, int ElementSize>
-bool TDynamicMeshOverlay<RealType, ElementSize>::IsSeamVertex(int vid) const
+bool TDynamicMeshOverlay<RealType, ElementSize>::IsSeamVertex(int vid, bool bBoundaryIsSeam) const
 {
 	// @todo can we do this more efficiently? At minimum we are looking up each triangle twice...
 	for (int edgeid : ParentMesh->VtxEdgesItr(vid))
 	{
+		if (!bBoundaryIsSeam && ParentMesh->IsBoundaryEdge(edgeid))
+		{
+			continue;
+		}
 		if (IsSeamEdge(edgeid))
 		{
 			return true;
@@ -410,10 +414,8 @@ void TDynamicMeshOverlay<RealType, ElementSize>::GetElementTriangles(int Element
 }
 
 
-
-
 template<typename RealType, int ElementSize>
-void TDynamicMeshOverlay<RealType, ElementSize>::OnRemoveTriangle(int TriangleID, bool bRemoveIsolatedVertices)
+void TDynamicMeshOverlay<RealType, ElementSize>::OnRemoveTriangle(int TriangleID)
 {
 	FIndex3i Triangle = GetTriangle(TriangleID);
 	if (Triangle.A < 0 && Triangle.B < 0 && Triangle.C < 0)
@@ -429,7 +431,7 @@ void TDynamicMeshOverlay<RealType, ElementSize>::OnRemoveTriangle(int TriangleID
 	{
 		int elemid = Triangle[j];
 		ElementsRefCounts.Decrement(elemid);
-		if (bRemoveIsolatedVertices && ElementsRefCounts.GetRefCount(elemid) == 1) 
+		if (ElementsRefCounts.GetRefCount(elemid) == 1) 
 		{
 			ElementsRefCounts.Decrement(elemid);
 			ParentVertices[elemid] = FDynamicMesh3::InvalidID;
diff --git a/Engine/Plugins/Experimental/GeometryProcessing/Source/DynamicMesh/Private/Operations/MeshPlaneCut.cpp b/Engine/Plugins/Experimental/GeometryProcessing/Source/DynamicMesh/Private/Operations/MeshPlaneCut.cpp
--- a/Engine/Plugins/Experimental/GeometryProcessing/Source/DynamicMesh/Private/Operations/MeshPlaneCut.cpp
+++ b/Engine/Plugins/Experimental/GeometryProcessing/Source/DynamicMesh/Private/Operations/MeshPlaneCut.cpp
@@ -100,7 +100,7 @@ bool FMeshPlaneCut::Cut()
 			continue; // edge split really shouldn't fail; skip the edge if it somehow does
 		}
 
-		NewEdges.Add(splitInfo.NewEdges.A); // TODO: why not OnCutEdges.Add(splitInfo.NewEdges.A);  ???
+		NewEdges.Add(splitInfo.NewEdges.A);
 		NewEdges.Add(splitInfo.NewEdges.B); OnCutEdges.Add(splitInfo.NewEdges.B);
 		if (splitInfo.NewEdges.C != FDynamicMesh3::InvalidID)
 		{
@@ -181,6 +181,16 @@ void FMeshPlaneCut::CollapseDegenerateEdges(const TSet<int>& OnCutEdges, const T
 				}
 
 				FIndex2i EV = Mesh->GetEdgeV(EID);
+				// if the vertex we'd remove is on a seam, try removing the other one instead
+				if (Mesh->HasAttributes() && Mesh->Attributes()->IsSeamVertex(EV.B, false))
+				{
+					Swap(EV.A, EV.B);
+					// if they were both on seams, then collapse should not happen?  (& would break OnCollapseEdge assumptions in overlay)
+					if (Mesh->HasAttributes() && Mesh->Attributes()->IsSeamVertex(EV.B, false))
+					{
+						continue;
+					}
+				}
 				FDynamicMesh3::FEdgeCollapseInfo CollapseInfo;
 				EMeshResult Result = Mesh->CollapseEdge(EV.A, EV.B, CollapseInfo);
 				if (Result == EMeshResult::Ok)
diff --git a/Engine/Plugins/Experimental/GeometryProcessing/Source/DynamicMesh/Public/DynamicMeshAttributeSet.h b/Engine/Plugins/Experimental/GeometryProcessing/Source/DynamicMesh/Public/DynamicMeshAttributeSet.h
--- a/Engine/Plugins/Experimental/GeometryProcessing/Source/DynamicMesh/Public/DynamicMeshAttributeSet.h
+++ b/Engine/Plugins/Experimental/GeometryProcessing/Source/DynamicMesh/Public/DynamicMeshAttributeSet.h
@@ -97,6 +97,9 @@ public:
 	/** @return true if the given edge is a seam edge in any overlay */
 	virtual bool IsSeamEdge(int EdgeID) const;
 
+	/** @return true if the given vertex is a seam vertex in any overlay */
+	virtual bool IsSeamVertex(int VertexID, bool bBoundaryIsSeam = true) const;
+
 	
 	//
 	// UV Layers 
@@ -214,7 +217,7 @@ protected:
 	// These functions are called by the FDynamicMesh3 to update the various
 	// attributes when the parent mesh topology has been modified.
 	virtual void OnNewTriangle(int TriangleID, bool bInserted);
-	virtual void OnRemoveTriangle(int TriangleID, bool bRemoveIsolatedVertices);
+	virtual void OnRemoveTriangle(int TriangleID);
 	virtual void OnReverseTriOrientation(int TriangleID);
 	virtual void OnSplitEdge(const FDynamicMesh3::FEdgeSplitInfo & splitInfo);
 	virtual void OnFlipEdge(const FDynamicMesh3::FEdgeFlipInfo & flipInfo);
diff --git a/Engine/Plugins/Experimental/GeometryProcessing/Source/DynamicMesh/Public/DynamicMeshOverlay.h b/Engine/Plugins/Experimental/GeometryProcessing/Source/DynamicMesh/Public/DynamicMeshOverlay.h
--- a/Engine/Plugins/Experimental/GeometryProcessing/Source/DynamicMesh/Public/DynamicMeshOverlay.h
+++ b/Engine/Plugins/Experimental/GeometryProcessing/Source/DynamicMesh/Public/DynamicMeshOverlay.h
@@ -222,7 +222,7 @@ public:
 	/** Returns true if the parent-mesh edge is a "Seam" in this overlay */
 	bool IsSeamEdge(int EdgeID) const;
 	/** Returns true if the parent-mesh vertex is connected to any seam edges */
-	bool IsSeamVertex(int VertexID) const;
+	bool IsSeamVertex(int VertexID, bool bBoundaryIsSeam = true) const;
 
 	/** find the elements associated with a given parent-mesh vertex */
 	void GetVertexElements(int VertexID, TArray<int>& OutElements) const;
@@ -242,7 +242,7 @@ public:
 	/** Set a triangle's element indices to InvalidID */
 	void InitializeNewTriangle(int TriangleID);
 	/** Remove a triangle from the overlay */
-	void OnRemoveTriangle(int TriangleID, bool bRemoveIsolatedVertices);
+	void OnRemoveTriangle(int TriangleID);
 	/** Reverse the orientation of a triangle's elements */
 	void OnReverseTriOrientation(int TriangleID);
 	/** Update the overlay to reflect an edge split in the parent mesh */
diff --git a/Engine/Plugins/Experimental/GeometryProcessing/Source/MeshConversion/Private/DynamicMeshToMeshDescription.cpp b/Engine/Plugins/Experimental/GeometryProcessing/Source/MeshConversion/Private/DynamicMeshToMeshDescription.cpp
--- a/Engine/Plugins/Experimental/GeometryProcessing/Source/MeshConversion/Private/DynamicMeshToMeshDescription.cpp
+++ b/Engine/Plugins/Experimental/GeometryProcessing/Source/MeshConversion/Private/DynamicMeshToMeshDescription.cpp
@@ -17,7 +17,7 @@ namespace DynamicMeshToMeshDescriptionConversionHelper
 	template <typename OutAttributeType, int VecLen, typename InAttributeType>
 	void SetAttributesFromOverlay(
 		const FDynamicMesh3* MeshInArg, const FMeshDescription& MeshOutArg,
-		TVertexInstanceAttributesRef<OutAttributeType>& InstanceAttrib, const TDynamicMeshVectorOverlay<float, VecLen, InAttributeType>* Overlay)
+		TVertexInstanceAttributesRef<OutAttributeType>& InstanceAttrib, const TDynamicMeshVectorOverlay<float, VecLen, InAttributeType>* Overlay, int AttribIndex=0)
 	{
 		const FPolygonArray& Polygons = MeshOutArg.Polygons();
 		int MeshInTriIdx = 0;
@@ -33,9 +33,9 @@ namespace DynamicMeshToMeshDescriptionConversionHelper
 				TArrayView<const FVertexInstanceID> InstanceTri = MeshOutArg.GetTriangleVertexInstances(TriangleID);
 
 				FIndex3i OverlayVertIndices = Overlay->GetTriangle(MeshInTriIdx);
-				InstanceAttrib.Set(InstanceTri[0], OutAttributeType(Overlay->GetElement(OverlayVertIndices.A)));
-				InstanceAttrib.Set(InstanceTri[1], OutAttributeType(Overlay->GetElement(OverlayVertIndices.B)));
-				InstanceAttrib.Set(InstanceTri[2], OutAttributeType(Overlay->GetElement(OverlayVertIndices.C)));
+				InstanceAttrib.Set(InstanceTri[0], AttribIndex, OutAttributeType(Overlay->GetElement(OverlayVertIndices.A)));
+				InstanceAttrib.Set(InstanceTri[1], AttribIndex, OutAttributeType(Overlay->GetElement(OverlayVertIndices.B)));
+				InstanceAttrib.Set(InstanceTri[2], AttribIndex, OutAttributeType(Overlay->GetElement(OverlayVertIndices.C)));
 			}
 		}
 	}
@@ -95,16 +95,17 @@ void FDynamicMeshToMeshDescription::UpdateAttributes(const FDynamicMesh3* MeshIn
 
 	if (bUpdateUVs)
 	{
-		const FDynamicMeshUVOverlay* Overlay = MeshIn->HasAttributes() ? MeshIn->Attributes()->PrimaryUV() : nullptr;
-
 		TVertexInstanceAttributesRef<FVector2D> InstanceAttrib =
 			MeshOut.VertexInstanceAttributes().GetAttributesRef<FVector2D>(MeshAttribute::VertexInstance::TextureCoordinate);
 		ensureMsgf(InstanceAttrib.IsValid(), TEXT("Trying to update UVs on a MeshDescription that has no texture coordinate attributes"));
 		if (InstanceAttrib.IsValid())
 		{
-			if (Overlay)
+			if (MeshIn->HasAttributes())
 			{
-				DynamicMeshToMeshDescriptionConversionHelper::SetAttributesFromOverlay(MeshIn, MeshOut, InstanceAttrib, Overlay);
+				for (int UVLayerIndex = 0, NumLayers = MeshIn->Attributes()->NumUVLayers(); UVLayerIndex < NumLayers; UVLayerIndex++)
+				{
+					DynamicMeshToMeshDescriptionConversionHelper::SetAttributesFromOverlay(MeshIn, MeshOut, InstanceAttrib, MeshIn->Attributes()->GetUVLayer(UVLayerIndex), UVLayerIndex);
+				}
 			}
 			else
 			{
@@ -199,7 +200,6 @@ void FDynamicMeshToMeshDescription::Convert_NoAttributes(const FDynamicMesh3* Me
 
 void FDynamicMeshToMeshDescription::Convert_SharedInstances(const FDynamicMesh3* MeshIn, FMeshDescription& MeshOut)
 {
-	const FDynamicMeshUVOverlay* UVOverlay = MeshIn->HasAttributes() ? MeshIn->Attributes()->PrimaryUV() : nullptr;
 	const FDynamicMeshNormalOverlay* NormalOverlay = MeshIn->HasAttributes() ? MeshIn->Attributes()->PrimaryNormals() : nullptr;
 
 	MeshOut.Empty();
@@ -252,29 +252,89 @@ void FDynamicMeshToMeshDescription::Convert_SharedInstances(const FDynamicMesh3*
 		}
 	}
 
-	// create new instances when seen
-	TMap<FIndex3i, FVertexInstanceID> InstanceList;
-	for (int TriID : MeshIn->TriangleIndicesItr())
+	// build all vertex instances (splitting as needed)
+	// store per-triangle instance ids
+	struct Tri
 	{
-		// @todo support additional overlays (requires IndexNi...)
-		FIndex3i Triangle = MeshIn->GetTriangle(TriID);
-		FIndex3i UVTriangle = (UVOverlay != nullptr) ? UVOverlay->GetTriangle(TriID) : FIndex3i(-1, -1, -1);
-		FIndex3i NormalTriangle = (NormalOverlay != nullptr) ? NormalOverlay->GetTriangle(TriID) : FIndex3i(-1, -1, -1);
-		FVertexInstanceID InstanceTri[3];
-		for (int j = 0; j < 3; ++j)
+		FVertexInstanceID V[3];
+
+		Tri() : V{ FVertexInstanceID::Invalid,FVertexInstanceID::Invalid,FVertexInstanceID::Invalid }
+		{}
+	};
+	TArray<Tri> TriVertInstances;
+	TriVertInstances.SetNum(MeshIn->MaxTriangleID());
+	TArray<int32> KnownInstanceIDs;
+	int NumUVLayers = MeshIn->HasAttributes() ? MeshIn->Attributes()->NumUVLayers() : 0;
+	Builder.SetNumUVLayers(NumUVLayers);
+	int KIItemLen = 1 + (NormalOverlay ? 1 : 0) + NumUVLayers;
+	for (int VertID : MeshIn->VertexIndicesItr())
+	{
+		KnownInstanceIDs.Reset();
+		for (int TriID : MeshIn->VtxTrianglesItr(VertID))
 		{
-			FIndex3i InstanceElem(Triangle[j], UVTriangle[j], NormalTriangle[j]);
-			if (InstanceList.Contains(InstanceElem) == false)
+			FIndex3i Tri = MeshIn->GetTriangle(TriID);
+			int SubIdx = IndexUtil::FindTriIndex(VertID, Tri);
+
+			int32 FoundInstance = FVertexInstanceID::Invalid.GetValue();
+			for (int KIItemIdx = 0; KIItemIdx < KnownInstanceIDs.Num(); KIItemIdx += KIItemLen)
 			{
-				FVertexInstanceID NewInstanceID = Builder.AppendInstance(MapV[Triangle[j]]);
-				InstanceList.Add(InstanceElem, NewInstanceID);
-				Builder.SetInstance(NewInstanceID,
-					(UVTriangle[j] == -1 || UVOverlay == nullptr) ? FVector2f::Zero() : UVOverlay->GetElement(UVTriangle[j]),
-					(NormalTriangle[j] == -1 || NormalOverlay == nullptr) ? FVector3f::UnitY() : NormalOverlay->GetElement(NormalTriangle[j]));
+				int KIItemInternalIdx = KIItemIdx;
+
+				if (NormalOverlay && KnownInstanceIDs[KIItemInternalIdx++] != NormalOverlay->GetTriangle(TriID)[SubIdx])
+				{
+					continue;
+				}
+
+				bool FoundInUVs = true;
+				for (int UVLayerIndex = 0; UVLayerIndex < NumUVLayers; UVLayerIndex++)
+				{
+					const FDynamicMeshUVOverlay* Overlay = MeshIn->Attributes()->GetUVLayer(UVLayerIndex);
+					if (KnownInstanceIDs[KIItemInternalIdx++] != Overlay->GetTriangle(TriID)[SubIdx])
+					{
+						FoundInUVs = false;
+						break;
+					}
+				}
+				if (!FoundInUVs)
+				{
+					continue;
+				}
+
+				FoundInstance = KnownInstanceIDs[KIItemInternalIdx++];
+				check(KIItemInternalIdx == KIItemIdx + KIItemLen);
+				break;
 			}
-			InstanceTri[j] = InstanceList[InstanceElem];
+			if (FoundInstance == FVertexInstanceID::Invalid.GetValue())
+			{
+				FVertexInstanceID NewInstanceID = Builder.AppendInstance(MapV[VertID]);
+				if (NormalOverlay)
+				{
+					int ElID = NormalOverlay->GetTriangle(TriID)[SubIdx];
+					KnownInstanceIDs.Add(int32(ElID));
+					Builder.SetInstanceNormal(NewInstanceID, ElID != -1 ? NormalOverlay->GetElement(ElID) : FVector3f::UnitZ());
+				}
+				else
+				{
+					Builder.SetInstanceNormal(NewInstanceID, FVector::UpVector);
+				}
+				for (int UVLayerIndex = 0; UVLayerIndex < NumUVLayers; UVLayerIndex++)
+				{
+					const FDynamicMeshUVOverlay* Overlay = MeshIn->Attributes()->GetUVLayer(UVLayerIndex);
+					int ElID = Overlay->GetTriangle(TriID)[SubIdx];
+					KnownInstanceIDs.Add(int32(ElID));
+
+					Builder.SetInstanceUV(NewInstanceID, ElID != -1 ? Overlay->GetElement(ElID) : FVector2f::Zero(), UVLayerIndex);
+				}
+				FoundInstance = NewInstanceID.GetValue();
+				KnownInstanceIDs.Add(FoundInstance);
+			}
+			TriVertInstances[TriID].V[SubIdx] = FVertexInstanceID(FoundInstance);
 		}
+	}
 
+	// build the polygons
+	for (int TriID : MeshIn->TriangleIndicesItr())
+	{
 		// transfer material index to MeshDescription polygon group (by convention)
 		FPolygonGroupID UsePolygonGroupID = ZeroPolygonGroupID;
 		if (MaterialIDAttrib)
@@ -284,7 +344,7 @@ void FDynamicMeshToMeshDescription::Convert_SharedInstances(const FDynamicMesh3*
 			UsePolygonGroupID = FPolygonGroupID(MaterialID);
 		}
 
-		FPolygonID NewPolygonID = Builder.AppendTriangle(InstanceTri[0], InstanceTri[1], InstanceTri[2], UsePolygonGroupID);
+		FPolygonID NewPolygonID = Builder.AppendTriangle(TriVertInstances[TriID].V[0], TriVertInstances[TriID].V[1], TriVertInstances[TriID].V[2], UsePolygonGroupID);
 
 		if (bCopyGroupToPolyGroup)
 		{
@@ -299,6 +359,8 @@ void FDynamicMeshToMeshDescription::Convert_SharedInstances(const FDynamicMesh3*
 
 void FDynamicMeshToMeshDescription::Convert_NoSharedInstances(const FDynamicMesh3* MeshIn, FMeshDescription& MeshOut)
 {
+	// TODO: update this path to support multiple UV layers
+
 	const FDynamicMeshUVOverlay* UVOverlay = MeshIn->HasAttributes() ? MeshIn->Attributes()->PrimaryUV() : nullptr;
 	const FDynamicMeshNormalOverlay* NormalOverlay = MeshIn->HasAttributes() ? MeshIn->Attributes()->PrimaryNormals() : nullptr;
 
diff --git a/Engine/Plugins/Experimental/GeometryProcessing/Source/MeshConversion/Private/MeshDescriptionBuilder.cpp b/Engine/Plugins/Experimental/GeometryProcessing/Source/MeshConversion/Private/MeshDescriptionBuilder.cpp
--- a/Engine/Plugins/Experimental/GeometryProcessing/Source/MeshConversion/Private/MeshDescriptionBuilder.cpp
+++ b/Engine/Plugins/Experimental/GeometryProcessing/Source/MeshConversion/Private/MeshDescriptionBuilder.cpp
@@ -116,6 +116,24 @@ void FMeshDescriptionBuilder::SetInstanceNormal(const FVertexInstanceID& Instanc
 }
 
 
+void FMeshDescriptionBuilder::SetInstanceUV(const FVertexInstanceID& InstanceID, const FVector2D& InstanceUV, int32 UVLayerIndex)
+{
+	if (InstanceUVs.IsValid() && ensure(UVLayerIndex < InstanceUVs.GetNumIndices()))
+	{
+		InstanceUVs.Set(InstanceID, UVLayerIndex, InstanceUV); 
+	}
+}
+
+
+void FMeshDescriptionBuilder::SetNumUVLayers(int32 NumUVLayers)
+{
+	if (ensure(InstanceUVs.IsValid()))
+	{
+		InstanceUVs.SetNumIndices(NumUVLayers);
+	}
+}
+
+
 void FMeshDescriptionBuilder::SetInstanceColor(const FVertexInstanceID& InstanceID, const FVector4& Color)
 {
 	if (InstanceColors.IsValid())
@@ -214,6 +232,8 @@ void FMeshDescriptionBuilder::SetPolyGroupID(const FPolygonID& PolygonID, int Gr
 
 void FMeshDescriptionBuilder::AppendMesh(const FDynamicMesh3* Mesh, bool bSetPolyGroups)
 {
+	// TODO: update to support multiple UV layers; & combine w/ extremely similar code in DynamicMeshToMeshDescription.cpp
+
 	// create vertices
 	TArray<FVertexID> MapV;
 	MapV.SetNum(Mesh->MaxVertexID());
diff --git a/Engine/Plugins/Experimental/GeometryProcessing/Source/MeshConversion/Private/MeshDescriptionToDynamicMesh.cpp b/Engine/Plugins/Experimental/GeometryProcessing/Source/MeshConversion/Private/MeshDescriptionToDynamicMesh.cpp
--- a/Engine/Plugins/Experimental/GeometryProcessing/Source/MeshConversion/Private/MeshDescriptionToDynamicMesh.cpp
+++ b/Engine/Plugins/Experimental/GeometryProcessing/Source/MeshConversion/Private/MeshDescriptionToDynamicMesh.cpp
@@ -139,9 +139,11 @@ void FMeshDescriptionToDynamicMesh::Convert(const FMeshDescription* MeshIn, FDyn
 	TVertexInstanceAttributesConstRef<FVector> InstanceNormals =
 		MeshIn->VertexInstanceAttributes().GetAttributesRef<FVector>(MeshAttribute::VertexInstance::Normal);
 
+	int NumUVLayers = InstanceUVs.GetNumIndices();
+
 	// enable attributes on output mesh
 	MeshOut.EnableAttributes();
-	FDynamicMeshUVOverlay* UVOverlay = MeshOut.Attributes()->PrimaryUV();
+	MeshOut.Attributes()->SetNumUVLayers(NumUVLayers);
 	FDynamicMeshNormalOverlay* NormalOverlay = MeshOut.Attributes()->PrimaryNormals();
 
 	TPolygonAttributesConstRef<int> PolyGroups =
@@ -164,7 +166,11 @@ void FMeshDescriptionToDynamicMesh::Convert(const FMeshDescription* MeshIn, FDyn
 	// reserve space in MeshOut?
 
 	// used to merge coincident elements so that we get actual topology
-	FUVWelder UVWelder(UVOverlay);
+	TArray<FUVWelder> UVWelders;
+	for (int UVLayerIndex = 0; UVLayerIndex < NumUVLayers; UVLayerIndex++)
+	{
+		UVWelders.Emplace(MeshOut.Attributes()->GetUVLayer(UVLayerIndex));
+	}
 	FNormalWelder NormalWelder(NormalOverlay);
 
 	// only enable material ID if we have more than one material
@@ -271,13 +277,14 @@ void FMeshDescriptionToDynamicMesh::Convert(const FMeshDescription* MeshIn, FDyn
 				TriToPolyTriMap.Insert(FIndex2i(PolygonID.GetValue(), TriIdx), NewTriangleID);
 			}
 
-			if (UVOverlay != nullptr)
+			for (int UVLayerIndex = 0; UVLayerIndex < NumUVLayers; UVLayerIndex++)
 			{
+				FDynamicMeshUVOverlay* UVOverlay = MeshOut.Attributes()->GetUVLayer(UVLayerIndex);
 				FIndex3i TriUV;
 				for (int j = 0; j < 3; ++j)
 				{
-					FVector2D UV = InstanceUVs.Get(InstanceTri[j]);
-					TriUV[j] = UVWelder.FindOrAddUnique(UV, Tri[j]);
+					FVector2D UV = InstanceUVs.Get(InstanceTri[j], UVLayerIndex);
+					TriUV[j] = UVWelders[UVLayerIndex].FindOrAddUnique(UV, Tri[j]);
 				}
 				UVOverlay->SetTriangle(NewTriangleID, TriUV);
 			}
@@ -304,9 +311,9 @@ void FMeshDescriptionToDynamicMesh::Convert(const FMeshDescription* MeshIn, FDyn
 
 	if (bPrintDebugMessages)
 	{
-		int NumUVs = (UVOverlay != nullptr) ? UVOverlay->MaxElementID() : 0;
+		int NumUVs = (NumUVLayers > 0) ? MeshOut.Attributes()->PrimaryUV()->MaxElementID() : 0;
 		int NumNormals = (NormalOverlay != nullptr) ? NormalOverlay->MaxElementID() : 0;
-		UE_LOG(LogTemp, Warning, TEXT("FMeshDescriptionToDynamicMesh:  FDynamicMesh verts %d triangles %d uvs %d normals %d"), MeshOut.MaxVertexID(), MeshOut.MaxTriangleID(), NumUVs, NumNormals);
+		UE_LOG(LogTemp, Warning, TEXT("FMeshDescriptionToDynamicMesh:  FDynamicMesh verts %d triangles %d (primary) uvs %d normals %d"), MeshOut.MaxVertexID(), MeshOut.MaxTriangleID(), NumUVs, NumNormals);
 	}
 
 }
diff --git a/Engine/Plugins/Experimental/GeometryProcessing/Source/MeshConversion/Public/MeshDescriptionBuilder.h b/Engine/Plugins/Experimental/GeometryProcessing/Source/MeshConversion/Public/MeshDescriptionBuilder.h
--- a/Engine/Plugins/Experimental/GeometryProcessing/Source/MeshConversion/Public/MeshDescriptionBuilder.h
+++ b/Engine/Plugins/Experimental/GeometryProcessing/Source/MeshConversion/Public/MeshDescriptionBuilder.h
@@ -48,6 +48,12 @@ public:
 	/** Set the Normal of a vertex instance*/
 	void SetInstanceNormal(const FVertexInstanceID& InstanceID, const FVector& Normal);
 
+	/** Set the UV of a vertex instance */
+	void SetInstanceUV(const FVertexInstanceID& InstanceID, const FVector2D& InstanceUV, int32 UVLayerIndex = 0);
+
+	/** Set the number of UV layers */
+	void SetNumUVLayers(int32 NumUVLayers);
+
 	/** Set the Color of a vertex instance*/
 	void SetInstanceColor(const FVertexInstanceID& InstanceID, const FVector4& Color);
 
