GeometryProcessing: add DynamicMeshOverlay::HasInteriorSeamEdges() to check for interior seams in overlay
#rnx

diff --git a/Engine/Plugins/Experimental/GeometryProcessing/Source/DynamicMesh/Private/DynamicMeshOverlay.cpp b/Engine/Plugins/Experimental/GeometryProcessing/Source/DynamicMesh/Private/DynamicMeshOverlay.cpp
--- a/Engine/Plugins/Experimental/GeometryProcessing/Source/DynamicMesh/Private/DynamicMeshOverlay.cpp
+++ b/Engine/Plugins/Experimental/GeometryProcessing/Source/DynamicMesh/Private/DynamicMeshOverlay.cpp
@@ -319,6 +319,37 @@ bool TDynamicMeshOverlay<RealType,ElementSize>::IsSeamEdge(int eid) const
 }
 
 
+template<typename RealType, int ElementSize>
+bool TDynamicMeshOverlay<RealType, ElementSize>::HasInteriorSeamEdges() const
+{
+	for (int eid : ParentMesh->EdgeIndicesItr())
+	{
+		FIndex2i et = ParentMesh->GetEdgeT(eid);
+		if (et.B != FDynamicMesh3::InvalidID)
+		{
+			FIndex2i ev = ParentMesh->GetEdgeV(eid);
+			int base_a = ev.A, base_b = ev.B;
+
+			FIndex3i Triangle0 = GetTriangle(et.A);
+			FIndex3i BaseTriangle0(ParentVertices[Triangle0.A], ParentVertices[Triangle0.B], ParentVertices[Triangle0.C]);
+			int idx_base_a1 = BaseTriangle0.IndexOf(base_a);
+			int idx_base_b1 = BaseTriangle0.IndexOf(base_b);
+
+			FIndex3i Triangle1 = GetTriangle(et.B);
+			FIndex3i BaseTriangle1(ParentVertices[Triangle1.A], ParentVertices[Triangle1.B], ParentVertices[Triangle1.C]);
+			int idx_base_a2 = BaseTriangle1.IndexOf(base_a);
+			int idx_base_b2 = BaseTriangle1.IndexOf(base_b);
+
+			if (!IndexUtil::SamePairUnordered(Triangle0[idx_base_a1], Triangle0[idx_base_b1], Triangle1[idx_base_a2], Triangle1[idx_base_b2]))
+			{
+				return true;
+			}
+		}
+	}
+	return false;
+}
+
+
 
 template<typename RealType, int ElementSize>
 bool TDynamicMeshOverlay<RealType, ElementSize>::IsSeamVertex(int vid, bool bBoundaryIsSeam) const
diff --git a/Engine/Plugins/Experimental/GeometryProcessing/Source/DynamicMesh/Public/DynamicMeshOverlay.h b/Engine/Plugins/Experimental/GeometryProcessing/Source/DynamicMesh/Public/DynamicMeshOverlay.h
--- a/Engine/Plugins/Experimental/GeometryProcessing/Source/DynamicMesh/Public/DynamicMeshOverlay.h
+++ b/Engine/Plugins/Experimental/GeometryProcessing/Source/DynamicMesh/Public/DynamicMeshOverlay.h
@@ -232,6 +232,9 @@ public:
 	/** find the triangles connected to an element */
 	void GetElementTriangles(int ElementID, TArray<int>& OutTriangles) const;
 
+	/** Currently an iteration every time */
+	bool HasInteriorSeamEdges() const;
+
 	/**
 	 * Checks that the overlay mesh is well-formed, ie all internal data structures are consistent
 	 */
