change default cylinder back to not having any intermediate verts to fix geometry processing tests

#rnx

diff --git a/Engine/Plugins/Experimental/GeometryProcessing/Source/DynamicMesh/Public/Generators/SweepGenerator.h b/Engine/Plugins/Experimental/GeometryProcessing/Source/DynamicMesh/Public/Generators/SweepGenerator.h
--- a/Engine/Plugins/Experimental/GeometryProcessing/Source/DynamicMesh/Public/Generators/SweepGenerator.h
+++ b/Engine/Plugins/Experimental/GeometryProcessing/Source/DynamicMesh/Public/Generators/SweepGenerator.h
@@ -246,7 +246,7 @@ public:
 	float Radius[2] = {1.0f, 1.0f};
 	float Height = 1.0f;
 	int AngleSamples = 16;
-	int LengthSamples = 4;
+	int LengthSamples = 0;
 	bool bCapped = false;
 	bool bUVScaleMatchSidesAndCaps = true;
 
