Add volume and area calc to mesh inspector

#rnx

diff --git a/Engine/Plugins/Experimental/GeometryProcessing/Source/DynamicMesh/Private/MeshAdapterUtil.cpp b/Engine/Plugins/Experimental/GeometryProcessing/Source/DynamicMesh/Private/MeshAdapterUtil.cpp
--- a/Engine/Plugins/Experimental/GeometryProcessing/Source/DynamicMesh/Private/MeshAdapterUtil.cpp
+++ b/Engine/Plugins/Experimental/GeometryProcessing/Source/DynamicMesh/Private/MeshAdapterUtil.cpp
@@ -72,4 +72,32 @@ FPointSetAdapterd MeshAdapterUtil::MakeBoundaryEdgeMidpointsAdapter(const FDynam
 	Adapter.GetPointNormal = [](int Idx) { return FVector3f::UnitY(); };
 
 	return Adapter;
+}
+
+
+FTriangleMeshAdapterd MeshAdapterUtil::MakeTransformedDynamicMeshAdapter(const FDynamicMesh3* Mesh, FTransform Transform)
+{
+	FTriangleMeshAdapterd Adapter = MakeDynamicMeshAdapter(Mesh);
+	Adapter.GetVertex = [Mesh, Transform](int32 Idx)
+	{
+		return FVector3d(Transform.TransformPosition(Mesh->GetVertex(Idx)));
+	};
+	return Adapter;
+}
+
+
+FTriangleMeshAdapterd MeshAdapterUtil::MakeDynamicMeshAdapter(const FDynamicMesh3* Mesh)
+{
+	return
+	FTriangleMeshAdapterd {
+		[Mesh](int Idx) { return Mesh->IsTriangle(Idx); },
+		[Mesh](int Idx) { return Mesh->IsVertex(Idx); },
+		[Mesh]() { return Mesh->MaxTriangleID(); },
+		[Mesh]() { return Mesh->MaxVertexID(); },
+		[Mesh]() { return Mesh->TriangleCount(); },
+		[Mesh]() { return Mesh->VertexCount(); },
+		[Mesh]() { return Mesh->GetShapeTimestamp(); },
+		[Mesh](int Idx) { return Mesh->GetTriangle(Idx); },
+		[Mesh](int Idx) { return Mesh->GetVertex(Idx); }
+	};
 }
\ No newline at end of file
diff --git a/Engine/Plugins/Experimental/GeometryProcessing/Source/DynamicMesh/Public/MeshAdapterUtil.h b/Engine/Plugins/Experimental/GeometryProcessing/Source/DynamicMesh/Public/MeshAdapterUtil.h
--- a/Engine/Plugins/Experimental/GeometryProcessing/Source/DynamicMesh/Public/MeshAdapterUtil.h
+++ b/Engine/Plugins/Experimental/GeometryProcessing/Source/DynamicMesh/Public/MeshAdapterUtil.h
@@ -5,13 +5,24 @@
 #include "CoreTypes.h"
 #include "DynamicMesh3.h"
 #include "PointSetAdapter.h"
+#include "MeshAdapter.h"
 
 /**
- * Utility functions for constructing PointSetAdapter instances from various 
- * point sets on a mesh
+ * Utility functions for constructing various PointSetAdapter and MeshAdapter instances from dynamic meshes
  */
 namespace MeshAdapterUtil
 {
+
+	/**
+	 * @return Transformed adapter of a FDynamicMesh
+	 */
+	FTriangleMeshAdapterd DYNAMICMESH_API MakeTransformedDynamicMeshAdapter(const FDynamicMesh3* Mesh, FTransform Transform);
+
+	/**
+	 * @return 1:1 adapter of a FDynamicMesh; can be used as a starting point to create other adapters
+	 */
+	FTriangleMeshAdapterd DYNAMICMESH_API MakeDynamicMeshAdapter(const FDynamicMesh3* Mesh);
+
 	/**
 	 * @return Mesh vertices as a point set
 	 */
diff --git a/Engine/Plugins/Experimental/GeometryProcessing/Source/GeometricObjects/Public/MeshAdapter.h b/Engine/Plugins/Experimental/GeometryProcessing/Source/GeometricObjects/Public/MeshAdapter.h
--- a/Engine/Plugins/Experimental/GeometryProcessing/Source/GeometricObjects/Public/MeshAdapter.h
+++ b/Engine/Plugins/Experimental/GeometryProcessing/Source/GeometricObjects/Public/MeshAdapter.h
@@ -30,10 +30,10 @@ struct TTriangleMeshAdapter
 	}
 };
 
-typedef TTriangleMeshAdapter<double> TTriangleMeshAdapterd;
-typedef TTriangleMeshAdapter<float> TTriangleMeshAdapterf;
+typedef TTriangleMeshAdapter<double> FTriangleMeshAdapterd;
+typedef TTriangleMeshAdapter<float> FTriangleMeshAdapterf;
 
-TTriangleMeshAdapterd GetArrayMesh(TArray<FVector>& Vertices, TArray<FIntVector>& Triangles)
+inline FTriangleMeshAdapterd GetArrayMesh(TArray<FVector>& Vertices, TArray<FIntVector>& Triangles)
 {
 	return {
 		[&](int) { return true; },
diff --git a/Engine/Plugins/Experimental/GeometryProcessing/Source/GeometricObjects/Public/MeshQueries.h b/Engine/Plugins/Experimental/GeometryProcessing/Source/GeometricObjects/Public/MeshQueries.h
--- a/Engine/Plugins/Experimental/GeometryProcessing/Source/GeometricObjects/Public/MeshQueries.h
+++ b/Engine/Plugins/Experimental/GeometryProcessing/Source/GeometricObjects/Public/MeshQueries.h
@@ -68,6 +68,35 @@ public:
 		Normal = VectorUtil::NormalArea(Triangle.V[0], Triangle.V[1], Triangle.V[2], Area);
 	}
 
+	static FVector2d GetVolumeArea(const TriangleMeshType& Mesh)
+	{
+		double Volume = 0.0;
+		double Area = 0;
+		for (int TriIdx = 0; TriIdx < Mesh.MaxTriangleID(); TriIdx++)
+		{
+			if (!Mesh.IsTriangle(TriIdx))
+			{
+				continue;
+			}
+
+			FVector3d V0, V1, V2;
+			Mesh.GetTriVertices(TriIdx, V0, V1, V2);
+
+			// Get cross product of edges and (un-normalized) normal vector.
+			FVector3d V1mV0 = V1 - V0;
+			FVector3d V2mV0 = V2 - V0;
+			FVector3d N = V2mV0.Cross(V1mV0);
+
+			Area += N.Length();
+
+			double tmp0 = V0.X + V1.X;
+			double f1x = tmp0 + V2.X;
+			Volume += N.X * f1x;
+		}
+
+		return FVector2d(Volume * (1.0/6.0), Area * .5f);
+	}
+
 	static FAxisAlignedBox3d GetTriBounds(const TriangleMeshType& Mesh, int TID)
 	{
 		FIndex3i TriInds = Mesh.GetTriangle(TID);
diff --git a/Engine/Plugins/Experimental/MeshModelingToolset/Source/MeshModelingTools/Private/MeshInspectorTool.cpp b/Engine/Plugins/Experimental/MeshModelingToolset/Source/MeshModelingTools/Private/MeshInspectorTool.cpp
--- a/Engine/Plugins/Experimental/MeshModelingToolset/Source/MeshModelingTools/Private/MeshInspectorTool.cpp
+++ b/Engine/Plugins/Experimental/MeshModelingToolset/Source/MeshModelingTools/Private/MeshInspectorTool.cpp
@@ -11,6 +11,9 @@
 #include "SimpleDynamicMeshComponent.h"
 #include "Materials/MaterialInstanceDynamic.h"
 
+#include "Properties/MeshStatisticsProperties.h"
+#include "Properties/MeshAnalysisProperties.h"
+
 #include "SceneManagement.h" // for FPrimitiveDrawInterface
 
 #define LOCTEXT_NAMESPACE "UMeshInspectorTool"
@@ -123,6 +126,10 @@ void UMeshInspectorTool::Setup()
 	UMeshStatisticsProperties* Statistics = NewObject<UMeshStatisticsProperties>(this);
 	Statistics->Update(*DynamicMeshComponent->GetMesh());
 	AddToolPropertySource(Statistics);
+
+	UMeshAnalysisProperties* MeshAnalysis = NewObject<UMeshAnalysisProperties>(this);
+	MeshAnalysis->Update(*DynamicMeshComponent->GetMesh(), ComponentTarget->GetWorldTransform());
+	AddToolPropertySource(MeshAnalysis);
 }
 
 
diff --git a/Engine/Plugins/Experimental/MeshModelingToolset/Source/MeshModelingTools/Private/Properties/MeshAnalysisProperties.cpp b/Engine/Plugins/Experimental/MeshModelingToolset/Source/MeshModelingTools/Private/Properties/MeshAnalysisProperties.cpp
new file mode 100644
--- /dev/null
+++ b/Engine/Plugins/Experimental/MeshModelingToolset/Source/MeshModelingTools/Private/Properties/MeshAnalysisProperties.cpp
@@ -0,0 +1,23 @@
+// Copyright 1998-2019 Epic Games, Inc. All Rights Reserved.
+
+#include "Properties/MeshAnalysisProperties.h"
+
+#include "DynamicMesh3.h"
+#include "MeshQueries.h"
+#include "MeshAdapter.h"
+#include "MeshAdapterUtil.h"
+
+#define LOCTEXT_NAMESPACE "UMeshAnalysisProperites"
+
+
+void UMeshAnalysisProperties::Update(const FDynamicMesh3& MeshIn, const FTransform& Transform)
+{
+	FTriangleMeshAdapterd TransformedMesh = MeshAdapterUtil::MakeTransformedDynamicMeshAdapter(&MeshIn, Transform);
+	FVector2d VolArea = TMeshQueries<FTriangleMeshAdapterd>::GetVolumeArea(TransformedMesh);
+	this->SurfaceArea = FString::Printf(TEXT("%.2f m^2"), VolArea.Y / 10000);
+	this->Volume = FString::Printf(TEXT("%.2f m^3"), VolArea.X / 1000000);
+}
+
+
+
+#undef LOCTEXT_NAMESPACE
diff --git a/Engine/Plugins/Experimental/MeshModelingToolset/Source/MeshModelingTools/Public/MeshInspectorTool.h b/Engine/Plugins/Experimental/MeshModelingToolset/Source/MeshModelingTools/Public/MeshInspectorTool.h
--- a/Engine/Plugins/Experimental/MeshModelingToolset/Source/MeshModelingTools/Public/MeshInspectorTool.h
+++ b/Engine/Plugins/Experimental/MeshModelingToolset/Source/MeshModelingTools/Public/MeshInspectorTool.h
@@ -9,7 +9,6 @@
 #include "DynamicMesh3.h"
 #include "DynamicMeshAABBTree3.h"
 #include "Properties/MeshMaterialProperties.h"
-#include "Properties/MeshStatisticsProperties.h"
 #include "MeshInspectorTool.generated.h"
 
 
diff --git a/Engine/Plugins/Experimental/MeshModelingToolset/Source/MeshModelingTools/Public/Properties/MeshAnalysisProperties.h b/Engine/Plugins/Experimental/MeshModelingToolset/Source/MeshModelingTools/Public/Properties/MeshAnalysisProperties.h
new file mode 100644
--- /dev/null
+++ b/Engine/Plugins/Experimental/MeshModelingToolset/Source/MeshModelingTools/Public/Properties/MeshAnalysisProperties.h
@@ -0,0 +1,31 @@
+// Copyright 1998-2019 Epic Games, Inc. All Rights Reserved.
+
+#pragma once
+
+#include "CoreMinimal.h"
+
+#include "InteractiveToolBuilder.h"
+
+#include "MeshAnalysisProperties.generated.h"
+
+
+// predeclarations
+class FDynamicMesh3;
+
+
+
+UCLASS()
+class MESHMODELINGTOOLS_API UMeshAnalysisProperties : public UInteractiveToolPropertySet
+{
+	GENERATED_BODY()
+public:
+
+	UPROPERTY(VisibleAnywhere, Category = MeshAnalysis)
+	FString SurfaceArea;
+
+	UPROPERTY(VisibleAnywhere, Category = MeshAnalysis)
+	FString Volume;
+
+	void Update(const FDynamicMesh3& Mesh, const FTransform& Transform);
+};
+
