normal repair option for edit normals tool

#rnx

diff --git a/Engine/Plugins/Experimental/GeometryProcessing/Source/DynamicMesh/Private/Operations/RepairOrientation.cpp b/Engine/Plugins/Experimental/GeometryProcessing/Source/DynamicMesh/Private/Operations/RepairOrientation.cpp
new file mode 100644
--- /dev/null
+++ b/Engine/Plugins/Experimental/GeometryProcessing/Source/DynamicMesh/Private/Operations/RepairOrientation.cpp
@@ -0,0 +1,162 @@
+// Copyright 1998-2019 Epic Games, Inc. All Rights Reserved.
+
+// Port of geometry3Sharp MeshRepairOrientation
+
+#include "Operations/RepairOrientation.h"
+
+#include "Async/ParallelFor.h"
+
+#include "DynamicMeshEditor.h"
+
+// TODO:
+//  - (in merge coincident) don't merge tris with same/opposite normals (option)
+//  - after orienting components, try to find adjacent open components and
+//    transfer orientation between them
+//  - orient via nesting
+
+
+void FMeshRepairOrientation::OrientComponents()
+{
+	Components.Reset();
+	
+	TSet<int> Remaining; Remaining.Reserve(Mesh->MaxTriangleID());
+	for (int TID : Mesh->TriangleIndicesItr())
+	{
+		Remaining.Add(TID);
+	}
+	auto PopOneTri = [&Remaining]()
+	{
+		for (int One : Remaining)
+		{
+			Remaining.Remove(One);
+			return One;
+		}
+		check(false);
+		return FDynamicMesh3::InvalidID;
+	};
+	TArray<int> Stack;
+	while (Remaining.Num() > 0)
+	{
+		Component& C = Components.Emplace_GetRef();
+	
+		Stack.Empty();
+		int Start = PopOneTri();
+		C.Triangles.Add(Start);
+		Stack.Add(Start);
+		while (Stack.Num() > 0) {
+			int Cur = Stack.Pop(false);
+			FIndex3i tcur = Mesh->GetTriangle(Cur);
+	
+			FIndex3i nbrs = Mesh->GetTriNeighbourTris(Cur);
+			for (int j = 0; j < 3; ++j) {
+				int nbr = nbrs[j];
+				if (Remaining.Contains(nbr) == false)
+				{
+					continue;
+				}
+	
+				int a = tcur[j];
+				int b = tcur[(j+1)%3];
+	
+				FIndex3i tnbr = Mesh->GetTriangle(nbr);
+				if (IndexUtil::FindTriOrderedEdge(b, a, tnbr) == IndexConstants::InvalidID)
+				{
+					ensure(Mesh->ReverseTriOrientation(nbr) == EMeshResult::Ok);
+				}
+				Stack.Add(nbr);
+				Remaining.Remove(nbr);
+				C.Triangles.Add(nbr);
+			}
+		}
+	}
+}
+
+	
+void FMeshRepairOrientation::ComputeStatistics(FDynamicMeshAABBTree3* Tree)
+{
+	for (Component& C : Components)
+	{
+		ComputeComponentStatistics(Tree, C);
+	}
+}
+
+
+void FMeshRepairOrientation::ComputeComponentStatistics(FDynamicMeshAABBTree3* Tree, Component& C)
+{
+	C.InFacing = C.OutFacing = 0;
+	double Dist = 2 * Mesh->GetCachedBounds().DiagonalLength();
+	
+	// only want to raycast triangles in this component
+	TSet<int> TrisInComponent; TrisInComponent.Append(C.Triangles);
+	Tree->TriangleFilterF = [&TrisInComponent](int TID)
+	{
+		return TrisInComponent.Contains(TID);
+	};
+	
+	// We want to try to figure out what is 'outside' relative to the world.
+	// Assumption is that faces we can hit from far away should be oriented outwards.
+	// So, for each triangle we construct far-away points in positive and negative normal
+	// direction, then raycast back towards the triangle. If we hit the triangle from
+	// one side and not the other, that is evidence we should keep/reverse that triangle.
+	// If it is not hit, or hit from both, that does not provide any evidence.
+	// We collect up this keep/reverse evidence and use the larger to decide on the global orientation.
+	
+	FCriticalSection Mutex;
+
+	// note the C# version of this explicitly set a block size for parallel groups, and used a spin lock; TODO: revisit whether that should be ported / if the naive version is causing badness
+	bool bNoParallel = false;
+	ParallelFor(C.Triangles.Num(), [this, &Tree, &Mutex, &C, &Dist](int32 Idx)
+	{
+		int TID = C.Triangles[Idx];
+		FVector3d Normal, Centroid;
+		double Area;
+		Mesh->GetTriInfo(TID, Normal, Area, Centroid);
+		if (Area < FMathf::ZeroTolerance)
+		{
+			return;
+		}
+
+		FVector3d PosPt = Centroid + Dist * Normal;
+		FVector3d NegPt = Centroid - Dist * Normal;
+		int HitPos = Tree->FindNearestHitTriangle(FRay3d(PosPt, -Normal));
+		int HitNeg = Tree->FindNearestHitTriangle(FRay3d(NegPt, Normal));
+		if (HitPos != TID && HitNeg != TID)
+		{
+			return; // no evidence
+		}
+		if (HitPos == TID && HitNeg == TID)
+		{
+			return; // no evidence?
+		}
+
+		{
+			FScopeLock Lock(&Mutex);
+			if (HitNeg == TID)
+			{
+				C.InFacing += Area;
+			}
+			else if (HitPos == TID)
+			{
+				C.OutFacing += Area;
+			}
+		}
+	},
+		bNoParallel);
+	
+	Tree->TriangleFilterF = nullptr;
+}
+
+
+void FMeshRepairOrientation::SolveGlobalOrientation(FDynamicMeshAABBTree3* Tree)
+{
+	check(Tree->GetMesh() == Mesh);
+	ComputeStatistics(Tree);
+	FDynamicMeshEditor Editor(Mesh);
+	for (const Component& C : Components)
+	{
+		if (C.InFacing > C.OutFacing)
+		{
+			Editor.ReverseTriangleOrientations(C.Triangles, true);
+		}
+	}
+}
diff --git a/Engine/Plugins/Experimental/GeometryProcessing/Source/DynamicMesh/Public/Operations/RepairOrientation.h b/Engine/Plugins/Experimental/GeometryProcessing/Source/DynamicMesh/Public/Operations/RepairOrientation.h
new file mode 100644
--- /dev/null
+++ b/Engine/Plugins/Experimental/GeometryProcessing/Source/DynamicMesh/Public/Operations/RepairOrientation.h
@@ -0,0 +1,58 @@
+// Copyright 1998-2019 Epic Games, Inc. All Rights Reserved.
+
+// Port of geometry3Sharp MeshRepairOrientation
+
+#pragma once
+
+#include "MathUtil.h"
+#include "VectorTypes.h"
+#include "DynamicMesh3.h"
+#include "DynamicMeshAABBTree3.h"
+
+/**
+ * Invert triangles as needed to a consistent ~"outward" orientation
+ */
+class DYNAMICMESH_API FMeshRepairOrientation
+{
+public:
+	//
+	// Inputs
+	//
+	FDynamicMesh3* Mesh;
+
+	FMeshRepairOrientation(FDynamicMesh3* Mesh) : Mesh(Mesh)
+	{
+	}
+	virtual ~FMeshRepairOrientation() {}
+	
+	
+	/**
+	 * @return EOperationValidationResult::Ok if we can apply operation, or error code if we cannot
+	 */
+	virtual EOperationValidationResult Validate()
+	{
+		// TODO: validate input
+
+		return EOperationValidationResult::Ok;
+	}
+	
+	// locally orient connected components
+	void OrientComponents();
+
+	// globally orient -- to be called after OrientComponents if a more globally consistent result is desired
+	void SolveGlobalOrientation(FDynamicMeshAABBTree3* Tree);
+
+protected:
+
+	struct Component
+	{
+		TArray<int> Triangles;
+		double OutFacing;
+		double InFacing;
+	};
+	TArray<Component> Components;
+
+	void ComputeStatistics(FDynamicMeshAABBTree3* Tree);
+	void ComputeComponentStatistics(FDynamicMeshAABBTree3* Tree, Component& c);
+};
+
diff --git a/Engine/Plugins/Experimental/MeshModelingToolset/Source/MeshModelingTools/Private/EditNormalsTool.cpp b/Engine/Plugins/Experimental/MeshModelingToolset/Source/MeshModelingTools/Private/EditNormalsTool.cpp
--- a/Engine/Plugins/Experimental/MeshModelingToolset/Source/MeshModelingTools/Private/EditNormalsTool.cpp
+++ b/Engine/Plugins/Experimental/MeshModelingToolset/Source/MeshModelingTools/Private/EditNormalsTool.cpp
@@ -64,12 +64,13 @@ UInteractiveTool* UEditNormalsToolBuilder::BuildTool(const FToolBuilderState& Sc
 
 UEditNormalsToolProperties::UEditNormalsToolProperties()
 {
+	bFixInconsistentNormals = false;
 	bInvertNormals = false;
 	bRecomputeNormals = true;
 	NormalCalculationMethod = ENormalCalculationMethod::AreaAngleWeighting;
 	bRecomputeNormalTopologyAndEdgeSharpness = false;
 	SharpEdgeAngleThreshold = 60;
-	bAllowSharpVertices = true;
+	bAllowSharpVertices = false;
 }
 
 UEditNormalsAdvancedProperties::UEditNormalsAdvancedProperties()
@@ -180,6 +181,7 @@ void UEditNormalsTool::SetAssetAPI(IToolsContextAssetAPI* AssetAPIIn)
 TSharedPtr<FDynamicMeshOperator> UEditNormalsOperatorFactory::MakeNewOperator()
 {
 	TSharedPtr<FEditNormalsOp> NormalsOp = MakeShared<FEditNormalsOp>();
+	NormalsOp->bFixInconsistentNormals = Tool->BasicProperties->bFixInconsistentNormals;
 	NormalsOp->bInvertNormals = Tool->BasicProperties->bInvertNormals;
 	NormalsOp->bRecomputeNormals = Tool->BasicProperties->bRecomputeNormals;
 	NormalsOp->bSplitNormals = Tool->BasicProperties->bRecomputeNormalTopologyAndEdgeSharpness;
@@ -263,7 +265,8 @@ void UEditNormalsTool::GenerateAsset(const TArray<TUniquePtr<FDynamicMeshOpResul
 		ComponentTargets[ComponentIdx]->CommitMesh([&Results, &ComponentIdx, this](FMeshDescription* MeshDescription)
 		{
 			FDynamicMeshToMeshDescription Converter;
-			if (BasicProperties->bRecomputeNormalTopologyAndEdgeSharpness || BasicProperties->bInvertNormals)
+			
+			if (BasicProperties->WillTopologyChange())
 			{
 				// full conversion if normal topology changed or faces were inverted
 				Converter.Convert(Results[ComponentIdx]->Mesh.Get(), *MeshDescription);
diff --git a/Engine/Plugins/Experimental/MeshModelingToolset/Source/MeshModelingTools/Public/EditNormalsTool.h b/Engine/Plugins/Experimental/MeshModelingToolset/Source/MeshModelingTools/Public/EditNormalsTool.h
--- a/Engine/Plugins/Experimental/MeshModelingToolset/Source/MeshModelingTools/Public/EditNormalsTool.h
+++ b/Engine/Plugins/Experimental/MeshModelingToolset/Source/MeshModelingTools/Public/EditNormalsTool.h
@@ -51,6 +51,10 @@ class MESHMODELINGTOOLS_API UEditNormalsToolProperties : public UInteractiveTool
 public:
 	UEditNormalsToolProperties();
 
+	bool WillTopologyChange()
+	{
+		return bFixInconsistentNormals || bInvertNormals || bRecomputeNormalTopologyAndEdgeSharpness;
+	}
 
 	/** Recompute all mesh normals */
 	UPROPERTY(EditAnywhere, Category = NormalsCalculation, meta = (EditCondition = "!bRecomputeNormalTopologyAndEdgeSharpness"))
@@ -60,6 +64,10 @@ public:
 	UPROPERTY(EditAnywhere, Category = NormalsCalculation)
 	ENormalCalculationMethod NormalCalculationMethod;
 
+	/** For meshes with inconsistent triangle orientations/normals, flip as needed to make the normals consistent */
+	UPROPERTY(EditAnywhere, Category = NormalsCalculation)
+	bool bFixInconsistentNormals;
+
 	/** Invert (flip) all mesh normals and associated triangle orientations */
 	UPROPERTY(EditAnywhere, Category = NormalsCalculation)
 	bool bInvertNormals;
diff --git a/Engine/Plugins/Experimental/MeshModelingToolset/Source/ModelingOperators/Private/CleaningOps/EditNormalsOp.cpp b/Engine/Plugins/Experimental/MeshModelingToolset/Source/ModelingOperators/Private/CleaningOps/EditNormalsOp.cpp
--- a/Engine/Plugins/Experimental/MeshModelingToolset/Source/ModelingOperators/Private/CleaningOps/EditNormalsOp.cpp
+++ b/Engine/Plugins/Experimental/MeshModelingToolset/Source/ModelingOperators/Private/CleaningOps/EditNormalsOp.cpp
@@ -12,7 +12,8 @@
 
 
 #include "MeshNormals.h"
-
+#include "Operations/RepairOrientation.h"
+#include "DynamicMeshAABBTree3.h"
 
 void FEditNormalsOp::SetTransform(const FTransform& Transform) {
 	ResultTransform = (FTransform3d)Transform;
@@ -35,6 +36,14 @@ void FEditNormalsOp::CalculateResult(FProgressCancel* Progress)
 	// if you split normals you must always recompute as well
 	bool bNeedsRecompute = bRecomputeNormals || bSplitNormals;
 
+	if (bFixInconsistentNormals)
+	{
+		FMeshRepairOrientation Repair(ResultMesh.Get());
+		Repair.OrientComponents();
+		FDynamicMeshAABBTree3 Tree(ResultMesh.Get());
+		Repair.SolveGlobalOrientation(&Tree);
+	}
+
 	if (bInvertNormals)
 	{
 		for (int TID : ResultMesh->TriangleIndicesItr())
diff --git a/Engine/Plugins/Experimental/MeshModelingToolset/Source/ModelingOperators/Public/CleaningOps/EditNormalsOp.h b/Engine/Plugins/Experimental/MeshModelingToolset/Source/ModelingOperators/Public/CleaningOps/EditNormalsOp.h
--- a/Engine/Plugins/Experimental/MeshModelingToolset/Source/ModelingOperators/Public/CleaningOps/EditNormalsOp.h
+++ b/Engine/Plugins/Experimental/MeshModelingToolset/Source/ModelingOperators/Public/CleaningOps/EditNormalsOp.h
@@ -27,6 +27,7 @@ public:
 	// inputs
 	TSharedPtr<FDynamicMesh3> OriginalMesh;
 
+	bool bFixInconsistentNormals;
 	bool bInvertNormals;
 	bool bRecomputeNormals;
 	ENormalCalculationMethod NormalCalculationMethod;
