Fix ue4game html5 compile issue in RepairOrientation.cpp
#rnx

diff --git a/Engine/Plugins/Experimental/GeometryProcessing/Source/DynamicMesh/Private/Operations/RepairOrientation.cpp b/Engine/Plugins/Experimental/GeometryProcessing/Source/DynamicMesh/Private/Operations/RepairOrientation.cpp
--- a/Engine/Plugins/Experimental/GeometryProcessing/Source/DynamicMesh/Private/Operations/RepairOrientation.cpp
+++ b/Engine/Plugins/Experimental/GeometryProcessing/Source/DynamicMesh/Private/Operations/RepairOrientation.cpp
@@ -8,6 +8,8 @@
 
 #include "DynamicMeshEditor.h"
 
+#include "Misc/ScopeLock.h"
+
 // TODO:
 //  - (in merge coincident) don't merge tris with same/opposite normals (option)
 //  - after orienting components, try to find adjacent open components and
