generic attribute sets can be registered on a mesh, with change tracking
support for per-vertex attributes
add change tracking for material id

diff --git a/Engine/Plugins/Experimental/GeometryProcessing/Source/DynamicMesh/Private/DynamicMeshAttributeSet.cpp b/Engine/Plugins/Experimental/GeometryProcessing/Source/DynamicMesh/Private/DynamicMeshAttributeSet.cpp
--- a/Engine/Plugins/Experimental/GeometryProcessing/Source/DynamicMesh/Private/DynamicMeshAttributeSet.cpp
+++ b/Engine/Plugins/Experimental/GeometryProcessing/Source/DynamicMesh/Private/DynamicMeshAttributeSet.cpp
@@ -42,8 +42,22 @@ bool FDynamicMeshAttributeSet::IsSeamVertex(int VID, bool bBoundaryIsSeam) const
 }
 
 
+void FDynamicMeshAttributeSet::OnNewVertex(int VertexID, bool bInserted)
+{
+	FDynamicAttributeSetBase::OnNewVertex(VertexID, bInserted);
+}
+
+
+void FDynamicMeshAttributeSet::OnRemoveVertex(int VertexID)
+{
+	FDynamicAttributeSetBase::OnRemoveVertex(VertexID);
+}
+
+
 void FDynamicMeshAttributeSet::OnNewTriangle(int TriangleID, bool bInserted)
 {
+	FDynamicAttributeSetBase::OnNewTriangle(TriangleID, bInserted);
+
 	for (FDynamicMeshUVOverlay& UVLayer : UVLayers)
 	{
 		UVLayer.InitializeNewTriangle(TriangleID);
@@ -60,6 +74,8 @@ void FDynamicMeshAttributeSet::OnNewTriangle(int TriangleID, bool bInserted)
 
 void FDynamicMeshAttributeSet::OnRemoveTriangle(int TriangleID)
 {
+	FDynamicAttributeSetBase::OnRemoveTriangle(TriangleID);
+
 	for (FDynamicMeshUVOverlay& UVLayer : UVLayers)
 	{
 		UVLayer.OnRemoveTriangle(TriangleID);
@@ -71,6 +87,8 @@ void FDynamicMeshAttributeSet::OnRemoveTriangle(int TriangleID)
 
 void FDynamicMeshAttributeSet::OnReverseTriOrientation(int TriangleID)
 {
+	FDynamicAttributeSetBase::OnReverseTriOrientation(TriangleID);
+
 	for (FDynamicMeshUVOverlay& UVLayer : UVLayers)
 	{
 		UVLayer.OnReverseTriOrientation(TriangleID);
@@ -80,22 +98,26 @@ void FDynamicMeshAttributeSet::OnReverseTriOrientation(int TriangleID)
 	// has no effect on MaterialIDAttrib
 }
 
-void FDynamicMeshAttributeSet::OnSplitEdge(const FDynamicMesh3::FEdgeSplitInfo & splitInfo)
+void FDynamicMeshAttributeSet::OnSplitEdge(const FDynamicMesh3::FEdgeSplitInfo& SplitInfo)
 {
+	FDynamicAttributeSetBase::OnSplitEdge(SplitInfo);
+
 	for (FDynamicMeshUVOverlay& UVLayer : UVLayers)
 	{
-		UVLayer.OnSplitEdge(splitInfo);
+		UVLayer.OnSplitEdge(SplitInfo);
 	}
-	Normals0.OnSplitEdge(splitInfo);
+	Normals0.OnSplitEdge(SplitInfo);
 
 	if (MaterialIDAttrib)
 	{
-		MaterialIDAttrib->OnSplitEdge(splitInfo);
+		MaterialIDAttrib->OnSplitEdge(SplitInfo);
 	}
 }
 
 void FDynamicMeshAttributeSet::OnFlipEdge(const FDynamicMesh3::FEdgeFlipInfo & flipInfo)
 {
+	FDynamicAttributeSetBase::OnFlipEdge(flipInfo);
+
 	for (FDynamicMeshUVOverlay& UVLayer : UVLayers)
 	{
 		UVLayer.OnFlipEdge(flipInfo);
@@ -111,6 +133,8 @@ void FDynamicMeshAttributeSet::OnFlipEdge(const FDynamicMesh3::FEdgeFlipInfo & f
 
 void FDynamicMeshAttributeSet::OnCollapseEdge(const FDynamicMesh3::FEdgeCollapseInfo & collapseInfo)
 {
+	FDynamicAttributeSetBase::OnCollapseEdge(collapseInfo);
+
 	for (FDynamicMeshUVOverlay& UVLayer : UVLayers)
 	{
 		UVLayer.OnCollapseEdge(collapseInfo);
@@ -125,6 +149,8 @@ void FDynamicMeshAttributeSet::OnCollapseEdge(const FDynamicMesh3::FEdgeCollapse
 
 void FDynamicMeshAttributeSet::OnPokeTriangle(const FDynamicMesh3::FPokeTriangleInfo & pokeInfo)
 {
+	FDynamicAttributeSetBase::OnPokeTriangle(pokeInfo);
+
 	for (FDynamicMeshUVOverlay& UVLayer : UVLayers)
 	{
 		UVLayer.OnPokeTriangle(pokeInfo);
@@ -139,6 +165,8 @@ void FDynamicMeshAttributeSet::OnPokeTriangle(const FDynamicMesh3::FPokeTriangle
 
 void FDynamicMeshAttributeSet::OnMergeEdges(const FDynamicMesh3::FMergeEdgesInfo & mergeInfo)
 {
+	FDynamicAttributeSetBase::OnMergeEdges(mergeInfo);
+
 	for (FDynamicMeshUVOverlay& UVLayer : UVLayers)
 	{
 		UVLayer.OnMergeEdges(mergeInfo);
diff --git a/Engine/Plugins/Experimental/GeometryProcessing/Source/DynamicMesh/Private/DynamicMeshChangeTracker.cpp b/Engine/Plugins/Experimental/GeometryProcessing/Source/DynamicMesh/Private/DynamicMeshChangeTracker.cpp
--- a/Engine/Plugins/Experimental/GeometryProcessing/Source/DynamicMesh/Private/DynamicMeshChangeTracker.cpp
+++ b/Engine/Plugins/Experimental/GeometryProcessing/Source/DynamicMesh/Private/DynamicMeshChangeTracker.cpp
@@ -393,6 +393,17 @@ void FDynamicMeshAttributeSetChangeTracker::BeginChange()
 			State.StartElements[ElemID] = true;
 		}
 	}
+
+	if (Attribs->HasMaterialID())
+	{
+		Change->MaterialIDAttribChange.Emplace();
+	}
+
+	for (int k = 0, NumRegAttrib = Attribs->NumRegisteredAttributes(); k < NumRegAttrib; k++)
+	{
+		Change->RegisteredAttributeChanges.Add(Attribs->GetRegisteredAttribute(k)->NewBlankChange());
+	}
+
 }
 
 
@@ -429,6 +440,25 @@ void FDynamicMeshAttributeSetChangeTracker::SaveInitialTriangle(int TriangleID)
 		SaveElement(NormTriangle.C, NormalStates[k], NormalLayer, NormalChange);
 		NormalChange.SaveInitialTriangle(NormalLayer, TriangleID);
 	}
+
+	if (Change->MaterialIDAttribChange.IsSet())
+	{
+		Change->MaterialIDAttribChange->SaveInitialTriangle(Attribs->GetMaterialID(), TriangleID);
+	}
+
+	for (int k = 0, NumRegAttrib = Attribs->NumRegisteredAttributes(); k < NumRegAttrib; k++)
+	{
+		Change->RegisteredAttributeChanges[k]->SaveInitialTriangle(Attribs->GetRegisteredAttribute(k), TriangleID);
+	}
+}
+
+
+void FDynamicMeshAttributeSetChangeTracker::SaveInitialVertex(int VertexID)
+{
+	for (int k = 0, NumRegAttrib = Attribs->NumRegisteredAttributes(); k < NumRegAttrib; k++)
+	{
+		Change->RegisteredAttributeChanges[k]->SaveInitialVertex(Attribs->GetRegisteredAttribute(k), VertexID);
+	}
 }
 
 
@@ -483,6 +513,24 @@ void FDynamicMeshAttributeSetChangeTracker::StoreAllFinalTriangles(const TArray<
 		}
 	}
 
+	if (Change->MaterialIDAttribChange.IsSet())
+	{
+		Change->MaterialIDAttribChange->StoreAllFinalTriangles(Attribs->GetMaterialID(), TriangleIDs);
+	}
+
+	for (int k = 0, NumRegAttrib = Attribs->NumRegisteredAttributes(); k < NumRegAttrib; k++)
+	{
+		Change->RegisteredAttributeChanges[k]->StoreAllFinalTriangles(Attribs->GetRegisteredAttribute(k), TriangleIDs);
+	}
+}
+
+
+void FDynamicMeshAttributeSetChangeTracker::StoreAllFinalVertices(const TArray<int>& VertexIDs)
+{
+	for (int k = 0, NumRegAttrib = Attribs->NumRegisteredAttributes(); k < NumRegAttrib; k++)
+	{
+		Change->RegisteredAttributeChanges[k]->StoreAllFinalVertices(Attribs->GetRegisteredAttribute(k), VertexIDs);
+	}
 }
 
 
@@ -504,6 +552,16 @@ bool FDynamicMeshAttributeChangeSet::Apply(FDynamicMeshAttributeSet* Attributes,
 		NormalChanges[k].Apply(NormalLayer, bRevert);
 	}
 
+	for (int k = 0, NumRegAttrib = Attributes->NumRegisteredAttributes(); k < NumRegAttrib; k++)
+	{
+		RegisteredAttributeChanges[k]->Apply(Attributes->GetRegisteredAttribute(k), bRevert);
+	}
+
+	if (MaterialIDAttribChange.IsSet())
+	{
+		MaterialIDAttribChange->Apply(Attributes->GetMaterialID(), bRevert);
+	}
+
 	return true;
 }
 
@@ -591,6 +649,11 @@ void FDynamicMeshChangeTracker::SaveVertex(int VertexID)
 	Change->SaveInitialVertex(Mesh, VertexID);
 
 	ChangedVertices[VertexID] = true;
+
+	if (AttribChangeTracker != nullptr)
+	{
+		AttribChangeTracker->SaveInitialVertex(VertexID);
+	}
 }
 
 
diff --git a/Engine/Plugins/Experimental/GeometryProcessing/Source/DynamicMesh/Private/DynamicMeshOverlay.cpp b/Engine/Plugins/Experimental/GeometryProcessing/Source/DynamicMesh/Private/DynamicMeshOverlay.cpp
--- a/Engine/Plugins/Experimental/GeometryProcessing/Source/DynamicMesh/Private/DynamicMeshOverlay.cpp
+++ b/Engine/Plugins/Experimental/GeometryProcessing/Source/DynamicMesh/Private/DynamicMeshOverlay.cpp
@@ -560,6 +560,13 @@ void TDynamicMeshOverlay<RealType, ElementSize>::OnFlipEdge(const FDynamicMesh3:
 {
 	int orig_t0 = FlipInfo.Triangles.A;
 	int orig_t1 = FlipInfo.Triangles.B;
+	bool bT0Set = IsSetTriangle(orig_t0), bT1Set = IsSetTriangle(orig_t1);
+	if (!bT0Set)
+	{
+		check(!bT1Set); // flipping across a set/unset boundary is not allowed?
+		return; // nothing to do on the overlay if both triangles are unset
+	}
+
 	int base_a = FlipInfo.OriginalVerts.A;
 	int base_b = FlipInfo.OriginalVerts.B;
 	int base_c = FlipInfo.OpposingVerts.A;
@@ -613,6 +620,8 @@ void TDynamicMeshOverlay<RealType, ElementSize>::OnFlipEdge(const FDynamicMesh3:
 template<typename RealType, int ElementSize>
 void TDynamicMeshOverlay<RealType, ElementSize>::OnCollapseEdge(const FDynamicMesh3::FEdgeCollapseInfo& collapseInfo)
 {
+	// TODO: support case where (either or both) triangles were unset in overlay
+
 	int vid_base_kept = collapseInfo.KeptVertex;
 	int vid_base_removed = collapseInfo.RemovedVertex;
 	int tid_removed0 = collapseInfo.RemovedTris.A;
@@ -701,6 +710,13 @@ void TDynamicMeshOverlay<RealType, ElementSize>::OnCollapseEdge(const FDynamicMe
 template<typename RealType, int ElementSize>
 void TDynamicMeshOverlay<RealType, ElementSize>::OnPokeTriangle(const FDynamicMesh3::FPokeTriangleInfo& PokeInfo)
 {
+	if (!IsSetTriangle(PokeInfo.OriginalTriangle))
+	{
+		InitializeNewTriangle(PokeInfo.NewTriangles.A);
+		InitializeNewTriangle(PokeInfo.NewTriangles.B);
+		return;
+	}
+
 	FIndex3i Triangle = GetTriangle(PokeInfo.OriginalTriangle);
 
 	// create new element at barycentric position
@@ -709,9 +725,9 @@ void TDynamicMeshOverlay<RealType, ElementSize>::OnPokeTriangle(const FDynamicMe
 	SetElementFromBary(CenterElemID, Triangle[0], Triangle[1], Triangle[2], BaryCoords);
 
 	// update orig triangle and two new ones. Winding orders here mirror FDynamicMesh3::PokeTriangle
-	SetTriangle(PokeInfo.OriginalTriangle, FIndex3i(Triangle[0], Triangle[1], CenterElemID) );
-	SetTriangle(PokeInfo.NewTriangles.A, FIndex3i(Triangle[1], Triangle[2], CenterElemID));
-	SetTriangle(PokeInfo.NewTriangles.B, FIndex3i(Triangle[2], Triangle[0], CenterElemID));
+	InternalSetTriangle(PokeInfo.OriginalTriangle, FIndex3i(Triangle[0], Triangle[1], CenterElemID), false );
+	InternalSetTriangle(PokeInfo.NewTriangles.A, FIndex3i(Triangle[1], Triangle[2], CenterElemID), false);
+	InternalSetTriangle(PokeInfo.NewTriangles.B, FIndex3i(Triangle[2], Triangle[0], CenterElemID), false);
 
 	ElementsRefCounts.Increment(Triangle[0]);
 	ElementsRefCounts.Increment(Triangle[1]);
diff --git a/Engine/Plugins/Experimental/GeometryProcessing/Source/DynamicMesh/Public/DynamicAttribute.h b/Engine/Plugins/Experimental/GeometryProcessing/Source/DynamicMesh/Public/DynamicAttribute.h
new file mode 100644
--- /dev/null
+++ b/Engine/Plugins/Experimental/GeometryProcessing/Source/DynamicMesh/Public/DynamicAttribute.h
@@ -0,0 +1,222 @@
+// Copyright 1998-2019 Epic Games, Inc. All Rights Reserved.
+
+#include "DynamicMesh3.h"
+
+#pragma once
+
+
+class FDynamicAttributeBase;
+
+
+/**
+* Generic base class for change tracking of an attribute layer
+*/
+class DYNAMICMESH_API FDynamicAttributeChangeBase
+{
+public:
+	virtual ~FDynamicAttributeChangeBase()
+	{
+	}
+
+	// default do-nothing implementations are provided because many attribute layers will only care about some kinds of elements and won't implement all of these
+
+	virtual void SaveInitialTriangle(const FDynamicAttributeBase* Attribute, int TriangleID)
+	{
+	}
+	virtual void SaveInitialVertex(const FDynamicAttributeBase* Attribute, int VertexID)
+	{
+	}
+
+	virtual void StoreAllFinalTriangles(const FDynamicAttributeBase* Attribute, const TArray<int>& TriangleIDs)
+	{
+	}
+	virtual void StoreAllFinalVertices(const FDynamicAttributeBase* Attribute, const TArray<int>& VertexIDs)
+	{
+	}
+
+	virtual bool Apply(FDynamicAttributeBase* Attribute, bool bRevert) const
+	{
+		return false;
+	}
+};
+
+
+/**
+ * Base class for attributes that live on a dynamic mesh (or similar dynamic object)
+ *
+ * Subclasses can override the On* functions to ensure the attribute remains up to date through changes to the dynamic object
+ */
+class DYNAMICMESH_API FDynamicAttributeBase
+{
+
+public:
+	virtual ~FDynamicAttributeBase()
+	{
+	}
+
+public:
+
+	/** Allocate a new copy of the attribute layer, optionally with a different parent mesh */
+	virtual FDynamicAttributeBase* MakeCopy(FDynamicMesh3* ParentMeshIn) const = 0;
+
+	/** Update to reflect an edge split in the parent mesh */
+	virtual void OnSplitEdge(const FDynamicMesh3::FEdgeSplitInfo& SplitInfo)
+	{
+	}
+
+	/** Update to reflect an edge flip in the parent mesh */
+	virtual void OnFlipEdge(const FDynamicMesh3::FEdgeFlipInfo& FlipInfo)
+	{
+	}
+
+	/** Update to reflect an edge collapse in the parent mesh */
+	virtual void OnCollapseEdge(const FDynamicMesh3::FEdgeCollapseInfo& CollapseInfo)
+	{
+	}
+
+	virtual void OnNewVertex(int VertexID, bool bInserted)
+	{
+	}
+
+	virtual void OnRemoveVertex(int VertexID)
+	{
+	}
+
+	virtual void OnNewTriangle(int TriangleID, bool bInserted)
+	{
+	}
+
+	virtual void OnRemoveTriangle(int TriangleID)
+	{
+	}
+
+	/** Update to reflect a face poke in the parent mesh */
+	virtual void OnPokeTriangle(const FDynamicMesh3::FPokeTriangleInfo& PokeInfo)
+	{
+	}
+
+	/** Update to reflect an edge merge in the parent mesh */
+	virtual void OnMergeEdges(const FDynamicMesh3::FMergeEdgesInfo& MergeInfo)
+	{
+	}
+
+	virtual void OnReverseTriOrientation(int TriangleID)
+	{
+	}
+
+
+	virtual TUniquePtr<FDynamicAttributeChangeBase> NewBlankChange() = 0;
+
+};
+
+
+/**
+* Generic base class for managing a set of registered attributes that must all be kept up to date
+*/
+class DYNAMICMESH_API FDynamicAttributeSetBase
+{
+protected:
+	// not managed by the base class; we should be able to register any attributes here that we want to be automatically updated
+	TArray<FDynamicAttributeBase*> RegisteredAttributes;
+
+	/**
+	 * Stores the given attribute pointer in the attribute register, so that it will be updated with mesh changes, but does not take ownership of the attribute memory.
+	 */
+	void RegisterExternalAttribute(FDynamicAttributeBase* Attribute)
+	{
+		RegisteredAttributes.Add(Attribute);
+	}
+
+	void ResetRegisteredAttributes()
+	{
+		RegisteredAttributes.Reset();
+	}
+
+public:
+	virtual ~FDynamicAttributeSetBase()
+	{
+	}
+
+	int NumRegisteredAttributes() const
+	{
+		return RegisteredAttributes.Num();
+	}
+
+	FDynamicAttributeBase* GetRegisteredAttribute(int Idx) const
+	{
+		return RegisteredAttributes[Idx];
+	}
+
+	// These functions are called by the FDynamicMesh3 to update the various
+	// attributes when the parent mesh topology has been modified.
+	virtual void OnNewTriangle(int TriangleID, bool bInserted)
+	{
+		for (FDynamicAttributeBase* A : RegisteredAttributes)
+		{
+			A->OnNewTriangle(TriangleID, bInserted);
+		}
+	}
+	virtual void OnNewVertex(int VertexID, bool bInserted)
+	{
+		for (FDynamicAttributeBase* A : RegisteredAttributes)
+		{
+			A->OnNewVertex(VertexID, bInserted);
+		}
+	}
+	virtual void OnRemoveTriangle(int TriangleID)
+	{
+		for (FDynamicAttributeBase* A : RegisteredAttributes)
+		{
+			A->OnRemoveTriangle(TriangleID);
+		}
+	}
+	virtual void OnRemoveVertex(int VertexID)
+	{
+		for (FDynamicAttributeBase* A : RegisteredAttributes)
+		{
+			A->OnRemoveTriangle(VertexID);
+		}
+	}
+	virtual void OnReverseTriOrientation(int TriangleID)
+	{
+		for (FDynamicAttributeBase* A : RegisteredAttributes)
+		{
+			A->OnReverseTriOrientation(TriangleID);
+		}
+	}
+	virtual void OnSplitEdge(const FDynamicMesh3::FEdgeSplitInfo& SplitInfo)
+	{
+		for (FDynamicAttributeBase* A : RegisteredAttributes)
+		{
+			A->OnSplitEdge(SplitInfo);
+		}
+	}
+	virtual void OnFlipEdge(const FDynamicMesh3::FEdgeFlipInfo& FlipInfo)
+	{
+		for (FDynamicAttributeBase* A : RegisteredAttributes)
+		{
+			A->OnFlipEdge(FlipInfo);
+		}
+	}
+	virtual void OnCollapseEdge(const FDynamicMesh3::FEdgeCollapseInfo& CollapseInfo)
+	{
+		for (FDynamicAttributeBase* A : RegisteredAttributes)
+		{
+			A->OnCollapseEdge(CollapseInfo);
+		}
+	}
+	virtual void OnPokeTriangle(const FDynamicMesh3::FPokeTriangleInfo& PokeInfo)
+	{
+		for (FDynamicAttributeBase* A : RegisteredAttributes)
+		{
+			A->OnPokeTriangle(PokeInfo);
+		}
+	}
+	virtual void OnMergeEdges(const FDynamicMesh3::FMergeEdgesInfo& MergeInfo)
+	{
+		for (FDynamicAttributeBase* A : RegisteredAttributes)
+		{
+			A->OnMergeEdges(MergeInfo);
+		}
+	}
+};
diff --git a/Engine/Plugins/Experimental/GeometryProcessing/Source/DynamicMesh/Public/DynamicMeshAttributeSet.h b/Engine/Plugins/Experimental/GeometryProcessing/Source/DynamicMesh/Public/DynamicMeshAttributeSet.h
--- a/Engine/Plugins/Experimental/GeometryProcessing/Source/DynamicMesh/Public/DynamicMeshAttributeSet.h
+++ b/Engine/Plugins/Experimental/GeometryProcessing/Source/DynamicMesh/Public/DynamicMeshAttributeSet.h
@@ -5,6 +5,7 @@
 
 #include "DynamicMeshOverlay.h"
 #include "DynamicMeshTriangleAttribute.h"
+#include "DynamicAttribute.h"
 
 /** Standard UV overlay type - 2-element float */
 typedef TDynamicMeshVectorOverlay<float, 2, FVector2f> FDynamicMeshUVOverlay;
@@ -22,7 +23,7 @@ typedef TDynamicMeshTriangleAttribute<int32, 1> FDynamicMeshMaterialAttribute;
  * 
  * @todo current internal structure is a work-in-progress
  */
-class DYNAMICMESH_API FDynamicMeshAttributeSet
+class DYNAMICMESH_API FDynamicMeshAttributeSet : public FDynamicAttributeSetBase
 {
 public:
 
@@ -52,6 +53,14 @@ public:
 			MaterialIDAttrib->Copy( *(Copy.MaterialIDAttrib) );
 		}
 
+		GenericAttributes.Reset();
+		ResetRegisteredAttributes();
+		for (int Idx = 0; Idx < Copy.GenericAttributes.Num(); Idx++)
+		{
+			const FDynamicAttributeBase* SourceAttrib = Copy.GenericAttributes[Idx].Get();
+			AttachAttribute(SourceAttrib->MakeCopy(ParentMesh));
+		}
+
 		// parent mesh is *not* copied!
 	}
 
@@ -186,6 +195,24 @@ public:
 		return MaterialIDAttrib.Get();
 	}
 
+	// Attach a new attribute (and transfer ownership of it to the attribute set)
+	int AttachAttribute(FDynamicAttributeBase* Attribute)
+	{
+		int AttributeID = GenericAttributes.Add(TUniquePtr<FDynamicAttributeBase>(Attribute));
+		RegisterExternalAttribute(Attribute);
+		return AttributeID;
+	}
+
+	FDynamicAttributeBase* GetAttachedAttribute(int AttributeID)
+	{
+		return GenericAttributes[AttributeID].Get();
+	}
+
+	int NumAttachedAttributes()
+	{
+		return GenericAttributes.Num();
+	}
+
 protected:
 	/** Parent mesh of this attribute set */
 	FDynamicMesh3* ParentMesh;
@@ -198,6 +225,7 @@ protected:
 
 	TUniquePtr<FDynamicMeshMaterialAttribute> MaterialIDAttrib;
 	
+	TArray<TUniquePtr<FDynamicAttributeBase>> GenericAttributes;
 
 protected:
 	friend class FDynamicMesh3;
@@ -216,8 +244,11 @@ protected:
 
 	// These functions are called by the FDynamicMesh3 to update the various
 	// attributes when the parent mesh topology has been modified.
+	// TODO: would it be better to register all the overlays and attributes with the base set and not overload these?  maybe!
 	virtual void OnNewTriangle(int TriangleID, bool bInserted);
+	virtual void OnNewVertex(int VertexID, bool bInserted);
 	virtual void OnRemoveTriangle(int TriangleID);
+	virtual void OnRemoveVertex(int VertexID);
 	virtual void OnReverseTriOrientation(int TriangleID);
 	virtual void OnSplitEdge(const FDynamicMesh3::FEdgeSplitInfo & splitInfo);
 	virtual void OnFlipEdge(const FDynamicMesh3::FEdgeFlipInfo & flipInfo);
diff --git a/Engine/Plugins/Experimental/GeometryProcessing/Source/DynamicMesh/Public/DynamicMeshChangeTracker.h b/Engine/Plugins/Experimental/GeometryProcessing/Source/DynamicMesh/Public/DynamicMeshChangeTracker.h
--- a/Engine/Plugins/Experimental/GeometryProcessing/Source/DynamicMesh/Public/DynamicMeshChangeTracker.h
+++ b/Engine/Plugins/Experimental/GeometryProcessing/Source/DynamicMesh/Public/DynamicMeshChangeTracker.h
@@ -89,11 +89,13 @@ typedef TDynamicMeshAttributeChange<float,3> FDynamicMeshNormalChange;
 /**
  * FDynamicMeshAttributeChangeSet stores a set of UV and Normal changes for a FDynamicMesh3
  */
-class DYNAMICMESH_API FDynamicMeshAttributeChangeSet
+class DYNAMICMESH_API FDynamicMeshAttributeChangeSet : FNoncopyable
 {
 public:
 	TArray<FDynamicMeshUVChange> UVChanges;
 	TArray<FDynamicMeshNormalChange> NormalChanges;
+	TOptional<FDynamicMeshTriangleAttributeChange<int32, 1>> MaterialIDAttribChange;
+	TArray<TUniquePtr<FDynamicAttributeChangeBase>> RegisteredAttributeChanges;
 
 	/** call ::Apply() on all the UV and Normal changes */
 	bool Apply(FDynamicMeshAttributeSet* Attributes, bool bRevert) const;
@@ -123,7 +125,7 @@ public:
 	/** Store the final state of a triangle */
 	void StoreFinalTriangle(const FDynamicMesh3* Mesh, int TriangleID);
 
-	/** Attach an attribute change set to this mesh change, which will the be applied/reverted autoamtically */
+	/** Attach an attribute change set to this mesh change, which will the be applied/reverted automatically */
 	void AttachAttributeChanges(TUniquePtr<FDynamicMeshAttributeChangeSet> AttribChanges)
 	{
 		this->AttributeChanges = MoveTemp(AttribChanges);
@@ -202,6 +204,12 @@ public:
 	/** store the final state of a set of triangles */
 	void StoreAllFinalTriangles(const TArray<int>& TriangleIDs);
 
+	/** Store the initial state of a vertex */
+	void SaveInitialVertex(int VertexID);
+
+	/** store the final state of a set of vertices */
+	void StoreAllFinalVertices(const TArray<int>& VertexIDs);
+
 
 protected:
 	const FDynamicMeshAttributeSet* Attribs = nullptr;
diff --git a/Engine/Plugins/Experimental/GeometryProcessing/Source/DynamicMesh/Public/DynamicMeshTriangleAttribute.h b/Engine/Plugins/Experimental/GeometryProcessing/Source/DynamicMesh/Public/DynamicMeshTriangleAttribute.h
--- a/Engine/Plugins/Experimental/GeometryProcessing/Source/DynamicMesh/Public/DynamicMeshTriangleAttribute.h
+++ b/Engine/Plugins/Experimental/GeometryProcessing/Source/DynamicMesh/Public/DynamicMeshTriangleAttribute.h
@@ -3,6 +3,40 @@
 #pragma once
 
 #include "DynamicMesh3.h"
+#include "DynamicAttribute.h"
+
+
+template<typename AttribValueType, int AttribDimension>
+class TDynamicMeshTriangleAttribute;
+
+
+template<typename AttribValueType, int AttribDimension>
+class DYNAMICMESH_API FDynamicMeshTriangleAttributeChange : public FDynamicAttributeChangeBase
+{
+private:
+	struct FChangeTriangleAttribute
+	{
+		AttribValueType Data[AttribDimension];
+		int TriangleID;
+	};
+
+	TArray<FChangeTriangleAttribute> OldTriangleAttributes, NewTriangleAttributes;
+
+public:
+	FDynamicMeshTriangleAttributeChange()
+	{}
+
+	virtual ~FDynamicMeshTriangleAttributeChange()
+	{}
+
+	inline virtual void SaveInitialTriangle(const FDynamicAttributeBase* Attribute, int TriangleID) override;
+
+	inline virtual void StoreAllFinalTriangles(const FDynamicAttributeBase* Attribute, const TArray<int>& TriangleIDs) override;
+
+	inline virtual bool Apply(FDynamicAttributeBase* Attribute, bool bRevert) const override;
+};
+
+
 
 /**
  * TDynamicMeshTriangleAttribute is an add-on to a FDynamicMesh3 that allows for 
@@ -12,7 +46,7 @@
  * can be mirrored to the overlay via OnSplitEdge(), etc.
  */
 template<typename AttribValueType, int AttribDimension>
-class TDynamicMeshTriangleAttribute
+class TDynamicMeshTriangleAttribute : public FDynamicAttributeBase
 {
 
 protected:
@@ -43,6 +77,13 @@ public:
 	/** @return the parent mesh for this overlay */
 	FDynamicMesh3* GetParentMesh() { return ParentMesh; }
 
+	virtual FDynamicAttributeBase* MakeCopy(FDynamicMesh3* ParentMeshIn) const override
+	{
+		TDynamicMeshTriangleAttribute<AttribValueType, AttribDimension>* ToFill = new TDynamicMeshTriangleAttribute<AttribValueType, AttribDimension>(ParentMeshIn);
+		ToFill->Copy(*this);
+		return ToFill;
+	}
+
 	/** Set this overlay to contain the same arrays as the copy overlay */
 	void Copy(const TDynamicMeshTriangleAttribute<AttribValueType, AttribDimension>& Copy)
 	{
@@ -54,7 +95,7 @@ public:
 	{
 		check(ParentMesh != nullptr);
 		AttribValues.Resize(0);
-		AttribValues.Resize( ParentMesh->MaxTriangleID() * AttribDimension );
+		AttribValues.Resize( ParentMesh->MaxTriangleID() * AttribDimension, InitialValue );
 	}
 
 	void SetNewValue(int NewTriangleID, const AttribValueType* Data)
@@ -149,6 +190,11 @@ public:
 	}
 
 
+	virtual TUniquePtr<FDynamicAttributeChangeBase> NewBlankChange() override
+	{
+		return MakeUnique<FDynamicMeshTriangleAttributeChange<AttribValueType, AttribDimension>>();
+	}
+
 
 public:
 
@@ -189,3 +235,40 @@ public:
 
 };
 
+
+
+template<typename AttribValueType, int AttribDimension>
+void FDynamicMeshTriangleAttributeChange<AttribValueType, AttribDimension>::SaveInitialTriangle(const FDynamicAttributeBase* Attribute, int TriangleID)
+{
+	FChangeTriangleAttribute& Change = OldTriangleAttributes.Emplace_GetRef();
+	Change.TriangleID = TriangleID;
+	const TDynamicMeshTriangleAttribute<AttribValueType, AttribDimension>* AttribCast = static_cast<const TDynamicMeshTriangleAttribute<AttribValueType, AttribDimension>*>(Attribute);
+	AttribCast->GetValue(TriangleID, Change.Data);
+}
+
+template<typename AttribValueType, int AttribDimension>
+void FDynamicMeshTriangleAttributeChange<AttribValueType, AttribDimension>::StoreAllFinalTriangles(const FDynamicAttributeBase* Attribute, const TArray<int>& TriangleIDs)
+{
+	const TDynamicMeshTriangleAttribute<AttribValueType, AttribDimension>* AttribCast = static_cast<const TDynamicMeshTriangleAttribute<AttribValueType, AttribDimension>*>(Attribute);
+	NewTriangleAttributes.Reserve(NewTriangleAttributes.Num() + TriangleIDs.Num());
+	for (int TriangleID : TriangleIDs)
+	{
+		FChangeTriangleAttribute& Change = NewTriangleAttributes.Emplace_GetRef();
+		Change.TriangleID = TriangleID;
+		AttribCast->GetValue(TriangleID, Change.Data);
+	}
+}
+
+template<typename AttribValueType, int AttribDimension>
+bool FDynamicMeshTriangleAttributeChange<AttribValueType, AttribDimension>::Apply(FDynamicAttributeBase* Attribute, bool bRevert) const
+{
+	TDynamicMeshTriangleAttribute<AttribValueType, AttribDimension>* AttribCast = static_cast<TDynamicMeshTriangleAttribute<AttribValueType, AttribDimension>*>(Attribute);
+	const TArray<FChangeTriangleAttribute> *Changes = bRevert ? &OldTriangleAttributes : &NewTriangleAttributes;
+	for (const FChangeTriangleAttribute& Change : *Changes)
+	{
+		check(AttribCast->GetParentMesh()->IsTriangle(Change.TriangleID));
+		AttribCast->SetValue(Change.TriangleID, Change.Data);
+	}
+	return true;
+}
+
diff --git a/Engine/Plugins/Experimental/GeometryProcessing/Source/DynamicMesh/Public/DynamicMeshVertexAttribute.h b/Engine/Plugins/Experimental/GeometryProcessing/Source/DynamicMesh/Public/DynamicMeshVertexAttribute.h
new file mode 100644
--- /dev/null
+++ b/Engine/Plugins/Experimental/GeometryProcessing/Source/DynamicMesh/Public/DynamicMeshVertexAttribute.h
@@ -0,0 +1,296 @@
+// Copyright 1998-2019 Epic Games, Inc. All Rights Reserved.
+
+#pragma once
+
+#include "DynamicMesh3.h"
+#include "DynamicAttribute.h"
+
+
+template<typename AttribValueType, int AttribDimension>
+class TDynamicMeshVertexAttribute;
+
+
+template<typename AttribValueType, int AttribDimension>
+class DYNAMICMESH_API FDynamicMeshVertexAttributeChange : public FDynamicAttributeChangeBase
+{
+private:
+	struct FChangeVertexAttribute
+	{
+		AttribValueType Data[AttribDimension];
+		int VertexID;
+	};
+
+	TArray<FChangeVertexAttribute> OldVertexAttributes, NewVertexAttributes;
+
+public:
+	FDynamicMeshVertexAttributeChange()
+	{}
+
+	virtual ~FDynamicMeshVertexAttributeChange()
+	{}
+
+	inline virtual void SaveInitialVertex(const FDynamicAttributeBase* Attribute, int VertexID) override;
+
+	inline virtual void StoreAllFinalVertices(const FDynamicAttributeBase* Attribute, const TArray<int>& VertexIDs) override;
+
+	inline virtual bool Apply(FDynamicAttributeBase* Attribute, bool bRevert) const override;
+};
+
+
+/**
+ * TDynamicMeshVertexAttribute is an add-on to a FDynamicMesh3 that allows for 
+ * per-triangle storage of an attribute value.
+ *
+ * The FDynamicMesh3 mesh topology operations (eg split/flip/collapse edge, poke face, etc)
+ * can be mirrored to the overlay via OnSplitEdge(), etc.
+ */
+template<typename AttribValueType, int AttribDimension>
+class DYNAMICMESH_API TDynamicMeshVertexAttribute : public FDynamicAttributeBase
+{
+
+protected:
+	/** The parent mesh this overlay belongs to */
+	FDynamicMesh3* ParentMesh;
+
+	/** List of per-triangle attribute values */
+	TDynamicVector<AttribValueType> AttribValues;
+
+	friend class FDynamicMesh3;
+	friend class FDynamicMeshAttributeSet;
+
+public:
+	/** Create an empty overlay */
+	TDynamicMeshVertexAttribute()
+	{
+		ParentMesh = nullptr;
+	}
+
+	/** Create an overlay for the given parent mesh */
+	TDynamicMeshVertexAttribute(FDynamicMesh3* ParentMeshIn)
+	{
+		ParentMesh = ParentMeshIn;
+	}
+
+	virtual ~TDynamicMeshVertexAttribute()
+	{
+	}
+
+	/** @return the parent mesh for this overlay */
+	const FDynamicMesh3* GetParentMesh() const { return ParentMesh; }
+	/** @return the parent mesh for this overlay */
+	FDynamicMesh3* GetParentMesh() { return ParentMesh; }
+
+	virtual FDynamicAttributeBase* MakeCopy(FDynamicMesh3* ParentMeshIn) const override
+	{
+		TDynamicMeshVertexAttribute<AttribValueType, AttribDimension>* ToFill = new TDynamicMeshVertexAttribute<AttribValueType, AttribDimension>(ParentMeshIn);
+		ToFill->Copy(*this);
+		return ToFill;
+	}
+
+	/** Set this overlay to contain the same arrays as the copy overlay */
+	void Copy(const TDynamicMeshVertexAttribute<AttribValueType, AttribDimension>& Copy)
+	{
+		AttribValues = Copy.AttribValues;
+	}
+
+	/** Initialize the attribute values to the given max triangle ID */
+	void Initialize(AttribValueType InitialValue = (AttribValueType)0)
+	{
+		check(ParentMesh != nullptr);
+		AttribValues.Resize(0);
+		AttribValues.Resize( ParentMesh->MaxVertexID() * AttribDimension, InitialValue );
+	}
+
+	void SetNewValue(int NewVertexID, const AttribValueType* Data)
+	{
+		int k = NewVertexID * AttribDimension;
+		for (int i = 0; i < AttribDimension; ++i)
+		{
+			AttribValues.InsertAt(Data[i], k + i);
+		}
+	}
+
+
+
+	//
+	// Accessors/Queries
+	//  
+
+	/** Get the element at a given index */
+	inline void GetValue(int VertexID, AttribValueType* Data) const
+	{
+		int k = VertexID * AttribDimension;
+		for (int i = 0; i < AttribDimension; ++i)
+		{
+			Data[i] = AttribValues[k + i];
+		}
+	}
+
+	/** Get the element at a given index */
+	template<typename AsType>
+	void GetValue(int VertexID, AsType& Data) const
+	{
+		int k = VertexID * AttribDimension;
+		for (int i = 0; i < AttribDimension; ++i)
+		{
+			Data[i] = AttribValues[k + i];
+		}
+	}
+
+
+	/** Set the element at a given index */
+	inline void SetValue(int VertexID, const AttribValueType* Data)
+	{
+		int k = VertexID * AttribDimension;
+		for (int i = 0; i < AttribDimension; ++i)
+		{
+			AttribValues[k + i] = Data[i];
+		}
+	}
+
+	/** Set the element at a given index */
+	template<typename AsType>
+	void SetValue(int VertexID, const AsType& Data)
+	{
+		int k = VertexID * AttribDimension;
+		for (int i = 0; i < AttribDimension; ++i)
+		{
+			AttribValues[k + i] = Data[i];
+		}
+	}
+
+	/**
+	 * Copy the attribute value at FromVertexID to ToVertexID
+	 */
+	inline void CopyValue(int FromVertexID, int ToVertexID)
+	{
+		int kA = FromVertexID * AttribDimension;
+		int kB = ToVertexID * AttribDimension;
+		for (int i = 0; i < AttribDimension; ++i)
+		{
+			AttribValues.InsertAt(AttribValues[kA+i], kB+i);
+		}
+	}
+
+
+public:
+
+	/** Update the overlay to reflect an edge split in the parent mesh */
+	void OnSplitEdge(const FDynamicMesh3::FEdgeSplitInfo& SplitInfo) override
+	{
+		SetAttributeFromLerp(SplitInfo.NewVertex, SplitInfo.OriginalVertices.A, SplitInfo.OriginalVertices.B, SplitInfo.SplitT);
+	}
+
+	/** Update the overlay to reflect an edge flip in the parent mesh */
+	void OnFlipEdge(const FDynamicMesh3::FEdgeFlipInfo& FlipInfo) override
+	{
+		// vertices unchanged
+	}
+
+	/** Update the overlay to reflect an edge collapse in the parent mesh */
+	void OnCollapseEdge(const FDynamicMesh3::FEdgeCollapseInfo& CollapseInfo) override
+	{
+		SetAttributeFromLerp(CollapseInfo.KeptVertex, CollapseInfo.KeptVertex, CollapseInfo.RemovedVertex, CollapseInfo.CollapseT);
+	}
+
+	virtual AttribValueType GetDefaultAttributeValue()
+	{
+		return (AttribValueType)0;
+	}
+
+	virtual void OnNewVertex(int VertexID, bool bInserted) override
+	{
+		if (ParentMesh->MaxVertexID() >= AttribValues.Num() * AttribDimension)
+		{
+			AttribValues.Resize((ParentMesh->MaxVertexID()+1) * AttribDimension, GetDefaultAttributeValue());
+		}
+	} 
+
+	/** Update the overlay to reflect a face poke in the parent mesh */
+	void OnPokeTriangle(const FDynamicMesh3::FPokeTriangleInfo& PokeInfo) override
+	{
+		FIndex3i Tri = PokeInfo.TriVertices;
+		SetAttributeFromBary(PokeInfo.NewVertex, Tri.A, Tri.B, Tri.C, PokeInfo.BaryCoords);
+	}
+
+	/** Update the overlay to reflect an edge merge in the parent mesh */
+	void OnMergeEdges(const FDynamicMesh3::FMergeEdgesInfo& MergeInfo) override
+	{
+		// just blend the attributes?
+		SetAttributeFromLerp(MergeInfo.KeptVerts.A, MergeInfo.KeptVerts.A, MergeInfo.RemovedVerts.A, .5);
+		SetAttributeFromLerp(MergeInfo.KeptVerts.B, MergeInfo.KeptVerts.B, MergeInfo.RemovedVerts.B, .5);
+	}
+
+	virtual TUniquePtr<FDynamicAttributeChangeBase> NewBlankChange() override
+	{
+		return MakeUnique<FDynamicMeshVertexAttributeChange<AttribValueType, AttribDimension>>();
+	}
+
+protected:
+
+	// interpolation functions; default implementation assumes your attributes can be interpolated as reals
+
+	/** Set the value at an Attribute to be a linear interpolation of two other Attributes */
+	virtual void SetAttributeFromLerp(int SetAttribute, int AttributeA, int AttributeB, double Alpha)
+	{
+		int IndexSet = AttribDimension * SetAttribute;
+		int IndexA = AttribDimension * AttributeA;
+		int IndexB = AttribDimension * AttributeB;
+		AttribValueType Beta = ((AttribValueType)1 - Alpha);
+		for (int i = 0; i < AttribDimension; ++i)
+		{
+			AttribValues[IndexSet+i] = Beta*AttribValues[IndexA+i] + Alpha*AttribValues[IndexB+i];
+		}
+	}
+
+	/** Set the value at an Attribute to be a barycentric interpolation of three other Attributes */
+	virtual void SetAttributeFromBary(int SetAttribute, int AttributeA, int AttributeB, int AttributeC, const FVector3d& BaryCoords)
+	{
+		int IndexSet = AttribDimension * SetAttribute;
+		int IndexA = AttribDimension * AttributeA;
+		int IndexB = AttribDimension * AttributeB;
+		int IndexC = AttribDimension * AttributeC;
+		for (int i = 0; i < AttribDimension; ++i)
+		{
+			AttribValues[IndexSet + i] = 
+				BaryCoords.X*AttribValues[IndexA+i] + BaryCoords.Y*AttribValues[IndexB+i] + BaryCoords.Z*AttribValues[IndexC+i];
+		}
+	}
+
+};
+
+
+template<typename AttribValueType, int AttribDimension>
+void FDynamicMeshVertexAttributeChange<AttribValueType, AttribDimension>::SaveInitialVertex(const FDynamicAttributeBase* Attribute, int VertexID)
+{
+	FChangeVertexAttribute& Change = OldVertexAttributes.Emplace_GetRef();
+	Change.VertexID = VertexID;
+	const TDynamicMeshVertexAttribute<AttribValueType, AttribDimension>* AttribCast = static_cast<const TDynamicMeshVertexAttribute<AttribValueType, AttribDimension>*>(Attribute);
+	AttribCast->GetValue(VertexID, Change.Data);
+}
+
+template<typename AttribValueType, int AttribDimension>
+void FDynamicMeshVertexAttributeChange<AttribValueType, AttribDimension>::StoreAllFinalVertices(const FDynamicAttributeBase* Attribute, const TArray<int>& VertexIDs)
+{
+	NewVertexAttributes.Reserve(NewVertexAttributes.Num() + VertexIDs.Num());
+	const TDynamicMeshVertexAttribute<AttribValueType, AttribDimension>* AttribCast = static_cast<const TDynamicMeshVertexAttribute<AttribValueType, AttribDimension>*>(Attribute);
+	for (int VertexID : VertexIDs)
+	{
+		FChangeVertexAttribute& Change = NewVertexAttributes.Emplace_GetRef();
+		Change.VertexID = VertexID;
+		AttribCast->GetValue(VertexID, Change.Data);
+	}
+}
+
+template<typename AttribValueType, int AttribDimension>
+bool FDynamicMeshVertexAttributeChange<AttribValueType, AttribDimension>::Apply(FDynamicAttributeBase* Attribute, bool bRevert) const
+{
+	const TArray<FChangeVertexAttribute> *Changes = bRevert ? &OldVertexAttributes : &NewVertexAttributes;
+	TDynamicMeshVertexAttribute<AttribValueType, AttribDimension>* AttribCast = static_cast<TDynamicMeshVertexAttribute<AttribValueType, AttribDimension>*>(Attribute);
+	for (const FChangeVertexAttribute& Change : *Changes)
+	{
+		check(AttribCast->GetParentMesh()->IsVertex(Change.VertexID));
+		AttribCast->SetValue(Change.VertexID, Change.Data);
+	}
+	return true;
+}
\ No newline at end of file
