generalize vertex attributes to apply to non-mesh objects
minor fixes and basic validation check

#rnx

diff --git a/Engine/Plugins/Experimental/GeometryProcessing/Source/DynamicMesh/Private/DynamicMesh3.cpp b/Engine/Plugins/Experimental/GeometryProcessing/Source/DynamicMesh/Private/DynamicMesh3.cpp
--- a/Engine/Plugins/Experimental/GeometryProcessing/Source/DynamicMesh/Private/DynamicMesh3.cpp
+++ b/Engine/Plugins/Experimental/GeometryProcessing/Source/DynamicMesh/Private/DynamicMesh3.cpp
@@ -905,11 +905,7 @@ bool FDynamicMesh3::CheckValidity(bool bAllowNonManifoldVertices, EValidityCheck
 
 	if (HasAttributes())
 	{
-		for (int UVLayerIndex = 0; UVLayerIndex < Attributes()->NumUVLayers(); UVLayerIndex++)
-		{
-			Attributes()->GetUVLayer(UVLayerIndex)->CheckValidity(true, FailMode);
-		}
-		Attributes()->PrimaryNormals()->CheckValidity(true, FailMode);
+		CheckOrFailF(Attributes()->CheckValidity(true, FailMode));
 	}
 
 	return is_ok;
diff --git a/Engine/Plugins/Experimental/GeometryProcessing/Source/DynamicMesh/Private/DynamicMesh3_Edits.cpp b/Engine/Plugins/Experimental/GeometryProcessing/Source/DynamicMesh/Private/DynamicMesh3_Edits.cpp
--- a/Engine/Plugins/Experimental/GeometryProcessing/Source/DynamicMesh/Private/DynamicMesh3_Edits.cpp
+++ b/Engine/Plugins/Experimental/GeometryProcessing/Source/DynamicMesh/Private/DynamicMesh3_Edits.cpp
@@ -40,6 +40,11 @@ int FDynamicMesh3::AppendVertex(const FVertexInfo& VtxInfo)
 
 	AllocateEdgesList(vid);
 
+	if (HasAttributes())
+	{
+		Attributes()->OnNewVertex(vid, false);
+	}
+
 	UpdateTimeStamp(true, true);
 	return vid;
 }
@@ -105,6 +110,11 @@ int FDynamicMesh3::AppendVertex(const FDynamicMesh3& from, int fromVID)
 
 	AllocateEdgesList(vid);
 
+	if (HasAttributes())
+	{
+		Attributes()->OnNewVertex(vid, false);
+	}
+
 	UpdateTimeStamp(true, true);
 	return vid;
 }
@@ -156,6 +166,11 @@ EMeshResult FDynamicMesh3::InsertVertex(int vid, const FVertexInfo& info, bool b
 
 	AllocateEdgesList(vid);
 
+	if (HasAttributes())
+	{
+		Attributes()->OnNewVertex(vid, true);
+	}
+
 	UpdateTimeStamp(true, true);
 	return EMeshResult::Ok;
 }
diff --git a/Engine/Plugins/Experimental/GeometryProcessing/Source/DynamicMesh/Public/DynamicAttribute.h b/Engine/Plugins/Experimental/GeometryProcessing/Source/DynamicMesh/Public/DynamicAttribute.h
--- a/Engine/Plugins/Experimental/GeometryProcessing/Source/DynamicMesh/Public/DynamicAttribute.h
+++ b/Engine/Plugins/Experimental/GeometryProcessing/Source/DynamicMesh/Public/DynamicAttribute.h
@@ -104,6 +104,18 @@ public:
 	{
 	}
 
+	/**
+	* Check validity of attribute
+	* 
+	* @param bAllowNonmanifold Accept non-manifold topology as valid. Note that this should almost always be true for attributes; non-manifold overlays are generally valid.
+	* @param FailMode Desired behavior if mesh is found invalid
+	*/
+	virtual bool CheckValidity(bool bAllowNonmanifold, EValidityCheckFailMode FailMode) const
+	{
+		// default impl just doesn't check anything; override with any useful sanity checks
+		return true;
+	}
+
 
 	virtual TUniquePtr<FDynamicAttributeChangeBase> NewBlankChange() = 0;
 
@@ -219,4 +231,20 @@ public:
 			A->OnMergeEdges(MergeInfo);
 		}
 	}
+
+	/**
+	* Check validity of attributes
+	* 
+	* @param bAllowNonmanifold Accept non-manifold topology as valid. Note that this should almost always be true for attributes; non-manifold overlays are generally valid.
+	* @param FailMode Desired behavior if mesh is found invalid
+	*/
+	virtual bool CheckValidity(bool bAllowNonmanifold, EValidityCheckFailMode FailMode) const
+	{
+		bool bValid = true;
+		for (FDynamicAttributeBase* A : RegisteredAttributes)
+		{
+			bValid = A->CheckValidity(bAllowNonmanifold, FailMode) && bValid;
+		}
+		return bValid;
+	}
 };
diff --git a/Engine/Plugins/Experimental/GeometryProcessing/Source/DynamicMesh/Public/DynamicMeshAttributeSet.h b/Engine/Plugins/Experimental/GeometryProcessing/Source/DynamicMesh/Public/DynamicMeshAttributeSet.h
--- a/Engine/Plugins/Experimental/GeometryProcessing/Source/DynamicMesh/Public/DynamicMeshAttributeSet.h
+++ b/Engine/Plugins/Experimental/GeometryProcessing/Source/DynamicMesh/Public/DynamicMeshAttributeSet.h
@@ -6,6 +6,7 @@
 #include "DynamicMeshOverlay.h"
 #include "DynamicMeshTriangleAttribute.h"
 #include "DynamicAttribute.h"
+#include "GeometryTypes.h"
 
 /** Standard UV overlay type - 2-element float */
 typedef TDynamicMeshVectorOverlay<float, 2, FVector2f> FDynamicMeshUVOverlay;
@@ -27,7 +28,7 @@ class DYNAMICMESH_API FDynamicMeshAttributeSet : public FDynamicAttributeSetBase
 {
 public:
 
-	FDynamicMeshAttributeSet(FDynamicMesh3* Mesh) 
+	FDynamicMeshAttributeSet(FDynamicMesh3* Mesh)
 		: ParentMesh(Mesh), Normals0(Mesh)
 	{
 		SetNumUVLayers(1);
@@ -50,7 +51,7 @@ public:
 		if (Copy.MaterialIDAttrib)
 		{
 			EnableMaterialID();
-			MaterialIDAttrib->Copy( *(Copy.MaterialIDAttrib) );
+			MaterialIDAttrib->Copy(*(Copy.MaterialIDAttrib));
 		}
 
 		GenericAttributes.Reset();
@@ -71,7 +72,7 @@ public:
 
 
 	/** @return number of UV layers */
-	virtual int NumUVLayers() const 
+	virtual int NumUVLayers() const
 	{
 		return UVLayers.Num();
 	}
@@ -109,7 +110,7 @@ public:
 	/** @return true if the given vertex is a seam vertex in any overlay */
 	virtual bool IsSeamVertex(int VertexID, bool bBoundaryIsSeam = true) const;
 
-	
+
 	//
 	// UV Layers 
 	//
@@ -127,7 +128,7 @@ public:
 	}
 
 	/** @return the primary UV layer (layer 0) */
-	FDynamicMeshUVOverlay* PrimaryUV() 
+	FDynamicMeshUVOverlay* PrimaryUV()
 	{
 		return &UVLayers[0];
 	}
@@ -224,7 +225,7 @@ protected:
 	TArray<FDynamicMeshNormalOverlay*> NormalLayers;
 
 	TUniquePtr<FDynamicMeshMaterialAttribute> MaterialIDAttrib;
-	
+
 	TArray<TUniquePtr<FDynamicAttributeBase>> GenericAttributes;
 
 protected:
@@ -255,5 +256,22 @@ protected:
 	virtual void OnCollapseEdge(const FDynamicMesh3::FEdgeCollapseInfo & collapseInfo);
 	virtual void OnPokeTriangle(const FDynamicMesh3::FPokeTriangleInfo & pokeInfo);
 	virtual void OnMergeEdges(const FDynamicMesh3::FMergeEdgesInfo & mergeInfo);
+
+	/**
+	 * Check validity of attributes
+	 * 
+	 * @param bAllowNonmanifold Accept non-manifold topology as valid. Note that this should almost always be true for attributes; non-manifold overlays are generally valid.
+	 * @param FailMode Desired behavior if mesh is found invalid
+	 */
+	virtual bool CheckValidity(bool bAllowNonmanifold, EValidityCheckFailMode FailMode) const
+	{
+		bool bValid = FDynamicAttributeSetBase::CheckValidity(bAllowNonmanifold, FailMode);
+		for (int UVLayerIndex = 0; UVLayerIndex < NumUVLayers(); UVLayerIndex++)
+		{
+			bValid = GetUVLayer(UVLayerIndex)->CheckValidity(bAllowNonmanifold, FailMode) && bValid;
+		}
+		bValid = PrimaryNormals()->CheckValidity(bAllowNonmanifold, FailMode) && bValid;
+		return bValid;
+	}
 };
 
diff --git a/Engine/Plugins/Experimental/GeometryProcessing/Source/DynamicMesh/Public/DynamicMeshVertexAttribute.h b/Engine/Plugins/Experimental/GeometryProcessing/Source/DynamicMesh/Public/DynamicMeshVertexAttribute.h
deleted file mode 100644
diff --git a/Engine/Plugins/Experimental/GeometryProcessing/Source/DynamicMesh/Public/DynamicVertexAttribute.h b/Engine/Plugins/Experimental/GeometryProcessing/Source/DynamicMesh/Public/DynamicVertexAttribute.h
new file mode 100644
--- /dev/null
+++ b/Engine/Plugins/Experimental/GeometryProcessing/Source/DynamicMesh/Public/DynamicVertexAttribute.h
@@ -0,0 +1,337 @@
+// Copyright 1998-2019 Epic Games, Inc. All Rights Reserved.
+
+#pragma once
+
+#include "DynamicMesh3.h"
+#include "DynamicAttribute.h"
+
+
+template<typename AttribValueType, int AttribDimension, typename ParentType>
+class TDynamicVertexAttribute;
+
+
+template<typename AttribValueType, int AttribDimension, typename ParentType>
+class DYNAMICMESH_API FDynamicVertexAttributeChange : public FDynamicAttributeChangeBase
+{
+private:
+	struct FChangeVertexAttribute
+	{
+		AttribValueType Data[AttribDimension];
+		int VertexID;
+	};
+
+	TArray<FChangeVertexAttribute> OldVertexAttributes, NewVertexAttributes;
+
+public:
+	FDynamicVertexAttributeChange()
+	{}
+
+	virtual ~FDynamicVertexAttributeChange()
+	{}
+
+	inline virtual void SaveInitialVertex(const FDynamicAttributeBase* Attribute, int VertexID) override;
+
+	inline virtual void StoreAllFinalVertices(const FDynamicAttributeBase* Attribute, const TArray<int>& VertexIDs) override;
+
+	inline virtual bool Apply(FDynamicAttributeBase* Attribute, bool bRevert) const override;
+};
+
+
+/**
+ * TDynamicVertexAttribute provides per-vertex storage of an attribute value
+ */
+template<typename AttribValueType, int AttribDimension, typename ParentType>
+class DYNAMICMESH_API TDynamicVertexAttribute : public FDynamicAttributeBase
+{
+
+protected:
+	/** The parent object (e.g. mesh, point set) this attribute belongs to */
+	ParentType* Parent;
+
+	/** List of per-triangle attribute values */
+	TDynamicVector<AttribValueType> AttribValues;
+
+public:
+	/** Create an empty overlay */
+	TDynamicVertexAttribute() : Parent(nullptr)
+	{
+	}
+
+	/** Create an attribute for the given parent */
+	TDynamicVertexAttribute(ParentType* ParentIn, bool bAutoInit = true) : Parent(ParentIn)
+	{
+		if (bAutoInit)
+		{
+			Initialize();
+		}
+	}
+
+	virtual ~TDynamicVertexAttribute()
+	{
+	}
+
+	/** @return the parent for this attribute */
+	const ParentType* GetParent() const { return Parent; }
+	/** @return the parent for this attribute */
+	ParentType* GetParent() { return Parent; }
+
+	virtual FDynamicAttributeBase* MakeCopy(ParentType* ParentIn) const override
+	{
+		TDynamicVertexAttribute<AttribValueType, AttribDimension, ParentType>* ToFill = new TDynamicVertexAttribute<AttribValueType, AttribDimension, ParentType>(ParentIn);
+		ToFill->Copy(*this);
+		return ToFill;
+	}
+
+	/** Set this overlay to contain the same arrays as the copy overlay */
+	void Copy(const TDynamicVertexAttribute<AttribValueType, AttribDimension, ParentType>& Copy)
+	{
+		AttribValues = Copy.AttribValues;
+	}
+
+	/** Initialize the attribute values to the given max triangle ID */
+	void Initialize(AttribValueType InitialValue = (AttribValueType)0)
+	{
+		check(Parent != nullptr);
+		AttribValues.Resize(0);
+		AttribValues.Resize( Parent->MaxVertexID() * AttribDimension, InitialValue );
+	}
+
+	void SetNewValue(int NewVertexID, const AttribValueType* Data)
+	{
+		int k = NewVertexID * AttribDimension;
+		for (int i = 0; i < AttribDimension; ++i)
+		{
+			AttribValues.InsertAt(Data[i], k + i);
+		}
+	}
+
+
+
+	//
+	// Accessors/Queries
+	//  
+
+	/** Get the element at a given index */
+	inline void GetValue(int VertexID, AttribValueType* Data) const
+	{
+		int k = VertexID * AttribDimension;
+		for (int i = 0; i < AttribDimension; ++i)
+		{
+			Data[i] = AttribValues[k + i];
+		}
+	}
+
+	/** Get the element at a given index */
+	template<typename AsType>
+	void GetValue(int VertexID, AsType& Data) const
+	{
+		int k = VertexID * AttribDimension;
+		for (int i = 0; i < AttribDimension; ++i)
+		{
+			Data[i] = AttribValues[k + i];
+		}
+	}
+
+
+	/** Set the element at a given index */
+	inline void SetValue(int VertexID, const AttribValueType* Data)
+	{
+		int k = VertexID * AttribDimension;
+		for (int i = 0; i < AttribDimension; ++i)
+		{
+			AttribValues[k + i] = Data[i];
+		}
+	}
+
+	/** Set the element at a given index */
+	template<typename AsType>
+	void SetValue(int VertexID, const AsType& Data)
+	{
+		int k = VertexID * AttribDimension;
+		for (int i = 0; i < AttribDimension; ++i)
+		{
+			AttribValues[k + i] = Data[i];
+		}
+	}
+
+	/**
+	 * Copy the attribute value at FromVertexID to ToVertexID
+	 */
+	inline void CopyValue(int FromVertexID, int ToVertexID)
+	{
+		int kA = FromVertexID * AttribDimension;
+		int kB = ToVertexID * AttribDimension;
+		for (int i = 0; i < AttribDimension; ++i)
+		{
+			AttribValues.InsertAt(AttribValues[kA+i], kB+i);
+		}
+	}
+
+
+public:
+
+	/** Update the overlay to reflect an edge split in the parent */
+	void OnSplitEdge(const FDynamicMesh3::FEdgeSplitInfo& SplitInfo) override
+	{
+		ResizeAttribStoreIfNeeded(SplitInfo.NewVertex);
+		SetAttributeFromLerp(SplitInfo.NewVertex, SplitInfo.OriginalVertices.A, SplitInfo.OriginalVertices.B, SplitInfo.SplitT);
+	}
+
+	/** Update the overlay to reflect an edge flip in the parent */
+	void OnFlipEdge(const FDynamicMesh3::FEdgeFlipInfo& FlipInfo) override
+	{
+		// vertices unchanged
+	}
+
+	/** Update the overlay to reflect an edge collapse in the parent */
+	void OnCollapseEdge(const FDynamicMesh3::FEdgeCollapseInfo& CollapseInfo) override
+	{
+		SetAttributeFromLerp(CollapseInfo.KeptVertex, CollapseInfo.KeptVertex, CollapseInfo.RemovedVertex, CollapseInfo.CollapseT);
+	}
+
+	virtual AttribValueType GetDefaultAttributeValue()
+	{
+		return (AttribValueType)0;
+	}
+
+	inline void ResizeAttribStoreIfNeeded(int VertexID)
+	{
+		int k = VertexID * AttribDimension;
+		if (k+AttribDimension > AttribValues.Num())
+		{
+			AttribValues.Resize(k + AttribDimension, GetDefaultAttributeValue());
+		}
+	}
+
+	virtual void OnNewVertex(int VertexID, bool bInserted) override
+	{
+		ResizeAttribStoreIfNeeded(VertexID);
+	}
+
+	/** Update the overlay to reflect a face poke in the parent */
+	void OnPokeTriangle(const FDynamicMesh3::FPokeTriangleInfo& PokeInfo) override
+	{
+		FIndex3i Tri = PokeInfo.TriVertices;
+		ResizeAttribStoreIfNeeded(PokeInfo.NewVertex);
+		SetAttributeFromBary(PokeInfo.NewVertex, Tri.A, Tri.B, Tri.C, PokeInfo.BaryCoords);
+	}
+
+	/** Update the overlay to reflect an edge merge in the parent */
+	void OnMergeEdges(const FDynamicMesh3::FMergeEdgesInfo& MergeInfo) override
+	{
+		// just blend the attributes?
+		SetAttributeFromLerp(MergeInfo.KeptVerts.A, MergeInfo.KeptVerts.A, MergeInfo.RemovedVerts.A, .5);
+		SetAttributeFromLerp(MergeInfo.KeptVerts.B, MergeInfo.KeptVerts.B, MergeInfo.RemovedVerts.B, .5);
+	}
+
+	virtual TUniquePtr<FDynamicAttributeChangeBase> NewBlankChange() override
+	{
+		return MakeUnique<FDynamicVertexAttributeChange<AttribValueType, AttribDimension, ParentType>>();
+	}
+
+	/**
+	* Check validity of attribute
+	* 
+	* @param bAllowNonmanifold Accept non-manifold topology as valid. Note that this should almost always be true for attributes; non-manifold overlays are generally valid.
+	* @param FailMode Desired behavior if mesh is found invalid
+	*/
+	virtual bool CheckValidity(bool bAllowNonmanifold, EValidityCheckFailMode FailMode) const override
+	{
+		// just check that the values buffer is big enough
+		if (Parent->MaxVertexID()*AttribDimension > AttribValues.Num())
+		{
+			switch (FailMode)
+			{
+			case EValidityCheckFailMode::Check:
+				check(false);
+			case EValidityCheckFailMode::Ensure:
+				ensure(false);
+			default:
+				return false;
+			}
+		}
+
+		return true;
+	}
+
+
+protected:
+
+	// interpolation functions; default implementation assumes your attributes can be interpolated as reals
+
+	/** Set the value at an Attribute to be a linear interpolation of two other Attributes */
+	virtual void SetAttributeFromLerp(int SetAttribute, int AttributeA, int AttributeB, double Alpha)
+	{
+		int IndexSet = AttribDimension * SetAttribute;
+		int IndexA = AttribDimension * AttributeA;
+		int IndexB = AttribDimension * AttributeB;
+		AttribValueType Beta = ((AttribValueType)1 - Alpha);
+		for (int i = 0; i < AttribDimension; ++i)
+		{
+			AttribValues[IndexSet+i] = Beta*AttribValues[IndexA+i] + Alpha*AttribValues[IndexB+i];
+		}
+	}
+
+	/** Set the value at an Attribute to be a barycentric interpolation of three other Attributes */
+	virtual void SetAttributeFromBary(int SetAttribute, int AttributeA, int AttributeB, int AttributeC, const FVector3d& BaryCoords)
+	{
+		int IndexSet = AttribDimension * SetAttribute;
+		int IndexA = AttribDimension * AttributeA;
+		int IndexB = AttribDimension * AttributeB;
+		int IndexC = AttribDimension * AttributeC;
+		for (int i = 0; i < AttribDimension; ++i)
+		{
+			AttribValues[IndexSet + i] = 
+				BaryCoords.X*AttribValues[IndexA+i] + BaryCoords.Y*AttribValues[IndexB+i] + BaryCoords.Z*AttribValues[IndexC+i];
+		}
+	}
+
+};
+
+
+template<typename AttribValueType, int AttribDimension>
+using TDynamicMeshVertexAttribute = TDynamicVertexAttribute<AttribValueType, AttribDimension, FDynamicMesh3>;
+
+
+template<typename AttribValueType, int AttribDimension, typename ParentType>
+void FDynamicVertexAttributeChange<AttribValueType, AttribDimension, ParentType>::SaveInitialVertex(const FDynamicAttributeBase* Attribute, int VertexID)
+{
+	FChangeVertexAttribute& Change = OldVertexAttributes.Emplace_GetRef();
+	Change.VertexID = VertexID;
+	const TDynamicVertexAttribute<AttribValueType, AttribDimension, ParentType>* AttribCast = static_cast<const TDynamicVertexAttribute<AttribValueType, AttribDimension, ParentType>*>(Attribute);
+	AttribCast->GetValue(VertexID, Change.Data);
+}
+
+template<typename AttribValueType, int AttribDimension, typename ParentType>
+void FDynamicVertexAttributeChange<AttribValueType, AttribDimension, ParentType>::StoreAllFinalVertices(const FDynamicAttributeBase* Attribute, const TArray<int>& VertexIDs)
+{
+	NewVertexAttributes.Reserve(NewVertexAttributes.Num() + VertexIDs.Num());
+	const TDynamicVertexAttribute<AttribValueType, AttribDimension, ParentType>* AttribCast = static_cast<const TDynamicVertexAttribute<AttribValueType, AttribDimension, ParentType>*>(Attribute);
+	for (int VertexID : VertexIDs)
+	{
+		FChangeVertexAttribute& Change = NewVertexAttributes.Emplace_GetRef();
+		Change.VertexID = VertexID;
+		AttribCast->GetValue(VertexID, Change.Data);
+	}
+}
+
+template<typename AttribValueType, int AttribDimension, typename ParentType>
+bool FDynamicVertexAttributeChange<AttribValueType, AttribDimension, ParentType>::Apply(FDynamicAttributeBase* Attribute, bool bRevert) const
+{
+	const TArray<FChangeVertexAttribute> *Changes = bRevert ? &OldVertexAttributes : &NewVertexAttributes;
+	TDynamicVertexAttribute<AttribValueType, AttribDimension, ParentType>* AttribCast = static_cast<TDynamicVertexAttribute<AttribValueType, AttribDimension, ParentType>*>(Attribute);
+	for (const FChangeVertexAttribute& Change : *Changes)
+	{
+		check(AttribCast->GetParent()->IsVertex(Change.VertexID));
+		AttribCast->SetValue(Change.VertexID, Change.Data);
+	}
+	return true;
+}
+
+
+template<typename AttribValueType, int AttribDimension>
+using TDynamicMeshVertexAttribute = TDynamicVertexAttribute<AttribValueType, AttribDimension, FDynamicMesh3>;
+
+template<typename AttribValueType, int AttribDimension>
+using FDynamicMeshVertexAttributeChange = FDynamicVertexAttributeChange<AttribValueType, AttribDimension, FDynamicMesh3>;
+
