fixes to embed polygons

#rnx

diff --git a/Engine/Plugins/Experimental/GeometryProcessing/Source/DynamicMesh/Private/Operations/EmbedSurfacePath.cpp b/Engine/Plugins/Experimental/GeometryProcessing/Source/DynamicMesh/Private/Operations/EmbedSurfacePath.cpp
--- a/Engine/Plugins/Experimental/GeometryProcessing/Source/DynamicMesh/Private/Operations/EmbedSurfacePath.cpp
+++ b/Engine/Plugins/Experimental/GeometryProcessing/Source/DynamicMesh/Private/Operations/EmbedSurfacePath.cpp
@@ -161,7 +161,7 @@ bool WalkMeshPlanar(const FDynamicMesh3* Mesh, int StartTri, FVector3d StartPt,
 	TArray<int> UnexploredEnds;
 	int BestKnownEnd = -1;
 	double BestKnownEndDistSq = FMathd::MaxReal;
-	TSet<int> ExploredTriangles;
+	TSet<int> ExploredTriangles, CrossedVertices;
 
 	bool bHasArrived = false;
 
@@ -269,32 +269,40 @@ bool WalkMeshPlanar(const FDynamicMesh3* Mesh, int StartTri, FVector3d StartPt,
 				{
 					FMeshSurfacePoint SurfPt(CandidateVertID);
 					FWalkIndices WalkInds(CurrentTri.V[TriSubIdx], CurrentEnd, -1);
+					double DSq = EndPt.DistanceSquared(CurrentTri.V[TriSubIdx]);
 
-					// TODO: extract this "next triangle candidate" logic to be used in more places??
-					// walking over a vertex is gross because we have to search the whole one ring for candidate next triangles and there might be multiple of them
-					// note that currently this means I compute signs for all vertices of neighboring triangles here, and do not re-use those signs when I actually process the triangle later; TODO reconsider if/when optimizing this fn
-					for (int32 NbrTriID : Mesh->VtxTrianglesItr(CandidateVertID))
+					// not allowed to go in a direction that gets us further from the destination than our initial point if backwards search not allowed
+					if ((bAllowBackwardsSearch || DSq <= InitialDistSq + 10 * FMathd::ZeroTolerance) && !CrossedVertices.Contains(CandidateVertID))
 					{
-						if (NbrTriID != TriID)
+						// consider going over this vertex
+						CrossedVertices.Add(CandidateVertID);
+
+						// TODO: extract this "next triangle candidate" logic to be used in more places??
+						// walking over a vertex is gross because we have to search the whole one ring for candidate next triangles and there might be multiple of them
+						// note that currently this means I compute signs for all vertices of neighboring triangles here, and do not re-use those signs when I actually process the triangle later; TODO reconsider if/when optimizing this fn
+						for (int32 NbrTriID : Mesh->VtxTrianglesItr(CandidateVertID))
 						{
-							FIndex3i NbrTriVertIDs = Mesh->GetTriangle(NbrTriID);
-							FTriangle3d NbrTri;
-							SetTriVertPositions(NbrTriVertIDs, NbrTri);
-							int SignsMultiplied = 1;
-							for (int NbrTriSubIdx = 0; NbrTriSubIdx < 3; NbrTriSubIdx++)
+							if (NbrTriID != TriID)
 							{
-								if (NbrTriVertIDs[NbrTriSubIdx] == CandidateVertID)
+								FIndex3i NbrTriVertIDs = Mesh->GetTriangle(NbrTriID);
+								FTriangle3d NbrTri;
+								SetTriVertPositions(NbrTriVertIDs, NbrTri);
+								int SignsMultiplied = 1;
+								for (int NbrTriSubIdx = 0; NbrTriSubIdx < 3; NbrTriSubIdx++)
 								{
-									continue;
+									if (NbrTriVertIDs[NbrTriSubIdx] == CandidateVertID)
+									{
+										continue;
+									}
+									double NbrSD = (NbrTri.V[NbrTriSubIdx] - StartPt).Dot(WalkPlaneNormal);
+									int NbrSign = FMathd::Abs(NbrSD) <= PtOnPlaneThresholdSq ? 0 : NbrSD > 0 ? 1 : -1;
+									SignsMultiplied *= NbrSign;
+								}
+								if (SignsMultiplied < 1) // plane will cross this triangle, so try walking it
+								{
+									WalkInds.WalkingOnTri = NbrTriID;
+									ComputedPointsAndSources.Emplace(SurfPt, WalkInds);
 								}
-								double NbrSD = (NbrTri.V[NbrTriSubIdx] - StartPt).Dot(WalkPlaneNormal);
-								int NbrSign = FMathd::Abs(NbrSD) <= PtOnPlaneThresholdSq ? 0 : NbrSD > 0 ? 1 : -1;
-								SignsMultiplied *= NbrSign;
-							}
-							if (SignsMultiplied < 1) // plane will cross this triangle, so try walking it
-							{
-								WalkInds.WalkingOnTri = NbrTriID;
-								ComputedPointsAndSources.Emplace(SurfPt, WalkInds);
 							}
 						}
 					}
@@ -334,7 +342,7 @@ bool WalkMeshPlanar(const FDynamicMesh3* Mesh, int StartTri, FVector3d StartPt,
 						continue;
 					}
 					double DSq = EndPt.DistanceSquared(CrossingP);
-					if (!bAllowBackwardsSearch && DSq > InitialDistSq + 100 * FMathd::ZeroTolerance)
+					if (!bAllowBackwardsSearch && DSq > InitialDistSq + 10 * FMathd::ZeroTolerance)
 					{
 						// not allowed to go in a direction that gets us further from the destination than our initial point if backwards search not allowed
 						continue;
@@ -350,6 +358,7 @@ bool WalkMeshPlanar(const FDynamicMesh3* Mesh, int StartTri, FVector3d StartPt,
 		{
 			double DistSq = EndPt.DistanceSquared(ComputedPointsAndSources[NewComputedPtIdx].Value.Position);
 			// TODO: reject cases that move us backwards if !bAllowBackwardsSearch
+			ensure(bAllowBackwardsSearch || DistSq < InitialDistSq + FMathd::ZeroTolerance);
 			if (BestCandidate == -1 || DistSq < BestCandidateDistSq)
 			{
 				BestCandidateDistSq = DistSq;
@@ -700,19 +709,22 @@ bool FMeshSurfacePath::EmbedSimplePath(bool bUpdatePath, TArray<int>& PathVertic
 		}
 		else
 		{
-			PathVertices.Add(EndPtUpdated.ElementID);
+			if (PathVertices.Num() == 0 || PathVertices.Last() != EndPtUpdated.ElementID)
+			{
+				PathVertices.Add(EndPtUpdated.ElementID);
+			}
 		}
 	}
 
 
 	// TODO: rm this debugging check
-	for (int PathIdx = InitialPathIdx; PathIdx + 1 < PathVertices.Num(); PathIdx++)
-	{
-		if (!ensure(Mesh->IsEdge(Mesh->FindEdge(PathVertices[PathIdx], PathVertices[PathIdx + 1]))))
-		{
-			return false;
-		}
-	}
+	//for (int PathIdx = InitialPathIdx; PathIdx + 1 < PathVertices.Num(); PathIdx++)
+	//{
+	//	if (!ensure(Mesh->IsEdge(Mesh->FindEdge(PathVertices[PathIdx], PathVertices[PathIdx + 1]))))
+	//	{
+	//		return false;
+	//	}
+	//}
 
 	if (bUpdatePath)
 	{
@@ -742,7 +754,6 @@ bool EmbedProjectedPath(FDynamicMesh3* Mesh, int StartTriID, FFrame3d Frame, con
 	OutVertexCorrespondence.Add(0);
 	for (int32 IdxA = 0; IdxA <= EndIdxA; IdxA++)
 	{
-		
 		int32 IdxB = (IdxA + 1) % Path2D.Num();
 		FMeshSurfacePath SurfacePath(Mesh);
 		int LastVert = -1;
@@ -762,19 +773,34 @@ bool EmbedProjectedPath(FDynamicMesh3* Mesh, int StartTriID, FFrame3d Frame, con
 		}
 		FVector2d WalkDir = Path2D[IdxB] - FVector2d(StartPos.X, StartPos.Y);
 		double WalkLen = WalkDir.Length();
-		if (WalkLen < PtSnapVertexOrEdgeThresholdSq)
+		bool bEmbedSuccess = true;
+		if (WalkLen >= PtSnapVertexOrEdgeThresholdSq || (LastVert != -1 && LastVert != OutPathVertices.Last()))
 		{
-			// we're already at the point
-			ensure(false); // TODO: can skip some processing for this case?
-		}
-		WalkDir /= WalkLen;
-		FVector3d WalkNormal(-WalkDir.Y, WalkDir.X, 0);
-		bool bWalkSuccess = SurfacePath.AddViaPlanarWalk(CurrentSeedTriID, StartPos, -1, LastVert, FVector3d(Path2D[IdxB].X, Path2D[IdxB].Y, 0), WalkNormal, ProjectToFrame, false, FMathf::ZeroTolerance, PtSnapVertexOrEdgeThresholdSq);
-		if (!bWalkSuccess)
-		{
-			return false;
+			WalkDir /= WalkLen;
+			FVector3d WalkNormal(-WalkDir.Y, WalkDir.X, 0);
+			bool bWalkSuccess = SurfacePath.AddViaPlanarWalk(CurrentSeedTriID, StartPos, -1, LastVert, FVector3d(Path2D[IdxB].X, Path2D[IdxB].Y, 0), WalkNormal, ProjectToFrame, false, FMathf::ZeroTolerance, PtSnapVertexOrEdgeThresholdSq);
+			if (!bWalkSuccess)
+			{
+				return false;
+			}
+			bEmbedSuccess = SurfacePath.EmbedSimplePath(false, OutPathVertices, true, PtSnapVertexOrEdgeThresholdSq);
 		}
-		bool bEmbedSuccess = SurfacePath.EmbedSimplePath(false, OutPathVertices, true, PtSnapVertexOrEdgeThresholdSq);
+		
+		
+
+		// TODO: remove this debugging check:
+		//TSet<int> VertPathSet;
+		//for (int VID : OutPathVertices)
+		//{
+		//	if (VID == LastVert || VID == OutPathVertices[0])
+		//	{
+		//		continue;
+		//	}
+		//	ensure(!VertPathSet.Contains(VID));
+		//	VertPathSet.Add(VID);
+		//}
+
+
 		OutVertexCorrespondence.Add(OutPathVertices.Num() - 1);
 		if (!bEmbedSuccess)
 		{
@@ -827,8 +853,27 @@ bool EmbedProjectedPath(FDynamicMesh3* Mesh, int StartTriID, FFrame3d Frame, con
 			}
 		}
 		OutVertexCorrespondence.Pop();
+
+		// wrap any trailing correspondence verts that happened to point to the last vertex
+		for (int i = OutVertexCorrespondence.Num() - 1; i >= 0; i--)
+		{
+			if (OutVertexCorrespondence[i] == OutPathVertices.Num())
+			{
+				OutVertexCorrespondence[i] = 0;
+			}
+			else
+			{
+				break;
+			}
+		}
 	}
 
+	// TODO rm debug check
+	//for (int Corresp : OutVertexCorrespondence)
+	//{
+	//	ensure(Corresp < OutPathVertices.Num());
+	//}
+
 	// TODO: check if the walk doesn't repeat any vertices / do something reasonable in cases where it does.
 
 	// If a selection was requested, flood fill to select the faces enclosed by the path
diff --git a/Engine/Plugins/Experimental/MeshModelingToolset/Source/ModelingOperatorsEditorOnly/Private/CuttingOps/EmbedPolygonsOp.cpp b/Engine/Plugins/Experimental/MeshModelingToolset/Source/ModelingOperatorsEditorOnly/Private/CuttingOps/EmbedPolygonsOp.cpp
--- a/Engine/Plugins/Experimental/MeshModelingToolset/Source/ModelingOperatorsEditorOnly/Private/CuttingOps/EmbedPolygonsOp.cpp
+++ b/Engine/Plugins/Experimental/MeshModelingToolset/Source/ModelingOperatorsEditorOnly/Private/CuttingOps/EmbedPolygonsOp.cpp
@@ -45,12 +45,27 @@ void FEmbedPolygonsOp::CalculateResult(FProgressCancel* Progress)
 	
 	TArray<TPair<float, int>> SortedHitTriangles;
 	TMeshQueries<FDynamicMesh3>::FindHitTriangles_LinearSearch(*ResultMesh, FRay3d(Frame.FromPlaneUV(Polygon[0]), Frame.Z()), SortedHitTriangles);
-	if (SortedHitTriangles.Num() < 1 || (SortedHitTriangles.Num() < 2 && Operation == EEmbeddedPolygonOpMethod::CutThrough))
+
+	if (SortedHitTriangles.Num() < 1)
 	{
 		// didn't hit the mesh 
 		return;
 	}
 
+	int SecondHit = 1;
+	if (Operation == EEmbeddedPolygonOpMethod::CutThrough)
+	{
+		while (SecondHit < SortedHitTriangles.Num() && FMath::IsNearlyEqual(SortedHitTriangles[SecondHit].Key, SortedHitTriangles[0].Key))
+		{
+			SecondHit++;
+		}
+		if (SecondHit >= SortedHitTriangles.Num())
+		{
+			// failed to find a second surface to connect to
+			return;
+		}
+	}
+
 	auto CutHole = [](FDynamicMesh3& Mesh, const FFrame3d& F, int TriStart, const FPolygon2d& PolygonArg, TArray<int>& PathVertIDs, TArray<int>& PathVertCorrespond)
 	{
 		FMeshFaceSelection Selection(&Mesh);
@@ -92,7 +107,7 @@ void FEmbedPolygonsOp::CalculateResult(FProgressCancel* Progress)
 	if (Operation == EEmbeddedPolygonOpMethod::CutThrough)
 	{
 		TArray<int> PathVertIDs2, PathVertCorrespond2;
-		bool bCutSide2 = CutHole(*ResultMesh, Frame, SortedHitTriangles[1].Value, Polygon, PathVertIDs2, PathVertCorrespond2);
+		bool bCutSide2 = CutHole(*ResultMesh, Frame, SortedHitTriangles[SecondHit].Value, Polygon, PathVertIDs2, PathVertCorrespond2);
 		if (!bCutSide2 || PathVertIDs2.Num() < 2)
 		{
 			return;
