add parenttype template to the base dynamicattribute classes to support dynamic attributes on objects other than meshes

#rnx

diff --git a/Engine/Plugins/Experimental/GeometryProcessing/Source/DynamicMesh/Private/DynamicMeshAttributeSet.cpp b/Engine/Plugins/Experimental/GeometryProcessing/Source/DynamicMesh/Private/DynamicMeshAttributeSet.cpp
--- a/Engine/Plugins/Experimental/GeometryProcessing/Source/DynamicMesh/Private/DynamicMeshAttributeSet.cpp
+++ b/Engine/Plugins/Experimental/GeometryProcessing/Source/DynamicMesh/Private/DynamicMeshAttributeSet.cpp
@@ -44,19 +44,19 @@ bool FDynamicMeshAttributeSet::IsSeamVertex(int VID, bool bBoundaryIsSeam) const
 
 void FDynamicMeshAttributeSet::OnNewVertex(int VertexID, bool bInserted)
 {
-	FDynamicAttributeSetBase::OnNewVertex(VertexID, bInserted);
+	FDynamicMeshAttributeSetBase::OnNewVertex(VertexID, bInserted);
 }
 
 
 void FDynamicMeshAttributeSet::OnRemoveVertex(int VertexID)
 {
-	FDynamicAttributeSetBase::OnRemoveVertex(VertexID);
+	FDynamicMeshAttributeSetBase::OnRemoveVertex(VertexID);
 }
 
 
 void FDynamicMeshAttributeSet::OnNewTriangle(int TriangleID, bool bInserted)
 {
-	FDynamicAttributeSetBase::OnNewTriangle(TriangleID, bInserted);
+	FDynamicMeshAttributeSetBase::OnNewTriangle(TriangleID, bInserted);
 
 	for (FDynamicMeshUVOverlay& UVLayer : UVLayers)
 	{
@@ -74,7 +74,7 @@ void FDynamicMeshAttributeSet::OnNewTriangle(int TriangleID, bool bInserted)
 
 void FDynamicMeshAttributeSet::OnRemoveTriangle(int TriangleID)
 {
-	FDynamicAttributeSetBase::OnRemoveTriangle(TriangleID);
+	FDynamicMeshAttributeSetBase::OnRemoveTriangle(TriangleID);
 
 	for (FDynamicMeshUVOverlay& UVLayer : UVLayers)
 	{
@@ -87,7 +87,7 @@ void FDynamicMeshAttributeSet::OnRemoveTriangle(int TriangleID)
 
 void FDynamicMeshAttributeSet::OnReverseTriOrientation(int TriangleID)
 {
-	FDynamicAttributeSetBase::OnReverseTriOrientation(TriangleID);
+	FDynamicMeshAttributeSetBase::OnReverseTriOrientation(TriangleID);
 
 	for (FDynamicMeshUVOverlay& UVLayer : UVLayers)
 	{
@@ -100,7 +100,7 @@ void FDynamicMeshAttributeSet::OnReverseTriOrientation(int TriangleID)
 
 void FDynamicMeshAttributeSet::OnSplitEdge(const FDynamicMesh3::FEdgeSplitInfo& SplitInfo)
 {
-	FDynamicAttributeSetBase::OnSplitEdge(SplitInfo);
+	FDynamicMeshAttributeSetBase::OnSplitEdge(SplitInfo);
 
 	for (FDynamicMeshUVOverlay& UVLayer : UVLayers)
 	{
@@ -116,7 +116,7 @@ void FDynamicMeshAttributeSet::OnSplitEdge(const FDynamicMesh3::FEdgeSplitInfo&
 
 void FDynamicMeshAttributeSet::OnFlipEdge(const FDynamicMesh3::FEdgeFlipInfo & flipInfo)
 {
-	FDynamicAttributeSetBase::OnFlipEdge(flipInfo);
+	FDynamicMeshAttributeSetBase::OnFlipEdge(flipInfo);
 
 	for (FDynamicMeshUVOverlay& UVLayer : UVLayers)
 	{
@@ -133,7 +133,7 @@ void FDynamicMeshAttributeSet::OnFlipEdge(const FDynamicMesh3::FEdgeFlipInfo & f
 
 void FDynamicMeshAttributeSet::OnCollapseEdge(const FDynamicMesh3::FEdgeCollapseInfo & collapseInfo)
 {
-	FDynamicAttributeSetBase::OnCollapseEdge(collapseInfo);
+	FDynamicMeshAttributeSetBase::OnCollapseEdge(collapseInfo);
 
 	for (FDynamicMeshUVOverlay& UVLayer : UVLayers)
 	{
@@ -149,7 +149,7 @@ void FDynamicMeshAttributeSet::OnCollapseEdge(const FDynamicMesh3::FEdgeCollapse
 
 void FDynamicMeshAttributeSet::OnPokeTriangle(const FDynamicMesh3::FPokeTriangleInfo & pokeInfo)
 {
-	FDynamicAttributeSetBase::OnPokeTriangle(pokeInfo);
+	FDynamicMeshAttributeSetBase::OnPokeTriangle(pokeInfo);
 
 	for (FDynamicMeshUVOverlay& UVLayer : UVLayers)
 	{
@@ -165,7 +165,7 @@ void FDynamicMeshAttributeSet::OnPokeTriangle(const FDynamicMesh3::FPokeTriangle
 
 void FDynamicMeshAttributeSet::OnMergeEdges(const FDynamicMesh3::FMergeEdgesInfo & mergeInfo)
 {
-	FDynamicAttributeSetBase::OnMergeEdges(mergeInfo);
+	FDynamicMeshAttributeSetBase::OnMergeEdges(mergeInfo);
 
 	for (FDynamicMeshUVOverlay& UVLayer : UVLayers)
 	{
diff --git a/Engine/Plugins/Experimental/GeometryProcessing/Source/DynamicMesh/Public/DynamicAttribute.h b/Engine/Plugins/Experimental/GeometryProcessing/Source/DynamicMesh/Public/DynamicAttribute.h
--- a/Engine/Plugins/Experimental/GeometryProcessing/Source/DynamicMesh/Public/DynamicAttribute.h
+++ b/Engine/Plugins/Experimental/GeometryProcessing/Source/DynamicMesh/Public/DynamicAttribute.h
@@ -4,75 +4,64 @@
 
 #include "DynamicMesh3.h"
 
-
-class FDynamicAttributeBase;
-
+template<typename ParentType>
+class TDynamicAttributeBase;
 
 /**
 * Generic base class for change tracking of an attribute layer
 */
-class DYNAMICMESH_API FDynamicAttributeChangeBase
+template<typename ParentType>
+class TDynamicAttributeChangeBase
 {
 public:
-	virtual ~FDynamicAttributeChangeBase()
+	virtual ~TDynamicAttributeChangeBase()
 	{
 	}
 
 	// default do-nothing implementations are provided because many attribute layers will only care about some kinds of elements and won't implement all of these
 
-	virtual void SaveInitialTriangle(const FDynamicAttributeBase* Attribute, int TriangleID)
+	virtual void SaveInitialTriangle(const TDynamicAttributeBase<ParentType>* Attribute, int TriangleID)
 	{
 	}
-	virtual void SaveInitialVertex(const FDynamicAttributeBase* Attribute, int VertexID)
+	virtual void SaveInitialVertex(const TDynamicAttributeBase<ParentType>* Attribute, int VertexID)
 	{
 	}
 
-	virtual void StoreAllFinalTriangles(const FDynamicAttributeBase* Attribute, const TArray<int>& TriangleIDs)
+	virtual void StoreAllFinalTriangles(const TDynamicAttributeBase<ParentType>* Attribute, const TArray<int>& TriangleIDs)
 	{
 	}
-	virtual void StoreAllFinalVertices(const FDynamicAttributeBase* Attribute, const TArray<int>& VertexIDs)
+	virtual void StoreAllFinalVertices(const TDynamicAttributeBase<ParentType>* Attribute, const TArray<int>& VertexIDs)
 	{
 	}
 
-	virtual bool Apply(FDynamicAttributeBase* Attribute, bool bRevert) const
+	virtual bool Apply(TDynamicAttributeBase<ParentType>* Attribute, bool bRevert) const
 	{
 		return false;
 	}
 };
 
+typedef TDynamicAttributeChangeBase<FDynamicMesh3> FDynamicMeshAttributeChangeBase;
+
 
 /**
  * Base class for attributes that live on a dynamic mesh (or similar dynamic object)
  *
  * Subclasses can override the On* functions to ensure the attribute remains up to date through changes to the dynamic object
  */
-class DYNAMICMESH_API FDynamicAttributeBase
+template<typename ParentType>
+class TDynamicAttributeBase
 {
 
 public:
-	virtual ~FDynamicAttributeBase()
+	virtual ~TDynamicAttributeBase()
 	{
 	}
 
 public:
 
 	/** Allocate a new copy of the attribute layer, optionally with a different parent mesh */
-	virtual FDynamicAttributeBase* MakeCopy(FDynamicMesh3* ParentMeshIn) const = 0;
-
-	/** Update to reflect an edge split in the parent mesh */
-	virtual void OnSplitEdge(const FDynamicMesh3::FEdgeSplitInfo& SplitInfo)
-	{
-	}
+	virtual TDynamicAttributeBase* MakeCopy(ParentType* ParentIn) const = 0;
 
-	/** Update to reflect an edge flip in the parent mesh */
-	virtual void OnFlipEdge(const FDynamicMesh3::FEdgeFlipInfo& FlipInfo)
-	{
-	}
-
-	/** Update to reflect an edge collapse in the parent mesh */
-	virtual void OnCollapseEdge(const FDynamicMesh3::FEdgeCollapseInfo& CollapseInfo)
-	{
-	}
 
 	virtual void OnNewVertex(int VertexID, bool bInserted)
 	{
@@ -90,16 +79,6 @@ public:
 	{
 	}
 
-	/** Update to reflect a face poke in the parent mesh */
-	virtual void OnPokeTriangle(const FDynamicMesh3::FPokeTriangleInfo& PokeInfo)
-	{
-	}
-
-	/** Update to reflect an edge merge in the parent mesh */
-	virtual void OnMergeEdges(const FDynamicMesh3::FMergeEdgesInfo& MergeInfo)
-	{
-	}
-
 	virtual void OnReverseTriOrientation(int TriangleID)
 	{
 	}
@@ -117,24 +96,55 @@ public:
 	}
 
 
-	virtual TUniquePtr<FDynamicAttributeChangeBase> NewBlankChange() = 0;
+	virtual TUniquePtr<TDynamicAttributeChangeBase<ParentType>> NewBlankChange() = 0;
+
+
+
+	/** Update to reflect an edge split in the parent mesh */
+	virtual void OnSplitEdge(const FDynamicMesh3::FEdgeSplitInfo& SplitInfo)
+	{
+	}
+
+	/** Update to reflect an edge flip in the parent mesh */
+	virtual void OnFlipEdge(const FDynamicMesh3::FEdgeFlipInfo& FlipInfo)
+	{
+	}
+
+	/** Update to reflect an edge collapse in the parent mesh */
+	virtual void OnCollapseEdge(const FDynamicMesh3::FEdgeCollapseInfo& CollapseInfo)
+	{
+	}
+
+	/** Update to reflect a face poke in the parent mesh */
+	virtual void OnPokeTriangle(const FDynamicMesh3::FPokeTriangleInfo& PokeInfo)
+	{
+	}
+
+	/** Update to reflect an edge merge in the parent mesh */
+	virtual void OnMergeEdges(const FDynamicMesh3::FMergeEdgesInfo& MergeInfo)
+	{
+	}
 
 };
 
 
+typedef TDynamicAttributeBase<FDynamicMesh3> FDynamicMeshAttributeBase;
+
+
 /**
 * Generic base class for managing a set of registered attributes that must all be kept up to date
 */
-class DYNAMICMESH_API FDynamicAttributeSetBase
+template<typename ParentType>
+class TDynamicAttributeSetBase
 {
 protected:
 	// not managed by the base class; we should be able to register any attributes here that we want to be automatically updated
-	TArray<FDynamicAttributeBase*> RegisteredAttributes;
+	TArray<TDynamicAttributeBase<ParentType>*> RegisteredAttributes;
 
 	/**
 	 * Stores the given attribute pointer in the attribute register, so that it will be updated with mesh changes, but does not take ownership of the attribute memory.
 	 */
-	void RegisterExternalAttribute(FDynamicAttributeBase* Attribute)
+	void RegisterExternalAttribute(TDynamicAttributeBase<ParentType>* Attribute)
 	{
 		RegisteredAttributes.Add(Attribute);
 	}
@@ -145,7 +155,7 @@ protected:
 	}
 
 public:
-	virtual ~FDynamicAttributeSetBase()
+	virtual ~TDynamicAttributeSetBase()
 	{
 	}
 
@@ -154,7 +164,7 @@ public:
 		return RegisteredAttributes.Num();
 	}
 
-	FDynamicAttributeBase* GetRegisteredAttribute(int Idx) const
+	TDynamicAttributeBase<ParentType>* GetRegisteredAttribute(int Idx) const
 	{
 		return RegisteredAttributes[Idx];
 	}
@@ -163,88 +173,95 @@ public:
 	// attributes when the parent mesh topology has been modified.
 	virtual void OnNewTriangle(int TriangleID, bool bInserted)
 	{
-		for (FDynamicAttributeBase* A : RegisteredAttributes)
+		for (TDynamicAttributeBase<ParentType>* A : RegisteredAttributes)
 		{
 			A->OnNewTriangle(TriangleID, bInserted);
 		}
 	}
 	virtual void OnNewVertex(int VertexID, bool bInserted)
 	{
-		for (FDynamicAttributeBase* A : RegisteredAttributes)
+		for (TDynamicAttributeBase<ParentType>* A : RegisteredAttributes)
 		{
 			A->OnNewVertex(VertexID, bInserted);
 		}
 	}
 	virtual void OnRemoveTriangle(int TriangleID)
 	{
-		for (FDynamicAttributeBase* A : RegisteredAttributes)
+		for (TDynamicAttributeBase<ParentType>* A : RegisteredAttributes)
 		{
 			A->OnRemoveTriangle(TriangleID);
 		}
 	}
 	virtual void OnRemoveVertex(int VertexID)
 	{
-		for (FDynamicAttributeBase* A : RegisteredAttributes)
+		for (TDynamicAttributeBase<ParentType>* A : RegisteredAttributes)
 		{
 			A->OnRemoveTriangle(VertexID);
 		}
 	}
 	virtual void OnReverseTriOrientation(int TriangleID)
 	{
-		for (FDynamicAttributeBase* A : RegisteredAttributes)
+		for (TDynamicAttributeBase<ParentType>* A : RegisteredAttributes)
 		{
 			A->OnReverseTriOrientation(TriangleID);
 		}
 	}
+
+	/**
+	* Check validity of attributes
+	* 
+	* @param bAllowNonmanifold Accept non-manifold topology as valid. Note that this should almost always be true for attributes; non-manifold overlays are generally valid.
+	* @param FailMode Desired behavior if mesh is found invalid
+	*/
+	virtual bool CheckValidity(bool bAllowNonmanifold, EValidityCheckFailMode FailMode) const
+	{
+		bool bValid = true;
+		for (TDynamicAttributeBase<ParentType>* A : RegisteredAttributes)
+		{
+			bValid = A->CheckValidity(bAllowNonmanifold, FailMode) && bValid;
+		}
+		return bValid;
+	}
+
+
+	// mesh-specific on* functions; may be split out
+public:
+
 	virtual void OnSplitEdge(const FDynamicMesh3::FEdgeSplitInfo& SplitInfo)
 	{
-		for (FDynamicAttributeBase* A : RegisteredAttributes)
+		for (TDynamicAttributeBase<ParentType>* A : RegisteredAttributes)
 		{
 			A->OnSplitEdge(SplitInfo);
 		}
 	}
 	virtual void OnFlipEdge(const FDynamicMesh3::FEdgeFlipInfo& FlipInfo)
 	{
-		for (FDynamicAttributeBase* A : RegisteredAttributes)
+		for (TDynamicAttributeBase<ParentType>* A : RegisteredAttributes)
 		{
 			A->OnFlipEdge(FlipInfo);
 		}
 	}
 	virtual void OnCollapseEdge(const FDynamicMesh3::FEdgeCollapseInfo& CollapseInfo)
 	{
-		for (FDynamicAttributeBase* A : RegisteredAttributes)
+		for (TDynamicAttributeBase<ParentType>* A : RegisteredAttributes)
 		{
 			A->OnCollapseEdge(CollapseInfo);
 		}
 	}
 	virtual void OnPokeTriangle(const FDynamicMesh3::FPokeTriangleInfo& PokeInfo)
 	{
-		for (FDynamicAttributeBase* A : RegisteredAttributes)
+		for (TDynamicAttributeBase<ParentType>* A : RegisteredAttributes)
 		{
 			A->OnPokeTriangle(PokeInfo);
 		}
 	}
 	virtual void OnMergeEdges(const FDynamicMesh3::FMergeEdgesInfo& MergeInfo)
 	{
-		for (FDynamicAttributeBase* A : RegisteredAttributes)
+		for (TDynamicAttributeBase<ParentType>* A : RegisteredAttributes)
 		{
 			A->OnMergeEdges(MergeInfo);
 		}
 	}
-
-	/**
-	* Check validity of attributes
-	* 
-	* @param bAllowNonmanifold Accept non-manifold topology as valid. Note that this should almost always be true for attributes; non-manifold overlays are generally valid.
-	* @param FailMode Desired behavior if mesh is found invalid
-	*/
-	virtual bool CheckValidity(bool bAllowNonmanifold, EValidityCheckFailMode FailMode) const
-	{
-		bool bValid = true;
-		for (FDynamicAttributeBase* A : RegisteredAttributes)
-		{
-			bValid = A->CheckValidity(bAllowNonmanifold, FailMode) && bValid;
-		}
-		return bValid;
-	}
 };
+
+typedef TDynamicAttributeSetBase<FDynamicMesh3> FDynamicMeshAttributeSetBase;
diff --git a/Engine/Plugins/Experimental/GeometryProcessing/Source/DynamicMesh/Public/DynamicMeshAttributeSet.h b/Engine/Plugins/Experimental/GeometryProcessing/Source/DynamicMesh/Public/DynamicMeshAttributeSet.h
--- a/Engine/Plugins/Experimental/GeometryProcessing/Source/DynamicMesh/Public/DynamicMeshAttributeSet.h
+++ b/Engine/Plugins/Experimental/GeometryProcessing/Source/DynamicMesh/Public/DynamicMeshAttributeSet.h
@@ -24,7 +24,7 @@ typedef TDynamicMeshTriangleAttribute<int32, 1> FDynamicMeshMaterialAttribute;
  * 
  * @todo current internal structure is a work-in-progress
  */
-class DYNAMICMESH_API FDynamicMeshAttributeSet : public FDynamicAttributeSetBase
+class DYNAMICMESH_API FDynamicMeshAttributeSet : public FDynamicMeshAttributeSetBase
 {
 public:
 
@@ -58,7 +58,7 @@ public:
 		ResetRegisteredAttributes();
 		for (int Idx = 0; Idx < Copy.GenericAttributes.Num(); Idx++)
 		{
-			const FDynamicAttributeBase* SourceAttrib = Copy.GenericAttributes[Idx].Get();
+			const FDynamicMeshAttributeBase* SourceAttrib = Copy.GenericAttributes[Idx].Get();
 			AttachAttribute(SourceAttrib->MakeCopy(ParentMesh));
 		}
 
@@ -197,14 +197,14 @@ public:
 	}
 
 	// Attach a new attribute (and transfer ownership of it to the attribute set)
-	int AttachAttribute(FDynamicAttributeBase* Attribute)
+	int AttachAttribute(FDynamicMeshAttributeBase* Attribute)
 	{
-		int AttributeID = GenericAttributes.Add(TUniquePtr<FDynamicAttributeBase>(Attribute));
+		int AttributeID = GenericAttributes.Add(TUniquePtr<FDynamicMeshAttributeBase>(Attribute));
 		RegisterExternalAttribute(Attribute);
 		return AttributeID;
 	}
 
-	FDynamicAttributeBase* GetAttachedAttribute(int AttributeID)
+	FDynamicMeshAttributeBase* GetAttachedAttribute(int AttributeID)
 	{
 		return GenericAttributes[AttributeID].Get();
 	}
@@ -226,7 +226,7 @@ protected:
 
 	TUniquePtr<FDynamicMeshMaterialAttribute> MaterialIDAttrib;
 
-	TArray<TUniquePtr<FDynamicAttributeBase>> GenericAttributes;
+	TArray<TUniquePtr<FDynamicMeshAttributeBase>> GenericAttributes;
 
 protected:
 	friend class FDynamicMesh3;
@@ -265,7 +265,7 @@ protected:
 	 */
 	virtual bool CheckValidity(bool bAllowNonmanifold, EValidityCheckFailMode FailMode) const
 	{
-		bool bValid = FDynamicAttributeSetBase::CheckValidity(bAllowNonmanifold, FailMode);
+		bool bValid = FDynamicMeshAttributeSetBase::CheckValidity(bAllowNonmanifold, FailMode);
 		for (int UVLayerIndex = 0; UVLayerIndex < NumUVLayers(); UVLayerIndex++)
 		{
 			bValid = GetUVLayer(UVLayerIndex)->CheckValidity(bAllowNonmanifold, FailMode) && bValid;
diff --git a/Engine/Plugins/Experimental/GeometryProcessing/Source/DynamicMesh/Public/DynamicMeshChangeTracker.h b/Engine/Plugins/Experimental/GeometryProcessing/Source/DynamicMesh/Public/DynamicMeshChangeTracker.h
--- a/Engine/Plugins/Experimental/GeometryProcessing/Source/DynamicMesh/Public/DynamicMeshChangeTracker.h
+++ b/Engine/Plugins/Experimental/GeometryProcessing/Source/DynamicMesh/Public/DynamicMeshChangeTracker.h
@@ -38,7 +38,7 @@
  * on its own (see comments in ApplyReplaceChange)
  */
 template<typename RealType, int ElementSize>
-class DYNAMICMESH_API TDynamicMeshAttributeChange
+class TDynamicMeshAttributeChange
 {
 public:
 	void SaveInitialElement(const TDynamicMeshOverlay<RealType,ElementSize>* Overlay, int ElementID);
@@ -95,7 +95,7 @@ public:
 	TArray<FDynamicMeshUVChange> UVChanges;
 	TArray<FDynamicMeshNormalChange> NormalChanges;
 	TOptional<FDynamicMeshTriangleAttributeChange<int32, 1>> MaterialIDAttribChange;
-	TArray<TUniquePtr<FDynamicAttributeChangeBase>> RegisteredAttributeChanges;
+	TArray<TUniquePtr<FDynamicMeshAttributeChangeBase>> RegisteredAttributeChanges;
 
 	/** call ::Apply() on all the UV and Normal changes */
 	bool Apply(FDynamicMeshAttributeSet* Attributes, bool bRevert) const;
diff --git a/Engine/Plugins/Experimental/GeometryProcessing/Source/DynamicMesh/Public/DynamicMeshTriangleAttribute.h b/Engine/Plugins/Experimental/GeometryProcessing/Source/DynamicMesh/Public/DynamicMeshTriangleAttribute.h
--- a/Engine/Plugins/Experimental/GeometryProcessing/Source/DynamicMesh/Public/DynamicMeshTriangleAttribute.h
+++ b/Engine/Plugins/Experimental/GeometryProcessing/Source/DynamicMesh/Public/DynamicMeshTriangleAttribute.h
@@ -11,7 +11,7 @@ class TDynamicMeshTriangleAttribute;
 
 
 template<typename AttribValueType, int AttribDimension>
-class DYNAMICMESH_API FDynamicMeshTriangleAttributeChange : public FDynamicAttributeChangeBase
+class FDynamicMeshTriangleAttributeChange : public FDynamicMeshAttributeChangeBase
 {
 private:
 	struct FChangeTriangleAttribute
@@ -29,11 +29,11 @@ public:
 	virtual ~FDynamicMeshTriangleAttributeChange()
 	{}
 
-	inline virtual void SaveInitialTriangle(const FDynamicAttributeBase* Attribute, int TriangleID) override;
+	inline virtual void SaveInitialTriangle(const FDynamicMeshAttributeBase* Attribute, int TriangleID) override;
 
-	inline virtual void StoreAllFinalTriangles(const FDynamicAttributeBase* Attribute, const TArray<int>& TriangleIDs) override;
+	inline virtual void StoreAllFinalTriangles(const FDynamicMeshAttributeBase* Attribute, const TArray<int>& TriangleIDs) override;
 
-	inline virtual bool Apply(FDynamicAttributeBase* Attribute, bool bRevert) const override;
+	inline virtual bool Apply(FDynamicMeshAttributeBase* Attribute, bool bRevert) const override;
 };
 
 
@@ -46,7 +46,7 @@ public:
  * can be mirrored to the overlay via OnSplitEdge(), etc.
  */
 template<typename AttribValueType, int AttribDimension>
-class TDynamicMeshTriangleAttribute : public FDynamicAttributeBase
+class TDynamicMeshTriangleAttribute : public FDynamicMeshAttributeBase
 {
 
 protected:
@@ -77,7 +77,7 @@ public:
 	/** @return the parent mesh for this overlay */
 	FDynamicMesh3* GetParentMesh() { return ParentMesh; }
 
-	virtual FDynamicAttributeBase* MakeCopy(FDynamicMesh3* ParentMeshIn) const override
+	virtual FDynamicMeshAttributeBase* MakeCopy(FDynamicMesh3* ParentMeshIn) const override
 	{
 		TDynamicMeshTriangleAttribute<AttribValueType, AttribDimension>* ToFill = new TDynamicMeshTriangleAttribute<AttribValueType, AttribDimension>(ParentMeshIn);
 		ToFill->Copy(*this);
@@ -190,7 +190,7 @@ public:
 	}
 
 
-	virtual TUniquePtr<FDynamicAttributeChangeBase> NewBlankChange() override
+	virtual TUniquePtr<FDynamicMeshAttributeChangeBase> NewBlankChange() override
 	{
 		return MakeUnique<FDynamicMeshTriangleAttributeChange<AttribValueType, AttribDimension>>();
 	}
@@ -238,7 +238,7 @@ public:
 
 
 template<typename AttribValueType, int AttribDimension>
-void FDynamicMeshTriangleAttributeChange<AttribValueType, AttribDimension>::SaveInitialTriangle(const FDynamicAttributeBase* Attribute, int TriangleID)
+void FDynamicMeshTriangleAttributeChange<AttribValueType, AttribDimension>::SaveInitialTriangle(const FDynamicMeshAttributeBase* Attribute, int TriangleID)
 {
 	FChangeTriangleAttribute& Change = OldTriangleAttributes.Emplace_GetRef();
 	Change.TriangleID = TriangleID;
@@ -247,7 +247,7 @@ void FDynamicMeshTriangleAttributeChange<AttribValueType, AttribDimension>::Save
 }
 
 template<typename AttribValueType, int AttribDimension>
-void FDynamicMeshTriangleAttributeChange<AttribValueType, AttribDimension>::StoreAllFinalTriangles(const FDynamicAttributeBase* Attribute, const TArray<int>& TriangleIDs)
+void FDynamicMeshTriangleAttributeChange<AttribValueType, AttribDimension>::StoreAllFinalTriangles(const FDynamicMeshAttributeBase* Attribute, const TArray<int>& TriangleIDs)
 {
 	const TDynamicMeshTriangleAttribute<AttribValueType, AttribDimension>* AttribCast = static_cast<const TDynamicMeshTriangleAttribute<AttribValueType, AttribDimension>*>(Attribute);
 	NewTriangleAttributes.Reserve(NewTriangleAttributes.Num() + TriangleIDs.Num());
@@ -260,7 +260,7 @@ void FDynamicMeshTriangleAttributeChange<AttribValueType, AttribDimension>::Stor
 }
 
 template<typename AttribValueType, int AttribDimension>
-bool FDynamicMeshTriangleAttributeChange<AttribValueType, AttribDimension>::Apply(FDynamicAttributeBase* Attribute, bool bRevert) const
+bool FDynamicMeshTriangleAttributeChange<AttribValueType, AttribDimension>::Apply(FDynamicMeshAttributeBase* Attribute, bool bRevert) const
 {
 	TDynamicMeshTriangleAttribute<AttribValueType, AttribDimension>* AttribCast = static_cast<TDynamicMeshTriangleAttribute<AttribValueType, AttribDimension>*>(Attribute);
 	const TArray<FChangeTriangleAttribute> *Changes = bRevert ? &OldTriangleAttributes : &NewTriangleAttributes;
diff --git a/Engine/Plugins/Experimental/GeometryProcessing/Source/DynamicMesh/Public/DynamicVertexAttribute.h b/Engine/Plugins/Experimental/GeometryProcessing/Source/DynamicMesh/Public/DynamicVertexAttribute.h
--- a/Engine/Plugins/Experimental/GeometryProcessing/Source/DynamicMesh/Public/DynamicVertexAttribute.h
+++ b/Engine/Plugins/Experimental/GeometryProcessing/Source/DynamicMesh/Public/DynamicVertexAttribute.h
@@ -11,7 +11,7 @@ class TDynamicVertexAttribute;
 
 
 template<typename AttribValueType, int AttribDimension, typename ParentType>
-class DYNAMICMESH_API FDynamicVertexAttributeChange : public FDynamicAttributeChangeBase
+class TDynamicVertexAttributeChange : public TDynamicAttributeChangeBase<ParentType>
 {
 private:
 	struct FChangeVertexAttribute
@@ -23,17 +23,17 @@ private:
 	TArray<FChangeVertexAttribute> OldVertexAttributes, NewVertexAttributes;
 
 public:
-	FDynamicVertexAttributeChange()
+	TDynamicVertexAttributeChange()
 	{}
 
-	virtual ~FDynamicVertexAttributeChange()
+	virtual ~TDynamicVertexAttributeChange()
 	{}
 
-	inline virtual void SaveInitialVertex(const FDynamicAttributeBase* Attribute, int VertexID) override;
+	inline virtual void SaveInitialVertex(const TDynamicAttributeBase<ParentType>* Attribute, int VertexID) override;
 
-	inline virtual void StoreAllFinalVertices(const FDynamicAttributeBase* Attribute, const TArray<int>& VertexIDs) override;
+	inline virtual void StoreAllFinalVertices(const TDynamicAttributeBase<ParentType>* Attribute, const TArray<int>& VertexIDs) override;
 
-	inline virtual bool Apply(FDynamicAttributeBase* Attribute, bool bRevert) const override;
+	inline virtual bool Apply(TDynamicAttributeBase<ParentType>* Attribute, bool bRevert) const override;
 };
 
 
@@ -41,7 +41,7 @@ public:
  * TDynamicVertexAttribute provides per-vertex storage of an attribute value
  */
 template<typename AttribValueType, int AttribDimension, typename ParentType>
-class DYNAMICMESH_API TDynamicVertexAttribute : public FDynamicAttributeBase
+class TDynamicVertexAttribute : public TDynamicAttributeBase<ParentType>
 {
 
 protected:
@@ -75,7 +75,7 @@ public:
 	/** @return the parent for this attribute */
 	ParentType* GetParent() { return Parent; }
 
-	virtual FDynamicAttributeBase* MakeCopy(ParentType* ParentIn) const override
+	virtual TDynamicAttributeBase<ParentType>* MakeCopy(ParentType* ParentIn) const override
 	{
 		TDynamicVertexAttribute<AttribValueType, AttribDimension, ParentType>* ToFill = new TDynamicVertexAttribute<AttribValueType, AttribDimension, ParentType>(ParentIn);
 		ToFill->Copy(*this);
@@ -224,9 +224,9 @@ public:
 		SetAttributeFromLerp(MergeInfo.KeptVerts.B, MergeInfo.KeptVerts.B, MergeInfo.RemovedVerts.B, .5);
 	}
 
-	virtual TUniquePtr<FDynamicAttributeChangeBase> NewBlankChange() override
+	virtual TUniquePtr<TDynamicAttributeChangeBase<ParentType>> NewBlankChange() override
 	{
-		return MakeUnique<FDynamicVertexAttributeChange<AttribValueType, AttribDimension, ParentType>>();
+		return MakeUnique<TDynamicVertexAttributeChange<AttribValueType, AttribDimension, ParentType>>();
 	}
 
 	/**
@@ -294,7 +294,7 @@ using TDynamicMeshVertexAttribute = TDynamicVertexAttribute<AttribValueType, Att
 
 
 template<typename AttribValueType, int AttribDimension, typename ParentType>
-void FDynamicVertexAttributeChange<AttribValueType, AttribDimension, ParentType>::SaveInitialVertex(const FDynamicAttributeBase* Attribute, int VertexID)
+void TDynamicVertexAttributeChange<AttribValueType, AttribDimension, ParentType>::SaveInitialVertex(const TDynamicAttributeBase<ParentType>* Attribute, int VertexID)
 {
 	FChangeVertexAttribute& Change = OldVertexAttributes.Emplace_GetRef();
 	Change.VertexID = VertexID;
@@ -303,7 +303,7 @@ void FDynamicVertexAttributeChange<AttribValueType, AttribDimension, ParentType>
 }
 
 template<typename AttribValueType, int AttribDimension, typename ParentType>
-void FDynamicVertexAttributeChange<AttribValueType, AttribDimension, ParentType>::StoreAllFinalVertices(const FDynamicAttributeBase* Attribute, const TArray<int>& VertexIDs)
+void TDynamicVertexAttributeChange<AttribValueType, AttribDimension, ParentType>::StoreAllFinalVertices(const TDynamicAttributeBase<ParentType>* Attribute, const TArray<int>& VertexIDs)
 {
 	NewVertexAttributes.Reserve(NewVertexAttributes.Num() + VertexIDs.Num());
 	const TDynamicVertexAttribute<AttribValueType, AttribDimension, ParentType>* AttribCast = static_cast<const TDynamicVertexAttribute<AttribValueType, AttribDimension, ParentType>*>(Attribute);
@@ -316,7 +316,7 @@ void FDynamicVertexAttributeChange<AttribValueType, AttribDimension, ParentType>
 }
 
 template<typename AttribValueType, int AttribDimension, typename ParentType>
-bool FDynamicVertexAttributeChange<AttribValueType, AttribDimension, ParentType>::Apply(FDynamicAttributeBase* Attribute, bool bRevert) const
+bool TDynamicVertexAttributeChange<AttribValueType, AttribDimension, ParentType>::Apply(TDynamicAttributeBase<ParentType>* Attribute, bool bRevert) const
 {
 	const TArray<FChangeVertexAttribute> *Changes = bRevert ? &OldVertexAttributes : &NewVertexAttributes;
 	TDynamicVertexAttribute<AttribValueType, AttribDimension, ParentType>* AttribCast = static_cast<TDynamicVertexAttribute<AttribValueType, AttribDimension, ParentType>*>(Attribute);
@@ -329,9 +329,4 @@ bool FDynamicVertexAttributeChange<AttribValueType, AttribDimension, ParentType>
 }
 
 
-template<typename AttribValueType, int AttribDimension>
-using TDynamicMeshVertexAttribute = TDynamicVertexAttribute<AttribValueType, AttribDimension, FDynamicMesh3>;
-
-template<typename AttribValueType, int AttribDimension>
-using FDynamicMeshVertexAttributeChange = FDynamicVertexAttributeChange<AttribValueType, AttribDimension, FDynamicMesh3>;
 
