GeometryProcessing: add box2/3 center+width constructor (makes a cube)
#rnx

diff --git a/Engine/Plugins/Experimental/GeometryProcessing/Source/GeometricObjects/Public/BoxTypes.h b/Engine/Plugins/Experimental/GeometryProcessing/Source/GeometricObjects/Public/BoxTypes.h
--- a/Engine/Plugins/Experimental/GeometryProcessing/Source/GeometricObjects/Public/BoxTypes.h
+++ b/Engine/Plugins/Experimental/GeometryProcessing/Source/GeometricObjects/Public/BoxTypes.h
@@ -199,6 +199,14 @@ struct TAxisAlignedBox3
 		this->Min = Min;
 		this->Max = Max;
 	}
+
+	TAxisAlignedBox3(const FVector3<RealType>& Center, RealType HalfWidth)
+	{
+		this->Min = FVector3<RealType>(Center.X-HalfWidth, Center.Y-HalfWidth, Center.Z-HalfWidth);
+		this->Max = FVector3<RealType>(Center.X+HalfWidth, Center.Y+HalfWidth, Center.Z+HalfWidth);
+	}
+
+
 	TAxisAlignedBox3(const TAxisAlignedBox3& Box, const TFunction<FVector3<RealType>(const FVector3<RealType>&)> TransformF)
 	{
 		if (TransformF == nullptr)
@@ -417,6 +425,12 @@ struct TAxisAlignedBox2
 	{
 	}
 
+	TAxisAlignedBox2(const FVector2<RealType>& Center, RealType HalfWidth)
+	{
+		this->Min = FVector2<RealType>(Center.X - HalfWidth, Center.Y - HalfWidth);
+		this->Max = FVector2<RealType>(Center.X + HalfWidth, Center.Y + HalfWidth);
+	}
+
 	operator FBox2D() const
 	{
 		FVector2D MinV((float)Min.X, (float)Min.Y);
