GeometryProcessing: add data structures for point sets and pointset octree
TDynamicPointSet3 - like a vertex-only FDynamicMesh3, templated on real type of the vertex coordinates. Optionally supports external attribute set.
FSparseDynamicPointOctree3 - simplification of FSparseDynamicOctree3 for points (which have no extents). Supports raycast (points-with-radius), range query, parallel range query. Optional support for dynamic-expansion insertion.
TDynamicVerticesOctree3 - templated adapter that allows for attaching FSparseDynamicPointOctree3 to a "dynamicXYZ" class, ie that supports similar vertex API to FDynamicMesh3 (eg new point set)
#rnx

diff --git a/Engine/Plugins/Experimental/GeometryProcessing/Source/DynamicMesh/Public/DynamicPointSet3.h b/Engine/Plugins/Experimental/GeometryProcessing/Source/DynamicMesh/Public/DynamicPointSet3.h
new file mode 100644
--- /dev/null
+++ b/Engine/Plugins/Experimental/GeometryProcessing/Source/DynamicMesh/Public/DynamicPointSet3.h
@@ -0,0 +1,328 @@
+// Copyright 1998-2019 Epic Games, Inc. All Rights Reserved.
+
+#pragma once
+
+#include "BoxTypes.h"
+#include "GeometryTypes.h"
+#include "MathUtil.h"
+#include "Util/DynamicVector.h"
+#include "Util/IndexUtil.h"
+#include "Util/IteratorUtil.h"
+#include "Util/RefCountVector.h"
+#include "VectorTypes.h"
+#include "VectorUtil.h"
+#include "DynamicAttribute.h"
+
+class FDynamicAttributeSetBase;
+
+
+/**
+ * TDynamicPointSet3 implements a dynamic 3D point set, templated on real-value type (float or double).
+ * The points are indexed and the class allows for gaps in the index space.
+ * 
+ * The points are referred to as "Vertices" for consistency with higher-level Mesh and Graph classes.
+ * This class can be used interchangeably with the Vertices of those classes in some cases
+ * (for example in spatial data structures like TDynamicVerticesOctree3)
+ *
+ * The point data is stored in via TDynamicVector<RealType>, which internally uses chunked storage
+ * and so is relatively efficient to grow/shrink/etc
+ *
+ * Reference counts for the points are stored via a FRefCountVector instance
+ *
+ * A per-point attribute set can optionally be attached. Currently the client is
+ * responsible for managing the memory for this attribute set. The attribute set
+ * will be automatically updated when points are added or removed.
+ */
+template<typename RealType>
+class TDynamicPointSet3
+{
+protected:
+	/** Reference counts of vertex indices. Iterate over this to find out which vertex indices are valid. */
+	FRefCountVector VertexRefCounts;
+	/** List of vertex positions */
+	TDynamicVector<RealType> Vertices;
+
+	/** The timestamp is incremented any time a function that modifies the point set is called */
+	int Timestamp = 0;
+
+
+	/** Base attribute set */
+	TDynamicAttributeSetBase<TDynamicPointSet3<RealType>>* BaseAttributeSet = nullptr;
+
+
+	//
+	// Attribute Set Support. At the level of this class, we do not have a specific attribute set.
+	// Clients can provide a pointer to their attribute set so that functions like AppendVertex()
+	// or DeleteVertex() will work.
+	//
+public:
+
+	void SetExternallyManagedAttributes(TDynamicAttributeSetBase<TDynamicPointSet3<RealType>>* AttributeSet)
+	{
+		BaseAttributeSet = AttributeSet;
+	}
+
+	TDynamicAttributeSetBase<TDynamicPointSet3<RealType>>* GetBaseAttributeSet() const
+	{
+		return BaseAttributeSet;
+	}
+
+
+public:
+
+
+
+	//
+	// Indices
+	//
+public:
+
+	/** @return number of vertices in the point set */
+	int VertexCount() const
+	{
+		return (int)VertexRefCounts.GetCount();
+	}
+
+	/** @return upper bound on vertex IDs used in the point set, ie all vertex IDs in use are < MaxVertexID */
+	int MaxVertexID() const
+	{
+		return (int)VertexRefCounts.GetMaxIndex();
+	}
+
+
+	/** @return true if VertexID is a valid vertex in this point set */
+	inline bool IsVertex(int VertexID) const
+	{
+		return VertexRefCounts.IsValid(VertexID);
+	}
+
+
+
+	//
+	// Construction
+	//
+public:
+
+	/** Discard all data */
+	void Clear()
+	{
+		Vertices.Clear();
+		VertexRefCounts = FRefCountVector();
+		BaseAttributeSet = nullptr;
+		Timestamp = 0;
+	}
+
+
+	/** Append vertex at position, returns vid */
+	int AppendVertex(const FVector3<RealType>& Position)
+	{
+		int vid = VertexRefCounts.Allocate();
+		int i = 3 * vid;
+		Vertices.InsertAt(Position[2], i + 2);
+		Vertices.InsertAt(Position[1], i + 1);
+		Vertices.InsertAt(Position[0], i);
+
+		if (GetBaseAttributeSet() != nullptr)
+		{
+			GetBaseAttributeSet()->OnNewVertex(vid, false);
+		}
+
+		UpdateTimeStamp();
+		return vid;
+	}
+
+	/** Call this before a set of unsafe InsertVertex() calls */
+	void BeginUnsafeVerticesInsert()
+	{
+		// do nothing...
+	}
+
+	/** Call after a set of unsafe InsertVertex() calls to rebuild free list */
+	void EndUnsafeVerticesInsert()
+	{
+		VertexRefCounts.RebuildFreeList();
+	}
+
+	/**
+	 * Insert vertex at given index, assuming it is unused.
+	 * If bUnsafe, we use fast id allocation that does not update free list.
+	 * You should only be using this between BeginUnsafeVerticesInsert() / EndUnsafeVerticesInsert() calls
+	 */
+	EMeshResult InsertVertex(int VertexID, const FVector3<RealType>& Position, bool bUnsafe = false)
+	{
+		if (VertexRefCounts.IsValid(VertexID))
+		{
+			return EMeshResult::Failed_VertexAlreadyExists;
+		}
+
+		bool bOK = (bUnsafe) ? VertexRefCounts.AllocateAtUnsafe(VertexID) : VertexRefCounts.AllocateAt(VertexID);
+		if (bOK == false)
+		{
+			return EMeshResult::Failed_CannotAllocateVertex;
+		}
+
+		int i = 3 * vid;
+		Vertices.InsertAt(Position[2], i + 2);
+		Vertices.InsertAt(Position[1], i + 1);
+		Vertices.InsertAt(Position[0], i);
+
+		if (GetBaseAttributeSet() != nullptr)
+		{
+			GetBaseAttributeSet()->OnNewVertex(vid, true);
+		}
+
+		UpdateTimeStamp();
+		return EMeshResult::Ok;
+	}
+
+
+
+	//
+	// get/set
+	//
+public:
+
+	/** @return the vertex position */
+	inline FVector3<RealType> GetVertex(int VertexID) const
+	{
+		check(IsVertex(VertexID));
+		int i = 3 * VertexID;
+		return FVector3<RealType>(Vertices[i], Vertices[i + 1], Vertices[i + 2]);
+	}
+
+	/** Set vertex position */
+	inline void SetVertex(int VertexID, const FVector3<RealType>& vNewPos)
+	{
+		//check(VectorUtil::IsFinite(vNewPos));
+		check(IsVertex(VertexID));
+		int i = 3 * VertexID;
+		Vertices[i] = vNewPos.X;
+		Vertices[i + 1] = vNewPos.Y;
+		Vertices[i + 2] = vNewPos.Z;
+		UpdateTimeStamp();
+	}
+
+
+
+	//
+	// iterators
+	//
+public:
+	/** @return enumerable object for valid vertex indices suitable for use with range-based for, ie for ( int i : VertexIndicesItr() ) */
+	FRefCountVector::IndexEnumerable VertexIndicesItr() const
+	{
+		return VertexRefCounts.Indices();
+	}
+
+	/** Enumerate positions of all points */
+	FRefCountVector::MappedEnumerable<FVector3<RealType>> VerticesItr() const
+	{
+		return VertexRefCounts.MappedIndices<FVector3<RealType>>([this](int VertexID) {
+			int i = 3 * VertexID;
+			return FVector3<RealType>(Vertices[i], Vertices[i + 1], Vertices[i + 2]);
+		});
+	}
+
+
+
+
+	/** @return true if vertex count == max vertex id */
+	bool IsCompact() const
+	{
+		return VertexRefCounts.IsDense();
+	}
+
+	/** returns measure of compactness in range [0,1], where 1 is fully compacted */
+	RealType CompactMetric() const
+	{
+		return ((RealType)VertexCount() / (RealType)MaxVertexID());
+	}
+
+	/** Timestamp is incremented any time any change is made to the point set */
+	inline int GetTimestamp() const
+	{
+		return Timestamp;
+	}
+
+	//
+	// Geometric queries
+	//
+public:
+
+	/** Returns bounding box of all points */
+	TAxisAlignedBox3<RealType> GetBounds() const
+	{
+		TAxisAlignedBox3<RealType> Box = TAxisAlignedBox3<RealType>::Empty();
+		for (int vi : VertexIndicesItr())
+		{
+			int k = 3 * vi;
+			Box.Contain( FVector3<RealType>(Vertices[k], Vertices[k + 1], Vertices[k + 2]) );
+		}
+		return Box;
+	}
+
+
+
+	//
+	// direct buffer access  (not a good idea to use)
+	//
+public:
+	const TDynamicVector<double>& GetVerticesBuffer()
+	{
+		return Vertices;
+	}
+	const FRefCountVector& GetVerticesRefCounts()
+	{
+		return VertexRefCounts;
+	}
+
+
+
+	//
+	// Edit operations
+	//
+public:
+
+	// @todo port this from DynamicMesh3
+	//void CompactInPlace(FCompactMaps* CompactInfo = nullptr);
+
+	/**
+	 * Remove point
+	 */
+	EMeshResult RemoveVertex(int VertexID)
+	{
+		if (VertexRefCounts.IsValid(VertexID) == false)
+		{
+			return EMeshResult::Failed_NotAVertex;
+		}
+		VertexRefCounts.Decrement(vID);
+		check(VertexRefCounts.IsValid(vID) == false);
+		VertexEdgeLists.Clear(vID);
+
+		if (GetBaseAttributeSet() != nullptr)
+		{
+			GetBaseAttributeSet()->OnRemoveVertex(VertexID);
+		}
+
+		UpdateTimeStamp();
+		return EMeshResult::Ok;
+	}
+
+
+
+protected:
+	void UpdateTimeStamp()
+	{
+		Timestamp++;
+	}
+
+
+
+};
+
+
+
+typedef TDynamicPointSet3<float> FDynamicPointSet3f;
+typedef TDynamicPointSet3<double> FDynamicPointSet3d;
+
+
diff --git a/Engine/Plugins/Experimental/GeometryProcessing/Source/DynamicMesh/Public/DynamicVerticesOctree3.h b/Engine/Plugins/Experimental/GeometryProcessing/Source/DynamicMesh/Public/DynamicVerticesOctree3.h
new file mode 100644
--- /dev/null
+++ b/Engine/Plugins/Experimental/GeometryProcessing/Source/DynamicMesh/Public/DynamicVerticesOctree3.h
@@ -0,0 +1,222 @@
+// Copyright 1998-2019 Epic Games, Inc. All Rights Reserved.
+
+#pragma once
+
+#include "Spatial/MeshAABBTree3.h"
+#include "Spatial/SparseDynamicPointOctree3.h"
+#include "MeshQueries.h"
+#include "Intersection/IntersectionUtil.h"
+#include "DynamicMesh3.h"
+
+
+/**
+ * TDynamicVerticesOctree3 is an extension of FSparseDynamicPointOctree3 for the vertices of 
+ * a FDynamicXYZ3 instance (eg FDynamicMesh3, TDynamicPointSet3, FDynamicGraph, etc)
+ * This extension does several things:
+ *   1) provides a simplified API based on vertex IDs to various Octree functions
+ *   2) tracks ModifiedBounds box of modified areas
+ *
+ * The template expansion requires that FDynamicXYZ3 has the following API:
+ *   1) a function int MaxVertexID() that returns the maximum vertex ID/index
+ *   2) a function bool IsVertex(int) that returns true if the index is valid
+ *   3) a function FVector3d GetVertex(int index) that returns the position of a vertex
+ *   4) a function FAxisAlignedBox3d GetBounds() that returns a bounding-box of the point set
+ */
+template<typename SourceType>
+class TDynamicVerticesOctree3 : public FSparseDynamicPointOctree3
+{
+	// potential optimizations:
+	//    - keep track of how many vertices are descendents of each cell? root cells?
+
+
+
+public:
+	/** parent mesh */
+	const SourceType* VertexSource;
+
+	/** bounding box of vertices that have been inserted/removed since last clear */
+	FAxisAlignedBox3d ModifiedBounds;
+
+	double HitSphereRadius = 1.0f;
+
+	/** 
+	 * Add all vertices of MeshIn to the octree 
+	 */
+	void Initialize(const SourceType* VertexSourceIn, bool bDynamicExpand)
+	{
+		this->VertexSource = VertexSourceIn;
+
+		HitSphereRadius = VertexSourceIn->GetBounds().Diagonal().Length() * 0.001;
+
+		int MaxVertexID = VertexSource->MaxVertexID();
+		for ( int VertexID = 0; VertexID < MaxVertexID; VertexID++ )
+		{
+			if (VertexSource->IsVertex(VertexID))
+			{
+				if (bDynamicExpand)
+				{
+					ModifiedBounds.Contain(VertexSource->GetVertex(VertexID));
+					const SourceType* VtxSource = VertexSource;
+					FSparseDynamicPointOctree3::InsertPoint_DynamicExpand(
+						VertexID,
+						[VtxSource](int k) { return VtxSource->GetVertex(k); });
+				}
+				else
+				{
+					InsertVertex(VertexID);
+				}
+			}
+		}
+	}
+
+	/**
+	 * Reset the internal ModifiedBounds box that tracks modified triangle bounds
+	 */
+	void ResetModifiedBounds()
+	{
+		ModifiedBounds = FAxisAlignedBox3d::Empty();
+	}
+
+	/**
+	 * Insert a triangle into the tree
+	 */
+	void InsertVertex(int32 VertexID)
+	{
+		FVector3d Point = VertexSource->GetVertex(VertexID);
+		ModifiedBounds.Contain(Point);
+		FSparseDynamicPointOctree3::InsertPoint(VertexID, Point);
+	}
+
+	/**
+	 * Insert a list of vertices into the tree
+	 */
+	void InsertVertices(const TArray<int>& Vertices)
+	{
+		int N = Vertices.Num();
+		for (int i = 0; i < N; ++i)
+		{
+			FVector3d Point = VertexSource->GetVertex(Vertices[i]);
+			ModifiedBounds.Contain(Point);
+			FSparseDynamicPointOctree3::InsertPoint(Vertices[i], Point);
+		}
+	}
+
+	/**
+	 * Insert a set of vertices into the tree
+	 */
+	void InsertVertices(const TSet<int>& Vertices)
+	{
+		for (int VertexID : Vertices)
+		{
+			FVector3d Point = VertexSource->GetVertex(VertexID);
+			ModifiedBounds.Contain(Point);
+			FSparseDynamicPointOctree3::InsertPoint(VertexID, Point);
+		}
+	}
+
+
+	/**
+	 * Remove a triangle from the tree
+	 */
+	bool RemovePoint(int32 VertexID)
+	{
+		FVector3d Point = VertexSource->GetVertex(VertexID);
+		ModifiedBounds.Contain(Point);
+		return FSparseDynamicPointOctree3::RemovePoint(VertexID);
+	}
+
+	/**
+	 * Remove a list of vertices from the tree
+	 */
+	void RemoveVertices(const TArray<int>& Vertices)
+	{
+		int N = Vertices.Num();
+		for ( int i = 0; i < N; ++i )
+		{
+			FVector3d Point = VertexSource->GetVertex(Vertices[i]);
+			ModifiedBounds.Contain(Point);
+			FSparseDynamicPointOctree3::RemovePoint(Vertices[i]);
+		}
+	}
+
+	/**
+	 * Remove a set of vertices from the tree
+	 */
+	void RemoveVertices(const TSet<int>& Vertices)
+	{
+		for (int VertexID : Vertices)
+		{
+			FVector3d Point = VertexSource->GetVertex(VertexID);
+			ModifiedBounds.Contain(Point);
+			FSparseDynamicPointOctree3::RemovePoint(VertexID);
+		}
+	}
+
+
+	/**
+	 * Reinsert a set of vertices into the tree
+	 */
+	void ReinsertVertices(const TSet<int>& Vertices)
+	{
+		for (int VertexID : Vertices)
+		{
+			FVector3d Point = VertexSource->GetVertex(VertexID);
+			ModifiedBounds.Contain(Point);
+			FSparseDynamicPointOctree3::ReinsertPoint(VertexID, Bounds);
+		}
+	}
+
+
+	/**
+	 * Include the current bounds of a triangle in the ModifiedBounds box
+	 */
+	void NotifyPendingModification(int VertexID)
+	{
+		FVector3d Point = VertexSource->GetVertex(VertexID);
+		ModifiedBounds.Contain(Point);
+	}
+
+	/**
+	 * Include the current bounds of a set of vertices in the ModifiedBounds box
+	 */
+	void NotifyPendingModification(const TSet<int>& Vertices)
+	{
+		for (int VertexID : Vertices)
+		{
+			FVector3d Point = VertexSource->GetVertex(VertexID);
+			ModifiedBounds.Contain(Point);
+		}
+	}
+
+
+	/**
+	 * Find the nearest triangle of the VertexSource that is hit by the ray
+	 */
+	int32 FindNearestHitVertex(const FRay3d& Ray,
+		double MaxDistance = TNumericLimits<double>::Max()) const
+	{
+		return FSparseDynamicPointOctree3::FindNearestHitPoint(Ray,
+			[this](int vid, const FRay3d& Ray) {
+				FVector3d Point = VertexSource->GetVertex(vid);
+				FLinearIntersection Hit = IntersectionUtil::RaySphereIntersection(Ray.Origin, Ray.Direction, Center, HitSphereRadius);
+				return (Hit.intersects) ? Hit.parameter.Min : TNumericLimits<double>::Max();
+			}, MaxDistance);
+	}
+
+
+	/**
+	 * Check that the Octree is internally valid
+	 */
+	void CheckValidity(
+		EValidityCheckFailMode FailMode = EValidityCheckFailMode::Check,
+		bool bVerbose = false,
+		bool bFailOnMissingPoints = false) const
+	{
+		FSparseDynamicPointOctree3::CheckValidity(
+			[&](int vid) { return VertexSource->IsVertex(vid); },
+			[this](int vid) { return VertexSource->GetVertex(vid); },
+			FailMode, bVerbose, bFailOnMissingPoints);
+	}
+
+
+};
\ No newline at end of file
diff --git a/Engine/Plugins/Experimental/GeometryProcessing/Source/GeometricObjects/Public/Spatial/SparseDynamicPointOctree3.h b/Engine/Plugins/Experimental/GeometryProcessing/Source/GeometricObjects/Public/Spatial/SparseDynamicPointOctree3.h
new file mode 100644
--- /dev/null
+++ b/Engine/Plugins/Experimental/GeometryProcessing/Source/GeometricObjects/Public/Spatial/SparseDynamicPointOctree3.h
@@ -0,0 +1,970 @@
+// Copyright 1998-2019 Epic Games, Inc. All Rights Reserved.
+
+#pragma once
+
+
+#include "CoreMinimal.h"
+#include "Containers/StaticArray.h"
+#include "GeometryTypes.h"
+#include "BoxTypes.h"
+#include "Util/DynamicVector.h"
+#include "Util/SmallListSet.h"
+#include "Spatial/SparseGrid3.h"
+#include "Intersection/IntrRay3AxisAlignedBox3.h"
+#include "Async/ParallelFor.h"
+#include "Spatial/SparseDynamicOctree3.h"    // for FDynamicFlagArray
+
+
+
+/**
+ * FSparsePointOctreeCell is a Node in a SparseDynamicOctree3. 
+ */
+struct FSparsePointOctreeCell
+{
+	static constexpr uint32 InvalidID = TNumericLimits<uint32>::Max();;
+	static constexpr uint8 InvalidLevel = TNumericLimits<uint8>::Max();
+
+	/** ID of cell (index into cell list) */
+	uint32 CellID;
+
+	/** Level of cell in octree */
+	uint8 Level;
+
+	/** i,j,k index of cell in level, relative to origin */
+	FVector3i Index;
+
+	/** CellID of each child, or InvalidID if that child does not exist */
+	TStaticArray<uint32, 8> Children;
+
+	FSparsePointOctreeCell()
+		: CellID(InvalidID), Level(0), Index(FVector3i::Zero())
+	{
+		Children[0] = Children[1] = Children[2] = Children[3] = InvalidID;
+		Children[4] = Children[5] = Children[6] = Children[7] = InvalidID;
+	}
+
+	FSparsePointOctreeCell(uint8 LevelIn, const FVector3i& IndexIn)
+		: CellID(InvalidID), Level(LevelIn), Index(IndexIn)
+	{
+		Children[0] = Children[1] = Children[2] = Children[3] = InvalidID;
+		Children[4] = Children[5] = Children[6] = Children[7] = InvalidID;
+	}
+
+	inline bool IsExistingCell() const
+	{
+		return CellID != InvalidID;
+	}
+
+	inline bool HasChild(int ChildIndex) const
+	{
+		return Children[ChildIndex] != InvalidID;
+	}
+
+	inline uint32 GetChildCellID(int ChildIndex) const
+	{
+		return Children[ChildIndex];
+	}
+
+	inline FSparsePointOctreeCell MakeChildCell(int ChildIndex) const
+	{
+		FVector3i IndexOffset(
+			((ChildIndex & 1) != 0) ? 1 : 0,
+			((ChildIndex & 2) != 0) ? 1 : 0,
+			((ChildIndex & 4) != 0) ? 1 : 0);
+		return FSparsePointOctreeCell(Level + 1, 2*Index + IndexOffset);
+	}
+
+	inline void SetChild(uint32 ChildIndex, const FSparsePointOctreeCell& ChildCell)
+	{
+		Children[ChildIndex] = ChildCell.CellID;
+	}
+
+};
+
+
+/**
+ * FSparseDynamicPointOctree3 sorts Points with axis-aligned bounding boxes into a dynamic
+ * sparse octree of axis-aligned uniform grid cells. At the top level we have an infinite
+ * grid of "root cells" of size RootDimension, which then contain 8 children, and so on.
+ * (So in fact each cell is a separate octree, and we have a uniform grid of octrees)
+ * 
+ * The Points and their bounding-boxes are not stored in the tree. You must have an
+ * integer identifier (PointID) for each Point, and call Insert(PointID, BoundingBox).
+ * Some query functions will require you to provide a lambda/etc that can be called to
+ * retrieve the bounding box for a given PointID.
+ * 
+ * By default Points are currently inserted at the maximum possible depth, ie smallest cell that
+ * will contain them, or MaxTreeDepth. 
+ * 
+ * The octree is dynamic. Points can be removed and re-inserted.
+ * 
+ */
+class FSparseDynamicPointOctree3
+{
+	// potential optimizations/improvements
+	//    - Cell sizes are known at each level...can keep a lookup table? will this improve performance?
+	//    - Store cells for each level in separate TDynamicVectors. CellID is then [Level:8 | Index:24].
+	//      This would allow level-grids to be processed separately / in-parallel (for example a cut at given level would be much faster)
+	//    - Currently insertion is max-depth but we do not dynamically expand more than once. So early
+	//      insertions end up in very large buckets. When a child expands we should check if any of its parents would fit.
+	//    - Currently insertion is max-depth so we end up with a huge number of single-Point cells. Should only go down a level
+	//      if enough Points exist in current cell. Can do this in a greedy fashion, less optimal but still acceptable...
+	//    - get rid of ValidPointIDs, I don't think we need it?
+
+public:
+
+	//
+	// Tree configuration parameters. It is not safe to change these after tree initialization!
+	// 
+
+	/**
+	 * Size of the Root cells of the octree. 
+	 */
+	double RootDimension = 1000.0;
+
+	/**
+	 * Points will not be inserted more than this many levels deep from a Root cell
+	 */
+	int MaxTreeDepth = 10;
+
+	/**
+	 * if using the InsertPoint_DynamicExpand() insertion path, then when a cell gets
+	 * to this many points, and has depth < MaxTreeDepth, it will be expanded and its points moved down
+	 */
+	int MaxPointsPerCell = 1000;
+
+public:
+
+	/**
+	 * Test if an Point is stored in the tree
+	 * @param PointID ID of the Point
+	 * @return true if PointID is stored in this octree
+	 */
+	inline bool ContainsPoint(int32 PointID) const;
+
+	/**
+	 * Insert PointID into the Octree at maximum depth. 
+	 * Currently only expands at most one depth-level per insertion. This is sub-optimal.
+	 * @param PointID ID of the Point to insert
+	 * @param Position position of the Point
+	 *
+	 */
+	inline void InsertPoint(int32 PointID, const FVector3d& Position);
+
+	/**
+	 * Insert PointID into the Octree. Dynamically expands cells that have more than MaxPointsPerCell
+	 * @param PointID ID of the Point to insert
+	 * @param GetPositionFunc function that returns position of point (Required to reinsert points on expand)
+	 */
+	inline void InsertPoint_DynamicExpand(int32 PointID, TFunctionRef<FVector3d(int)> GetPositionFunc);
+
+
+	/**
+	 * Remove an Point from the octree
+	 * @param PointID ID of the Point
+	 * @return true if the Point was in the tree and removed
+	 */
+	inline bool RemovePoint(int32 PointID);
+
+	/**
+	 * Update the position of an Point in the octree. This is more efficient than doing a remove+insert
+	 * @param PointID ID of the Point
+	 * @param NewBounds enw bounding box of the Point
+	 */
+	inline void ReinsertPoint(int32 PointID, const FVector3d& NewPosition);
+
+	/**
+	 * Find nearest ray-hit point with Points in tree
+	 * @param Ray the ray 
+	 * @param HitPointDistFunc function that returns distance along ray to hit-point on Point identified by PointID (or TNumericLimits<double>::Max() on miss)
+	 * @param MaxDistance maximum hit distance
+	 * @return PointID of hit Point, or -1 on miss
+	 */
+	inline int32 FindNearestHitPoint(const FRay3d& Ray,
+		TFunctionRef<double(int, const FRay3d&)> HitPointDistFunc,
+		double MaxDistance = TNumericLimits<double>::Max()) const;
+
+	/**
+	 * Collect PointIDs from all the cells with bounding boxes that intersect Bounds,
+	 * where PredicateFunc passes
+	 * @param Bounds query box
+	 * @param PredicateFunc return true if point should be included (ie this is a filter)
+	 * @param PointIDsOut collected PointIDs are stored here 
+	 * @param TempBuffer optional temporary buffer to store a queue of cells (to avoid memory allocations)
+	 */
+	inline void RangeQuery(const FAxisAlignedBox3d& Bounds,
+		TFunctionRef<bool(int)> PredicateFunc,
+		TArray<int>& PointIDsOut,
+		TArray<const FSparsePointOctreeCell*>* TempBuffer = nullptr) const;
+
+
+	/**
+	 * Collect PointIDs from all the cells with bounding boxes that intersect Bounds, where PredicateFunc passes
+	 * Query is parallelized across Root cells. So if there is only one Root cell, it's actually slower
+	 * than RangeQuery() due to internal temp buffers.
+	 * @param Bounds query box
+	 * @param PredicateFunc return true if point should be included (ie this is a filter)
+	 * @param PointIDsOut collected PointIDs are stored here
+	 * @param TempBuffer optional temporary buffer to store a queue of cells (to avoid memory allocations)
+	 */
+	inline void ParallelRangeQuery(const FAxisAlignedBox3d& Bounds,
+		TFunctionRef<bool(int)> PredicateFunc,
+		TArray<int>& PointIDsOut,
+		TArray<const FSparsePointOctreeCell*>* TempBuffer = nullptr) const;
+
+
+	/**
+	 * Check that the octree is internally valid
+	 * @param IsValidPointIDFunc function that returns true if given PointID is valid
+	 * @param GetPointFunc function that returns point position identified by PointID
+	 * @param FailMode how should validity checks fail
+	 * @param bVerbose if true, print some debug info via UE_LOG
+	 * @param bFailOnMissingPoints if true, assume PointIDs are dense and that all PointIDs must be in the tree
+	 */
+	inline void CheckValidity(
+		TFunctionRef<bool(int)> IsValidPointIDFunc,
+		TFunctionRef<FVector3d(int)> GetPointFunc,
+		EValidityCheckFailMode FailMode = EValidityCheckFailMode::Check,
+		bool bVerbose = false,
+		bool bFailOnMissingPoints = false) const;
+
+	/**
+	 * statistics about internal structure of the octree
+	 */
+	struct FStatistics
+	{
+		int32 Levels;
+		TArray<int32> LevelBoxCounts;
+		TArray<int32> LevelObjCounts;
+		TArray<int32> LevelMaxObjCounts;
+		TArray<float> LevelCellSizes;
+		inline FString ToString() const;
+	};
+
+	/**
+	 * Populate given FStatistics with info about the octree
+	 */
+	inline void ComputeStatistics(FStatistics& StatsOut) const;
+
+protected:
+	// this identifier is used for unknown cells
+	static constexpr uint32 InvalidCellID = FSparsePointOctreeCell::InvalidID;
+
+	// reference counts for Cells list. We don't actually need reference counts here, but we need a free
+	// list and iterators, and FRefCountVector provides this
+	FRefCountVector CellRefCounts;
+
+	// list of cells. Note that some cells may be unused, depending on CellRefCounts
+	TDynamicVector<FSparsePointOctreeCell> Cells;
+
+	FSmallListSet CellPointLists;			// per-cell Point ID lists
+
+	TDynamicVector<uint32> PointIDToCellMap;	// map from external Point IDs to which cell the Point is in (or invalid)
+	FDynamicFlagArray ValidPointIDs;			// set of PointIDs in the tree. This is perhaps not necessary...couldn't we rely on PointIDToCellMap?
+
+	// RootCells are the top-level cells of the octree, of size RootDimension. 
+	// So the elements of this sparse grid are CellIDs
+	TSparseGrid3<uint32> RootCells;
+
+
+	// calculate the base width of a cell at a given level
+	inline double GetCellWidth(uint32 Level) const
+	{
+		double Divisor = (double)( (uint64)1 << (Level & 0x1F) );
+		double CellWidth = RootDimension / Divisor;
+		return CellWidth;
+	}
+
+
+	FAxisAlignedBox3d GetCellBox(uint32 Level, const FVector3i& Index) const
+	{
+		double CellWidth = GetCellWidth(Level);
+		double MinX = (CellWidth * (double)Index.X);
+		double MinY = (CellWidth * (double)Index.Y);
+		double MinZ = (CellWidth * (double)Index.Z);
+		return FAxisAlignedBox3d(
+			FVector3d(MinX, MinY, MinZ),
+			FVector3d(MinX + CellWidth, MinY + CellWidth, MinZ + CellWidth));
+	}
+	inline FAxisAlignedBox3d GetCellBox(const FSparsePointOctreeCell& Cell) const
+	{
+		return GetCellBox(Cell.Level, Cell.Index);
+	}
+	FVector3d GetCellCenter(const FSparsePointOctreeCell& Cell) const
+	{
+		double CellWidth = GetCellWidth(Cell.Level);
+		double MinX = CellWidth * (double)Cell.Index.X;
+		double MinY = CellWidth * (double)Cell.Index.Y;
+		double MinZ = CellWidth * (double)Cell.Index.Z;
+		CellWidth *= 0.5;
+		return FVector3d(MinX + CellWidth, MinY + CellWidth, MinZ + CellWidth);
+	}
+
+
+	FVector3i PointToIndex(uint32 Level, const FVector3d& Position) const
+	{
+		double CellWidth = GetCellWidth(Level);
+		int32 i = (int32)FMathd::Floor(Position.X / CellWidth);
+		int32 j = (int32)FMathd::Floor(Position.Y / CellWidth);
+		int32 k = (int32)FMathd::Floor(Position.Z / CellWidth);
+		return FVector3i(i, j, k);
+	}
+
+
+	int ToChildCellIndex(const FSparsePointOctreeCell& Cell, const FVector3d& Position) const
+	{
+		FVector3d Center = GetCellCenter(Cell);
+		int ChildIndex =
+			((Position.X < Center.X) ? 0 : 1) +
+			((Position.Y < Center.Y) ? 0 : 2) +
+			((Position.Z < Center.Z) ? 0 : 4);
+		return ChildIndex;
+	}
+
+	bool CellContains(const FSparsePointOctreeCell& Cell, const FVector3d& Position) const
+	{
+		FAxisAlignedBox3d CellBox = GetCellBox(Cell);
+		return CellBox.Contains(Position);
+	}
+
+	uint32 GetCellForPoint(int32 PointID) const
+	{
+		if (PointID >= 0 && PointID < PointIDToCellMap.Num())
+		{
+			return PointIDToCellMap[PointID];
+		}
+		return InvalidCellID;
+	}
+
+	inline FSparsePointOctreeCell FindCurrentContainingCell(const FVector3d& Position) const;
+
+
+	inline void Insert_NewRoot(int32 PointID, const FVector3d& Position, FSparsePointOctreeCell NewRootCell);
+	inline void Insert_ToCell(int32 PointID, const FVector3d& Position, const FSparsePointOctreeCell& ExistingCell);
+	inline void Insert_NewChildCell(int32 PointID, const FVector3d& Position, int ParentCellID, FSparsePointOctreeCell NewChildCell, int ChildIdx);
+
+	inline double FindNearestRayCellIntersection(const FSparsePointOctreeCell& Cell, const FRay3d& Ray) const;
+
+};
+
+
+
+bool FSparseDynamicPointOctree3::ContainsPoint(int32 PointID) const
+{
+	return ValidPointIDs.Get(PointID);
+}
+
+
+
+void FSparseDynamicPointOctree3::InsertPoint(int32 PointID, const FVector3d& Position)
+{
+	check(ContainsPoint(PointID) == false);
+
+	FSparsePointOctreeCell CurrentCell = FindCurrentContainingCell(Position);
+	check(CurrentCell.Level != FSparsePointOctreeCell::InvalidLevel);
+
+	// if we found a containing root cell but it doesn't exist, create it and insert
+	if (CurrentCell.Level == 0 && CurrentCell.IsExistingCell() == false)
+	{
+		Insert_NewRoot(PointID, Position, CurrentCell);
+		return;
+	}
+
+	// YIKES this currently does max-depth insertion...
+	//   desired behavior is that parent cell accumulates and then splits later!
+
+	int PotentialChildIdx = ToChildCellIndex(CurrentCell, Position);
+	// if current cell does not have this child we might fit there so try it
+	if (CurrentCell.HasChild(PotentialChildIdx) == false)
+	{
+		// todo can we do a fast check based on level and dimensions??
+		FSparsePointOctreeCell NewChildCell = CurrentCell.MakeChildCell(PotentialChildIdx);
+		if (NewChildCell.Level <= MaxTreeDepth && CellContains(NewChildCell, Position))
+		{
+			Insert_NewChildCell(PointID, Position, CurrentCell.CellID, NewChildCell, PotentialChildIdx);
+			return;
+		}
+	}
+
+	// insert into current cell if
+	//   1) child cell exists (in which case we didn't fit or FindCurrentContainingCell would have returned)
+	//   2) we tried to fit in child cell and failed
+	Insert_ToCell(PointID, Position, CurrentCell);
+}
+
+
+void FSparseDynamicPointOctree3::Insert_NewChildCell(int32 PointID, const FVector3d& Position,
+	int ParentCellID, FSparsePointOctreeCell NewChildCell, int ChildIdx)
+{
+	FSparsePointOctreeCell& OrigParentCell = Cells[ParentCellID];
+	check(OrigParentCell.HasChild(ChildIdx) == false);
+
+	NewChildCell.CellID = CellRefCounts.Allocate();
+	Cells.InsertAt(NewChildCell, NewChildCell.CellID);
+
+	PointIDToCellMap.InsertAt(NewChildCell.CellID, PointID);
+	ValidPointIDs.Set(PointID, true);
+
+	CellPointLists.AllocateAt(NewChildCell.CellID);
+	CellPointLists.Insert(NewChildCell.CellID, PointID);
+
+	OrigParentCell.SetChild(ChildIdx, NewChildCell);
+
+	check(CellContains(NewChildCell, Position));
+	check(PointToIndex(NewChildCell.Level, Position) == NewChildCell.Index);
+}
+
+
+
+void FSparseDynamicPointOctree3::Insert_ToCell(int32 PointID, const FVector3d& Position, const FSparsePointOctreeCell& ExistingCell)
+{
+	check(CellRefCounts.IsValid(ExistingCell.CellID));
+
+	PointIDToCellMap.InsertAt(ExistingCell.CellID, PointID);
+	ValidPointIDs.Set(PointID, true);
+
+	CellPointLists.Insert(ExistingCell.CellID, PointID);
+
+	check(CellContains(ExistingCell, Position));
+	check(PointToIndex(ExistingCell.Level, Position) == ExistingCell.Index);
+}
+
+
+
+void FSparseDynamicPointOctree3::Insert_NewRoot(int32 PointID, const FVector3d& Position, FSparsePointOctreeCell NewRootCell)
+{
+	check(RootCells.Has(NewRootCell.Index) == false);
+
+	NewRootCell.CellID = CellRefCounts.Allocate();
+	Cells.InsertAt(NewRootCell, NewRootCell.CellID);
+
+	PointIDToCellMap.InsertAt(NewRootCell.CellID, PointID);
+	ValidPointIDs.Set(PointID, true);
+
+	uint32* RootCellElem = RootCells.Get(NewRootCell.Index, true);
+	*RootCellElem = NewRootCell.CellID;
+
+	CellPointLists.AllocateAt(NewRootCell.CellID);
+	CellPointLists.Insert(NewRootCell.CellID, PointID);
+}
+
+
+
+
+
+
+
+void FSparseDynamicPointOctree3::InsertPoint_DynamicExpand(
+	int32 PointID, 
+	TFunctionRef<FVector3d(int)> GetPositionFunc)
+{
+	check(ContainsPoint(PointID) == false);
+
+	FVector3d Position = GetPositionFunc(PointID);
+
+	FSparsePointOctreeCell CurrentCell = FindCurrentContainingCell(Position);
+	check(CurrentCell.Level != FSparsePointOctreeCell::InvalidLevel);
+
+	// if we found a containing root cell but it doesn't exist, create it and insert
+	if (CurrentCell.Level == 0 && CurrentCell.IsExistingCell() == false)
+	{
+		Insert_NewRoot(PointID, Position, CurrentCell);
+		return;
+	}
+
+	// insert into current cell if there is space, or if we hit max depth
+	int NumPointsInCell = CellPointLists.GetCount(CurrentCell.CellID);
+	if (NumPointsInCell + 1 < MaxPointsPerCell || CurrentCell.Level == MaxTreeDepth)
+	{
+		Insert_ToCell(PointID, Position, CurrentCell);
+		return;
+	}
+
+	const FSparsePointOctreeCell* ParentCell = &Cells[CurrentCell.CellID];
+
+	// this function inserts a point into one of the child cells of CurrentCell
+	auto InsertToChildLocalFunc = [this,&GetPositionFunc,ParentCell](int InsertPointID)
+	{
+		FVector3d MovePosition = GetPositionFunc(InsertPointID);
+		int ChildIdx = this->ToChildCellIndex(*ParentCell, MovePosition);
+		if (ParentCell->HasChild(ChildIdx))
+		{
+			uint32 ChildCellID = ParentCell->GetChildCellID(ChildIdx);
+			const FSparsePointOctreeCell* ChildCell = &Cells[ChildCellID];
+			this->Insert_ToCell(InsertPointID, MovePosition, *ChildCell);
+		}
+		else
+		{
+			FSparsePointOctreeCell NewChildCell = ParentCell->MakeChildCell(ChildIdx);
+			this->Insert_NewChildCell(InsertPointID, MovePosition, ParentCell->CellID, NewChildCell, ChildIdx);
+		}
+	};
+
+	// otherwise current cell got too big, so move points to child cells
+	for (int32 MovePointID : CellPointLists.Values(ParentCell->CellID))
+	{
+		InsertToChildLocalFunc(MovePointID);
+	}
+
+	// we removed all the points from ParentCell
+	CellPointLists.Clear(ParentCell->CellID);
+
+	// now insert ourself into one of these child cells
+	InsertToChildLocalFunc(PointID);
+}
+
+
+
+
+
+bool FSparseDynamicPointOctree3::RemovePoint(int32 PointID)
+{
+	if (ContainsPoint(PointID) == false)
+	{
+		return false;
+	}
+
+	uint32 CellID = GetCellForPoint(PointID);
+	if (CellID == InvalidCellID)
+	{
+		return false;
+	}
+
+	PointIDToCellMap[PointID] = InvalidCellID;
+	ValidPointIDs.Set(PointID, false);
+
+	bool bInList = CellPointLists.Remove(CellID, PointID);
+	check(bInList);
+	return true;
+}
+
+
+
+
+
+
+void FSparseDynamicPointOctree3::ReinsertPoint(int32 PointID, const FVector3d& NewPosition)
+{
+	if (ContainsPoint(PointID))
+	{
+		uint32 CellID = GetCellForPoint(PointID);
+		if (CellID != InvalidCellID)
+		{
+			FSparsePointOctreeCell& CurrentCell = Cells[CellID];
+			if (CellContains(CurrentCell, NewPosition))
+			{
+				return;		// everything is fine
+			}
+
+		}
+	}
+
+	RemovePoint(PointID);
+	InsertPoint(PointID, NewPosition);
+}
+
+
+
+
+
+
+double FSparseDynamicPointOctree3::FindNearestRayCellIntersection(const FSparsePointOctreeCell& Cell, const FRay3d& Ray) const
+{
+	FAxisAlignedBox3d Box = GetCellBox(Cell);
+	double ray_t = TNumericLimits<double>::Max();
+	if (FIntrRay3AxisAlignedBox3d::FindIntersection(Ray, Box, ray_t))
+	{
+		return ray_t;
+	}
+	else
+	{
+		return TNumericLimits<double>::Max();
+	}
+}
+
+
+
+
+int32 FSparseDynamicPointOctree3::FindNearestHitPoint(const FRay3d& Ray,
+	TFunctionRef<double(int, const FRay3d&)> HitPointDistFunc,
+	double MaxDistance) const
+{
+	// this should take advantage of raster!
+
+	// we use queue instead of recursion
+	TArray<const FSparsePointOctreeCell*> Queue;
+	Queue.Reserve(64);
+
+	// push all root cells onto queue if they are hit by ray
+	RootCells.AllocatedIteration([&](const uint32* RootCellID)
+	{
+		const FSparsePointOctreeCell* RootCell = &Cells[*RootCellID];
+		double RayHitParam = FindNearestRayCellIntersection(*RootCell, Ray);
+		if (RayHitParam < MaxDistance)
+		{
+			Queue.Add(&Cells[*RootCellID]);
+		}
+	});
+
+	// test cells until the queue is empty
+	int32 HitPointID = -1;
+	while (Queue.Num() > 0)
+	{
+		const FSparsePointOctreeCell* CurCell = Queue.Pop(false);
+		
+		// process elements
+		for (int PointID : CellPointLists.Values(CurCell->CellID))
+		{
+			double HitDist = HitPointDistFunc(PointID, Ray);
+			if (HitDist < MaxDistance)
+			{
+				MaxDistance = HitDist;
+				HitPointID = PointID;
+			}
+		}
+
+		// descend to child cells
+		// sort by distance? use DDA?
+		for (int k = 0; k < 8; ++k) 
+		{
+			if (CurCell->HasChild(k))
+			{
+				const FSparsePointOctreeCell* ChildCell = &Cells[CurCell->GetChildCellID(k)];
+				double RayHitParam = FindNearestRayCellIntersection(*ChildCell, Ray);
+				if (RayHitParam < MaxDistance)
+				{
+					Queue.Add(ChildCell);
+				}
+			}
+		}
+	}
+
+	return HitPointID;
+}
+
+
+
+
+
+
+
+void FSparseDynamicPointOctree3::RangeQuery(
+	const FAxisAlignedBox3d& Bounds,
+	TFunctionRef<bool(int)> PredicateFunc,
+	TArray<int>& PointIDs,
+	TArray<const FSparsePointOctreeCell*>* TempBuffer) const
+{
+	// todo: this should take advantage of raster!
+
+	TArray<const FSparsePointOctreeCell*> InternalBuffer;
+	TArray<const FSparsePointOctreeCell*>& Queue =
+		(TempBuffer == nullptr) ? InternalBuffer : *TempBuffer;
+	Queue.Reset();
+	Queue.Reserve(128);
+
+	// start at root cells
+	RootCells.AllocatedIteration([&](const uint32* RootCellID)
+	{
+		const FSparsePointOctreeCell* RootCell = &Cells[*RootCellID];
+		if (GetCellBox(*RootCell).Intersects(Bounds))
+		{
+			Queue.Add(&Cells[*RootCellID]);
+		}
+	});
+
+
+	while (Queue.Num() > 0)
+	{
+		const FSparsePointOctreeCell* CurCell = Queue.Pop(false);
+
+		// process elements
+		for (int PointID : CellPointLists.Values(CurCell->CellID))
+		{
+			if (PredicateFunc(PointID))
+			{
+				PointIDs.Add(PointID);
+			}
+		}
+
+		for (int k = 0; k < 8; ++k)
+		{
+			if (CurCell->HasChild(k))
+			{
+				const FSparsePointOctreeCell* ChildCell = &Cells[CurCell->GetChildCellID(k)];
+				if (GetCellBox(*ChildCell).Intersects(Bounds))
+				{
+					Queue.Add(ChildCell);
+				}
+			}
+		}
+	}
+}
+
+
+
+
+
+
+
+
+
+void FSparseDynamicPointOctree3::ParallelRangeQuery(
+	const FAxisAlignedBox3d& Bounds,
+	TFunctionRef<bool(int)> PredicateFunc,
+	TArray<int>& PointIDs,
+	TArray<const FSparsePointOctreeCell*>* TempBuffer) const
+{
+	// todo: this should take advantage of raster!
+
+	TArray<const FSparsePointOctreeCell*> ProcessRootCells;
+	// start at root cells
+	RootCells.AllocatedIteration([&](const uint32* RootCellID)
+	{
+		const FSparsePointOctreeCell* RootCell = &Cells[*RootCellID];
+		if (GetCellBox(*RootCell).Intersects(Bounds))
+		{
+			ProcessRootCells.Add(&Cells[*RootCellID]);
+		}
+	});
+
+	FCriticalSection OutputLock;
+	TArray<int> PerThreadCounts;
+	TArray<int> PerThreadTraversed;
+
+	ParallelFor(ProcessRootCells.Num(), [&](int idx)
+	{
+		const FSparsePointOctreeCell* RootCell = ProcessRootCells[idx];
+		TArray<const FSparsePointOctreeCell*, TInlineAllocator<128>> Queue;
+		Queue.Add(RootCell);
+
+		TArray<int, TInlineAllocator<256>> LocalPointIDs;
+		int Processed = 0;
+
+		while (Queue.Num() > 0)
+		{
+			const FSparsePointOctreeCell* CurCell = Queue.Pop(false);
+
+			// process elements
+			for (int PointID : CellPointLists.Values(CurCell->CellID))
+			{
+				if (PredicateFunc(PointID))
+				{
+					LocalPointIDs.Add(PointID);
+				}
+			}
+
+			for (int k = 0; k < 8; ++k)
+			{
+				if (CurCell->HasChild(k))
+				{
+					Processed++;
+					const FSparsePointOctreeCell* ChildCell = &Cells[CurCell->GetChildCellID(k)];
+					if (GetCellBox(*ChildCell).Intersects(Bounds))
+					{
+						Queue.Add(ChildCell);
+					}
+				}
+			}
+		}
+
+		if (LocalPointIDs.Num() > 0)
+		{
+			OutputLock.Lock();
+			for (int ObjID : LocalPointIDs )
+			{
+				PointIDs.Add(ObjID);
+			}
+			PerThreadCounts.Add(LocalPointIDs.Num());
+			PerThreadTraversed.Add(Processed);
+			OutputLock.Unlock();
+		}
+
+	});
+
+
+	
+}
+
+
+
+
+
+
+FSparsePointOctreeCell FSparseDynamicPointOctree3::FindCurrentContainingCell(const FVector3d& Position) const
+{
+	// look up root cell, which may not exist
+	FVector3i RootIndex = PointToIndex(0, Position);
+	const uint32* RootCellID = RootCells.Get(RootIndex);
+	if (RootCellID == nullptr)
+	{
+		return FSparsePointOctreeCell(0, RootIndex);
+	}
+	check(CellRefCounts.IsValid(*RootCellID));
+
+	// check if point is contained in root cell
+	// (should we do this before checking for existence? we can...)
+	const FSparsePointOctreeCell* RootCell = &Cells[*RootCellID];
+	if (CellContains(*RootCell, Position) == false)
+	{
+		return FSparsePointOctreeCell(FSparsePointOctreeCell::InvalidLevel, FVector3i::Zero());
+	}
+
+	const FSparsePointOctreeCell* CurrentCell = RootCell;
+	do
+	{
+		int ChildIdx = ToChildCellIndex(*CurrentCell, Position);
+		if (CurrentCell->HasChild(ChildIdx))
+		{
+			int32 ChildCellID = CurrentCell->GetChildCellID(ChildIdx);
+			check(CellRefCounts.IsValid(ChildCellID));
+			const FSparsePointOctreeCell* ChildCell = &Cells[ChildCellID];
+			if (CellContains(*ChildCell, Position))
+			{
+				CurrentCell = ChildCell;
+				continue;
+			}
+		}
+
+		return *CurrentCell;
+
+	} while (true);		// loop will always terminate
+
+	return FSparsePointOctreeCell(FSparsePointOctreeCell::InvalidLevel, FVector3i::Zero());
+}
+
+
+
+void FSparseDynamicPointOctree3::CheckValidity(
+	TFunctionRef<bool(int)> IsValidPointIDFunc,
+	TFunctionRef<FVector3d(int)> GetPointFunc,
+	EValidityCheckFailMode FailMode,
+	bool bVerbose,
+	bool bFailOnMissingPoints) const
+{
+	bool is_ok = true;
+	TFunction<void(bool)> CheckOrFailF = [&](bool b)
+	{
+		is_ok = is_ok && b;
+	};
+	if (FailMode == EValidityCheckFailMode::Check)
+	{
+		CheckOrFailF = [&](bool b)
+		{
+			checkf(b, TEXT("FSparseDynamicPointOctree3::CheckValidity failed!"));
+			is_ok = is_ok && b;
+		};
+	}
+	else if (FailMode == EValidityCheckFailMode::Ensure)
+	{
+		CheckOrFailF = [&](bool b)
+		{
+			ensureMsgf(b, TEXT("FSparseDynamicPointOctree3::CheckValidity failed!"));
+			is_ok = is_ok && b;
+		};
+	}
+
+	TArray<int> CellsAtLevels, PointsAtLevel;
+	CellsAtLevels.Init(0, 32);
+	PointsAtLevel.Init(0, 32);
+	uint32 MissingPointCount = 0;
+	uint8 MaxLevel = 0;
+
+	// check that all Point IDs in per-cell Point lists is valid
+	for (int32 CellID : CellRefCounts.Indices())
+	{
+		for (int32 PointID : CellPointLists.Values(CellID))
+		{
+			CheckOrFailF(IsValidPointIDFunc(PointID));
+		}
+	}
+
+	uint32 NumPointIDs = PointIDToCellMap.Num();
+	for (uint32 PointID = 0; PointID < NumPointIDs; PointID++)
+	{
+		bool IsValidPointID = IsValidPointIDFunc(PointID);
+		if (IsValidPointID)
+		{
+			FVector3d Position = GetPointFunc(PointID);
+
+			CheckOrFailF(PointID < PointIDToCellMap.Num());
+			uint32 CellID = PointIDToCellMap[PointID];
+
+			if (bFailOnMissingPoints)
+			{
+				CheckOrFailF(CellID != InvalidCellID);
+			}
+
+			if (CellID == InvalidCellID)
+			{
+				MissingPointCount++;
+			}
+			else
+			{
+				CheckOrFailF(CellRefCounts.IsValid(CellID));
+				FSparsePointOctreeCell Cell = Cells[CellID];
+				FAxisAlignedBox3d CellBounds = GetCellBox(Cell);
+				CheckOrFailF(CellBounds.Contains(Position));
+				CheckOrFailF(CellPointLists.Contains(CellID, PointID));
+
+				PointsAtLevel[Cell.Level]++;
+			}
+		}
+	}
+
+
+	for (int32 CellID : CellRefCounts.Indices())
+	{
+		const FSparsePointOctreeCell& Cell = Cells[CellID];
+		CellsAtLevels[Cell.Level]++;
+		MaxLevel = FMath::Max(MaxLevel, Cell.Level);
+	}
+
+	if (bVerbose)
+	{
+		UE_LOG(LogTemp, Warning, TEXT("FSparseDynamicPointOctree3::CheckValidity: MaxLevel %d  MissingCount %d"), MaxLevel, MissingPointCount);
+		for (uint32 k = 0; k <= MaxLevel; ++k)
+		{
+			UE_LOG(LogTemp, Warning, TEXT("    Level %4d  Cells %4d  Points %4d"), k, CellsAtLevels[k], PointsAtLevel[k]);
+		}
+	}
+}
+
+
+
+
+void FSparseDynamicPointOctree3::ComputeStatistics(FStatistics& StatsOut) const
+{
+	StatsOut.Levels = 0;
+	for (int32 CellID : CellRefCounts.Indices())
+	{
+		const FSparsePointOctreeCell& Cell = Cells[CellID];
+		StatsOut.Levels = FMath::Max(StatsOut.Levels, (int32)Cell.Level);
+	}
+	StatsOut.Levels++;
+	StatsOut.LevelObjCounts.Init(0, StatsOut.Levels);
+	StatsOut.LevelMaxObjCounts.Init(0, StatsOut.Levels);
+	StatsOut.LevelBoxCounts.Init(0, StatsOut.Levels);
+	StatsOut.LevelCellSizes.Init(0, StatsOut.Levels);
+	for (int32 CellID : CellRefCounts.Indices())
+	{
+		const FSparsePointOctreeCell& Cell = Cells[CellID];
+		StatsOut.LevelBoxCounts[Cell.Level]++;
+		int ObjCount = CellPointLists.GetCount(CellID);
+		StatsOut.LevelObjCounts[Cell.Level] += ObjCount;
+		StatsOut.LevelMaxObjCounts[Cell.Level] = FMath::Max(StatsOut.LevelMaxObjCounts[Cell.Level], ObjCount);
+		StatsOut.LevelCellSizes[Cell.Level] = GetCellWidth(Cell.Level);
+	}
+}
+
+
+FString FSparseDynamicPointOctree3::FStatistics::ToString() const
+{
+	FString Result = FString::Printf(
+		TEXT("Levels %2d\r\n"), Levels);
+	for (int k = 0; k < Levels; ++k)
+	{
+		Result += FString::Printf(TEXT("  Level %2d:  Cells %8d  Objs %8d  AvgObjs %5.3f  MaxObjs %8d  Dimension %5.5f\r\n"), 
+			k, LevelBoxCounts[k], LevelObjCounts[k], ((float)LevelObjCounts[k] / (float)LevelBoxCounts[k]), LevelMaxObjCounts[k], LevelCellSizes[k] );
+	}
+	return Result;
+}
\ No newline at end of file
