GeometryProcessing: reorganize FDynamicMesh3 constructors so that accidentally calling copy constructor w/ a pointer results in a compile error instead of implicit cast to bool
#rnx

diff --git a/Engine/Plugins/Experimental/GeometryProcessing/Source/DynamicMesh/Public/DynamicMesh3.h b/Engine/Plugins/Experimental/GeometryProcessing/Source/DynamicMesh/Public/DynamicMesh3.h
--- a/Engine/Plugins/Experimental/GeometryProcessing/Source/DynamicMesh/Public/DynamicMesh3.h
+++ b/Engine/Plugins/Experimental/GeometryProcessing/Source/DynamicMesh/Public/DynamicMesh3.h
@@ -209,9 +209,14 @@ protected:
 public:
 	virtual ~FDynamicMesh3();
 
-	explicit FDynamicMesh3(bool bWantNormals = true, bool bWantColors = false, bool bWantUVs = false, bool bWantTriGroups = false);
+	FDynamicMesh3() 
+		: FDynamicMesh3(true, false, false, false)
+	{
+	}
+
+	explicit FDynamicMesh3(bool bWantNormals, bool bWantColors, bool bWantUVs, bool bWantTriGroups);
 
-	FDynamicMesh3(EMeshComponents flags)
+	explicit FDynamicMesh3(EMeshComponents flags)
 		: FDynamicMesh3(((int)flags & (int)EMeshComponents::VertexNormals) != 0, ((int)flags & (int)EMeshComponents::VertexColors) != 0,
 						((int)flags & (int)EMeshComponents::VertexUVs) != 0, ((int)flags & (int)EMeshComponents::FaceGroups) != 0)
 	{
