add a MeshDescription adapter class to help people directly use FMeshDescription meshes w/ GeometryProcessing classes that template the triangle mesh type.

diff --git a/Engine/Plugins/Experimental/GeometryProcessing/Source/MeshConversion/Public/MeshDescriptionAdapter.h b/Engine/Plugins/Experimental/GeometryProcessing/Source/MeshConversion/Public/MeshDescriptionAdapter.h
new file mode 100644
--- /dev/null
+++ b/Engine/Plugins/Experimental/GeometryProcessing/Source/MeshConversion/Public/MeshDescriptionAdapter.h
@@ -0,0 +1,81 @@
+// Copyright 1998-2019 Epic Games, Inc. All Rights Reserved.
+
+#pragma once
+
+#include "CoreMinimal.h"
+#include "MeshDescription.h"
+#include "StaticMeshAttributes.h"
+#include "Spatial/MeshAABBTree3.h"
+
+
+/**
+ * Basic struct to adapt a FMeshDescription for use by GeometryProcessing classes that template the mesh type and expect a standard set of basic accessors
+ * For example, this adapter will let you use a FMeshDescription with GeometryProcessing's TMeshAABBTree3
+ *
+ *  Usage example -- given some const FMeshDescription* Mesh:
+ *    FMeshDescriptionAABBAdapter MeshAdapter(Mesh); // adapt the mesh
+ *    TMeshAABBTree3<const FMeshDescriptionTriangleMeshAdapter> AABBTree(&MeshAdapter); // provide the adapter to a templated class like TMeshAABBTree3
+ */
+struct MESHCONVERSION_API FMeshDescriptionTriangleMeshAdapter
+{
+protected:
+	const FMeshDescription* Mesh;
+	TVertexAttributesConstRef<FVector> VertexPositions;
+
+public:
+	FMeshDescriptionTriangleMeshAdapter(const FMeshDescription* MeshIn) : Mesh(MeshIn)
+	{
+		VertexPositions = MeshIn->VertexAttributes().GetAttributesRef<FVector>(MeshAttribute::Vertex::Position);
+	}
+
+	bool IsTriangle(int32 TID) const
+	{
+		return TID >= 0 && TID < Mesh->Triangles().Num();
+	}
+	bool IsVertex(int32 VID) const
+	{
+		return VID >= 0 && VID < Mesh->Vertices().Num();
+	}
+	// ID and Count are the same for MeshDescription because it's compact
+	int32 MaxTriangleID() const
+	{
+		return Mesh->Triangles().Num();
+	}
+	int32 TriangleCount() const
+	{
+		return Mesh->Triangles().Num();
+	}
+	int32 MaxVertexID() const
+	{
+		return Mesh->Vertices().Num();
+	}
+	int32 VertexCount() const
+	{
+		return Mesh->Vertices().Num();
+	}
+	int32 GetShapeTimestamp() const
+	{
+		// MeshDescription doesn't provide any mechanism to know if it's been modified so just return 0
+		// and leave it to the caller to not build an aabb and then change the underlying mesh
+		return 0;
+	}
+	FIndex3i GetTriangle(int32 IDValue) const
+	{
+		const TStaticArray<FVertexID, 3> TriVertIDs = Mesh->GetTriangleVertices(FTriangleID(IDValue));
+		return FIndex3i(TriVertIDs[0].GetValue(), TriVertIDs[1].GetValue(), TriVertIDs[2].GetValue());
+	}
+	FVector3d GetVertex(int32 IDValue) const
+	{
+		return FVector3d(VertexPositions[FVertexID(IDValue)]);
+	}
+
+	inline void GetTriVertices(int32 IDValue, FVector3d& V0, FVector3d& V1, FVector3d& V2) const
+	{
+		const TStaticArray<FVertexID, 3> TriVertIDs = Mesh->GetTriangleVertices(FTriangleID(IDValue));
+		V0 = FVector3d(VertexPositions[TriVertIDs[0]]);
+		V1 = FVector3d(VertexPositions[TriVertIDs[1]]);
+		V2 = FVector3d(VertexPositions[TriVertIDs[2]]);
+	}
+};
+
+
