formalize max supported depth in sparse dynamic octree

#rnx

diff --git a/Engine/Plugins/Experimental/GeometryProcessing/Source/GeometricObjects/Public/Spatial/SparseDynamicOctree3.h b/Engine/Plugins/Experimental/GeometryProcessing/Source/GeometricObjects/Public/Spatial/SparseDynamicOctree3.h
--- a/Engine/Plugins/Experimental/GeometryProcessing/Source/GeometricObjects/Public/Spatial/SparseDynamicOctree3.h
+++ b/Engine/Plugins/Experimental/GeometryProcessing/Source/GeometricObjects/Public/Spatial/SparseDynamicOctree3.h
@@ -165,6 +165,25 @@ public:
 	 */
 	double MaxExpandFactor = 0.25;
 
+	static constexpr uint32 MaxSupportedTreeDepth = 0x1F;
+
+	int GetMaxTreeDepth()
+	{
+		return MaxTreeDepth;
+	}
+	/**
+	 * Sets max tree depth w/ protection against setting beyond supported max
+	 */
+	void SetMaxTreeDepth(int MaxTreeDepthIn)
+	{
+		if (!ensure((uint32)MaxTreeDepthIn <= MaxSupportedTreeDepth))
+		{
+			MaxTreeDepthIn = (int32)MaxSupportedTreeDepth;
+		}
+
+		MaxTreeDepth = MaxTreeDepthIn;
+	}
+protected:
 	/**
 	 * Objects will not be inserted more than this many levels deep from a Root cell
 	 */
@@ -292,7 +311,8 @@ protected:
 	// calculate the base width of a cell at a given level
 	inline double GetCellWidth(uint32 Level) const
 	{
-		double Divisor = (double)( (uint64)1 << (Level & 0x1F) );
+		checkSlow(Level < MaxSupportedTreeDepth);
+		double Divisor = (double)( (uint64)1 << (Level & MaxSupportedTreeDepth) );
 		double CellWidth = RootDimension / Divisor;
 		return CellWidth;
 	}
