octree level is allowed to be exactly maxsupportedtreedepth

#rnx

diff --git a/Engine/Plugins/Experimental/GeometryProcessing/Source/GeometricObjects/Public/Spatial/SparseDynamicOctree3.h b/Engine/Plugins/Experimental/GeometryProcessing/Source/GeometricObjects/Public/Spatial/SparseDynamicOctree3.h
--- a/Engine/Plugins/Experimental/GeometryProcessing/Source/GeometricObjects/Public/Spatial/SparseDynamicOctree3.h
+++ b/Engine/Plugins/Experimental/GeometryProcessing/Source/GeometricObjects/Public/Spatial/SparseDynamicOctree3.h
@@ -311,7 +311,7 @@ protected:
 	// calculate the base width of a cell at a given level
 	inline double GetCellWidth(uint32 Level) const
 	{
-		checkSlow(Level < MaxSupportedTreeDepth);
+		checkSlow(Level <= MaxSupportedTreeDepth);
 		double Divisor = (double)( (uint64)1 << (Level & MaxSupportedTreeDepth) );
 		double CellWidth = RootDimension / Divisor;
 		return CellWidth;
