add an index mesh-specific adapter + test

#rnx

diff --git a/Engine/Plugins/Experimental/GeometryProcessing/Source/GeometricObjects/Public/MeshAdapter.h b/Engine/Plugins/Experimental/GeometryProcessing/Source/GeometricObjects/Public/MeshAdapter.h
--- a/Engine/Plugins/Experimental/GeometryProcessing/Source/GeometricObjects/Public/MeshAdapter.h
+++ b/Engine/Plugins/Experimental/GeometryProcessing/Source/GeometricObjects/Public/MeshAdapter.h
@@ -8,6 +8,9 @@
 #include "Math/IntVector.h"
 #include "Math/Vector.h"
 
+/**
+ * Most generic / lazy example of a triangle mesh adapter; possibly useful for prototyping / building on top of (but slower than making a more specific-case adapter)
+ */
 template <typename RealType>
 struct TTriangleMeshAdapter
 {
@@ -33,6 +36,11 @@ struct TTriangleMeshAdapter
 typedef TTriangleMeshAdapter<double> FTriangleMeshAdapterd;
 typedef TTriangleMeshAdapter<float> FTriangleMeshAdapterf;
 
+/**
+ * Example function to generate a generic mesh adapter from arrays
+ * @param Vertices Array of mesh vertices
+ * @param Triangles Array of int-vectors, one per triangle, indexing into the vertices array
+ */
 inline FTriangleMeshAdapterd GetArrayMesh(TArray<FVector>& Vertices, TArray<FIntVector>& Triangles)
 {
 	return {
@@ -46,3 +54,76 @@ inline FTriangleMeshAdapterd GetArrayMesh(TArray<FVector>& Vertices, TArray<FInt
 		[&](int Idx) { return FIndex3i(Triangles[Idx]); },
 		[&](int Idx) { return FVector3d(Vertices[Idx]); }};
 }
+
+
+/**
+ * Faster adapter specifically for the common index mesh case
+ */
+template<typename IndexType, typename OutRealType>
+struct TIndexMeshArrayAdapter
+{
+	const TArray<FVector>* SourceVertices;
+	const TArray<IndexType>* SourceTriangles;
+
+	TIndexMeshArrayAdapter(const TArray<FVector>* SourceVerticesIn, const TArray<uint32>* SourceTrianglesIn) : SourceVertices(SourceVerticesIn), SourceTriangles(SourceTrianglesIn)
+	{
+	}
+
+	FORCEINLINE bool IsTriangle(int32 Index) const
+	{
+		return SourceTriangles->IsValidIndex(Index * 3);
+	}
+
+	FORCEINLINE bool IsVertex(int32 Index) const
+	{
+		return SourceVertices->IsValidIndex(Index);
+	}
+
+	FORCEINLINE int32 MaxTriangleID() const
+	{
+		return SourceTriangles->Num() / 3;
+	}
+
+	FORCEINLINE int32 MaxVertexID() const
+	{
+		return SourceVertices->Num();
+	}
+
+	// Counts are same as MaxIDs, because these are compact meshes
+	FORCEINLINE int32 TriangleCount() const
+	{
+		return SourceTriangles->Num() / 3;
+	}
+
+	FORCEINLINE int32 VertexCount() const
+	{
+		return SourceVertices->Num();
+	}
+
+	FORCEINLINE int32 GetShapeTimestamp() const
+	{
+		return 0; // source data doesn't have a timestamp concept
+	}
+
+	FORCEINLINE FIndex3i GetTriangle(int32 Index) const
+	{
+		int32 Start = Index * 3;
+		return FIndex3i((int)(*SourceTriangles)[Start], (int)(*SourceTriangles)[Start+1], (int)(*SourceTriangles)[Start+2]);
+	}
+
+	FORCEINLINE FVector3<OutRealType> GetVertex(int32 Index) const
+	{
+		return FVector3<OutRealType>((*SourceVertices)[Index]);
+	}
+
+	FORCEINLINE void GetTriVertices(int32 TriIndex, FVector3<OutRealType>& V0, FVector3<OutRealType>& V1, FVector3<OutRealType>& V2) const
+	{
+		int32 Start = TriIndex * 3;
+		V0 = FVector3<OutRealType>((*SourceVertices)[(*SourceTriangles)[Start]]);
+		V1 = FVector3<OutRealType>((*SourceVertices)[(*SourceTriangles)[Start+1]]);
+		V2 = FVector3<OutRealType>((*SourceVertices)[(*SourceTriangles)[Start+2]]);
+	}
+
+};
+
+typedef TIndexMeshArrayAdapter<uint32, double> FIndexMeshArrayAdapterd;
