minor geometryprocessing aabbtree api changes to make the class easier to use as an alternative line trace

#rnx

diff --git a/Engine/Plugins/Experimental/GeometryProcessing/Source/GeometricObjects/Public/MeshAdapter.h b/Engine/Plugins/Experimental/GeometryProcessing/Source/GeometricObjects/Public/MeshAdapter.h
--- a/Engine/Plugins/Experimental/GeometryProcessing/Source/GeometricObjects/Public/MeshAdapter.h
+++ b/Engine/Plugins/Experimental/GeometryProcessing/Source/GeometricObjects/Public/MeshAdapter.h
@@ -65,6 +65,16 @@ struct TIndexMeshArrayAdapter
 	const TArray<FVector>* SourceVertices;
 	const TArray<IndexType>* SourceTriangles;
 
+	void SetSources(const TArray<FVector>* SourceVerticesIn, const TArray<uint32>* SourceTrianglesIn)
+	{
+		SourceVertices = SourceVerticesIn;
+		SourceTriangles = SourceTrianglesIn;
+	}
+
+	TIndexMeshArrayAdapter() : SourceVertices(nullptr), SourceTriangles(nullptr)
+	{
+	}
+
 	TIndexMeshArrayAdapter(const TArray<FVector>* SourceVerticesIn, const TArray<uint32>* SourceTrianglesIn) : SourceVertices(SourceVerticesIn), SourceTriangles(SourceTrianglesIn)
 	{
 	}
diff --git a/Engine/Plugins/Experimental/GeometryProcessing/Source/GeometricObjects/Public/Spatial/MeshAABBTree3.h b/Engine/Plugins/Experimental/GeometryProcessing/Source/GeometricObjects/Public/Spatial/MeshAABBTree3.h
--- a/Engine/Plugins/Experimental/GeometryProcessing/Source/GeometricObjects/Public/Spatial/MeshAABBTree3.h
+++ b/Engine/Plugins/Experimental/GeometryProcessing/Source/GeometricObjects/Public/Spatial/MeshAABBTree3.h
@@ -98,6 +98,8 @@ public:
 		find_nearest_tri(RootIndex, P, NearestDistSqr, tNearID);
 		return tNearID;
 	}
+
+protected:
 	void find_nearest_tri(int IBox, const FVector3d& P, double& NearestDistSqr, int& TID)
 	{
 		int idx = BoxToIndex[IBox];
@@ -164,28 +166,38 @@ public:
 		}
 	}
 
+public:
 	virtual bool SupportsTriangleRayIntersection() override
 	{
 		return true;
 	}
 
-	virtual int FindNearestHitTriangle(const FRay3d& Ray, double MaxDist = TNumericLimits<double>::Max()) override
+	inline virtual int FindNearestHitTriangle(const FRay3d& Ray, double MaxDistance = TNumericLimits<double>::Max()) override
+	{
+		double NearestT;
+		int TID;
+		FindNearestHitTriangle(Ray, NearestT, TID, MaxDistance);
+		return TID;
+	}
+
+	virtual bool FindNearestHitTriangle(const FRay3d& Ray, double& NearestT, int& TID, double MaxDist = TNumericLimits<double>::Max()) override
 	{
+		TID = IndexConstants::InvalidID;
+		NearestT = (MaxDist < TNumericLimits<double>::Max()) ? MaxDist : TNumericLimits<float>::Max();
+
 		check(MeshTimestamp == Mesh->GetShapeTimestamp());
 		check(RootIndex >= 0);
 		if (RootIndex < 0)
 		{
-			return IndexConstants::InvalidID;
+			return false;
 		}
 		// TODO: check( ray_is_normalized)
 
 		// [RMS] note: using float.MaxValue here because we need to use <= to compare Box hit
 		//   to NearestT, and Box hit returns double.MaxValue on no-hit. So, if we set
 		//   nearestT to double.MaxValue, then we will test all boxes (!)
-		double NearestT = (MaxDist < TNumericLimits<double>::Max()) ? MaxDist : TNumericLimits<float>::Max();
-		int tNearID = IndexConstants::InvalidID;
-		FindHitTriangle(RootIndex, Ray, NearestT, tNearID);
-		return tNearID;
+		FindHitTriangle(RootIndex, Ray, NearestT, TID);
+		return TID != IndexConstants::InvalidID;
 	}
 
 	void FindHitTriangle(int IBox, const FRay3d& Ray, double& NearestT, int& TID)
diff --git a/Engine/Plugins/Experimental/GeometryProcessing/Source/GeometricObjects/Public/Spatial/SpatialInterfaces.h b/Engine/Plugins/Experimental/GeometryProcessing/Source/GeometricObjects/Public/Spatial/SpatialInterfaces.h
--- a/Engine/Plugins/Experimental/GeometryProcessing/Source/GeometricObjects/Public/Spatial/SpatialInterfaces.h
+++ b/Engine/Plugins/Experimental/GeometryProcessing/Source/GeometricObjects/Public/Spatial/SpatialInterfaces.h
@@ -61,7 +61,25 @@ public:
 	 * @param MaxDistance maximum hit distance
 	 * @return ID of triangle intersected by ray within MaxDistance, or InvalidID if not found
 	 */
-	virtual int FindNearestHitTriangle(const FRay3d& Ray, double MaxDistance = TNumericLimits<double>::Max()) = 0;
+	inline virtual int FindNearestHitTriangle(const FRay3d& Ray, double MaxDistance = TNumericLimits<double>::Max())
+	{
+		double NearestT;
+		int TID;
+		FindNearestHitTriangle(Ray, NearestT, TID, MaxDistance);
+		return TID;
+	}
+
+	/**
+	 * Find nearest triangle from the given ray
+	 * @param Ray query ray
+	 * @param NearestT returned-by-reference parameter of the nearest hit
+	 * @param TID returned-by-reference ID of triangle intersected by ray within MaxDistance, or InvalidID if not found
+	 * @param MaxDistance maximum hit distance
+	 * @return true if hit, false if no hit found
+	 */
+	virtual bool FindNearestHitTriangle(const FRay3d& Ray, double& NearestT, int& TID, double MaxDistance = TNumericLimits<double>::Max()) = 0;
+
+
 };
 
 
