change FNoncopyable to UE_NONCOPYABLE for attribute change set

#rnx

diff --git a/Engine/Plugins/Experimental/GeometryProcessing/Source/DynamicMesh/Public/DynamicMeshChangeTracker.h b/Engine/Plugins/Experimental/GeometryProcessing/Source/DynamicMesh/Public/DynamicMeshChangeTracker.h
--- a/Engine/Plugins/Experimental/GeometryProcessing/Source/DynamicMesh/Public/DynamicMeshChangeTracker.h
+++ b/Engine/Plugins/Experimental/GeometryProcessing/Source/DynamicMesh/Public/DynamicMeshChangeTracker.h
@@ -89,9 +89,12 @@ typedef TDynamicMeshAttributeChange<float,3> FDynamicMeshNormalChange;
 /**
  * FDynamicMeshAttributeChangeSet stores a set of UV and Normal changes for a FDynamicMesh3
  */
-class DYNAMICMESH_API FDynamicMeshAttributeChangeSet : FNoncopyable
+class DYNAMICMESH_API FDynamicMeshAttributeChangeSet
 {
 public:
+	UE_NONCOPYABLE(FDynamicMeshAttributeChangeSet);
+	FDynamicMeshAttributeChangeSet() = default;
+
 	TArray<FDynamicMeshUVChange> UVChanges;
 	TArray<FDynamicMeshNormalChange> NormalChanges;
 	TOptional<FDynamicMeshTriangleAttributeChange<int32, 1>> MaterialIDAttribChange;
