add removeoccludedtriangles as a dmesh operation + test

#rnx

diff --git a/Engine/Plugins/Experimental/GeometryProcessing/Source/DynamicMesh/Private/Operations/RemoveOccludedTriangles.cpp b/Engine/Plugins/Experimental/GeometryProcessing/Source/DynamicMesh/Private/Operations/RemoveOccludedTriangles.cpp
new file mode 100644
--- /dev/null
+++ b/Engine/Plugins/Experimental/GeometryProcessing/Source/DynamicMesh/Private/Operations/RemoveOccludedTriangles.cpp
@@ -0,0 +1,144 @@
+// Copyright 1998-2019 Epic Games, Inc. All Rights Reserved.
+
+// Port of geometry3Sharp RemoveOccludedTriangles
+
+#include "Operations/RemoveOccludedTriangles.h"
+
+#include "DynamicMeshEditor.h"
+#include "MeshNormals.h"
+
+#include "Async/ParallelFor.h"
+#include "Misc/ScopeLock.h"
+
+
+
+bool FRemoveOccludedTriangles::Apply(FDynamicMeshAABBTree3* Spatial)
+{
+	FDynamicMeshAABBTree3 SpatialIfNotPassed;
+	if (!Spatial)
+	{
+		Spatial = &SpatialIfNotPassed;
+		Spatial->SetMesh(Mesh, true);
+	}
+	TFastWindingTree<FDynamicMesh3> FastWinding(Spatial, InsideMode == ECalculationMode::FastWindingNumber);
+
+	if (Cancelled())
+	{
+		return false;
+	}
+
+	// ray directions
+	TArray<FVector3d> RayDirs; int NR = 0;
+	if (InsideMode == ECalculationMode::SimpleOcclusionTest)
+	{
+		RayDirs.Add(FVector3d::UnitX()); RayDirs.Add(-FVector3d::UnitX());
+		RayDirs.Add(FVector3d::UnitY()); RayDirs.Add(-FVector3d::UnitY());
+		RayDirs.Add(FVector3d::UnitZ()); RayDirs.Add(-FVector3d::UnitZ());
+		NR = RayDirs.Num();
+	}
+
+	auto IsOccludedFWN = [this, &FastWinding](const FVector3d& Pt)
+	{
+		return FastWinding.FastWindingNumber(Pt) > WindingIsoValue;
+	};
+	auto IsOccludedSimple = [this, &Spatial, &RayDirs, &NR](const FVector3d& Pt)
+	{
+		FRay3d Ray;
+		Ray.Origin = Pt;
+		for (int RayIdx = 0; RayIdx < NR; ++RayIdx)
+		{
+			Ray.Direction = RayDirs[RayIdx];
+			int HitTID = Spatial->FindNearestHitTriangle(Ray);
+			if (HitTID == IndexConstants::InvalidID)
+			{
+				return false;
+			}
+		}
+		return true;
+	};
+
+	TFunctionRef<bool(const FVector3d& Pt)> IsOccludedF =
+			(InsideMode == ECalculationMode::FastWindingNumber) ? 
+		(TFunctionRef<bool(const FVector3d& Pt)>)IsOccludedFWN : 
+		(TFunctionRef<bool(const FVector3d& Pt)>)IsOccludedSimple;
+
+	bool bForceSingleThread = false;
+
+	TArray<bool> VertexOccluded;
+	if ( TriangleSamplingMethod == ETriangleSampling::PerVertex )
+	{
+		VertexOccluded.SetNum(Mesh->VertexCount());
+
+		FMeshNormals Normals(Mesh);
+		if (Mesh->HasVertexNormals() == false)
+		{
+			Normals.ComputeVertexNormals();
+		}
+
+		ParallelFor(Mesh->MaxVertexID(), [this, &Normals, &VertexOccluded, &IsOccludedF](int32 VID)
+		{
+			if (!Mesh->IsVertex(VID))
+			{
+				return;
+			}
+			FVector3d SamplePos = Mesh->GetVertex(VID);
+			FVector3d Normal = Mesh->HasVertexNormals() ? (FVector3d)Mesh->GetVertexNormal(VID) : Normals[VID];
+			SamplePos += Normal * NormalOffset;
+			VertexOccluded[VID] = IsOccludedF(SamplePos);
+		}, bForceSingleThread);
+	}
+	if (Cancelled())
+	{
+		return false;
+	}
+
+	RemovedT.Empty();
+	FCriticalSection RemoveTMutex;
+	ParallelFor(Mesh->MaxTriangleID(), [this, &VertexOccluded, &IsOccludedF, &RemoveTMutex](int32 TID)
+	{
+		if (!Mesh->IsTriangle(TID))
+		{
+			return;
+		}
+
+		bool bInside = false;
+		if (TriangleSamplingMethod == ETriangleSampling::PerVertex)
+		{
+			FIndex3i Tri = Mesh->GetTriangle(TID);
+			bInside = VertexOccluded[Tri.A] || VertexOccluded[Tri.B] || VertexOccluded[Tri.C];
+		}
+		else
+		{
+			FVector3d Centroid, Normal;
+			double Area;
+			Mesh->GetTriInfo(TID, Normal, Area, Centroid);
+			FVector3d SamplePos = Centroid + Normal * NormalOffset;
+			bInside = IsOccludedF(SamplePos);
+		}
+		if (bInside)
+		{
+			FScopeLock RemoveLock(&RemoveTMutex);
+			RemovedT.Add(TID);
+		}
+	}, bForceSingleThread);
+
+
+	if (Cancelled())
+	{
+		return false;
+	}
+
+	if (RemovedT.Num() > 0)
+	{
+		FDynamicMeshEditor Editor(Mesh);
+		bool bOK = Editor.RemoveTriangles(RemovedT, true);
+		if (!bOK)
+		{
+			bRemoveFailed = true;
+			return false;
+		}
+		// TODO: do we want to consider if we have made the mesh non-manifold or do any cleanup?
+	} 
+
+	return true;
+}
\ No newline at end of file
diff --git a/Engine/Plugins/Experimental/GeometryProcessing/Source/DynamicMesh/Public/Operations/RemoveOccludedTriangles.h b/Engine/Plugins/Experimental/GeometryProcessing/Source/DynamicMesh/Public/Operations/RemoveOccludedTriangles.h
new file mode 100644
--- /dev/null
+++ b/Engine/Plugins/Experimental/GeometryProcessing/Source/DynamicMesh/Public/Operations/RemoveOccludedTriangles.h
@@ -0,0 +1,101 @@
+// Copyright 1998-2019 Epic Games, Inc. All Rights Reserved.
+
+// Port of geometry3Sharp RemoveOccludedTriangles
+
+#pragma once
+
+#include "MathUtil.h"
+#include "VectorTypes.h"
+#include "DynamicMesh3.h"
+#include "DynamicMeshAABBTree3.h"
+#include "Spatial/FastWinding.h"
+
+#include "Util/ProgressCancel.h"
+
+/**
+ * Remove "occluded" triangles, i.e. triangles on the "inside" of the mesh. 
+ * This is a fuzzy definition, current implementation has a couple of options
+ * including a winding number-based version and an ambient-occlusion-ish version,
+ * where if face is occluded for all test rays, then we classify it as inside and remove it.
+ */
+class DYNAMICMESH_API FRemoveOccludedTriangles
+{
+public:
+
+	FDynamicMesh3* Mesh;
+
+	FRemoveOccludedTriangles(FDynamicMesh3* Mesh) : Mesh(Mesh)
+	{
+	}
+	virtual ~FRemoveOccludedTriangles() {}
+
+	/**
+	 * Remove the occluded triangles
+	 *
+	 * @param Spatial optional precomputed AABBTree for the mesh; if null it will be computed as needed
+	 * @return true on success
+	 */
+	virtual bool Apply(FDynamicMeshAABBTree3* Spatial = nullptr);
+
+	/**
+	 * @return EOperationValidationResult::Ok if we can apply operation, or error code if we cannot
+	 */
+	virtual EOperationValidationResult Validate()
+	{
+		// TODO: validate input
+		return EOperationValidationResult::Ok;
+	}
+
+	//
+	// Input settings
+	//
+
+	// how/where to sample triangles when testing for occlusion
+	enum class ETriangleSampling
+	{
+		PerVertex,
+		PerCentroid
+	};
+	ETriangleSampling TriangleSamplingMethod = ETriangleSampling::PerCentroid;
+
+	// we nudge points out by this amount to try to counteract numerical issues
+	double NormalOffset = FMathd::ZeroTolerance;
+
+	// use this as winding isovalue for WindingNumber mode
+	double WindingIsoValue = 0.5;
+
+	enum class ECalculationMode
+	{
+		FastWindingNumber,
+		SimpleOcclusionTest
+	};
+	ECalculationMode InsideMode = ECalculationMode::FastWindingNumber;
+
+	/**
+	 * Set this to be able to cancel running operation
+	 */
+	FProgressCancel* Progress = nullptr;
+
+
+	//
+	// Outputs
+	//
+
+	// indices of removed triangles. will be empty if nothing removed
+	TArray<int> RemovedT;
+
+	// true if it wanted to remove triangles but the actual remove operation failed
+	bool bRemoveFailed = false;
+
+
+protected:
+	/**
+	 * if this returns true, abort computation. 
+	 */
+	virtual bool Cancelled()
+	{
+		return (Progress == nullptr) ? false : Progress->Cancelled();
+	}
+
+};
+
diff --git a/Engine/Plugins/Experimental/GeometryProcessing/Source/GeometricObjects/Public/Spatial/FastWinding.h b/Engine/Plugins/Experimental/GeometryProcessing/Source/GeometricObjects/Public/Spatial/FastWinding.h
--- a/Engine/Plugins/Experimental/GeometryProcessing/Source/GeometricObjects/Public/Spatial/FastWinding.h
+++ b/Engine/Plugins/Experimental/GeometryProcessing/Source/GeometricObjects/Public/Spatial/FastWinding.h
@@ -86,7 +86,7 @@ namespace FastTriWinding
 	/**
 	 *  Evaluate first-order FWN approximation at point Q, relative to Center c
 	 */
-	double EvaluateOrder1Approx(const FVector3d& Center, const FVector3d& Order1Coeff, const FVector3d& Q)
+	inline double EvaluateOrder1Approx(const FVector3d& Center, const FVector3d& Order1Coeff, const FVector3d& Q)
 	{
 		FVector3d dpq = (Center - Q);
 		double len = dpq.Length();
@@ -97,7 +97,7 @@ namespace FastTriWinding
 	/**
 	 *  Evaluate second-order FWN approximation at point Q, relative to Center c
 	 */
-	double EvaluateOrder2Approx(const FVector3d& Center, const FVector3d& Order1Coeff, const FMatrix3d& Order2Coeff, const FVector3d& Q)
+	inline double EvaluateOrder2Approx(const FVector3d& Center, const FVector3d& Order1Coeff, const FMatrix3d& Order2Coeff, const FVector3d& Q)
 	{
 		FVector3d dpq = (Center - Q);
 		double len = dpq.Length();
@@ -125,7 +125,7 @@ namespace FastTriWinding
 	// triangle-winding-number first-order approximation.
 	// T is triangle, P is 'Center' of cluster of dipoles, Q is evaluation point
 	// (This is really just for testing)
-	double Order1Approx(const FTriangle3d& T, const FVector3d& P, const FVector3d& XN, double XA, const FVector3d& Q)
+	inline double Order1Approx(const FTriangle3d& T, const FVector3d& P, const FVector3d& XN, double XA, const FVector3d& Q)
 	{
 		FVector3d at0 = XA * XN;
 
@@ -139,7 +139,7 @@ namespace FastTriWinding
 	// triangle-winding-number second-order approximation
 	// T is triangle, P is 'Center' of cluster of dipoles, Q is evaluation point
 	// (This is really just for testing)
-	double Order2Approx(const FTriangle3d& T, const FVector3d& P, const FVector3d& XN, double XA, const FVector3d& Q)
+	inline double Order2Approx(const FTriangle3d& T, const FVector3d& P, const FVector3d& XN, double XA, const FVector3d& Q)
 	{
 		FVector3d dpq = (P - Q);
 
@@ -199,15 +199,18 @@ public:
 	 */
 	int FWNApproxOrder = 2;
 
-	TFastWindingTree(TMeshAABBTree3<TriangleMeshType>* TreeToRef)
+	TFastWindingTree(TMeshAABBTree3<TriangleMeshType>* TreeToRef, bool bAutoBuild = true)
 	{
-		SetTree(TreeToRef);
+		SetTree(TreeToRef, bAutoBuild);
 	}
 
-	void SetTree(TMeshAABBTree3<TriangleMeshType>* TreeToRef)
+	void SetTree(TMeshAABBTree3<TriangleMeshType>* TreeToRef, bool bAutoBuild = true)
 	{
 		this->Tree = TreeToRef;
-		Build(true);
+		if (bAutoBuild)
+		{
+			Build(true);
+		}
 	}
 
 	void Build(bool bAlwaysBuildRegardlessOfTimestamp = true)
