GeometryProcessing: resolve some template compile failures caught by strict compile
#rnx

diff --git a/Engine/Plugins/Experimental/GeometryProcessing/Source/DynamicMesh/Public/DynamicPointSet3.h b/Engine/Plugins/Experimental/GeometryProcessing/Source/DynamicMesh/Public/DynamicPointSet3.h
--- a/Engine/Plugins/Experimental/GeometryProcessing/Source/DynamicMesh/Public/DynamicPointSet3.h
+++ b/Engine/Plugins/Experimental/GeometryProcessing/Source/DynamicMesh/Public/DynamicPointSet3.h
@@ -161,14 +161,14 @@ public:
 			return EMeshResult::Failed_CannotAllocateVertex;
 		}
 
-		int i = 3 * vid;
+		int i = 3 * VertexID;
 		Vertices.InsertAt(Position[2], i + 2);
 		Vertices.InsertAt(Position[1], i + 1);
 		Vertices.InsertAt(Position[0], i);
 
 		if (GetBaseAttributeSet() != nullptr)
 		{
-			GetBaseAttributeSet()->OnNewVertex(vid, true);
+			GetBaseAttributeSet()->OnNewVertex(VertexID, true);
 		}
 
 		UpdateTimeStamp();
@@ -295,9 +295,8 @@ public:
 		{
 			return EMeshResult::Failed_NotAVertex;
 		}
-		VertexRefCounts.Decrement(vID);
-		check(VertexRefCounts.IsValid(vID) == false);
-		VertexEdgeLists.Clear(vID);
+		VertexRefCounts.Decrement(VertexID);
+		check(VertexRefCounts.IsValid(VertexID) == false);
 
 		if (GetBaseAttributeSet() != nullptr)
 		{
diff --git a/Engine/Plugins/Experimental/GeometryProcessing/Source/DynamicMesh/Public/DynamicVerticesOctree3.h b/Engine/Plugins/Experimental/GeometryProcessing/Source/DynamicMesh/Public/DynamicVerticesOctree3.h
--- a/Engine/Plugins/Experimental/GeometryProcessing/Source/DynamicMesh/Public/DynamicVerticesOctree3.h
+++ b/Engine/Plugins/Experimental/GeometryProcessing/Source/DynamicMesh/Public/DynamicVerticesOctree3.h
@@ -162,7 +162,7 @@ public:
 		{
 			FVector3d Point = VertexSource->GetVertex(VertexID);
 			ModifiedBounds.Contain(Point);
-			FSparseDynamicPointOctree3::ReinsertPoint(VertexID, Bounds);
+			FSparseDynamicPointOctree3::ReinsertPoint(VertexID, ModifiedBounds);
 		}
 	}
 
@@ -198,7 +198,7 @@ public:
 		return FSparseDynamicPointOctree3::FindNearestHitPoint(Ray,
 			[this](int vid, const FRay3d& Ray) {
 				FVector3d Point = VertexSource->GetVertex(vid);
-				FLinearIntersection Hit = IntersectionUtil::RaySphereIntersection(Ray.Origin, Ray.Direction, Center, HitSphereRadius);
+				FLinearIntersection Hit = IntersectionUtil::RaySphereIntersection(Ray.Origin, Ray.Direction, Point, HitSphereRadius);
 				return (Hit.intersects) ? Hit.parameter.Min : TNumericLimits<double>::Max();
 			}, MaxDistance);
 	}
