fix geometryprocessingtest compile

#rnx

diff --git a/Engine/Plugins/Experimental/GeometryProcessing/Source/DynamicMesh/Private/Operations/RemoveOccludedTriangles.cpp b/Engine/Plugins/Experimental/GeometryProcessing/Source/DynamicMesh/Private/Operations/RemoveOccludedTriangles.cpp
deleted file mode 100644
diff --git a/Engine/Plugins/Experimental/GeometryProcessing/Source/DynamicMesh/Public/Operations/RemoveOccludedTriangles.h b/Engine/Plugins/Experimental/GeometryProcessing/Source/DynamicMesh/Public/Operations/RemoveOccludedTriangles.h
--- a/Engine/Plugins/Experimental/GeometryProcessing/Source/DynamicMesh/Public/Operations/RemoveOccludedTriangles.h
+++ b/Engine/Plugins/Experimental/GeometryProcessing/Source/DynamicMesh/Public/Operations/RemoveOccludedTriangles.h
@@ -11,6 +11,11 @@
 #include "Spatial/FastWinding.h"
 
 #include "MeshAdapter.h"
+#include "MeshNormals.h"
+
+#include "Async/ParallelFor.h"
+#include "Misc/ScopeLock.h"
+
 
 #include "Util/ProgressCancel.h"
 
@@ -50,11 +55,138 @@ public:
 	 * See alternative versions below if you would like to not precompute acceleration structures; these are mainly provided so the N-mutual-occluding-objects case can run on each mesh separately without 
 	 *
 	 * @param LocalToWorld Transform to take the local mesh into the space of the occluders
-	 * @param Occluders AABB trees for all occluders to consider
+	 * @param Spatial AABB trees for all occluders to consider
 	 * @param FastWindingTrees Precomputed fast winding trees for occluder
 	 * @return true on success
 	 */
-	virtual bool Apply(FTransform3d MeshLocalToOccluderSpace, TMeshAABBTree3<OccluderTriangleMeshType>* Occluder, TFastWindingTree<OccluderTriangleMeshType>* FastWindingTree);
+	virtual bool Apply(FTransform3d MeshLocalToOccluderSpace, TMeshAABBTree3<OccluderTriangleMeshType>* Spatial, TFastWindingTree<OccluderTriangleMeshType>* FastWindingTree)
+	{
+		if (Cancelled())
+		{
+			return false;
+		}
+
+		// ray directions
+		TArray<FVector3d> RayDirs; int NR = 0;
+
+		if (InsideMode == EOcclusionCalculationMode::SimpleOcclusionTest)
+		{
+			RayDirs.Add(FVector3d::UnitX()); RayDirs.Add(-FVector3d::UnitX());
+			RayDirs.Add(FVector3d::UnitY()); RayDirs.Add(-FVector3d::UnitY());
+			RayDirs.Add(FVector3d::UnitZ()); RayDirs.Add(-FVector3d::UnitZ());
+
+			for (int AddRayIdx = 0; AddRayIdx < AddRandomRays; AddRayIdx++)
+			{
+				RayDirs.Add(FVector3d(FMath::VRand()));
+			}
+			NR = RayDirs.Num();
+		}
+
+		auto IsOccludedFWN = [this, &FastWindingTree](const FVector3d& Pt)
+		{
+			return FastWindingTree->FastWindingNumber(Pt) > WindingIsoValue;
+		};
+		auto IsOccludedSimple = [this, &Spatial, &RayDirs, &NR](const FVector3d& Pt)
+		{
+			FRay3d Ray;
+			Ray.Origin = Pt;
+
+			for (int RayIdx = 0; RayIdx < NR; ++RayIdx)
+			{
+				Ray.Direction = RayDirs[RayIdx];
+				int HitTID = Spatial->FindNearestHitTriangle(Ray);
+				if (HitTID == IndexConstants::InvalidID)
+				{
+					return false;
+				}
+			}
+			return true;
+		};
+
+		TFunctionRef<bool(const FVector3d& Pt)> IsOccludedF =
+			(InsideMode == EOcclusionCalculationMode::FastWindingNumber) ? 
+			(TFunctionRef<bool(const FVector3d& Pt)>)IsOccludedFWN : 
+			(TFunctionRef<bool(const FVector3d& Pt)>)IsOccludedSimple;
+
+		bool bForceSingleThread = false;
+
+		TArray<bool> VertexOccluded;
+		if ( TriangleSamplingMethod == EOcclusionTriangleSampling::Vertices || TriangleSamplingMethod == EOcclusionTriangleSampling::VerticesAndCentroids )
+		{
+			VertexOccluded.SetNum(Mesh->MaxVertexID());
+
+			// do not trust source mesh normals; safer to recompute
+			FMeshNormals Normals(Mesh);
+			Normals.ComputeVertexNormals();
+
+			ParallelFor(Mesh->MaxVertexID(), [this, &Normals, &VertexOccluded, &IsOccludedF, &MeshLocalToOccluderSpace](int32 VID)
+			{
+				if (!Mesh->IsVertex(VID))
+				{
+					return;
+				}
+				FVector3d SamplePos = Mesh->GetVertex(VID);
+				FVector3d Normal = Normals[VID];
+				SamplePos += Normal * NormalOffset;
+				VertexOccluded[VID] = IsOccludedF(MeshLocalToOccluderSpace.TransformPosition(SamplePos));
+			}, bForceSingleThread);
+		}
+		if (Cancelled())
+		{
+			return false;
+		}
+
+		RemovedT.Empty();
+		FCriticalSection RemoveTMutex;
+		ParallelFor(Mesh->MaxTriangleID(), [this, &VertexOccluded, &IsOccludedF, &RemoveTMutex, &MeshLocalToOccluderSpace](int32 TID)
+		{
+			if (!Mesh->IsTriangle(TID))
+			{
+				return;
+			}
+
+			bool bInside = true;
+			if (TriangleSamplingMethod == EOcclusionTriangleSampling::Vertices || TriangleSamplingMethod == EOcclusionTriangleSampling::VerticesAndCentroids)
+			{
+				FIndex3i Tri = Mesh->GetTriangle(TID);
+				bInside = VertexOccluded[Tri.A] && VertexOccluded[Tri.B] && VertexOccluded[Tri.C];
+			}
+			if (bInside && (TriangleSamplingMethod == EOcclusionTriangleSampling::Centroids || TriangleSamplingMethod == EOcclusionTriangleSampling::VerticesAndCentroids))
+			{
+				FVector3d Centroid, Normal;
+				double Area;
+				Mesh->GetTriInfo(TID, Normal, Area, Centroid);
+
+				FVector3d SamplePos = Centroid + Normal * NormalOffset;
+				bInside = IsOccludedF(MeshLocalToOccluderSpace.TransformPosition(SamplePos));
+			}
+			if (bInside)
+			{
+				FScopeLock RemoveLock(&RemoveTMutex);
+				RemovedT.Add(TID);
+			}
+		}, bForceSingleThread);
+
+
+		if (Cancelled())
+		{
+			return false;
+		}
+
+		if (RemovedT.Num() > 0)
+		{
+			FDynamicMeshEditor Editor(Mesh);
+			bool bOK = Editor.RemoveTriangles(RemovedT, true);
+			if (!bOK)
+			{
+				bRemoveFailed = true;
+				return false;
+			}
+			// TODO: do we want to consider if we have made the mesh non-manifold or do any cleanup?
+		} 
+
+		return true;
+	}
 
 	/**
 	 * Remove the occluded triangles
@@ -126,10 +258,3 @@ protected:
 
 };
 
-
-
-template class DYNAMICMESH_API TRemoveOccludedTriangles<TIndexMeshArrayAdapter<int, double, FVector3d>>;
-
-template class DYNAMICMESH_API TRemoveOccludedTriangles<FDynamicMesh3>;
-
-
