Fix RemoveOccludedTriangles.h compile error
#rnx

diff --git a/Engine/Plugins/Experimental/GeometryProcessing/Source/DynamicMesh/Public/Operations/RemoveOccludedTriangles.h b/Engine/Plugins/Experimental/GeometryProcessing/Source/DynamicMesh/Public/Operations/RemoveOccludedTriangles.h
--- a/Engine/Plugins/Experimental/GeometryProcessing/Source/DynamicMesh/Public/Operations/RemoveOccludedTriangles.h
+++ b/Engine/Plugins/Experimental/GeometryProcessing/Source/DynamicMesh/Public/Operations/RemoveOccludedTriangles.h
@@ -12,6 +12,7 @@
 
 #include "MeshAdapter.h"
 #include "MeshNormals.h"
+#include "DynamicMeshEditor.h"
 
 #include "Async/ParallelFor.h"
 #include "Misc/ScopeLock.h"
