ersatz move (delete/add) mesh Generators from DynamicMesh over to GeometricObjects, as they aren't specific to DynamicMesh at all

#rnx

diff --git a/Engine/Plugins/Experimental/GeometryProcessing/Source/DynamicMesh/Private/Generators/PlanarPolygonMeshGenerator.cpp b/Engine/Plugins/Experimental/GeometryProcessing/Source/DynamicMesh/Private/Generators/PlanarPolygonMeshGenerator.cpp
deleted file mode 100644
diff --git a/Engine/Plugins/Experimental/GeometryProcessing/Source/DynamicMesh/Private/Generators/RectangleMeshGenerator.cpp b/Engine/Plugins/Experimental/GeometryProcessing/Source/DynamicMesh/Private/Generators/RectangleMeshGenerator.cpp
deleted file mode 100644
diff --git a/Engine/Plugins/Experimental/GeometryProcessing/Source/DynamicMesh/Public/Generators/GridBoxMeshGenerator.h b/Engine/Plugins/Experimental/GeometryProcessing/Source/DynamicMesh/Public/Generators/GridBoxMeshGenerator.h
deleted file mode 100644
diff --git a/Engine/Plugins/Experimental/GeometryProcessing/Source/DynamicMesh/Public/Generators/MeshShapeGenerator.h b/Engine/Plugins/Experimental/GeometryProcessing/Source/DynamicMesh/Public/Generators/MeshShapeGenerator.h
deleted file mode 100644
diff --git a/Engine/Plugins/Experimental/GeometryProcessing/Source/DynamicMesh/Public/Generators/MinimalBoxMeshGenerator.h b/Engine/Plugins/Experimental/GeometryProcessing/Source/DynamicMesh/Public/Generators/MinimalBoxMeshGenerator.h
deleted file mode 100644
diff --git a/Engine/Plugins/Experimental/GeometryProcessing/Source/DynamicMesh/Public/Generators/PlanarPolygonMeshGenerator.h b/Engine/Plugins/Experimental/GeometryProcessing/Source/DynamicMesh/Public/Generators/PlanarPolygonMeshGenerator.h
deleted file mode 100644
diff --git a/Engine/Plugins/Experimental/GeometryProcessing/Source/DynamicMesh/Public/Generators/RectangleMeshGenerator.h b/Engine/Plugins/Experimental/GeometryProcessing/Source/DynamicMesh/Public/Generators/RectangleMeshGenerator.h
deleted file mode 100644
diff --git a/Engine/Plugins/Experimental/GeometryProcessing/Source/DynamicMesh/Public/Generators/SphereGenerator.h b/Engine/Plugins/Experimental/GeometryProcessing/Source/DynamicMesh/Public/Generators/SphereGenerator.h
deleted file mode 100644
diff --git a/Engine/Plugins/Experimental/GeometryProcessing/Source/DynamicMesh/Public/Generators/SweepGenerator.h b/Engine/Plugins/Experimental/GeometryProcessing/Source/DynamicMesh/Public/Generators/SweepGenerator.h
deleted file mode 100644
diff --git a/Engine/Plugins/Experimental/GeometryProcessing/Source/GeometricObjects/Private/Generators/PlanarPolygonMeshGenerator.cpp b/Engine/Plugins/Experimental/GeometryProcessing/Source/GeometricObjects/Private/Generators/PlanarPolygonMeshGenerator.cpp
new file mode 100644
--- /dev/null
+++ b/Engine/Plugins/Experimental/GeometryProcessing/Source/GeometricObjects/Private/Generators/PlanarPolygonMeshGenerator.cpp
@@ -0,0 +1,73 @@
+// Copyright 1998-2019 Epic Games, Inc. All Rights Reserved.
+
+#include "Generators/PlanarPolygonMeshGenerator.h"
+#include "CompGeom/PolygonTriangulation.h"
+
+FPlanarPolygonMeshGenerator::FPlanarPolygonMeshGenerator()
+{
+	Normal = FVector3f::UnitZ();
+	IndicesMap = FIndex2i(0, 1);
+}
+
+
+
+void FPlanarPolygonMeshGenerator::SetPolygon(const TArray<FVector2D>& PolygonVerts)
+{
+	Polygon = FPolygon2d();
+	int NumVerts = PolygonVerts.Num();
+	for (int i = 0; i < NumVerts; ++i)
+	{
+		Polygon.AppendVertex(PolygonVerts[i]);
+	}
+}
+
+
+
+
+FMeshShapeGenerator& FPlanarPolygonMeshGenerator::Generate()
+{
+	int NumVertices = Polygon.VertexCount();
+	check(NumVertices >= 3);
+	if (NumVertices < 3)
+	{
+		return *this;
+	}
+
+	TArray<FIndex3i> TriangleList;
+	PolygonTriangulation::TriangulateSimplePolygon(Polygon.GetVertices(), TriangleList);
+
+	int NumTriangles = TriangleList.Num();
+	SetBufferSizes(NumVertices, NumTriangles, NumVertices, NumVertices);
+
+	FAxisAlignedBox2d BoundingBox = Polygon.Bounds();
+	double Width = BoundingBox.Width(), Height = BoundingBox.Height();
+	double UVScale = FMath::Max(Width, Height);
+
+	for (int VertIdx = 0; VertIdx < NumVertices; ++VertIdx)
+	{
+		FVector2d Pos = Polygon[VertIdx];
+
+		UVs[VertIdx] = FVector2f(
+			(float)((Pos.X - BoundingBox.Min.X) / UVScale),
+			(float)((Pos.Y - BoundingBox.Min.Y) / UVScale));
+		UVParentVertex[VertIdx] = VertIdx;
+
+		Normals[VertIdx] = Normal;
+		NormalParentVertex[VertIdx] = VertIdx;
+
+		Vertices[VertIdx] = MakeVertex(Pos.X, Pos.Y);
+	}
+
+
+	for (int TriIdx = 0; TriIdx < NumTriangles; TriIdx++)
+	{
+		FIndex3i Tri = TriangleList[TriIdx];
+		SetTriangle(TriIdx, Tri);
+		SetTriangleUVs(TriIdx, Tri);
+		SetTriangleNormals(TriIdx, Tri);
+		SetTrianglePolygon(TriIdx, 0);
+	}
+
+
+	return *this;
+}
\ No newline at end of file
diff --git a/Engine/Plugins/Experimental/GeometryProcessing/Source/GeometricObjects/Private/Generators/RectangleMeshGenerator.cpp b/Engine/Plugins/Experimental/GeometryProcessing/Source/GeometricObjects/Private/Generators/RectangleMeshGenerator.cpp
new file mode 100644
--- /dev/null
+++ b/Engine/Plugins/Experimental/GeometryProcessing/Source/GeometricObjects/Private/Generators/RectangleMeshGenerator.cpp
@@ -0,0 +1,104 @@
+// Copyright 1998-2019 Epic Games, Inc. All Rights Reserved.
+
+#include "Generators/RectangleMeshGenerator.h"
+
+
+FRectangleMeshGenerator::FRectangleMeshGenerator()
+{
+	Origin = FVector3d::Zero();
+	Width = 10.0f;
+	Height = 10.0f;
+	WidthVertexCount = HeightVertexCount = 8;
+	Normal = FVector3f::UnitZ();
+	IndicesMap = FIndex2i(0, 1);
+	bScaleUVByAspectRatio = true;
+}
+
+
+
+FMeshShapeGenerator& FRectangleMeshGenerator::Generate()
+{
+	check(IndicesMap.A >= 0 && IndicesMap.A <= 2);
+	check(IndicesMap.B >= 0 && IndicesMap.B <= 2);
+
+	int WidthNV = (WidthVertexCount > 1) ? WidthVertexCount : 2;
+	int HeightNV = (HeightVertexCount > 1) ? HeightVertexCount : 2;
+
+	int TotalNumVertices = WidthNV * HeightNV;
+	int TotalNumTriangles = 2 * (WidthNV - 1) * (HeightNV - 1);
+	SetBufferSizes(TotalNumVertices, TotalNumTriangles, TotalNumVertices, TotalNumVertices);
+
+	// corner vertices
+	FVector3d v00 = MakeVertex(-Width / 2.0f, -Height / 2.0f);
+	FVector3d v01 = MakeVertex(Width / 2.0f, -Height / 2.0f);
+	FVector3d v11 = MakeVertex(Width / 2.0f, Height / 2.0f);
+	FVector3d v10 = MakeVertex(-Width / 2.0f, Height / 2.0f);
+
+	// corner UVs
+	float uvleft = 0.0f, uvright = 1.0f, uvbottom = 0.0f, uvtop = 1.0f;
+	if (bScaleUVByAspectRatio && Width != Height)
+	{
+		if (Width > Height)
+		{
+			uvtop = Height / Width;
+		}
+		else
+		{
+			uvright = Width / Height;
+		}
+	}
+
+	FVector2f uv00 = FVector2f(uvleft, uvbottom);
+	FVector2f uv01 = FVector2f(uvright, uvbottom);
+	FVector2f uv11 = FVector2f(uvright, uvtop);
+	FVector2f uv10 = FVector2f(uvleft, uvtop);
+
+
+	int vi = 0;
+	int ti = 0;
+
+	// add vertex rows
+	int start_vi = vi;
+	for (int yi = 0; yi < HeightNV; ++yi)
+	{
+		double ty = (double)yi / (double)(HeightNV - 1);
+		for (int xi = 0; xi < WidthNV; ++xi)
+		{
+			double tx = (double)xi / (double)(WidthNV - 1);
+			Normals[vi] = Normal;
+			NormalParentVertex[vi] = vi;
+			UVs[vi] = BilinearInterp(uv00, uv01, uv11, uv10, (float)tx, (float)ty);
+			UVParentVertex[vi] = vi;
+			Vertices[vi++] = BilinearInterp(v00, v01, v11, v10, tx, ty);
+		}
+	}
+
+	// add triangulated quads
+	int PolyIndex = 0;
+	for (int y0 = 0; y0 < HeightNV - 1; ++y0)
+	{
+		for (int x0 = 0; x0 < WidthNV - 1; ++x0)
+		{
+			int i00 = start_vi + y0 * WidthNV + x0;
+			int i10 = start_vi + (y0 + 1)*WidthNV + x0;
+			int i01 = i00 + 1, i11 = i10 + 1;
+
+			SetTriangle(ti, i00, i11, i01);
+			SetTrianglePolygon(ti, PolyIndex);
+			SetTriangleUVs(ti, i00, i11, i01);
+			SetTriangleNormals(ti, i00, i11, i01);
+
+			ti++;
+
+			SetTriangle(ti, i00, i10, i11);
+			SetTrianglePolygon(ti, PolyIndex);
+			SetTriangleUVs(ti, i00, i10, i11);
+			SetTriangleNormals(ti, i00, i10, i11);
+
+			ti++;
+			PolyIndex++;
+		}
+	}
+
+	return *this;
+}
\ No newline at end of file
diff --git a/Engine/Plugins/Experimental/GeometryProcessing/Source/GeometricObjects/Public/Generators/GridBoxMeshGenerator.h b/Engine/Plugins/Experimental/GeometryProcessing/Source/GeometricObjects/Public/Generators/GridBoxMeshGenerator.h
new file mode 100644
--- /dev/null
+++ b/Engine/Plugins/Experimental/GeometryProcessing/Source/GeometricObjects/Public/Generators/GridBoxMeshGenerator.h
@@ -0,0 +1,266 @@
+// Copyright 1998-2019 Epic Games, Inc. All Rights Reserved.
+
+#pragma once
+
+#include "MeshShapeGenerator.h"
+#include "OrientedBoxTypes.h"
+#include "Util/IndexUtil.h"
+
+/**
+ *  Generate a mesh of a box that has "gridded" faces, i.e. grid of triangulated quads, 
+ *  with EdgeVertices setting the number of verts along each edge in each dimension
+ */
+class /*GEOMETRICOBJECTS_API*/ FGridBoxMeshGenerator : public FMeshShapeGenerator
+{
+public:
+	/** 3D box */
+	FOrientedBox3d Box;
+
+	FIndex3i EdgeVertices { 8,8,8 };
+
+	/** If true (default), UVs are scaled so that there is no stretching. If false, UVs are scaled to fill unit square */
+	bool bScaleUVByAspectRatio = true;
+
+public:
+
+	/** Generate the mesh */
+	virtual FMeshShapeGenerator& Generate() override
+	{
+		FIndex3i N = EdgeVertices;
+		N.A = FMath::Max(2, N.A);
+		N.B = FMath::Max(2, N.B);
+		N.C = FMath::Max(2, N.C);
+
+		FVector3d Nscale(1.0 / (N.A-1), 1.0 / (N.B-1), 1.0 / (N.C-1));
+
+		FIndex3i NInternal = N;
+		NInternal.A -= 2; NInternal.B -= 2; NInternal.C -= 2;
+
+		FIndex3i NTri = N;
+		NTri.A--; NTri.B--; NTri.C--;
+		int NumUVsAndNormals = 2 * (N.A * N.B + N.B * N.C + N.C * N.A);
+		int NumVertices = 8 + (NInternal.A + NInternal.B + NInternal.C) * 4 + (NInternal.A*NInternal.B + NInternal.B*NInternal.C + NInternal.C*NInternal.A) * 2;
+		int NumTriangles = 4 * (NTri.A * NTri.B + NTri.B * NTri.C + NTri.C * NTri.A);
+		SetBufferSizes(NumVertices, NumTriangles, NumUVsAndNormals, NumUVsAndNormals);
+
+		int FaceDimOrder[3]{ 1, 2, 0 }; // ordering from IndexUtil.cpp
+		int D[2][3]{ { 1,2,0 }, { 2,0,1 } }; // helper mapping to go from major dimension to the two canonical sub-dimensions in order; is just D[i][j] -> (j+i+1)%3 (TODO: replace with the modulus version?)
+
+		// allocate a big mapping from faces to vertices, to make life easier when creating triangles later
+		FIndex3i VerticesPerFace(N.B * N.C, N.A * N.C, N.A * N.B);
+		TArray<TArray<int>> FaceVertIndices;
+		FaceVertIndices.SetNum(6);
+		for (int Dim = 0; Dim < 3; Dim++)
+		{
+			int FaceIdxForDim = FaceDimOrder[Dim] * 2;
+			int VertexNum = VerticesPerFace[Dim];
+			FaceVertIndices[FaceIdxForDim].SetNum(VertexNum);
+			FaceVertIndices[FaceIdxForDim+1].SetNum(VertexNum);
+		}
+
+		auto ToFaceV = [&N, &D](int Dim, int D0, int D1)
+		{
+			return D0 + N[D[0][Dim]] * D1;
+		};
+
+		// create the corners and distribute them into the face mapping
+		for (int i = 0; i < 8; ++i) 
+		{
+			Vertices[i] = Box.GetCorner(i);
+			FIndex3i CornerSides = FOrientedBox3d::GetCornerSide(i);
+			for (int Dim = 0; Dim < 3; Dim++)
+			{
+				int FaceIdx = FaceDimOrder[Dim]*2 + CornerSides[Dim];
+				int D0Ind = CornerSides[D[0][Dim]] * (N[D[0][Dim]] - 1);
+				int D1Ind = CornerSides[D[1][Dim]] * (N[D[1][Dim]] - 1);
+				int CornerVInd = ToFaceV(Dim, D0Ind, D1Ind);
+				FaceVertIndices[FaceIdx][CornerVInd] = i;
+			}
+		}
+
+		// create the internal (non-corner) edge vertices and distribute them into the face mapping
+		int CurrentVertIndex = 8;
+		for (int Dim = 0; Dim < 3; Dim++)
+		{
+			int EdgeLen = N[Dim];
+			if (EdgeLen <= 2)
+			{
+				continue; // no internal edge vertices on this dimension
+			}
+			int MajorFaceInd = FaceDimOrder[Dim] * 2;
+			int FaceInds[2]{ -1,-1 };
+			int DSides[2]{ 0,0 };
+			for (DSides[0] = 0; DSides[0] < 2; DSides[0]++)
+			{
+				FaceInds[0] = FaceDimOrder[D[0][Dim]] * 2 + DSides[0];
+				for (DSides[1] = 0; DSides[1] < 2; DSides[1]++)
+				{
+					FaceInds[1] = FaceDimOrder[D[1][Dim]] * 2 + DSides[1];
+					int MajorCornerInd = ToFaceV(Dim, DSides[0] * (N[D[0][Dim]] - 1), DSides[1] * (N[D[1][Dim]] - 1));
+					FVector3d Corners[2]
+					{
+						Vertices[FaceVertIndices[MajorFaceInd  ][MajorCornerInd]],
+						Vertices[FaceVertIndices[MajorFaceInd+1][MajorCornerInd]]
+					};
+					
+					for (int EdgeVert = 1; EdgeVert + 1 < EdgeLen; EdgeVert++)
+					{
+						Vertices[CurrentVertIndex] = FVector3d::Lerp(Corners[0], Corners[1], EdgeVert * Nscale[Dim]);
+						for (int WhichFace = 0; WhichFace < 2; WhichFace++) // each edge is shared by two faces (w/ major axes of subdim 0 and subdim 1 respectively)
+						{
+							int FaceInd = FaceInds[WhichFace];
+
+							int FaceDim = D[WhichFace][Dim];
+							int SubDims[2];
+							SubDims[1 - WhichFace] = EdgeVert;
+							SubDims[WhichFace] = NTri[D[WhichFace][FaceDim]] * DSides[1-WhichFace];
+							int FaceV = ToFaceV(FaceDim, SubDims[0], SubDims[1]);
+							FaceVertIndices[FaceInds[WhichFace]][FaceV] = CurrentVertIndex;
+						}
+						CurrentVertIndex++;
+					}
+				}
+			}
+		}
+
+		// create the internal (non-corner, non-edge) face vertices and distribute them into the face mapping
+		for (int Dim = 0; Dim < 3; Dim++)
+		{
+			int FaceIdxBase = FaceDimOrder[Dim]*2;
+			int FaceInternalVNum = NInternal[D[0][Dim]] * NInternal[D[1][Dim]];
+			if (FaceInternalVNum <= 0)
+			{
+				continue;
+			}
+
+			for (int Side = 0; Side < 2; Side++)
+			{
+				int MajorFaceInd = FaceIdxBase + Side;
+				for (int D0 = 1; D0 + 1 < N[D[0][Dim]]; D0++)
+				{
+					int BotInd = ToFaceV(Dim, D0, 0);
+					int TopInd = ToFaceV(Dim, D0, N[D[1][Dim]] - 1);
+
+					FVector3d Edges[2]
+					{
+						Vertices[FaceVertIndices[MajorFaceInd][BotInd]],
+						Vertices[FaceVertIndices[MajorFaceInd][TopInd]]
+					};
+					for (int D1 = 1; D1 + 1 < N[D[1][Dim]]; D1++)
+					{
+						Vertices[CurrentVertIndex] = FVector3d::Lerp(Edges[0], Edges[1], D1 * Nscale[D[1][Dim]]);
+						FaceVertIndices[MajorFaceInd][ToFaceV(Dim, D0, D1)] = CurrentVertIndex;
+						CurrentVertIndex++;
+					}
+				}
+			}
+		}
+
+		double MaxDimension = 2.0*Box.Extents.MaxAbs();
+		float UVScale = (bScaleUVByAspectRatio) ? (1.0f / (float)MaxDimension) : 1.0f;
+
+		// create the face triangles and UVs+normals
+		int CurrentTriIdx = 0;
+		int CurrentUVIdx = 0;
+		for (int Dim = 0; Dim < 3; Dim++)
+		{
+			int FaceIdxBase = FaceDimOrder[Dim]*2;
+
+			// UV-specific minor axes + flips; manually set to match default ue4 cube texture arrangement
+			int Minor1Flip[3] = { -1, 1, 1 };
+			int Minor2Flip[3] = { -1, -1, 1 };
+
+			// UV scales for D0, D1
+			double FaceWidth = FMathd::Abs(Box.Extents[D[0][Dim]]) * 2.0;
+			double FaceHeight = FMathd::Abs(Box.Extents[D[1][Dim]]) * 2.0;
+			double WidthUVScale = FaceWidth * UVScale;
+			double HeightUVScale = FaceHeight * UVScale;
+
+			
+			for (int Side = 0; Side < 2; Side++)
+			{
+				int SideOpp = 1 - Side;
+				float SideSign = Side * 2 - 1;
+
+				FVector3f Normal(0, 0, 0);
+				Normal[Dim] = (2 * Side - 1);
+				int MajorFaceInd = FaceIdxBase + Side;
+
+				int FaceUVStartInd = CurrentUVIdx;
+				// set all the UVs and normals
+				FVector2f UV;
+				int UVXDim = Dim == 1 ? 1 : 0;	// which dim (of D0,D1) follows the horizontal UV coordinate
+				int UVYDim = 1 - UVXDim;		// which dim (of D0,D1) follows the vertical UV coordinate
+				for (int D0 = 0; D0 < N[D[0][Dim]]; D0++)
+				{
+					for (int D1 = 0; D1 < N[D[1][Dim]]; D1++)
+					{
+						// put the grid coordinates (centered at 0,0) into the UVs
+						UV[UVXDim] = D0 * Nscale[D[0][Dim]] - .5;
+						UV[UVYDim] = D1 * Nscale[D[1][Dim]] - .5;
+						// invert axes to match the desired UV patterns & so the opp faces are not backwards
+						UV.X *= SideSign * Minor1Flip[Dim];
+						UV.Y *= Minor2Flip[Dim];
+						// recenter and scale up
+						UV[UVXDim] = (UV[UVXDim] + .5) * WidthUVScale;
+						UV[UVYDim] = (UV[UVYDim] + .5) * HeightUVScale;
+						UVs[CurrentUVIdx] = UV;
+						Normals[CurrentUVIdx] = Normal;
+						UVParentVertex[CurrentUVIdx] = FaceVertIndices[MajorFaceInd][ToFaceV(Dim, D0, D1)];
+						NormalParentVertex[CurrentUVIdx] = FaceVertIndices[MajorFaceInd][ToFaceV(Dim, D0, D1)];
+						CurrentUVIdx++;
+					}
+				}
+
+				// set all the triangles
+				for (int D0 = 0; D0 + 1 < N[D[0][Dim]]; D0++)
+				{
+					for (int D1 = 0; D1 + 1 < N[D[1][Dim]]; D1++)
+					{
+						SetTriangle(CurrentTriIdx,
+								FaceVertIndices[MajorFaceInd][ToFaceV(Dim, D0, D1)],
+								FaceVertIndices[MajorFaceInd][ToFaceV(Dim, D0+SideOpp, D1+Side)],
+								FaceVertIndices[MajorFaceInd][ToFaceV(Dim, D0+1, D1+1)]
+							);
+
+						SetTriangleUVs(CurrentTriIdx,
+							FaceUVStartInd + D1 + (D0) * N[D[1][Dim]],
+							FaceUVStartInd + D1+Side + (D0+SideOpp) * N[D[1][Dim]],
+							FaceUVStartInd + D1+1 + (D0+1) * N[D[1][Dim]]
+						);
+						SetTriangleNormals(CurrentTriIdx,
+							FaceUVStartInd + D1 + (D0) * N[D[1][Dim]],
+							FaceUVStartInd + D1+Side + (D0+SideOpp) * N[D[1][Dim]],
+							FaceUVStartInd + D1+1 + (D0+1) * N[D[1][Dim]]
+						);
+						SetTrianglePolygon(CurrentTriIdx, MajorFaceInd);
+						CurrentTriIdx++;
+
+						SetTriangle(CurrentTriIdx,
+							FaceVertIndices[MajorFaceInd][ToFaceV(Dim, D0, D1)],
+							FaceVertIndices[MajorFaceInd][ToFaceV(Dim, D0+1, D1+1)],
+							FaceVertIndices[MajorFaceInd][ToFaceV(Dim, D0+Side, D1+SideOpp)]
+						);
+						SetTriangleUVs(CurrentTriIdx,
+							FaceUVStartInd + D1 + (D0) * N[D[1][Dim]],
+							FaceUVStartInd + D1+1 + (D0+1) * N[D[1][Dim]],
+							FaceUVStartInd + D1+SideOpp + (D0+Side) * N[D[1][Dim]]
+						);
+						SetTriangleNormals(CurrentTriIdx,
+							FaceUVStartInd + D1 + (D0) * N[D[1][Dim]],
+							FaceUVStartInd + D1+1 + (D0+1) * N[D[1][Dim]],
+							FaceUVStartInd + D1+SideOpp + (D0+Side) * N[D[1][Dim]]
+						);
+						SetTrianglePolygon(CurrentTriIdx, MajorFaceInd);
+						CurrentTriIdx++;
+					}
+				}
+			}
+		}
+
+
+		return *this;
+	}
+
+};
+
diff --git a/Engine/Plugins/Experimental/GeometryProcessing/Source/GeometricObjects/Public/Generators/MeshShapeGenerator.h b/Engine/Plugins/Experimental/GeometryProcessing/Source/GeometricObjects/Public/Generators/MeshShapeGenerator.h
new file mode 100644
--- /dev/null
+++ b/Engine/Plugins/Experimental/GeometryProcessing/Source/GeometricObjects/Public/Generators/MeshShapeGenerator.h
@@ -0,0 +1,253 @@
+// Copyright 1998-2019 Epic Games, Inc. All Rights Reserved.
+
+#pragma once
+
+#include "VectorTypes.h"
+#include "IndexTypes.h"
+
+/**
+ * Base class for triangle mesh generators (eg like to generate sphere, cylinder, etc)
+ * Subclasses must implement ::Generate() 
+ */
+class GEOMETRICOBJECTS_API FMeshShapeGenerator
+{
+public:
+	/** Array of vertex positions */
+	TArray<FVector3d> Vertices;
+	
+	/** Array of UV positions. These are completely independent of vertex list (ie not per-vertex unless that's what generator produces) */
+	TArray<FVector2f> UVs;
+	/** Parent vertex index for each UV. Same length as UVs array. */
+	TArray<int> UVParentVertex;
+
+	/** Array of Normals. These are completely independent of vertex list (ie not per-vertex unless that's what generator produces) */
+	TArray<FVector3f> Normals;
+	/** Parent vertex index for each Normal. Same length as Normals array. */
+	TArray<int> NormalParentVertex;
+
+	/** Array of triangle corner positions, stored as tuples of indices into Vertices array */
+	TArray<FIndex3i> Triangles;
+	/** Array of triangle corner UVs, stored as tuples of indices into UVs array. Same length as Triangles array. */
+	TArray<FIndex3i> TriangleUVs;
+	/** Array of triangle corner Normals, stored as tuples of indices into Normals array. Same length as Triangles array. */
+	TArray<FIndex3i> TriangleNormals;
+	/** ARray of per-triangle integer polygon IDs. Same length as Triangles array. */
+	TArray<int> TrianglePolygonIDs;
+
+
+public:
+	/** If true, reverse orientation of created mesh */
+	bool bReverseOrientation = false;
+
+	virtual ~FMeshShapeGenerator()
+	{
+	}
+
+
+	/**
+	 * Subclasses implement this to generate mesh
+	 */
+	virtual FMeshShapeGenerator& Generate() = 0;
+
+
+	/** clear arrays so that Generate() can be run again */
+	void Reset()
+	{
+		Vertices.Reset();
+		UVs.Reset();
+		UVParentVertex.Reset();
+		Normals.Reset();
+		NormalParentVertex.Reset();
+		Triangles.Reset();
+		TriangleUVs.Reset();
+		TriangleNormals.Reset();
+		TrianglePolygonIDs.Reset();
+	}
+
+
+	/**
+	 * Set the various internal buffers to the correct sizes for the given element counts
+	 */
+	void SetBufferSizes(int NumVertices, int NumTriangles, int NumUVs, int NumNormals)
+	{
+		if (NumVertices > 0)
+		{
+			Vertices.SetNum(NumVertices);
+		}
+		if (NumTriangles > 0)
+		{
+			Triangles.SetNum(NumTriangles);
+			TriangleUVs.SetNum(NumTriangles);
+			TriangleNormals.SetNum(NumTriangles);
+			TrianglePolygonIDs.SetNum(NumTriangles);
+		}
+		if (NumUVs > 0)
+		{
+			UVs.SetNum(NumUVs);
+			UVParentVertex.SetNum(NumUVs);
+		}
+		if (NumNormals > 0)
+		{
+			Normals.SetNum(NumNormals);
+			NormalParentVertex.SetNum(NumNormals);
+		}
+	}
+	/**
+	 * Extends the various internal buffers to the correct sizes for the given additional element counts
+	 */
+	void ExtendBufferSizes(int AddVertices, int AddTriangles, int AddUVs, int AddNormals)
+	{
+		if (AddVertices > 0)
+		{
+			int NumVertices = Vertices.Num() + AddVertices;
+			Vertices.SetNum(NumVertices);
+		}
+		if (AddTriangles > 0)
+		{
+			int NumTriangles = Triangles.Num() + AddTriangles;
+			Triangles.SetNum(NumTriangles);
+			TriangleUVs.SetNum(NumTriangles);
+			TriangleNormals.SetNum(NumTriangles);
+			TrianglePolygonIDs.SetNum(NumTriangles);
+		}
+		if (AddUVs > 0)
+		{
+			int NumUVs = UVs.Num() + AddUVs;
+			UVs.SetNum(NumUVs);
+			UVParentVertex.SetNum(NumUVs);
+		}
+		if (AddNormals > 0)
+		{
+			int NumNormals = Normals.Num() + AddNormals;
+			Normals.SetNum(NumNormals);
+			NormalParentVertex.SetNum(NumNormals);
+		}
+	}
+
+
+	/** Set vertex at Index to given Position */
+	inline void SetVertex(int Index, const FVector3d& Position)
+	{
+		Vertices[Index] = Position;
+	}
+	/** Append a new vertex at the given Position */
+	inline int AppendVertex(const FVector3d& Position)
+	{
+		Vertices.Add(Position);
+		return Vertices.Num()-1;
+	}
+
+	/** Set UV at Index to given value with given ParentVertex */
+	inline void SetUV(int Index, const FVector2f& UV, int ParentVertex)
+	{
+		UVs[Index] = UV;
+		UVParentVertex[Index] = ParentVertex;
+	}
+	inline int AppendUV(const FVector2f& UV, int ParentVertex)
+	{
+		UVs.Add(UV);
+		UVParentVertex.Add(ParentVertex);
+		return UVs.Num() - 1;
+	}
+
+	/** Set Normal at Index to given value with given ParentVertex */
+	inline void SetNormal(int Index, const FVector3f& Normal, int ParentVertex)
+	{
+		Normals[Index] = Normal;
+		NormalParentVertex[Index] = ParentVertex;
+	}
+	inline int AppendNormal(const FVector3f& Normal, int ParentVertex)
+	{
+		Normals.Add(Normal);
+		NormalParentVertex.Add(ParentVertex);
+		return Normals.Num() - 1;
+	}
+
+
+	inline void SetTriangle(int Index, const FIndex3i& Tri)
+	{
+		Triangles[Index] = bReverseOrientation ? FIndex3i(Tri.C, Tri.B, Tri.A) : Tri;
+	}
+	inline void SetTriangle(int Index, int A, int B, int C)
+	{
+		Triangles[Index] = bReverseOrientation ? FIndex3i(C, B, A) : FIndex3i(A, B, C);
+	}
+	inline void SetTriangle(int Index, int A, int B, int C, bool bClockwiseOverride)
+	{
+		Triangles[Index] = (bReverseOrientation != bClockwiseOverride) ? FIndex3i(C, B, A) : FIndex3i(A, B, C);
+	}
+
+	
+	inline int AppendTriangle(int Index, int A, int B, int C)
+	{
+		Triangles.Add( bReverseOrientation ? FIndex3i(C, B, A) : FIndex3i(A, B, C) );
+		return Triangles.Num() - 1;
+	}
+
+
+	inline void SetTriangleUVs(int Index, const FIndex3i& Tri)
+	{
+		TriangleUVs[Index] = bReverseOrientation ? FIndex3i(Tri.C, Tri.B, Tri.A) : Tri;
+	}
+	inline void SetTriangleUVs(int Index, int A, int B, int C)
+	{
+		TriangleUVs[Index] = bReverseOrientation ? FIndex3i(C, B, A) : FIndex3i(A, B, C);
+	}
+	inline void SetTriangleUVs(int Index, int A, int B, int C, bool bClockwiseOverride)
+	{
+		TriangleUVs[Index] = (bReverseOrientation != bClockwiseOverride) ? FIndex3i(C, B, A) : FIndex3i(A, B, C);
+	}
+
+
+	inline void SetTriangleNormals(int Index, const FIndex3i& Tri)
+	{
+		TriangleNormals[Index] = bReverseOrientation ? FIndex3i(Tri.C, Tri.B, Tri.A) : Tri;
+	}
+	inline void SetTriangleNormals(int Index, int A, int B, int C)
+	{
+		TriangleNormals[Index] = bReverseOrientation ? FIndex3i(C, B, A) : FIndex3i(A, B, C);
+	}
+	inline void SetTriangleNormals(int Index, int A, int B, int C, bool bClockwiseOverride)
+	{
+		TriangleNormals[Index] = (bReverseOrientation != bClockwiseOverride) ? FIndex3i(C, B, A) : FIndex3i(A, B, C);
+	}
+
+	inline void SetTrianglePolygon(int Index, int PolygonID)
+	{
+		TrianglePolygonIDs[Index] = PolygonID;
+	}
+
+
+
+
+	static FVector3d BilinearInterp(const FVector3d &v00, const FVector3d &v10, const FVector3d &v11, const FVector3d &v01, double tx, double ty)
+	{
+		FVector3d a = FVector3d::Lerp(v00, v01, ty);
+		FVector3d b = FVector3d::Lerp(v10, v11, ty);
+		return FVector3d::Lerp(a, b, tx);
+	}
+
+	static FVector2d BilinearInterp(const FVector2d &v00, const FVector2d &v10, const FVector2d &v11, const FVector2d &v01, double tx, double ty)
+	{
+		FVector2d a = FVector2d::Lerp(v00, v01, ty);
+		FVector2d b = FVector2d::Lerp(v10, v11, ty);
+		return FVector2d::Lerp(a, b, tx);
+	}
+
+	static FVector2f BilinearInterp(const FVector2f &v00, const FVector2f &v10, const FVector2f &v11, const FVector2f &v01, float tx, float ty)
+	{
+		FVector2f a = FVector2f::Lerp(v00, v01, ty);
+		FVector2f b = FVector2f::Lerp(v10, v11, ty);
+		return FVector2f::Lerp(a, b, tx);
+	}
+
+
+	static FVector3i LinearInterp(const FVector3i &a, const FVector3i &b, double t)
+	{
+		FVector3d c = FVector3d::Lerp((FVector3d)a, (FVector3d)b, t);
+		return FVector3i((int)round(c.X), (int)round(c.Y), (int)round(c.Z));
+	}
+
+
+
+};
\ No newline at end of file
diff --git a/Engine/Plugins/Experimental/GeometryProcessing/Source/GeometricObjects/Public/Generators/MinimalBoxMeshGenerator.h b/Engine/Plugins/Experimental/GeometryProcessing/Source/GeometricObjects/Public/Generators/MinimalBoxMeshGenerator.h
new file mode 100644
--- /dev/null
+++ b/Engine/Plugins/Experimental/GeometryProcessing/Source/GeometricObjects/Public/Generators/MinimalBoxMeshGenerator.h
@@ -0,0 +1,86 @@
+// Copyright 1998-2019 Epic Games, Inc. All Rights Reserved.
+
+#pragma once
+
+#include "MeshShapeGenerator.h"
+#include "OrientedBoxTypes.h"
+#include "Util/IndexUtil.h"
+
+/**
+ * Generate an oriented Box mesh with the smallest number of triangles possible (6 vertices, 12 triangles)
+ */
+class /*GEOMETRICOBJECTS_API*/ FMinimalBoxMeshGenerator : public FMeshShapeGenerator
+{
+public:
+	/** 3D box */
+	FOrientedBox3d Box;
+
+	/** If true (default), UVs are scaled so that there is no stretching. If false, UVs are scaled to fill unit square */
+	bool bScaleUVByAspectRatio = true;
+
+public:
+
+	/** Generate the mesh */
+	virtual FMeshShapeGenerator& Generate() override
+	{
+		// box has 8 corners, 6 quad-faces, each face has 2 triangles and 4 attribute-corners
+		SetBufferSizes(8, 12, 24, 24);
+
+
+		for (int i = 0; i < 8; ++i) 
+		{
+			Vertices[i] = Box.GetCorner(i);
+		}
+
+		double MaxDimension = 2.0*Box.Extents.MaxAbs();
+		float UVScale = (bScaleUVByAspectRatio) ? (1.0f / (float)MaxDimension) : 1.0f;
+
+		int TriIndex = 0;
+		int AttribIndex = 0;
+		for (int fi = 0; fi < 6; ++fi) 
+		{		
+			int NormalAxisIdx = IndexUtil::BoxFaceNormals[fi];
+			FVector3f FaceNormal = (FVector3f)(FMath::Sign(NormalAxisIdx) * Box.GetAxis(FMath::Abs(NormalAxisIdx) - 1));
+
+			double FaceWidth = Vertices[IndexUtil::BoxFaces[fi][0]].Distance(Vertices[IndexUtil::BoxFaces[fi][1]]);
+			double FaceHeight = Vertices[IndexUtil::BoxFaces[fi][1]].Distance(Vertices[IndexUtil::BoxFaces[fi][2]]);
+			double WidthUVScale = FaceWidth * UVScale;
+			double HeightUVScale = FaceHeight * UVScale;
+
+			int ElementIndices[4];
+
+			for (int j = 0; j < 4; ++j)
+			{
+				Normals[AttribIndex] = FaceNormal;
+				FVector2i CornerUV = IndexUtil::BoxFacesUV[j];
+				UVs[AttribIndex] = FVector2f(WidthUVScale*(float)CornerUV.X, HeightUVScale*(float)CornerUV.Y);
+				ElementIndices[j] = AttribIndex;
+				AttribIndex++;
+			}
+
+
+			SetTriangle(TriIndex,
+				IndexUtil::BoxFaces[fi][0], IndexUtil::BoxFaces[fi][1], IndexUtil::BoxFaces[fi][2]);
+			SetTrianglePolygon(TriIndex, fi);
+			SetTriangleUVs(TriIndex, ElementIndices[0], ElementIndices[1], ElementIndices[2]);
+			SetTriangleNormals(TriIndex, ElementIndices[0], ElementIndices[1], ElementIndices[2]);
+			TriIndex++;
+
+			SetTriangle(TriIndex,
+				IndexUtil::BoxFaces[fi][0], IndexUtil::BoxFaces[fi][2], IndexUtil::BoxFaces[fi][3]);
+			SetTrianglePolygon(TriIndex, fi);
+			SetTriangleUVs(TriIndex, ElementIndices[0], ElementIndices[2], ElementIndices[3]);
+			SetTriangleNormals(TriIndex, ElementIndices[0], ElementIndices[2], ElementIndices[3]);
+			TriIndex++;
+
+
+		}
+
+		return *this;
+	}
+
+
+
+
+
+};
\ No newline at end of file
diff --git a/Engine/Plugins/Experimental/GeometryProcessing/Source/GeometricObjects/Public/Generators/PlanarPolygonMeshGenerator.h b/Engine/Plugins/Experimental/GeometryProcessing/Source/GeometricObjects/Public/Generators/PlanarPolygonMeshGenerator.h
new file mode 100644
--- /dev/null
+++ b/Engine/Plugins/Experimental/GeometryProcessing/Source/GeometricObjects/Public/Generators/PlanarPolygonMeshGenerator.h
@@ -0,0 +1,41 @@
+// Copyright 1998-2019 Epic Games, Inc. All Rights Reserved.
+
+#pragma once
+
+#include "MeshShapeGenerator.h"
+#include "Polygon2.h"
+
+/**
+ * Generate planar triangulation of a Polygon.
+ */
+class GEOMETRICOBJECTS_API FPlanarPolygonMeshGenerator : public FMeshShapeGenerator
+{
+public:
+	/** Polygon to triangulate. If Polygon has self-intersections or degenerate edges, result is undefined. */
+	FPolygon2d Polygon;
+
+	/** Normal vector of all vertices will be set to this value. Default is +Z axis. */
+	FVector3f Normal;
+
+	/** How to map 2D indices to 3D. Default is (0,1) = (x,y,0). */
+	FIndex2i IndicesMap;
+
+public:
+	FPlanarPolygonMeshGenerator();
+
+	/** Initialize the polygon from an array of 2D vertices */
+	void SetPolygon(const TArray<FVector2D>& PolygonVerts);
+
+	/** Generate the triangulation */
+	virtual FMeshShapeGenerator& Generate() override;
+
+
+	/** Create vertex at position under IndicesMap, shifted to Origin*/
+	inline FVector3d MakeVertex(double x, double y)
+	{
+		FVector3d v(0, 0, 0);
+		v[IndicesMap.A] = x;
+		v[IndicesMap.B] = y;
+		return v;
+	}
+};
\ No newline at end of file
diff --git a/Engine/Plugins/Experimental/GeometryProcessing/Source/GeometricObjects/Public/Generators/RectangleMeshGenerator.h b/Engine/Plugins/Experimental/GeometryProcessing/Source/GeometricObjects/Public/Generators/RectangleMeshGenerator.h
new file mode 100644
--- /dev/null
+++ b/Engine/Plugins/Experimental/GeometryProcessing/Source/GeometricObjects/Public/Generators/RectangleMeshGenerator.h
@@ -0,0 +1,52 @@
+// Copyright 1998-2019 Epic Games, Inc. All Rights Reserved.
+
+#pragma once
+
+#include "MeshShapeGenerator.h"
+
+/**
+ * Generate planar rectangular mesh with variable number of subdivisions along width and height.
+ * By default, center of rectangle is centered at (0,0,0) origin
+ */
+class GEOMETRICOBJECTS_API FRectangleMeshGenerator : public FMeshShapeGenerator
+{
+public:
+	/** Rectangle will be translated so that center is at this point */
+	FVector3d Origin;
+	/** Normal vector of all vertices will be set to this value. Default is +Z axis. */
+	FVector3f Normal;
+
+	/** Width of rectangle */
+	double Width;
+	/** Height of rectangle */
+	double Height;
+
+	/** Number of vertices along Width axis */
+	int WidthVertexCount;
+	/** Number of vertices along Height axis */
+	int HeightVertexCount;
+
+	/** If true (default), UVs are scaled so that there is no stretching. If false, UVs are scaled to fill unit square */
+	bool bScaleUVByAspectRatio = true;
+
+	/** Specifices how 2D indices are mapped to 3D points. Default is (0,1) = (x,y,0). */ 
+	FIndex2i IndicesMap;
+
+public:
+	FRectangleMeshGenerator();
+
+	/** Generate the mesh */
+	virtual FMeshShapeGenerator& Generate() override;
+
+	/** Create vertex at position under IndicesMap, shifted to Origin*/
+	virtual FVector3d MakeVertex(double x, double y)
+	{
+		FVector3d v(0, 0, 0);
+		v[IndicesMap.A] = x;
+		v[IndicesMap.B] = y;
+		return Origin + v;
+	}
+
+
+
+};
\ No newline at end of file
diff --git a/Engine/Plugins/Experimental/GeometryProcessing/Source/GeometricObjects/Public/Generators/SphereGenerator.h b/Engine/Plugins/Experimental/GeometryProcessing/Source/GeometricObjects/Public/Generators/SphereGenerator.h
new file mode 100644
--- /dev/null
+++ b/Engine/Plugins/Experimental/GeometryProcessing/Source/GeometricObjects/Public/Generators/SphereGenerator.h
@@ -0,0 +1,206 @@
+// Copyright 1998-2019 Epic Games, Inc. All Rights Reserved.
+
+#pragma once
+
+#include "MeshShapeGenerator.h"
+#include "OrientedBoxTypes.h"
+#include "Util/IndexUtil.h"
+
+/**
+ * Generate a sphere mesh, with UVs wrapped cylindrically
+ */
+class /*GEOMETRICOBJECTS_API*/ FSphereGenerator : public FMeshShapeGenerator
+{
+public:
+	/** Radius */
+	double Radius = 1;
+
+	int NumPhi = 16; // number of vertices along vertical extent from north pole to south pole
+	int NumTheta = 16; // number of vertices around circles
+
+public:
+
+	inline static FVector3d SphericalToCartesian(double r, double theta, double phi)
+	{
+		double Sphi = sin(phi);
+		double Cphi = cos(phi);
+		double Ctheta = cos(theta);
+		double Stheta = sin(theta);
+
+		return FVector3d(r * Ctheta * Sphi, r * Stheta * Sphi, r * Cphi);
+	}
+
+	/** Generate the mesh */
+	virtual FMeshShapeGenerator& Generate() override
+	{
+		// enforce sane values for vertex counts
+		if (NumPhi < 3)
+		{
+			NumPhi = 3;
+		}
+		if (NumTheta < 3)
+		{
+			NumTheta = 3;
+		}
+		int32 NumVertices = (NumPhi - 2) * NumTheta + 2;
+		int32 NumUVs = NumPhi * (NumTheta + 1);
+		int32 NumTris = (NumPhi - 2) * NumTheta * 2;
+		SetBufferSizes(NumVertices, NumTris, NumUVs, NumVertices);
+
+		double Dphi = FMathd::Pi / double(NumPhi - 1);
+		double Dtheta = FMathd::TwoPi / double(NumTheta);
+		float DUVphi = 1.0 / float(NumPhi - 1);
+		float DUVtheta = 1.0 / float(NumTheta);
+
+		TArray<int32> Grid2d; Grid2d.SetNumUninitialized(NumTheta * NumPhi);
+		for (int32 i = 0; i < NumTheta * NumPhi; ++i)
+		{
+			Grid2d[i] = -1;
+		}
+
+
+		auto GridInterface = [&Grid2d, this](int32 itheta, int32 jphi)->int32&
+		{
+			// make periodic in j
+			while (itheta >= NumTheta) { itheta = itheta - NumTheta; }
+
+			int32 offset = itheta + jphi * NumTheta;
+
+			return Grid2d[offset];
+		};
+
+		int32 VtxIdx = 0;
+		for (int32 p = 1; p < NumPhi - 1; ++p) // NB: this skips the poles.
+		{
+			double Phi = p * Dphi;
+
+			for (int32 t = 0; t < NumTheta; ++t)
+			{
+				double Theta = t * Dtheta;
+				FVector3d Pos = SphericalToCartesian(1., Theta, Phi);
+
+				Vertices[VtxIdx] = Pos * Radius;
+				Normals[VtxIdx] = FVector3f(Pos);
+				NormalParentVertex[VtxIdx] = VtxIdx;
+				GridInterface(t, p) = VtxIdx;
+				VtxIdx++;
+			}
+		}
+
+
+		// add a single point at the North Pole
+		{
+			FVector3d Pos = SphericalToCartesian(1., 0., 0.);
+			Vertices[VtxIdx] = Pos * Radius;
+			Normals[VtxIdx] = FVector3f(Pos);
+			NormalParentVertex[VtxIdx] = VtxIdx;
+
+			for (int32 t = 0; t < NumTheta; ++t)
+			{
+				GridInterface(t, 0) = VtxIdx;
+			}
+
+			VtxIdx++;
+		}
+		// add a single point at the south pole
+		{
+			FVector3d Pos = SphericalToCartesian(-1., 0., 0.);
+			Vertices[VtxIdx] = Pos * Radius;
+			Normals[VtxIdx] = FVector3f(Pos);
+			NormalParentVertex[VtxIdx] = VtxIdx;
+
+			for (int32 t = 0; t < NumTheta; ++t)
+			{
+				GridInterface(t, NumPhi - 1) = VtxIdx;
+			}
+
+			VtxIdx++;
+		}
+
+		int32 UVIdx = 0;
+		auto UVIdxLookup = [this](int32 t, int32 p)
+		{
+			check(t >= 0 && t <= NumTheta);
+			return p * (NumTheta+1) + t;
+		};
+		for (int32 p = 0; p < NumPhi; ++p) // NB: this skips the poles.
+		{
+			float Phi = p * DUVphi;
+
+			for (int32 t = 0; t <= NumTheta; ++t)
+			{
+				float Theta = 1 - t * DUVtheta;
+				UVs[UVIdx] = FVector2f(Theta, Phi);
+				UVParentVertex[UVIdx] = GridInterface(t, p);
+				ensure(UVIdxLookup(t, p) == UVIdx);
+
+				UVIdx++;
+			}
+		}
+
+		int32 TriIdx = 0, PolyIdx = 0;
+		for (int32 p = 1; p < NumPhi - 2; ++p)
+		{
+			for (int32 t = 0; t < NumTheta; ++t)
+			{
+				// (p,t), (p, t+1), (p+1, t), (p+1, t+1) are the 4 corners of a quad.
+				// counter clockwise.
+				int32 Corners[4] = { GridInterface(t, p + 1), GridInterface(t + 1, p + 1), GridInterface(t + 1, p), GridInterface(t, p) };
+				int32 UVCorners[4] = { UVIdxLookup(t, p + 1),   UVIdxLookup(t + 1, p + 1),   UVIdxLookup(t + 1, p),   UVIdxLookup(t, p) };
+
+				// convert each quad into 2 triangles.
+
+				SetTriangle(TriIdx, Corners[0], Corners[2], Corners[1]);
+				SetTrianglePolygon(TriIdx, PolyIdx);
+				SetTriangleUVs(TriIdx, UVCorners[0], UVCorners[2], UVCorners[1]);
+				SetTriangleNormals(TriIdx, Corners[0], Corners[2], Corners[1]);
+				TriIdx++;
+
+				SetTriangle(TriIdx, Corners[2], Corners[0], Corners[3]);
+				SetTrianglePolygon(TriIdx, PolyIdx);
+				SetTriangleUVs(TriIdx, UVCorners[2], UVCorners[0], UVCorners[3]);
+				SetTriangleNormals(TriIdx, Corners[2], Corners[0], Corners[3]);
+				TriIdx++;
+				PolyIdx++;
+			}
+		}
+
+		// Do Triangles that connect to north pole
+		{
+			int32 p = 0; // North pole 
+			for (int32 t = 0; t < NumTheta; ++t)
+			{
+				int32 Corners[3] = { GridInterface(t, p + 1), GridInterface(t + 1, p + 1), GridInterface(t + 1, p) };
+				int32 UVCorners[3] = { UVIdxLookup(t, p + 1),   UVIdxLookup(t + 1, p + 1),   UVIdxLookup(t + 1, p) };
+
+				SetTriangle(TriIdx, Corners[0], Corners[2], Corners[1]);
+				SetTrianglePolygon(TriIdx, PolyIdx);
+				SetTriangleUVs(TriIdx, UVCorners[0], UVCorners[2], UVCorners[1]);
+				SetTriangleNormals(TriIdx, Corners[0], Corners[2], Corners[1]);
+				TriIdx++;
+				PolyIdx++;
+			}
+		}
+
+		// Do Triangles that connect to south pole
+		{
+			int32 p = NumPhi - 2; // South pole 
+			for (int32 t = 0; t < NumTheta; ++t)
+			{
+				int32 Corners[4] = { GridInterface(t, p + 1), GridInterface(t + 1, p), GridInterface(t, p) };
+				int32 UVCorners[3] = { UVIdxLookup(t, p + 1),   UVIdxLookup(t + 1, p),   UVIdxLookup(t, p) };
+
+				SetTriangle(TriIdx, Corners[0], Corners[2], Corners[1]);
+				SetTrianglePolygon(TriIdx, PolyIdx);
+				SetTriangleUVs(TriIdx, UVCorners[0], UVCorners[2], UVCorners[1]);
+				SetTriangleNormals(TriIdx, Corners[0], Corners[2], Corners[1]);
+				TriIdx++;
+				PolyIdx++;
+			}
+		}
+
+		return *this;
+	}
+
+
+};
\ No newline at end of file
diff --git a/Engine/Plugins/Experimental/GeometryProcessing/Source/GeometricObjects/Public/Generators/SweepGenerator.h b/Engine/Plugins/Experimental/GeometryProcessing/Source/GeometricObjects/Public/Generators/SweepGenerator.h
new file mode 100644
--- /dev/null
+++ b/Engine/Plugins/Experimental/GeometryProcessing/Source/GeometricObjects/Public/Generators/SweepGenerator.h
@@ -0,0 +1,409 @@
+// Copyright 1998-2019 Epic Games, Inc. All Rights Reserved.
+
+// Port of geometry3Sharp CylinderGenerator
+
+#pragma once
+
+#include "Math/UnrealMathUtility.h"
+#include "MeshShapeGenerator.h"
+#include "Misc/AssertionMacros.h"
+
+#include "CompGeom/PolygonTriangulation.h"
+#include "FrameTypes.h"
+#include "MatrixTypes.h"
+#include "Polygon2.h"
+#include "Curve/CurveUtil.h"
+
+/**
+ * ECapType indicates the type of cap to use on a sweep
+ */
+enum class /*GEOMETRICOBJECTS_API*/ ECapType
+{
+	None = 0,
+	FlatTriangulation = 1
+	// TODO: Cone, other caps ...
+};
+
+class /*GEOMETRICOBJECTS_API*/ FSweepGeneratorBase : public FMeshShapeGenerator
+{
+public:
+	virtual ~FSweepGeneratorBase()
+	{
+	}
+
+protected:
+	int32 CapVertStart[2], CapNormalStart[2], CapUVStart[2], CapTriangleStart[2], CapPolygonStart[2];
+
+	/**
+	 * Shared logic for creating vertex buffers and triangulations across all sweep primitives
+	 * Note: Does not set vertex positions or normals; a separate call must do that.
+	 */
+	void ConstructMeshTopology(const FPolygon2d& CrossSection,
+							   const TArrayView<const int32>& UVSections,
+							   const TArrayView<const int32>& NormalSections,
+							   int32 NumCrossSections,
+							   const ECapType Caps[2],
+							   FVector2d UVScale, FVector2d UVOffset)
+	{
+		// per cross section
+		int32 XVerts = CrossSection.VertexCount();
+		int32 XNormals = XVerts + NormalSections.Num();
+		int32 XUVs = XVerts + UVSections.Num() + 1;
+
+		int32 NumVerts = XVerts * NumCrossSections;
+		int32 NumNormals = NumCrossSections > 1 ? XNormals * NumCrossSections : 0;
+		int32 NumUVs = NumCrossSections > 1 ? XUVs * NumCrossSections : 0;
+		int32 NumPolygons = (NumCrossSections - 1) * XVerts;
+		int32 NumTriangles = NumPolygons * 2;
+
+		TArray<FIndex3i> OutTriangles;
+
+		for (int32 CapIdx = 0; CapIdx < 2; CapIdx++)
+		{
+			CapVertStart[CapIdx] = NumVerts;
+			CapNormalStart[CapIdx] = NumNormals;
+			CapUVStart[CapIdx] = NumUVs;
+			CapTriangleStart[CapIdx] = NumTriangles;
+			CapPolygonStart[CapIdx] = NumPolygons;
+
+			if (Caps[CapIdx] == ECapType::FlatTriangulation)
+			{
+				NumTriangles += XVerts - 2;
+				NumPolygons++;
+				NumUVs += XVerts;
+				NumNormals += XVerts;
+			}
+
+			// TODO: support more cap type; e.g.:
+			//else if (Caps[CapIdx] == ECapType::FlatMidpointFan)
+			//{
+			//	NumTriangles += XVerts;
+			//	NumPolygons += XVerts;
+			//	NumUVs += XVerts + 1;
+			//	NumNormals += XVerts + 1;
+			//	NumVerts += 1;
+			//}
+			//else if (Caps[CapIdx] == ECapType::Cone)
+			//{
+			//	NumTriangles += XVerts;
+			//	NumPolygons += XVerts;
+			//	NumUVs += XVerts + 1;
+			//	NumNormals += XVerts * 2;
+			//	NumVerts += 1;
+			//}
+		}
+
+		SetBufferSizes(NumVerts, NumTriangles, NumUVs, NumNormals);
+
+		for (int32 CapIdx = 0; CapIdx < 2; CapIdx++)
+		{
+
+			if (Caps[CapIdx] == ECapType::FlatTriangulation)
+			{
+				int32 VertOffset = CapIdx * (XVerts * (NumCrossSections - 1));
+
+				PolygonTriangulation::TriangulateSimplePolygon(CrossSection.GetVertices(), OutTriangles);
+				int32 TriIdx = CapTriangleStart[CapIdx];
+				int32 PolyIdx = CapPolygonStart[CapIdx];
+				for (const FIndex3i& Triangle : OutTriangles)
+				{
+					bool Flipped = CapIdx == 0;
+					Flipped = !Flipped;
+					SetTriangle(TriIdx,
+								Triangle.A + VertOffset, Triangle.B + VertOffset, Triangle.C + VertOffset,
+								Flipped);
+					SetTriangleUVs(TriIdx,
+								   Triangle.A + CapUVStart[CapIdx],
+								   Triangle.B + CapUVStart[CapIdx],
+								   Triangle.C + CapUVStart[CapIdx],
+								   Flipped);
+					SetTriangleNormals(TriIdx,
+									   Triangle.A + CapNormalStart[CapIdx],
+									   Triangle.B + CapNormalStart[CapIdx],
+									   Triangle.C + CapNormalStart[CapIdx],
+									   Flipped);
+					SetTrianglePolygon(TriIdx, PolyIdx);
+					TriIdx++;
+				}
+				float SideScale = 2 * CapIdx - 1;
+				for (int32 Idx = 0; Idx < XVerts; Idx++)
+				{
+					FVector2d CenteredVert = CrossSection.GetVertices()[Idx] * UVScale + UVOffset;
+					SetUV(CapUVStart[CapIdx] + Idx, FVector2f(CenteredVert.X * SideScale, CenteredVert.Y), VertOffset + Idx);
+					SetNormal(CapNormalStart[CapIdx] + Idx, FVector3f::Zero(), VertOffset + Idx);
+				}
+			}
+		}
+
+		// fill in UVs and triangles along length
+		if (NumCrossSections > 1)
+		{
+			int32 UVSection = 0, UVSubIdx = 0;
+
+			int32 NumSections = UVSections.Num();
+			int32 NextDupVertIdx = UVSection < NumSections ? UVSections[UVSection] : -1;
+			for (int32 VertSubIdx = 0; VertSubIdx < XVerts; UVSubIdx++)
+			{
+				float UVX = VertSubIdx / float(XVerts);
+				for (int32 XIdx = 0; XIdx < NumCrossSections; XIdx++)
+				{
+					float UVY = XIdx / float(NumCrossSections - 1);
+					SetUV(XIdx * XUVs + UVSubIdx, FVector2f(1-UVX, 1-UVY), XIdx * XVerts + VertSubIdx);
+				}
+
+				if (VertSubIdx == NextDupVertIdx)
+				{
+					NextDupVertIdx = UVSection < NumSections ? UVSections[UVSection] : -1;
+				}
+				else
+				{
+					for (int32 XIdx = 0; XIdx + 1 < NumCrossSections; XIdx++)
+					{
+						SetTriangleUVs(
+							XVerts * 2 * XIdx + 2 * VertSubIdx,
+							XIdx * XUVs + UVSubIdx,
+							XIdx * XUVs + UVSubIdx + 1,
+							(XIdx + 1) * XUVs + UVSubIdx, true);
+						SetTriangleUVs(
+							XVerts * 2 * XIdx + 2 * VertSubIdx + 1,
+							(XIdx + 1) * XUVs + UVSubIdx + 1,
+							(XIdx + 1) * XUVs + UVSubIdx,
+							XIdx * XUVs + UVSubIdx + 1, true);
+					}
+					VertSubIdx++;
+				}
+			}
+			{
+				// final UV
+				float UVX = 1.0f;
+				int32 VertSubIdx = 0;
+				for (int32 XIdx = 0; XIdx < NumCrossSections; XIdx++)
+				{
+					float UVY = XIdx / float(NumCrossSections - 1);
+					SetUV(XIdx * XUVs + UVSubIdx, FVector2f(1-UVX, 1-UVY), XIdx * XVerts + VertSubIdx);
+				}
+			}
+			NumSections = NormalSections.Num();
+			int32 NormalSection = 0;
+			NextDupVertIdx = NormalSection < NumSections ? NormalSections[NormalSection] : -1;
+			check(NextDupVertIdx < XVerts);
+			for (int32 VertSubIdx = 0, NormalSubIdx = 0; VertSubIdx < XVerts; NormalSubIdx++)
+			{
+				for (int32 XIdx = 0; XIdx < NumCrossSections; XIdx++)
+				{
+					// just set the normal parent; don't compute normal yet
+					SetNormal(XIdx * XNormals + NormalSubIdx, FVector3f(0, 0, 0), XIdx * XVerts + VertSubIdx);
+				}
+
+				if (VertSubIdx == NextDupVertIdx)
+				{
+					NextDupVertIdx = NormalSection < NumSections ? NormalSections[NormalSection] : -1;
+					check(NextDupVertIdx < XVerts);
+				}
+				else
+				{
+					int32 WrappedNextNormalSubIdx = (NormalSubIdx + 1) % XNormals;
+					int32 WrappedNextVertexSubIdx = (VertSubIdx + 1) % XVerts;
+					for (int32 XIdx = 0; XIdx + 1 < NumCrossSections; XIdx++)
+					{
+						int32 T0Idx = XVerts * 2 * XIdx + 2 * VertSubIdx;
+						int32 T1Idx = T0Idx + 1;
+						int32 PIdx = XVerts * XIdx + VertSubIdx;
+						SetTrianglePolygon(T0Idx, PIdx);
+						SetTrianglePolygon(T1Idx, PIdx);
+						SetTriangle(T0Idx,
+									XIdx * XVerts + VertSubIdx,
+									XIdx * XVerts + WrappedNextVertexSubIdx,
+									(XIdx + 1) * XVerts + VertSubIdx, true);
+						SetTriangle(T1Idx,
+									(XIdx + 1) * XVerts + WrappedNextVertexSubIdx,
+									(XIdx + 1) * XVerts + VertSubIdx,
+									XIdx * XVerts + WrappedNextVertexSubIdx, true);
+						SetTriangleNormals(
+							T0Idx,
+							XIdx * XNormals + NormalSubIdx,
+							XIdx * XNormals + WrappedNextNormalSubIdx,
+							(XIdx + 1) * XNormals + NormalSubIdx, true);
+						SetTriangleNormals(
+							T1Idx,
+							(XIdx + 1) * XNormals + WrappedNextNormalSubIdx,
+							(XIdx + 1) * XNormals + NormalSubIdx,
+							XIdx * XNormals + WrappedNextNormalSubIdx, true);
+					}
+					VertSubIdx++;
+				}
+			}
+		}
+	}
+};
+
+/**
+ * Generate a cylinder with optional end caps
+ */
+class /*DYNAMICMESH_API*/ FCylinderGenerator : public FSweepGeneratorBase
+{
+public:
+	float Radius[2] = {1.0f, 1.0f};
+	float Height = 1.0f;
+	int AngleSamples = 16;
+	int LengthSamples = 0;
+	bool bCapped = false;
+	bool bUVScaleMatchSidesAndCaps = true;
+
+public:
+	/** Generate the mesh */
+	virtual FMeshShapeGenerator& Generate() override
+	{
+		FPolygon2d X = FPolygon2d::MakeCircle(1.0, AngleSamples);
+		const TArray<FVector2d>& XVerts = X.GetVertices();
+		ECapType Caps[2] = {ECapType::None, ECapType::None};
+
+		FVector2d NormalSide = (FVector2d(Radius[1], Height) - FVector2d(Radius[0], 0)).Perp().Normalized();
+
+		if (bCapped)
+		{
+			Caps[0] = ECapType::FlatTriangulation;
+			Caps[1] = ECapType::FlatTriangulation;
+		}
+
+		int NumX = LengthSamples + 2;
+		ConstructMeshTopology(X, {}, {}, NumX, Caps, FVector2d(.5, .5), FVector2d(.5, .5));
+
+		// set vertex positions and normals for all cross sections along length
+		double LengthFactor = 1.0 / double(NumX-1);
+		for (int SubIdx = 0; SubIdx < X.VertexCount(); SubIdx++)
+		{
+			for (int XIdx = 0; XIdx < NumX; ++XIdx)
+			{
+				double Along = XIdx * LengthFactor;
+				double AlongRadius = FMath::Lerp(Radius[0], Radius[1], Along);
+				Vertices[SubIdx + XIdx * AngleSamples] =
+					FVector3d(XVerts[SubIdx].X * AlongRadius, XVerts[SubIdx].Y * AlongRadius, Height * Along);
+				Normals[SubIdx + XIdx * AngleSamples] = FVector3f(XVerts[SubIdx].X*NormalSide.X, XVerts[SubIdx].Y*NormalSide.X, NormalSide.Y);
+			}
+		}
+		// if capped, set top/bottom normals
+		if (bCapped)
+		{
+			for (int SubIdx = 0; SubIdx < X.VertexCount(); SubIdx++)
+			{
+				for (int XBotTop = 0; XBotTop < 2; ++XBotTop)
+				{
+					Normals[CapNormalStart[XBotTop] + SubIdx] = FVector3f(0, 0, 2 * XBotTop - 1);
+				}
+			}
+		}
+
+		for (int k = 0; k < Normals.Num(); ++k)
+		{
+			Normals[k].Normalize();
+		}
+
+		if (bUVScaleMatchSidesAndCaps)
+		{
+			float MaxAbsRad = FMathf::Max(FMathf::Abs(Radius[0]), FMathf::Abs(Radius[1]));
+			float AbsHeight = FMathf::Abs(Height);
+			float MaxAbsCircumference = MaxAbsRad * FMathf::TwoPi;
+			
+			// scales to put each differently-scaled UV coordinate into the same space
+			float ThetaScale = MaxAbsCircumference;
+			float HeightScale = AbsHeight;
+			float CapScale = MaxAbsRad*2;
+
+			float MaxScale = FMathf::Max3(ThetaScale, HeightScale, CapScale);
+			ThetaScale /= MaxScale;
+			HeightScale /= MaxScale;
+			CapScale /= MaxScale;
+			for (int UVIdx = 0; UVIdx < CapUVStart[0]; UVIdx++)
+			{
+				UVs[UVIdx].X *= ThetaScale;
+				UVs[UVIdx].Y *= HeightScale;
+			}
+			for (int UVIdx = CapUVStart[0]; UVIdx < UVs.Num(); UVIdx++)
+			{
+				UVs[UVIdx] *= CapScale;
+			}
+		}
+
+		return *this;
+	}
+};
+
+
+/**
+ * Sweep a 2D Profile Polygon along a 3D Path.
+ * 
+ * TODO: 
+ *  - Loop path support
+ *  - Mitering cross sections support?
+ */
+class /*DYNAMICMESH_API*/ FGeneralizedCylinderGenerator : public FSweepGeneratorBase
+{
+public:
+	FPolygon2d CrossSection;
+	TArray<FVector3d> Path;
+
+	FFrame3d InitialFrame;
+
+	bool bCapped = false;
+
+public:
+	/** Generate the mesh */
+	virtual FMeshShapeGenerator& Generate() override
+	{
+		const bool bLoop = false; // TODO: support loops -- note this requires updating FSweepGeneratorBase to set up the mesh topology correctly!
+		
+		const TArray<FVector2d>& XVerts = CrossSection.GetVertices();
+		ECapType Caps[2] = {ECapType::None, ECapType::None};
+
+		if (bCapped)
+		{
+			Caps[0] = ECapType::FlatTriangulation;
+			Caps[1] = ECapType::FlatTriangulation;
+		}
+		int PathNum = Path.Num();
+		ConstructMeshTopology(CrossSection, {}, {}, PathNum, Caps, FVector2d(.5, .5), FVector2d(.5, .5));
+
+		int XNum = CrossSection.VertexCount();
+		TArray<FVector2d> XNormals; XNormals.SetNum(XNum);
+		for (int Idx = 0; Idx < XNum; Idx++)
+		{
+			XNormals[Idx] = CrossSection.GetNormal_FaceAvg(Idx);
+		}
+
+		FFrame3d CrossSectionFrame = InitialFrame;
+		for (int PathIdx = 0; PathIdx < PathNum; ++PathIdx)
+		{
+			FVector3d Tangent = TCurveUtil<double>::Tangent(Path, PathIdx, bLoop);
+			CrossSectionFrame.AlignAxis(2, Tangent);
+			FVector3d C = Path[PathIdx];
+			FVector3d X = CrossSectionFrame.X();
+			FVector3d Y = CrossSectionFrame.Y();
+			for (int SubIdx = 0; SubIdx < XNum; SubIdx++)
+			{
+				FVector2d XP = CrossSection[SubIdx];
+				FVector2d XN = XNormals[SubIdx];
+				Vertices[SubIdx + PathIdx * XNum] = C + X * XP.X + Y * XP.Y;
+				Normals[SubIdx + PathIdx * XNum] = FVector3f(X * XN.X + Y * XN.Y);
+			}
+		}
+		if (bCapped && !bLoop)
+		{
+			for (int CapIdx = 0; CapIdx < 2; CapIdx++)
+			{
+				FVector3f Normal(TCurveUtil<double>::Tangent(Path, CapIdx * (PathNum-1), bLoop) * (CapIdx * 2 - 1));
+				for (int SubIdx = 0; SubIdx < XNum; SubIdx++)
+				{
+					Normals[CapNormalStart[CapIdx] + SubIdx] = Normal;
+				}
+			}
+		}
+
+		for (int k = 0; k < Normals.Num(); ++k)
+		{
+			Normals[k].Normalize();
+		}
+
+		return *this;
+	}
+};
