add spherical box to addprimitivetool

#rnx

diff --git a/Engine/Plugins/Experimental/GeometryProcessing/Source/GeometricObjects/Public/Generators/BoxSphereGenerator.h b/Engine/Plugins/Experimental/GeometryProcessing/Source/GeometricObjects/Public/Generators/BoxSphereGenerator.h
new file mode 100644
--- /dev/null
+++ b/Engine/Plugins/Experimental/GeometryProcessing/Source/GeometricObjects/Public/Generators/BoxSphereGenerator.h
@@ -0,0 +1,70 @@
+// Copyright 1998-2019 Epic Games, Inc. All Rights Reserved.
+
+// Port of geometry3Sharp SphereGenerators
+
+#pragma once
+
+#include "MeshShapeGenerator.h"
+#include "OrientedBoxTypes.h"
+#include "Util/IndexUtil.h"
+
+#include "GridBoxMeshGenerator.h"
+
+/**
+ *  Generate a sphere by pushing a boxes vertices onto a sphere -- often more useful
+ */
+class /*GEOMETRICOBJECTS_API*/ FBoxSphereGenerator : public FGridBoxMeshGenerator
+{
+public:
+
+	/** Sphere radius */
+	float Radius = 1;
+
+	/** Options for methods to project cube vertices to the sphere */
+	enum class ESphereProjectionMethod
+	{
+		NormalizedVector,
+		CubeMapping             // produces more even distribution of quads
+								// see http://catlikecoding.com/unity/tutorials/cube-sphere/
+								// or http://mathproofs.blogspot.ca/2005/07/mapping-cube-to-sphere.html
+	};
+	ESphereProjectionMethod SphereProjectionMethod = ESphereProjectionMethod::CubeMapping;
+
+	FMeshShapeGenerator& Generate() override
+	{
+		FGridBoxMeshGenerator::Generate();
+
+		FVector3d Center = Box.Center();
+		FVector3d AxX = Box.AxisX(), AxY = Box.AxisY(), AxZ = Box.AxisZ();
+		for (int32 VertIdx = 0; VertIdx < Vertices.Num(); VertIdx++)
+		{
+			FVector3d V = Vertices[VertIdx] - Center;
+			if (SphereProjectionMethod == ESphereProjectionMethod::CubeMapping) {
+				double x = V.Dot(AxX) / Box.Extents.X;
+				double y = V.Dot(AxY) / Box.Extents.Y;
+				double z = V.Dot(AxZ) / Box.Extents.Z;
+				double x2 = x * x, y2 = y * y, z2 = z * z;
+				double sx = x * FMath::Sqrt(1.0 - y2*0.5 - z2*0.5 + y2*z2/3.0);
+				double sy = y * FMath::Sqrt(1.0 - x2*0.5 - z2*0.5 + x2*z2/3.0);
+				double sz = z * FMath::Sqrt(1.0 - x2*0.5 - y2*0.5 + x2*y2/3.0);
+				V = sx*AxX + sy*AxY + sz*AxZ;
+			}
+			V.Normalize();
+			Vertices[VertIdx] = V;
+		}
+
+		for (int32 NormalIdx = 0; NormalIdx < Normals.Num(); NormalIdx++)
+		{
+			Normals[NormalIdx] = (FVector3f)Vertices[NormalParentVertex[NormalIdx]];
+		}
+
+		for (int32 VertIdx = 0; VertIdx < Vertices.Num(); VertIdx++)
+		{
+			Vertices[VertIdx] = Center + Radius*Vertices[VertIdx];
+		}
+
+		return *this;
+	}
+
+};
+
diff --git a/Engine/Plugins/Experimental/MeshModelingToolset/Source/MeshModelingTools/Private/AddPrimitiveTool.cpp b/Engine/Plugins/Experimental/MeshModelingToolset/Source/MeshModelingTools/Private/AddPrimitiveTool.cpp
--- a/Engine/Plugins/Experimental/MeshModelingToolset/Source/MeshModelingTools/Private/AddPrimitiveTool.cpp
+++ b/Engine/Plugins/Experimental/MeshModelingToolset/Source/MeshModelingTools/Private/AddPrimitiveTool.cpp
@@ -12,6 +12,7 @@
 #include "Generators/GridBoxMeshGenerator.h"
 #include "Generators/RectangleMeshGenerator.h"
 #include "Generators/SphereGenerator.h"
+#include "Generators/BoxSphereGenerator.h"
 #include "DynamicMeshAttributeSet.h"
 #include "StaticMeshComponentBuilder.h"
 #include "Drawing/MeshDebugDrawing.h"
@@ -91,7 +92,7 @@ namespace
 	  { TEXT("PlaceMode"),     EMakeMeshShapeType::All },
 	  { TEXT("PivotLocation"), EMakeMeshShapeType::All },
 	  { TEXT("Slices"),        EMakeMeshShapeType::Cylinder | EMakeMeshShapeType::Cone | EMakeMeshShapeType::Sphere },
-	  { TEXT("Subdivisions"),  EMakeMeshShapeType::Box | EMakeMeshShapeType::Plane | EMakeMeshShapeType::Cylinder | EMakeMeshShapeType::Cone }
+	  { TEXT("Subdivisions"),  EMakeMeshShapeType::Box | EMakeMeshShapeType::Plane | EMakeMeshShapeType::Cylinder | EMakeMeshShapeType::Cone | EMakeMeshShapeType::SphericalBox }
 	};
 };
 
@@ -272,6 +273,10 @@ void UAddPrimitiveTool::UpdatePreviewMesh()
 		GenerateSphere(&NewMesh);
 		break;
 
+	case EMakeMeshShapeType::SphericalBox:
+		GenerateSphericalBox(&NewMesh);
+		break;
+
 	case EMakeMeshShapeType::Box:
 	default:
 		GenerateBox(&NewMesh);
@@ -394,4 +399,16 @@ void UAddPrimitiveTool::GenerateSphere(FDynamicMesh3* OutMesh)
 }
 
 
+void UAddPrimitiveTool::GenerateSphericalBox(FDynamicMesh3* OutMesh)
+{
+	FBoxSphereGenerator SphereGen;
+	SphereGen.Radius = ShapeSettings->Width * 0.5f;
+	SphereGen.Box = FOrientedBox3d(FVector3d::Zero(), 0.5*FVector3d(ShapeSettings->Width, ShapeSettings->Width, ShapeSettings->Width));
+	int EdgeNum = ShapeSettings->Subdivisions + 3;
+	SphereGen.EdgeVertices = FIndex3i(EdgeNum, EdgeNum, EdgeNum);
+	SphereGen.Generate();
+	OutMesh->Copy(&SphereGen);
+}
+
+
 #undef LOCTEXT_NAMESPACE
diff --git a/Engine/Plugins/Experimental/MeshModelingToolset/Source/MeshModelingTools/Public/AddPrimitiveTool.h b/Engine/Plugins/Experimental/MeshModelingToolset/Source/MeshModelingTools/Public/AddPrimitiveTool.h
--- a/Engine/Plugins/Experimental/MeshModelingToolset/Source/MeshModelingTools/Public/AddPrimitiveTool.h
+++ b/Engine/Plugins/Experimental/MeshModelingToolset/Source/MeshModelingTools/Public/AddPrimitiveTool.h
@@ -43,7 +43,8 @@ enum class EMakeMeshShapeType : uint8
 	Cylinder = 0x02 UMETA(DisplayName = "Cylinder"),
 	Cone     = 0x04 UMETA(DisplayName = "Cone"),
 	Plane    = 0x08 UMETA(DisplayName = "Plane"),
-	Sphere   = 0x10 UMETA(DisplayName = "Sphere")
+	Sphere   = 0x10 UMETA(DisplayName = "Sphere"),
+	SphericalBox = 0x20 UMETA(DisplayName = "Spherical Box")
 };
 ENUM_CLASS_FLAGS(EMakeMeshShapeType);
 
@@ -188,4 +189,5 @@ protected:
 	void GenerateBox(FDynamicMesh3* OutMesh);
 	void GeneratePlane(FDynamicMesh3* OutMesh);
 	void GenerateSphere(FDynamicMesh3* OutMesh);
+	void GenerateSphericalBox(FDynamicMesh3* OutMesh);
 };
