Adding missing include
#rnx

diff --git a/Engine/Plugins/Experimental/GeometryProcessing/Source/GeometricObjects/Public/MeshAdapter.h b/Engine/Plugins/Experimental/GeometryProcessing/Source/GeometricObjects/Public/MeshAdapter.h
--- a/Engine/Plugins/Experimental/GeometryProcessing/Source/GeometricObjects/Public/MeshAdapter.h
+++ b/Engine/Plugins/Experimental/GeometryProcessing/Source/GeometricObjects/Public/MeshAdapter.h
@@ -7,6 +7,7 @@
 
 #include "Math/IntVector.h"
 #include "Math/Vector.h"
+#include "IndexTypes.h"
 
 /**
  * Most generic / lazy example of a triangle mesh adapter; possibly useful for prototyping / building on top of (but slower than making a more specific-case adapter)
