add support for normal discontinuity on generated sweeps, generalize cylinder code + add arrow generator.  add arrows to add primitive tool.

#rnx

diff --git a/Engine/Plugins/Experimental/GeometryProcessing/Source/GeometricObjects/Public/Generators/RectangleMeshGenerator.h b/Engine/Plugins/Experimental/GeometryProcessing/Source/GeometricObjects/Public/Generators/RectangleMeshGenerator.h
--- a/Engine/Plugins/Experimental/GeometryProcessing/Source/GeometricObjects/Public/Generators/RectangleMeshGenerator.h
+++ b/Engine/Plugins/Experimental/GeometryProcessing/Source/GeometricObjects/Public/Generators/RectangleMeshGenerator.h
@@ -29,7 +29,7 @@ public:
 	/** If true (default), UVs are scaled so that there is no stretching. If false, UVs are scaled to fill unit square */
 	bool bScaleUVByAspectRatio = true;
 
-	/** Specifices how 2D indices are mapped to 3D points. Default is (0,1) = (x,y,0). */ 
+	/** Specifies how 2D indices are mapped to 3D points. Default is (0,1) = (x,y,0). */ 
 	FIndex2i IndicesMap;
 
 public:
diff --git a/Engine/Plugins/Experimental/GeometryProcessing/Source/GeometricObjects/Public/Generators/SweepGenerator.h b/Engine/Plugins/Experimental/GeometryProcessing/Source/GeometricObjects/Public/Generators/SweepGenerator.h
--- a/Engine/Plugins/Experimental/GeometryProcessing/Source/GeometricObjects/Public/Generators/SweepGenerator.h
+++ b/Engine/Plugins/Experimental/GeometryProcessing/Source/GeometricObjects/Public/Generators/SweepGenerator.h
@@ -41,6 +41,7 @@ protected:
 	void ConstructMeshTopology(const FPolygon2d& CrossSection,
 							   const TArrayView<const int32>& UVSections,
 							   const TArrayView<const int32>& NormalSections,
+							   const TArrayView<const int32>& SharpNormalsAlongLength,
 							   int32 NumCrossSections,
 							   const ECapType Caps[2],
 							   FVector2d UVScale, FVector2d UVOffset)
@@ -52,6 +53,7 @@ protected:
 
 		int32 NumVerts = XVerts * NumCrossSections;
 		int32 NumNormals = NumCrossSections > 1 ? XNormals * NumCrossSections : 0;
+		NumNormals += XNormals * SharpNormalsAlongLength.Num();
 		int32 NumUVs = NumCrossSections > 1 ? XUVs * NumCrossSections : 0;
 		int32 NumPolygons = (NumCrossSections - 1) * XVerts;
 		int32 NumTriangles = NumPolygons * 2;
@@ -189,10 +191,18 @@ protected:
 			check(NextDupVertIdx < XVerts);
 			for (int32 VertSubIdx = 0, NormalSubIdx = 0; VertSubIdx < XVerts; NormalSubIdx++)
 			{
-				for (int32 XIdx = 0; XIdx < NumCrossSections; XIdx++)
+				int SharpNormalIdx = 0;
+				for (int32 XIdx = 0, NormalXIdx = 0; XIdx < NumCrossSections; XIdx++, NormalXIdx++)
 				{
 					// just set the normal parent; don't compute normal yet
-					SetNormal(XIdx * XNormals + NormalSubIdx, FVector3f(0, 0, 0), XIdx * XVerts + VertSubIdx);
+					SetNormal(NormalXIdx * XNormals + NormalSubIdx, FVector3f(0, 0, 0), XIdx * XVerts + VertSubIdx);
+					// duplicate normals for cross sections that are 'sharp'
+					if (SharpNormalIdx < SharpNormalsAlongLength.Num() && XIdx == SharpNormalsAlongLength[SharpNormalIdx])
+					{
+						NormalXIdx++;
+						SetNormal(NormalXIdx * XNormals + NormalSubIdx, FVector3f(0, 0, 0), XIdx * XVerts + VertSubIdx);
+						SharpNormalIdx++;
+					}
 				}
 
 				if (VertSubIdx == NextDupVertIdx)
@@ -204,7 +214,8 @@ protected:
 				{
 					int32 WrappedNextNormalSubIdx = (NormalSubIdx + 1) % XNormals;
 					int32 WrappedNextVertexSubIdx = (VertSubIdx + 1) % XVerts;
-					for (int32 XIdx = 0; XIdx + 1 < NumCrossSections; XIdx++)
+					SharpNormalIdx = 0;
+					for (int32 XIdx = 0, NXIdx = 0; XIdx + 1 < NumCrossSections; XIdx++, NXIdx++)
 					{
 						int32 T0Idx = XVerts * 2 * XIdx + 2 * VertSubIdx;
 						int32 T1Idx = T0Idx + 1;
@@ -221,14 +232,19 @@ protected:
 									XIdx * XVerts + WrappedNextVertexSubIdx, true);
 						SetTriangleNormals(
 							T0Idx,
-							XIdx * XNormals + NormalSubIdx,
-							XIdx * XNormals + WrappedNextNormalSubIdx,
-							(XIdx + 1) * XNormals + NormalSubIdx, true);
+							NXIdx * XNormals + NormalSubIdx,
+							NXIdx * XNormals + WrappedNextNormalSubIdx,
+							(NXIdx + 1) * XNormals + NormalSubIdx, true);
 						SetTriangleNormals(
 							T1Idx,
-							(XIdx + 1) * XNormals + WrappedNextNormalSubIdx,
-							(XIdx + 1) * XNormals + NormalSubIdx,
-							XIdx * XNormals + WrappedNextNormalSubIdx, true);
+							(NXIdx + 1) * XNormals + WrappedNextNormalSubIdx,
+							(NXIdx + 1) * XNormals + NormalSubIdx,
+							NXIdx * XNormals + WrappedNextNormalSubIdx, true);
+						if (SharpNormalIdx < SharpNormalsAlongLength.Num() && XIdx+1 == SharpNormalsAlongLength[SharpNormalIdx])
+						{
+							NXIdx++;
+							SharpNormalIdx++;
+						}
 					}
 					VertSubIdx++;
 				}
@@ -238,48 +254,106 @@ protected:
 };
 
 /**
- * Generate a cylinder with optional end caps
- */
-class /*DYNAMICMESH_API*/ FCylinderGenerator : public FSweepGeneratorBase
+* Generate a cylinder with optional end caps
+*/
+class /*DYNAMICMESH_API*/ FVerticalCylinderGeneratorBase : public FSweepGeneratorBase
 {
 public:
-	float Radius[2] = {1.0f, 1.0f};
-	float Height = 1.0f;
 	int AngleSamples = 16;
-	int LengthSamples = 0;
 	bool bCapped = false;
 	bool bUVScaleMatchSidesAndCaps = true;
 
-public:
-	/** Generate the mesh */
-	virtual FMeshShapeGenerator& Generate() override
+	static float ComputeSegLengths(const TArrayView<float>& Radii, const TArrayView<float>& Heights, TArray<float>& AlongPercents)
+	{
+		float LenAlong = 0;
+		int32 NumX = Radii.Num();
+		AlongPercents.SetNum(NumX);
+		AlongPercents[0] = 0;
+		for (int XIdx = 0; XIdx+1 < NumX; XIdx++)
+		{
+			float Dist = FVector2d(Radii[XIdx], Heights[XIdx]).Distance(FVector2d(Radii[XIdx + 1], Heights[XIdx + 1]));
+			LenAlong += Dist;
+			AlongPercents[XIdx + 1] = LenAlong;
+		}
+		for (int XIdx = 0; XIdx+1 < NumX; XIdx++)
+		{
+			AlongPercents[XIdx+1] /= LenAlong;
+		}
+		return LenAlong;
+	}
+
+	bool GenerateVerticalCircleSweep(const TArrayView<float>& Radii, const TArrayView<float>& Heights, const TArrayView<int>& SharpNormalsAlongLength)
 	{
 		FPolygon2d X = FPolygon2d::MakeCircle(1.0, AngleSamples);
 		const TArray<FVector2d>& XVerts = X.GetVertices();
 		ECapType Caps[2] = {ECapType::None, ECapType::None};
 
-		FVector2d NormalSide = (FVector2d(Radius[1], Height) - FVector2d(Radius[0], 0)).Perp().Normalized();
-
 		if (bCapped)
 		{
 			Caps[0] = ECapType::FlatTriangulation;
 			Caps[1] = ECapType::FlatTriangulation;
 		}
 
-		int NumX = LengthSamples + 2;
-		ConstructMeshTopology(X, {}, {}, NumX, Caps, FVector2d(.5, .5), FVector2d(.5, .5));
+		int NumX = Radii.Num();
+		if (!ensure(NumX == Heights.Num()))
+		{
+			return false;
+		}
+		// first and last cross sections can't be sharp, so can't have more than NumX-2 sharp normal indices
+		if (!ensure(SharpNormalsAlongLength.Num() + 2 <= NumX))
+		{
+			return false;
+		}
+
+		TArray<float> AlongPercents;
+		float LenAlong = ComputeSegLengths(Radii, Heights, AlongPercents);
+
+		ConstructMeshTopology(X, {}, {}, SharpNormalsAlongLength, NumX, Caps, FVector2d(.5, .5), FVector2d(.5, .5));
+
+		TArray<FVector2d> NormalSides; NormalSides.SetNum(NumX - 1);
+		for (int XIdx = 0; XIdx+1 < NumX; XIdx++)
+		{
+			NormalSides[XIdx] = (FVector2d(Radii[XIdx+1], Heights[XIdx+1]) - FVector2d(Radii[XIdx], Heights[XIdx])).Perp().Normalized();
+		}
+		TArray<FVector2d> SmoothedNormalSides; SmoothedNormalSides.SetNum(NumX);
+		// smooth internal normals
+		SmoothedNormalSides[0] = NormalSides[0];
+		SmoothedNormalSides.Last() = NormalSides.Last();
+		for (int XIdx = 1; XIdx + 1 < NumX; XIdx++)
+		{
+			SmoothedNormalSides[XIdx] = (NormalSides[XIdx] + NormalSides[XIdx - 1]).Normalized();
+		}
+
 
 		// set vertex positions and normals for all cross sections along length
-		double LengthFactor = 1.0 / double(NumX-1);
 		for (int SubIdx = 0; SubIdx < X.VertexCount(); SubIdx++)
 		{
-			for (int XIdx = 0; XIdx < NumX; ++XIdx)
+			int SharpNormalIdx = 0;
+			for (int XIdx = 0, NormalXIdx = 0; XIdx < NumX; ++XIdx, ++NormalXIdx)
 			{
-				double Along = XIdx * LengthFactor;
-				double AlongRadius = FMath::Lerp(Radius[0], Radius[1], Along);
+				double Along = AlongPercents[XIdx];
+				double AlongRadius = Radii[XIdx];
 				Vertices[SubIdx + XIdx * AngleSamples] =
-					FVector3d(XVerts[SubIdx].X * AlongRadius, XVerts[SubIdx].Y * AlongRadius, Height * Along);
-				Normals[SubIdx + XIdx * AngleSamples] = FVector3f(XVerts[SubIdx].X*NormalSide.X, XVerts[SubIdx].Y*NormalSide.X, NormalSide.Y);
+					FVector3d(XVerts[SubIdx].X * AlongRadius, XVerts[SubIdx].Y * AlongRadius, Heights[XIdx]);
+				if (SharpNormalIdx < SharpNormalsAlongLength.Num() && XIdx == SharpNormalsAlongLength[SharpNormalIdx])
+				{
+					// write sharp normals
+					if (ensure(XIdx > 0)) // very first index cannot be sharp
+					{
+						Normals[SubIdx + NormalXIdx * AngleSamples] = FVector3f(XVerts[SubIdx].X*NormalSides[XIdx-1].X, XVerts[SubIdx].Y*NormalSides[XIdx-1].X, NormalSides[XIdx-1].Y);
+					}
+					NormalXIdx++;
+					if (ensure(XIdx + 1 < NumX)) // very last index cannot be sharp
+					{
+						Normals[SubIdx + NormalXIdx * AngleSamples] = FVector3f(XVerts[SubIdx].X*NormalSides[XIdx].X, XVerts[SubIdx].Y*NormalSides[XIdx].X, NormalSides[XIdx].Y);
+					}
+					SharpNormalIdx++;
+				}
+				else
+				{
+					// write smoothed normal
+					Normals[SubIdx + NormalXIdx * AngleSamples] = FVector3f(XVerts[SubIdx].X*SmoothedNormalSides[XIdx].X, XVerts[SubIdx].Y*SmoothedNormalSides[XIdx].X, SmoothedNormalSides[XIdx].Y);
+				}
 			}
 		}
 		// if capped, set top/bottom normals
@@ -294,17 +368,16 @@ public:
 			}
 		}
 
-		for (int k = 0; k < Normals.Num(); ++k)
-		{
-			Normals[k].Normalize();
-		}
-
 		if (bUVScaleMatchSidesAndCaps)
 		{
-			float MaxAbsRad = FMathf::Max(FMathf::Abs(Radius[0]), FMathf::Abs(Radius[1]));
-			float AbsHeight = FMathf::Abs(Height);
+			float MaxAbsRad = FMathf::Abs(Radii[0]);
+			for (int XIdx = 0; XIdx < NumX; XIdx++)
+			{
+				MaxAbsRad = FMathf::Max(FMathf::Abs(Radii[XIdx]), MaxAbsRad);
+			}
+			float AbsHeight = LenAlong;
 			float MaxAbsCircumference = MaxAbsRad * FMathf::TwoPi;
-			
+
 			// scales to put each differently-scaled UV coordinate into the same space
 			float ThetaScale = MaxAbsCircumference;
 			float HeightScale = AbsHeight;
@@ -325,11 +398,103 @@ public:
 			}
 		}
 
+		return true;
+	}
+};
+
+/**
+ * Generate a cylinder with optional end caps
+ */
+class /*DYNAMICMESH_API*/ FCylinderGenerator : public FVerticalCylinderGeneratorBase
+{
+public:
+	float Radius[2] = {1.0f, 1.0f};
+	float Height = 1.0f;
+	int LengthSamples = 0;
+
+public:
+	/** Generate the mesh */
+	virtual FMeshShapeGenerator& Generate() override
+	{
+		TArray<float> Radii, Heights;
+
+		Radii.Add(Radius[0]);
+		Heights.Add(0);
+		for (int ExtraIdx = 0; ExtraIdx < LengthSamples; ExtraIdx++)
+		{
+			float Along = float(ExtraIdx + 1) / float(LengthSamples + 1);
+			Radii.Add(FMath::Lerp(Radius[0], Radius[1], Along));
+			Heights.Add(Height * Along);
+		}
+		Radii.Add(Radius[1]);
+		Heights.Add(Height);
+
+		GenerateVerticalCircleSweep(Radii, Heights, {});
+
+		return *this;
+	}
+};
+
+/**
+* Generate a 3D arrow
+*/
+class /*DYNAMICMESH_API*/ FArrowGenerator : public FVerticalCylinderGeneratorBase
+{
+public:
+	float StickRadius = 0.5f;
+	float StickLength = 1.0f;
+	float HeadBaseRadius = 1.0f;
+	float TipRadius = 0.01f;
+	float HeadLength = 0.5f;
+
+	int AdditionalLengthSamples[3]{ 0,0,0 }; // additional length-wise samples on the three segments (along stick, along arrow base, along arrow cone)
+
+	void DistributeAdditionalLengthSamples(int TargetSamples)
+	{
+		TArray<float> AlongPercents;
+		TArray<float> Radii{ StickRadius, StickRadius, HeadBaseRadius, TipRadius };
+		TArray<float> Heights{ 0, StickLength, StickLength, StickLength + HeadLength };
+		float LenAlong = ComputeSegLengths(Radii, Heights, AlongPercents);
+		for (int Idx = 0; Idx < 3; Idx++)
+		{
+			AdditionalLengthSamples[Idx] = (int)(.5f+AlongPercents[Idx + 1] * TargetSamples);
+		}
+	}
+
+public:
+	/** Generate the mesh */
+	virtual FMeshShapeGenerator& Generate() override
+	{
+		TArray<float> Radii, Heights;
+		TArray<int> SharpNormalsAlongLength;
+		float SrcRadii[] { StickRadius, StickRadius, HeadBaseRadius, TipRadius };
+		float SrcHeights[] { 0, StickLength, StickLength, StickLength + HeadLength };
+
+		int SegIdx = 0;
+		for (; SegIdx < 3; SegIdx++)
+		{
+			Radii.Add(SrcRadii[SegIdx]);
+			Heights.Add(SrcHeights[SegIdx]);
+			for (int ExtraIdx = 0; ExtraIdx < AdditionalLengthSamples[SegIdx]; ExtraIdx++)
+			{
+				float Along = float(ExtraIdx + 1) / float(AdditionalLengthSamples[SegIdx] + 1);
+				Radii.Add(FMath::Lerp(SrcRadii[SegIdx], SrcRadii[SegIdx + 1], Along));
+				Heights.Add(FMath::Lerp(SrcHeights[SegIdx], SrcHeights[SegIdx + 1], Along));
+			}
+		}
+		SharpNormalsAlongLength.Add(1 + AdditionalLengthSamples[0]);
+		SharpNormalsAlongLength.Add(SharpNormalsAlongLength.Last() + 1 + AdditionalLengthSamples[1]);
+		Radii.Add(SrcRadii[SegIdx]);
+		Heights.Add(SrcHeights[SegIdx]);
+		
+		GenerateVerticalCircleSweep(Radii, Heights, SharpNormalsAlongLength);
+
 		return *this;
 	}
 };
 
 
+
 /**
  * Sweep a 2D Profile Polygon along a 3D Path.
  * 
@@ -362,7 +527,7 @@ public:
 			Caps[1] = ECapType::FlatTriangulation;
 		}
 		int PathNum = Path.Num();
-		ConstructMeshTopology(CrossSection, {}, {}, PathNum, Caps, FVector2d(.5, .5), FVector2d(.5, .5));
+		ConstructMeshTopology(CrossSection, {}, {}, {}, PathNum, Caps, FVector2d(.5, .5), FVector2d(.5, .5));
 
 		int XNum = CrossSection.VertexCount();
 		TArray<FVector2d> XNormals; XNormals.SetNum(XNum);
diff --git a/Engine/Plugins/Experimental/MeshModelingToolset/Source/MeshModelingTools/Private/AddPrimitiveTool.cpp b/Engine/Plugins/Experimental/MeshModelingToolset/Source/MeshModelingTools/Private/AddPrimitiveTool.cpp
--- a/Engine/Plugins/Experimental/MeshModelingToolset/Source/MeshModelingTools/Private/AddPrimitiveTool.cpp
+++ b/Engine/Plugins/Experimental/MeshModelingToolset/Source/MeshModelingTools/Private/AddPrimitiveTool.cpp
@@ -87,12 +87,13 @@ namespace
 	{ 
 	  { TEXT("Shape"),         EMakeMeshShapeType::All },
 	  { TEXT("Width"),         EMakeMeshShapeType::All },
-	  { TEXT("Height"),        EMakeMeshShapeType::Box | EMakeMeshShapeType::Cylinder | EMakeMeshShapeType::Cone | EMakeMeshShapeType::Plane },
+	  { TEXT("Height"),        EMakeMeshShapeType::Box | EMakeMeshShapeType::Cylinder | EMakeMeshShapeType::Cone | EMakeMeshShapeType::Arrow | EMakeMeshShapeType::Plane },
 	  { TEXT("Rotation"),      EMakeMeshShapeType::All },
 	  { TEXT("PlaceMode"),     EMakeMeshShapeType::All },
 	  { TEXT("PivotLocation"), EMakeMeshShapeType::All },
-	  { TEXT("Slices"),        EMakeMeshShapeType::Cylinder | EMakeMeshShapeType::Cone | EMakeMeshShapeType::Sphere },
-	  { TEXT("Subdivisions"),  EMakeMeshShapeType::Box | EMakeMeshShapeType::Plane | EMakeMeshShapeType::Cylinder | EMakeMeshShapeType::Cone | EMakeMeshShapeType::SphericalBox }
+	  { TEXT("Slices"),        EMakeMeshShapeType::Cylinder | EMakeMeshShapeType::Cone | EMakeMeshShapeType::Arrow | EMakeMeshShapeType::Sphere },
+	  { TEXT("Subdivisions"),  EMakeMeshShapeType::Box | EMakeMeshShapeType::Plane | EMakeMeshShapeType::Cylinder |
+							   EMakeMeshShapeType::Cone | EMakeMeshShapeType::Arrow | EMakeMeshShapeType::SphericalBox }
 	};
 };
 
@@ -269,6 +270,10 @@ void UAddPrimitiveTool::UpdatePreviewMesh()
 		GenerateCone(&NewMesh);
 		break;
 
+	case EMakeMeshShapeType::Arrow:
+		GenerateArrow(&NewMesh);
+		break;
+
 	case EMakeMeshShapeType::Sphere:
 		GenerateSphere(&NewMesh);
 		break;
@@ -388,6 +393,22 @@ void UAddPrimitiveTool::GenerateCone(FDynamicMesh3* OutMesh)
 }
 
 
+void UAddPrimitiveTool::GenerateArrow(FDynamicMesh3* OutMesh)
+{
+	FArrowGenerator ArrowGen;
+	ArrowGen.StickRadius = ShapeSettings->Width * .25f;
+	ArrowGen.StickLength = ShapeSettings->Height * .25f;
+	ArrowGen.HeadBaseRadius = ShapeSettings->Width * .5f;
+	ArrowGen.TipRadius = .01f;
+	ArrowGen.HeadLength = ShapeSettings->Height * .75f;
+	ArrowGen.AngleSamples = ShapeSettings->Slices;
+	ArrowGen.bCapped = true;
+	ArrowGen.DistributeAdditionalLengthSamples(ShapeSettings->Subdivisions);
+	ArrowGen.Generate();
+	OutMesh->Copy(&ArrowGen);
+}
+
+
 void UAddPrimitiveTool::GenerateSphere(FDynamicMesh3* OutMesh)
 {
 	FSphereGenerator SphereGen;
diff --git a/Engine/Plugins/Experimental/MeshModelingToolset/Source/MeshModelingTools/Public/AddPrimitiveTool.h b/Engine/Plugins/Experimental/MeshModelingToolset/Source/MeshModelingTools/Public/AddPrimitiveTool.h
--- a/Engine/Plugins/Experimental/MeshModelingToolset/Source/MeshModelingTools/Public/AddPrimitiveTool.h
+++ b/Engine/Plugins/Experimental/MeshModelingToolset/Source/MeshModelingTools/Public/AddPrimitiveTool.h
@@ -42,9 +42,11 @@ enum class EMakeMeshShapeType : uint8
 	Box      = 0x01 UMETA(DisplayName = "Box"),
 	Cylinder = 0x02 UMETA(DisplayName = "Cylinder"),
 	Cone     = 0x04 UMETA(DisplayName = "Cone"),
-	Plane    = 0x08 UMETA(DisplayName = "Plane"),
-	Sphere   = 0x10 UMETA(DisplayName = "Sphere"),
-	SphericalBox = 0x20 UMETA(DisplayName = "Spherical Box")
+	Arrow	 = 0x08 UMETA(DisplayName = "Arrow"),
+	Plane    = 0x10 UMETA(DisplayName = "Plane"),
+	Sphere   = 0x20 UMETA(DisplayName = "Sphere"),
+	SphericalBox = 0x40 UMETA(DisplayName = "Spherical Box")
+	
 };
 ENUM_CLASS_FLAGS(EMakeMeshShapeType);
 
@@ -190,4 +192,5 @@ protected:
 	void GeneratePlane(FDynamicMesh3* OutMesh);
 	void GenerateSphere(FDynamicMesh3* OutMesh);
 	void GenerateSphericalBox(FDynamicMesh3* OutMesh);
+	void GenerateArrow(FDynamicMesh3* OutMesh);
 };
