add rounded rectangles to add primitive tool

#rnx

diff --git a/Engine/Plugins/Experimental/GeometryProcessing/Source/GeometricObjects/Private/Generators/RectangleMeshGenerator.cpp b/Engine/Plugins/Experimental/GeometryProcessing/Source/GeometricObjects/Private/Generators/RectangleMeshGenerator.cpp
--- a/Engine/Plugins/Experimental/GeometryProcessing/Source/GeometricObjects/Private/Generators/RectangleMeshGenerator.cpp
+++ b/Engine/Plugins/Experimental/GeometryProcessing/Source/GeometricObjects/Private/Generators/RectangleMeshGenerator.cpp
@@ -100,5 +100,197 @@ FMeshShapeGenerator& FRectangleMeshGenerator::Generate()
 		}
 	}
 
+	return *this;
+}
+
+
+
+FRoundedRectangleMeshGenerator::FRoundedRectangleMeshGenerator()
+{
+	Radius = 1.0;
+	AngleSamples = 16;
+	SharpCorners = ERoundedRectangleCorner::None;
+}
+
+
+
+FMeshShapeGenerator& FRoundedRectangleMeshGenerator::Generate()
+{
+	check(IndicesMap.A >= 0 && IndicesMap.A <= 2);
+	check(IndicesMap.B >= 0 && IndicesMap.B <= 2);
+
+	int WidthNV = (WidthVertexCount > 1) ? WidthVertexCount : 2;
+	int HeightNV = (HeightVertexCount > 1) ? HeightVertexCount : 2;
+	int RoundNV = (AngleSamples > 0) ? AngleSamples : 1;
+	WidthNV += 2;
+	HeightNV += 2;
+
+	int NumRound = 4-NumSharpCorners(SharpCorners);
+
+	int TotalNumVertices = WidthNV * HeightNV + NumRound * (RoundNV-1);
+	int TotalNumTriangles = 2 * (WidthNV - 1) * (HeightNV - 1) + NumRound * (RoundNV - 1);
+	SetBufferSizes(TotalNumVertices, TotalNumTriangles, TotalNumVertices, TotalNumVertices);
+
+	// make UVs and normals self-parented, and assign all normals to constant
+	for (int VertIdx = 0; VertIdx < TotalNumVertices; VertIdx++)
+	{
+		NormalParentVertex[VertIdx] = VertIdx;
+		UVParentVertex[VertIdx] = VertIdx;
+		Normals[VertIdx] = Normal;
+	}
+
+	float TotWidth = Radius * 2 + Width;
+	float TotHeight = Radius * 2 + Height;
+
+	// corner vertices
+	FVector3d v00 = MakeVertex(-TotWidth / 2.0f, -TotHeight / 2.0f);
+	FVector3d v01 = MakeVertex(TotWidth / 2.0f, -TotHeight / 2.0f);
+	FVector3d v11 = MakeVertex(TotWidth / 2.0f, TotHeight / 2.0f);
+	FVector3d v10 = MakeVertex(-TotWidth / 2.0f, TotHeight / 2.0f);
+
+	// corner UVs
+	float uvleft = 0.0f, uvright = 1.0f, uvbottom = 0.0f, uvtop = 1.0f;
+	if (bScaleUVByAspectRatio && TotWidth != TotHeight)
+	{
+		if (TotWidth > TotHeight)
+		{
+			uvtop = TotHeight / TotWidth;
+		}
+		else
+		{
+			uvright = TotWidth / TotHeight;
+		}
+	}
+
+	FVector2f uv00 = FVector2f(uvleft, uvbottom);
+	FVector2f uv01 = FVector2f(uvright, uvbottom);
+	FVector2f uv11 = FVector2f(uvright, uvtop);
+	FVector2f uv10 = FVector2f(uvleft, uvtop);
+
+
+	int vi = 0;
+	int ti = 0;
+
+	TArray<double> XFrac, YFrac; XFrac.SetNum(WidthNV); YFrac.SetNum(HeightNV);
+	XFrac[0] = 0; XFrac.Last() = 1;
+	YFrac[0] = 0; YFrac.Last() = 1;
+	for (int XIdx = 1; XIdx + 1 < XFrac.Num(); XIdx++)
+	{
+		XFrac[XIdx] = (Radius + Width * (XIdx - 1) / (WidthNV - 3)) / TotWidth;
+	}
+	for (int YIdx = 1; YIdx + 1 < YFrac.Num(); YIdx++)
+	{
+		YFrac[YIdx] = (Radius + Height * (YIdx - 1) / (HeightNV - 3)) / TotHeight;
+	}
+
+	// add vertex rows
+	int start_vi = vi;
+	for (int yi = 0; yi < HeightNV; ++yi)
+	{
+		double ty = YFrac[yi];
+		for (int xi = 0; xi < WidthNV; ++xi)
+		{
+			double tx = XFrac[xi];
+			UVs[vi] = BilinearInterp(uv00, uv01, uv11, uv10, (float)tx, (float)ty);
+			Vertices[vi++] = BilinearInterp(v00, v01, v11, v10, tx, ty);
+		}
+	}
+
+	// add triangulated quads
+	int PolyIndex = 0;
+	for (int y0 = 0; y0 < HeightNV - 1; ++y0)
+	{
+		bool bIsYEdge = y0 == 0 || y0 == HeightNV - 2;
+		for (int x0 = 0; x0 < WidthNV - 1; ++x0)
+		{
+			if (bIsYEdge)
+			{
+				bool bIsXEdge = x0 == 0 || x0 == WidthNV - 2;
+				if (bIsXEdge && !SideInCorners(x0 == 0 ? 0 : 1, y0 == 0 ? 0 : 1, SharpCorners))
+				{
+					// rounded corner; skip triangulation; we'll get it in a second pass
+					continue;
+				}
+			}
+			int i00 = start_vi + y0 * WidthNV + x0;
+			int i10 = start_vi + (y0 + 1)*WidthNV + x0;
+			int i01 = i00 + 1, i11 = i10 + 1;
+
+			SetTriangle(ti, i00, i11, i01);
+			SetTrianglePolygon(ti, PolyIndex);
+			SetTriangleUVs(ti, i00, i11, i01);
+			SetTriangleNormals(ti, i00, i11, i01);
+
+			ti++;
+
+			SetTriangle(ti, i00, i10, i11);
+			SetTrianglePolygon(ti, PolyIndex);
+			SetTriangleUVs(ti, i00, i10, i11);
+			SetTriangleNormals(ti, i00, i10, i11);
+
+			ti++;
+			PolyIndex++;
+		}
+	}
+
+	for (int SideX = 0; SideX < 2; SideX++)
+	{
+		for (int SideY = 0; SideY < 2; SideY++)
+		{
+			if (!SideInCorners(SideX, SideY, SharpCorners))
+			{
+				// rounded corner
+				int CornerY = SideY * (HeightNV - 1);
+				int CornerX = SideX * (WidthNV - 1);
+				int InCornerY = SideY ? HeightNV - 2 : 1;
+				int InCornerX = SideX ? WidthNV - 2 : 1;
+
+				// vertex to re-purpose
+				int UseVIdx = start_vi + CornerY * WidthNV + CornerX;
+
+				int VCenterIdx = start_vi + InCornerY * WidthNV + InCornerX;
+				int OffXIdx = start_vi + InCornerY * WidthNV + CornerX;
+				int OffYIdx = start_vi + CornerY * WidthNV + InCornerX;
+				int ActingCosIdx = SideY == SideX ? OffYIdx : OffXIdx;
+				int ActingSinIdx = SideY == SideX ? OffXIdx : OffYIdx;
+				FVector3d CenterV = Vertices[VCenterIdx];
+				FVector3d CosV = Vertices[ActingCosIdx] - CenterV;
+				FVector3d SinV = Vertices[ActingSinIdx] - CenterV;
+				FVector2f CenterUV = UVs[VCenterIdx];
+				FVector2f CosUV = UVs[ActingCosIdx] - CenterUV;
+				FVector2f SinUV = UVs[ActingSinIdx] - CenterUV;
+				
+				int LastUsedIdx = ActingCosIdx;
+
+				for (int AngleSampleIdx = 1; AngleSampleIdx < RoundNV+1; AngleSampleIdx++)
+				{
+					float Angle = FMathf::HalfPi * AngleSampleIdx / float(RoundNV+1);
+					float CosAngle = FMath::Cos(Angle);
+					float SinAngle = FMath::Sin(Angle);
+					Vertices[UseVIdx] = CenterV + CosAngle * CosV + SinAngle * SinV;
+					UVs[UseVIdx] = CenterUV + CosAngle * CosUV + SinAngle * SinUV;
+					SetTriangle(ti, VCenterIdx, LastUsedIdx, UseVIdx);
+					SetTrianglePolygon(ti, PolyIndex);
+					SetTriangleUVs(ti, VCenterIdx, LastUsedIdx, UseVIdx);
+					SetTriangleNormals(ti, VCenterIdx, LastUsedIdx, UseVIdx);
+
+					LastUsedIdx = UseVIdx;
+					UseVIdx = vi++;
+					ti++;
+				}
+				vi--;
+
+				// add last triangle in fan
+				SetTriangle(ti, VCenterIdx, LastUsedIdx, ActingSinIdx);
+				SetTrianglePolygon(ti, PolyIndex);
+				SetTriangleUVs(ti, VCenterIdx, LastUsedIdx, ActingSinIdx);
+				SetTriangleNormals(ti, VCenterIdx, LastUsedIdx, ActingSinIdx);
+				ti++;
+
+				PolyIndex++;
+			}
+		}
+	}
+
 	return *this;
 }
\ No newline at end of file
diff --git a/Engine/Plugins/Experimental/GeometryProcessing/Source/GeometricObjects/Public/Generators/RectangleMeshGenerator.h b/Engine/Plugins/Experimental/GeometryProcessing/Source/GeometricObjects/Public/Generators/RectangleMeshGenerator.h
--- a/Engine/Plugins/Experimental/GeometryProcessing/Source/GeometricObjects/Public/Generators/RectangleMeshGenerator.h
+++ b/Engine/Plugins/Experimental/GeometryProcessing/Source/GeometricObjects/Public/Generators/RectangleMeshGenerator.h
@@ -3,6 +3,7 @@
 #pragma once
 
 #include "MeshShapeGenerator.h"
+#include "Misc/EnumClassFlags.h"
 
 /**
  * Generate planar rectangular mesh with variable number of subdivisions along width and height.
@@ -47,6 +48,58 @@ public:
 		return Origin + v;
 	}
 
+};
+
+enum class ERoundedRectangleCorner : uint8
+{
+	None = 0,
+	BottomLeft = 1,
+	BottomRight = 2,
+	TopLeft = 4,
+	TopRight = 8,
+	All = 15
+};
+ENUM_CLASS_FLAGS(ERoundedRectangleCorner);
+
+/**
+* Adds rounded corners to the rectangle mesh
+*/
+class GEOMETRICOBJECTS_API FRoundedRectangleMeshGenerator : public FRectangleMeshGenerator
+{
+public:
+	/** Radius of rounded corners */
+	double Radius;
+
+	/** Number of samples to put per rounded corner */
+	int AngleSamples;
+
+
+	ERoundedRectangleCorner SharpCorners;
+
+	static inline bool SideInCorners(int SideX, int SideY, ERoundedRectangleCorner Corners)
+	{
+		int SideInCorners = int(Corners)&((SideY ? 4 : 1) * (SideX ? 2 : 1));
+		return SideInCorners != 0;
+	}
+
+	static inline int NumSharpCorners(ERoundedRectangleCorner SharpFlags)
+	{
+		int NumSharp = 0;
+		int FlagsValue = (int)SharpFlags;
+		while (FlagsValue)
+		{
+			FlagsValue &= (FlagsValue - 1);
+			NumSharp++;
+		}
+		return NumSharp;
+	}
+
+public:
+	FRoundedRectangleMeshGenerator();
+
+	/** Generate the mesh */
+	virtual FMeshShapeGenerator& Generate() override;
+
 
 
-};
\ No newline at end of file
+};
diff --git a/Engine/Plugins/Experimental/MeshModelingToolset/Source/MeshModelingTools/Private/AddPrimitiveTool.cpp b/Engine/Plugins/Experimental/MeshModelingToolset/Source/MeshModelingTools/Private/AddPrimitiveTool.cpp
--- a/Engine/Plugins/Experimental/MeshModelingToolset/Source/MeshModelingTools/Private/AddPrimitiveTool.cpp
+++ b/Engine/Plugins/Experimental/MeshModelingToolset/Source/MeshModelingTools/Private/AddPrimitiveTool.cpp
@@ -87,12 +87,12 @@ namespace
 	{ 
 	  { TEXT("Shape"),         EMakeMeshShapeType::All },
 	  { TEXT("Width"),         EMakeMeshShapeType::All },
-	  { TEXT("Height"),        EMakeMeshShapeType::Box | EMakeMeshShapeType::Cylinder | EMakeMeshShapeType::Cone | EMakeMeshShapeType::Arrow | EMakeMeshShapeType::Plane },
+	  { TEXT("Height"),        EMakeMeshShapeType::Box | EMakeMeshShapeType::Cylinder | EMakeMeshShapeType::Cone | EMakeMeshShapeType::Arrow | EMakeMeshShapeType::Rectangle | EMakeMeshShapeType::RoundedRectangle },
 	  { TEXT("Rotation"),      EMakeMeshShapeType::All },
 	  { TEXT("PlaceMode"),     EMakeMeshShapeType::All },
 	  { TEXT("PivotLocation"), EMakeMeshShapeType::All },
-	  { TEXT("Slices"),        EMakeMeshShapeType::Cylinder | EMakeMeshShapeType::Cone | EMakeMeshShapeType::Arrow | EMakeMeshShapeType::Sphere },
-	  { TEXT("Subdivisions"),  EMakeMeshShapeType::Box | EMakeMeshShapeType::Plane | EMakeMeshShapeType::Cylinder |
+	  { TEXT("Slices"),        EMakeMeshShapeType::Cylinder | EMakeMeshShapeType::Cone | EMakeMeshShapeType::Arrow | EMakeMeshShapeType::RoundedRectangle | EMakeMeshShapeType::Sphere },
+	  { TEXT("Subdivisions"),  EMakeMeshShapeType::Box | EMakeMeshShapeType::Rectangle | EMakeMeshShapeType::RoundedRectangle | EMakeMeshShapeType::Cylinder |
 							   EMakeMeshShapeType::Cone | EMakeMeshShapeType::Arrow | EMakeMeshShapeType::SphericalBox }
 	};
 };
@@ -258,8 +258,12 @@ void UAddPrimitiveTool::UpdatePreviewMesh()
 	FDynamicMesh3 NewMesh;
 	switch (ShapeSettings->Shape)
 	{
-	case EMakeMeshShapeType::Plane:
-		GeneratePlane(&NewMesh);
+	case EMakeMeshShapeType::Rectangle:
+		GenerateRectangle(&NewMesh);
+		break;
+
+	case EMakeMeshShapeType::RoundedRectangle:
+		GenerateRoundedRectangle(&NewMesh);
 		break;
 
 	case EMakeMeshShapeType::Cylinder:
@@ -353,7 +357,7 @@ void UAddPrimitiveTool::GenerateBox(FDynamicMesh3* OutMesh)
 }
 
 
-void UAddPrimitiveTool::GeneratePlane(FDynamicMesh3* OutMesh)
+void UAddPrimitiveTool::GenerateRectangle(FDynamicMesh3* OutMesh)
 {
 	FRectangleMeshGenerator RectGen;
 	RectGen.Width = ShapeSettings->Width;
@@ -363,6 +367,18 @@ void UAddPrimitiveTool::GeneratePlane(FDynamicMesh3* OutMesh)
 	OutMesh->Copy(&RectGen);
 }
 
+void UAddPrimitiveTool::GenerateRoundedRectangle(FDynamicMesh3* OutMesh)
+{
+	FRoundedRectangleMeshGenerator RectGen;
+	RectGen.Width = ShapeSettings->Width;
+	RectGen.Height = ShapeSettings->Height;
+	RectGen.WidthVertexCount = RectGen.HeightVertexCount = ShapeSettings->Subdivisions + 2;
+	RectGen.Radius = ShapeSettings->Width * .5; // TODO replace with a radius parameter
+	RectGen.AngleSamples = ShapeSettings->Slices;
+	RectGen.Generate();
+	OutMesh->Copy(&RectGen);
+}
+
 
 void UAddPrimitiveTool::GenerateCylinder(FDynamicMesh3* OutMesh)
 {
diff --git a/Engine/Plugins/Experimental/MeshModelingToolset/Source/MeshModelingTools/Public/AddPrimitiveTool.h b/Engine/Plugins/Experimental/MeshModelingToolset/Source/MeshModelingTools/Public/AddPrimitiveTool.h
--- a/Engine/Plugins/Experimental/MeshModelingToolset/Source/MeshModelingTools/Public/AddPrimitiveTool.h
+++ b/Engine/Plugins/Experimental/MeshModelingToolset/Source/MeshModelingTools/Public/AddPrimitiveTool.h
@@ -37,15 +37,16 @@ public:
 UENUM()
 enum class EMakeMeshShapeType : uint8
 {
-	None     = 0x00 UMETA(DisplayName = "None", Hidden),
-	All      = 0xff UMETA(DisplayName = "All", Hidden), 
-	Box      = 0x01 UMETA(DisplayName = "Box"),
-	Cylinder = 0x02 UMETA(DisplayName = "Cylinder"),
-	Cone     = 0x04 UMETA(DisplayName = "Cone"),
-	Arrow	 = 0x08 UMETA(DisplayName = "Arrow"),
-	Plane    = 0x10 UMETA(DisplayName = "Plane"),
-	Sphere   = 0x20 UMETA(DisplayName = "Sphere"),
-	SphericalBox = 0x40 UMETA(DisplayName = "Spherical Box")
+	None			 = 0x00 UMETA(DisplayName = "None", Hidden),
+	All				 = 0xff UMETA(DisplayName = "All", Hidden), 
+	Box				 = 0x01 UMETA(DisplayName = "Box"),
+	Cylinder		 = 0x02 UMETA(DisplayName = "Cylinder"),
+	Cone			 = 0x04 UMETA(DisplayName = "Cone"),
+	Arrow			 = 0x08 UMETA(DisplayName = "Arrow"),
+	Rectangle		 = 0x10 UMETA(DisplayName = "Rectangle"),
+	RoundedRectangle = 0x20 UMETA(DisplayName = "Rounded Rectangle"),
+	Sphere			 = 0x40 UMETA(DisplayName = "Sphere"),
+	SphericalBox	 = 0x80 UMETA(DisplayName = "Spherical Box")
 	
 };
 ENUM_CLASS_FLAGS(EMakeMeshShapeType);
@@ -189,7 +190,8 @@ protected:
 	void GenerateCylinder(FDynamicMesh3* OutMesh);
 	void GenerateCone(FDynamicMesh3* OutMesh);
 	void GenerateBox(FDynamicMesh3* OutMesh);
-	void GeneratePlane(FDynamicMesh3* OutMesh);
+	void GenerateRectangle(FDynamicMesh3* OutMesh);
+	void GenerateRoundedRectangle(FDynamicMesh3* OutMesh);
 	void GenerateSphere(FDynamicMesh3* OutMesh);
 	void GenerateSphericalBox(FDynamicMesh3* OutMesh);
 	void GenerateArrow(FDynamicMesh3* OutMesh);
