add disc and punctured disc generators + corresponding options in add primitive tool

#rnx

diff --git a/Engine/Plugins/Experimental/GeometryProcessing/Source/GeometricObjects/Private/Generators/DiscMeshGenerator.cpp b/Engine/Plugins/Experimental/GeometryProcessing/Source/GeometricObjects/Private/Generators/DiscMeshGenerator.cpp
new file mode 100644
--- /dev/null
+++ b/Engine/Plugins/Experimental/GeometryProcessing/Source/GeometricObjects/Private/Generators/DiscMeshGenerator.cpp
@@ -0,0 +1,182 @@
+// Copyright 1998-2019 Epic Games, Inc. All Rights Reserved.
+
+#include "Generators/DiscMeshGenerator.h"
+
+FDiscMeshGenerator::FDiscMeshGenerator()
+{
+	Normal = FVector3f::UnitZ();
+	IndicesMap = FIndex2i(0, 1);
+	AngleSamples = 10;
+	RadialSamples = 2;
+	StartAngle = 0;
+	EndAngle = 360;
+	Radius = 1;
+}
+
+
+FPuncturedDiscMeshGenerator::FPuncturedDiscMeshGenerator()
+{
+	HoleRadius = .5;
+}
+
+
+
+
+FMeshShapeGenerator& FDiscMeshGenerator::Generate()
+{
+	int AngleNV = AngleSamples > 2 ? AngleSamples : 2;
+	int RadialNV = RadialSamples > 1 ? RadialSamples : 1;
+	int NumVertices = AngleNV * RadialNV + 1;
+	
+	bool bFullDisc = (EndAngle - StartAngle) >= 359.999;
+	int AngleNT = bFullDisc ? AngleNV : AngleNV - 1;
+
+	int NumTriangles = AngleNT * (2 * RadialNV - 1);
+	SetBufferSizes(NumVertices, NumTriangles, NumVertices, NumVertices);
+
+	Vertices[0] = MakeVertex(0, 0); // center
+	UVs[0] = FVector2f(.5f, .5f);
+	double IdxToRadiansOffset = StartAngle * FMathd::DegToRad;
+	double IdxToRadiansScale = (EndAngle - StartAngle) * FMathd::DegToRad / double(AngleNV - (bFullDisc ? 0 : 1));
+	double IdxToRadiusScale = Radius / double(RadialNV);
+	double IdxToUVRadiusScale = .5 / double(RadialNV);
+	for (int AngleIdx = 0; AngleIdx < AngleNV; AngleIdx++)
+	{
+		double Angle = IdxToRadiansOffset + IdxToRadiansScale * AngleIdx;
+		double CosA = FMathd::Cos(Angle);
+		double SinA = FMathd::Sin(Angle);
+		for (int RadIdx = 1; RadIdx <= RadialNV; RadIdx++)
+		{
+			int VertIdx = AngleIdx + 1 + (RadIdx-1) * AngleNV;
+			double R = RadIdx * IdxToRadiusScale;
+			Vertices[VertIdx] = MakeVertex(R*CosA, R*SinA);
+			UVs[VertIdx] = FVector2f(.5f + CosA * RadIdx * IdxToUVRadiusScale, .5f + SinA * RadIdx * IdxToUVRadiusScale);
+		}
+	}
+
+	for (int Idx = 0; Idx < Vertices.Num(); Idx++)
+	{
+		UVParentVertex[Idx] = Idx;
+		NormalParentVertex[Idx] = Idx;
+		Normals[Idx] = Normal;
+	}
+
+	int TriIdx = 0, PolyIdx = 0;
+	for (int AngleIdx = 0; AngleIdx+1 < AngleNV; AngleIdx++)
+	{
+		SetTriangleWithMatchedUVNormal(TriIdx, 0, AngleIdx + 2, AngleIdx + 1);
+		SetTrianglePolygon(TriIdx, PolyIdx++);
+		TriIdx++;
+	}
+	if (bFullDisc)
+	{
+		SetTriangleWithMatchedUVNormal(TriIdx, 0, 1, AngleNV);
+		SetTrianglePolygon(TriIdx, PolyIdx++);
+		TriIdx++;
+	}
+
+	for (int RadiusIdx = 0; RadiusIdx+1 < RadialNV; RadiusIdx++)
+	{
+		int Inner = 1 + RadiusIdx * AngleNV;
+		int Outer = Inner + AngleNV;
+		for (int AngleIdx = 0; AngleIdx + 1 < AngleNV; AngleIdx++)
+		{
+			SetTriangleWithMatchedUVNormal(TriIdx, Inner + AngleIdx, Outer + AngleIdx + 1, Outer + AngleIdx);
+			SetTrianglePolygon(TriIdx, PolyIdx);
+			TriIdx++;
+
+			SetTriangleWithMatchedUVNormal(TriIdx, Inner + AngleIdx, Inner + AngleIdx + 1, Outer + AngleIdx + 1);
+			SetTrianglePolygon(TriIdx, PolyIdx);
+			TriIdx++;
+
+			PolyIdx++;
+		}
+		if (bFullDisc)
+		{
+			SetTriangleWithMatchedUVNormal(TriIdx, Inner + AngleNV - 1, Outer, Outer + AngleNV - 1);
+			SetTrianglePolygon(TriIdx, PolyIdx);
+			TriIdx++;
+
+			SetTriangleWithMatchedUVNormal(TriIdx, Inner + AngleNV - 1, Inner, Outer);
+			SetTrianglePolygon(TriIdx, PolyIdx);
+			TriIdx++;
+
+			PolyIdx++;
+		}
+	}
+
+
+	return *this;
+}
+
+FMeshShapeGenerator& FPuncturedDiscMeshGenerator::Generate()
+{
+	int AngleNV = AngleSamples > 2 ? AngleSamples : 2;
+	int RadialNV = RadialSamples > 2 ? RadialSamples : 2;
+	int NumVertices = AngleNV * RadialNV;
+
+	bool bFullDisc = (EndAngle - StartAngle) >= 359.999;
+	int AngleNT = bFullDisc ? AngleNV : AngleNV - 1;
+
+	int NumTriangles = AngleNT * 2 * (RadialNV-1);
+	SetBufferSizes(NumVertices, NumTriangles, NumVertices, NumVertices);
+
+	double IdxToRadiansOffset = StartAngle * FMathd::DegToRad;
+	double IdxToRadiansScale = (EndAngle - StartAngle) * FMathd::DegToRad / double(AngleNV - (bFullDisc ? 0 : 1));
+	double IdxToRadiusScale = (Radius - HoleRadius) / double(RadialNV-1);
+	double IdxToRadiusOffset = HoleRadius;
+	double RadiusToUVScale = .5 / Radius;
+	for (int AngleIdx = 0; AngleIdx < AngleNV; AngleIdx++)
+	{
+		double Angle = IdxToRadiansOffset + IdxToRadiansScale * AngleIdx;
+		double CosA = FMathd::Cos(Angle);
+		double SinA = FMathd::Sin(Angle);
+		for (int RadIdx = 0; RadIdx < RadialNV; RadIdx++)
+		{
+			int VertIdx = AngleIdx + RadIdx * AngleNV;
+			double R = RadIdx * IdxToRadiusScale + IdxToRadiusOffset;
+			Vertices[VertIdx] = MakeVertex(R*CosA, R*SinA);
+			UVs[VertIdx] = FVector2f(.5f + CosA * R * RadiusToUVScale, .5f + SinA * R * RadiusToUVScale);
+		}
+	}
+
+	for (int Idx = 0; Idx < Vertices.Num(); Idx++)
+	{
+		UVParentVertex[Idx] = Idx;
+		NormalParentVertex[Idx] = Idx;
+		Normals[Idx] = Normal;
+	}
+
+	int TriIdx = 0, PolyIdx = 0;
+	for (int RadiusIdx = 0; RadiusIdx+1 < RadialNV; RadiusIdx++)
+	{
+		int Inner = RadiusIdx * AngleNV;
+		int Outer = Inner + AngleNV;
+		for (int AngleIdx = 0; AngleIdx + 1 < AngleNV; AngleIdx++)
+		{
+			SetTriangleWithMatchedUVNormal(TriIdx, Inner + AngleIdx, Outer + AngleIdx + 1, Outer + AngleIdx);
+			SetTrianglePolygon(TriIdx, PolyIdx);
+			TriIdx++;
+
+			SetTriangleWithMatchedUVNormal(TriIdx, Inner + AngleIdx, Inner + AngleIdx + 1, Outer + AngleIdx + 1);
+			SetTrianglePolygon(TriIdx, PolyIdx);
+			TriIdx++;
+
+			PolyIdx++;
+		}
+		if (bFullDisc)
+		{
+			SetTriangleWithMatchedUVNormal(TriIdx, Inner + AngleNV - 1, Outer, Outer + AngleNV - 1);
+			SetTrianglePolygon(TriIdx, PolyIdx);
+			TriIdx++;
+
+			SetTriangleWithMatchedUVNormal(TriIdx, Inner + AngleNV - 1, Inner, Outer);
+			SetTrianglePolygon(TriIdx, PolyIdx);
+			TriIdx++;
+
+			PolyIdx++;
+		}
+	}
+
+	return *this;
+}
diff --git a/Engine/Plugins/Experimental/GeometryProcessing/Source/GeometricObjects/Public/Generators/DiscMeshGenerator.h b/Engine/Plugins/Experimental/GeometryProcessing/Source/GeometricObjects/Public/Generators/DiscMeshGenerator.h
new file mode 100644
--- /dev/null
+++ b/Engine/Plugins/Experimental/GeometryProcessing/Source/GeometricObjects/Public/Generators/DiscMeshGenerator.h
@@ -0,0 +1,65 @@
+// Copyright 1998-2019 Epic Games, Inc. All Rights Reserved.
+
+#pragma once
+
+#include "MeshShapeGenerator.h"
+
+/**
+ * Generate planar disc (a circle polygon)
+ */
+class GEOMETRICOBJECTS_API FDiscMeshGenerator : public FMeshShapeGenerator
+{
+public:
+	/** Radius */
+	float Radius;
+
+	/** Normal vector of all vertices will be set to this value. Default is +Z axis. */
+	FVector3f Normal;
+
+	/** Number of vertices around circumference */
+	int AngleSamples;
+
+	/** Number of vertices along radial spokes */
+	int RadialSamples;
+
+	/** Start of angle range spanned by disc, in degrees */
+	float StartAngle;
+
+	/** End of angle range spanned by disc, in degrees */
+	float EndAngle;
+
+	/** How to map 2D indices to 3D. Default is (0,1) = (x,y,0). */
+	FIndex2i IndicesMap;
+
+public:
+	FDiscMeshGenerator();
+
+	/** Generate the disc */
+	virtual FMeshShapeGenerator& Generate() override;
+
+
+	/** Create vertex at position under IndicesMap, shifted to Origin*/
+	inline FVector3d MakeVertex(double x, double y)
+	{
+		FVector3d v(0, 0, 0);
+		v[IndicesMap.A] = x;
+		v[IndicesMap.B] = y;
+		return v;
+	}
+};
+
+/**
+* Generate planar disc with a hole
+*/
+class GEOMETRICOBJECTS_API FPuncturedDiscMeshGenerator : public FDiscMeshGenerator
+{
+public:
+	/** Hole Radius */
+	float HoleRadius;
+
+public:
+	FPuncturedDiscMeshGenerator();
+
+	/** Generate the disc */
+	virtual FMeshShapeGenerator& Generate() override;
+};
\ No newline at end of file
diff --git a/Engine/Plugins/Experimental/GeometryProcessing/Source/GeometricObjects/Public/Generators/MeshShapeGenerator.h b/Engine/Plugins/Experimental/GeometryProcessing/Source/GeometricObjects/Public/Generators/MeshShapeGenerator.h
--- a/Engine/Plugins/Experimental/GeometryProcessing/Source/GeometricObjects/Public/Generators/MeshShapeGenerator.h
+++ b/Engine/Plugins/Experimental/GeometryProcessing/Source/GeometricObjects/Public/Generators/MeshShapeGenerator.h
@@ -217,7 +217,16 @@ public:
 		TrianglePolygonIDs[Index] = PolygonID;
 	}
 
-
+	/**
+	 * Set triangle and UVs and normals with matching indices
+	 * Convenience function for shapes with no uv or normal seams
+	 */
+	inline void SetTriangleWithMatchedUVNormal(int Index, int A, int B, int C)
+	{
+		SetTriangle(Index, A, B, C);
+		SetTriangleNormals(Index, A, B, C);
+		SetTriangleUVs(Index, A, B, C);
+	}
 
 
 	static FVector3d BilinearInterp(const FVector3d &v00, const FVector3d &v10, const FVector3d &v11, const FVector3d &v01, double tx, double ty)
diff --git a/Engine/Plugins/Experimental/MeshModelingToolset/Source/MeshModelingTools/Private/AddPrimitiveTool.cpp b/Engine/Plugins/Experimental/MeshModelingToolset/Source/MeshModelingTools/Private/AddPrimitiveTool.cpp
--- a/Engine/Plugins/Experimental/MeshModelingToolset/Source/MeshModelingTools/Private/AddPrimitiveTool.cpp
+++ b/Engine/Plugins/Experimental/MeshModelingToolset/Source/MeshModelingTools/Private/AddPrimitiveTool.cpp
@@ -13,6 +13,7 @@
 #include "Generators/RectangleMeshGenerator.h"
 #include "Generators/SphereGenerator.h"
 #include "Generators/BoxSphereGenerator.h"
+#include "Generators/DiscMeshGenerator.h"
 #include "DynamicMeshAttributeSet.h"
 #include "StaticMeshComponentBuilder.h"
 #include "Drawing/MeshDebugDrawing.h"
@@ -91,8 +92,8 @@ namespace
 	  { TEXT("Rotation"),      EMakeMeshShapeType::All },
 	  { TEXT("PlaceMode"),     EMakeMeshShapeType::All },
 	  { TEXT("PivotLocation"), EMakeMeshShapeType::All },
-	  { TEXT("Slices"),        EMakeMeshShapeType::Cylinder | EMakeMeshShapeType::Cone | EMakeMeshShapeType::Arrow | EMakeMeshShapeType::RoundedRectangle | EMakeMeshShapeType::Sphere },
-	  { TEXT("Subdivisions"),  EMakeMeshShapeType::Box | EMakeMeshShapeType::Rectangle | EMakeMeshShapeType::RoundedRectangle | EMakeMeshShapeType::Cylinder |
+	  { TEXT("Slices"),        EMakeMeshShapeType::Cylinder | EMakeMeshShapeType::Cone | EMakeMeshShapeType::Arrow | EMakeMeshShapeType::RoundedRectangle | EMakeMeshShapeType::Disc | EMakeMeshShapeType::PuncturedDisc | EMakeMeshShapeType::Sphere },
+	  { TEXT("Subdivisions"),  EMakeMeshShapeType::Box | EMakeMeshShapeType::Rectangle | EMakeMeshShapeType::RoundedRectangle | EMakeMeshShapeType::Disc | EMakeMeshShapeType::PuncturedDisc | EMakeMeshShapeType::Cylinder |
 							   EMakeMeshShapeType::Cone | EMakeMeshShapeType::Arrow | EMakeMeshShapeType::SphericalBox }
 	};
 };
@@ -266,6 +267,14 @@ void UAddPrimitiveTool::UpdatePreviewMesh()
 		GenerateRoundedRectangle(&NewMesh);
 		break;
 
+	case EMakeMeshShapeType::Disc:
+		GenerateDisc(&NewMesh);
+		break;
+
+	case EMakeMeshShapeType::PuncturedDisc:
+		GeneratePuncturedDisc(&NewMesh);
+		break;
+
 	case EMakeMeshShapeType::Cylinder:
 		GenerateCylinder(&NewMesh);
 		break;
@@ -380,6 +389,29 @@ void UAddPrimitiveTool::GenerateRoundedRectangle(FDynamicMesh3* OutMesh)
 }
 
 
+void UAddPrimitiveTool::GenerateDisc(FDynamicMesh3* OutMesh)
+{
+	FDiscMeshGenerator Gen;
+	Gen.Radius = ShapeSettings->Width * 0.5f;
+	Gen.AngleSamples = ShapeSettings->Slices;
+	Gen.RadialSamples = ShapeSettings->Subdivisions;
+	Gen.Generate();
+	OutMesh->Copy(&Gen);
+}
+
+
+void UAddPrimitiveTool::GeneratePuncturedDisc(FDynamicMesh3* OutMesh)
+{
+	FPuncturedDiscMeshGenerator Gen;
+	Gen.Radius = ShapeSettings->Width * 0.5f;
+	Gen.HoleRadius = Gen.Radius * .5;
+	Gen.AngleSamples = ShapeSettings->Slices;
+	Gen.RadialSamples = ShapeSettings->Subdivisions;
+	Gen.Generate();
+	OutMesh->Copy(&Gen);
+}
+
+
 void UAddPrimitiveTool::GenerateCylinder(FDynamicMesh3* OutMesh)
 {
 	FCylinderGenerator CylGen;
diff --git a/Engine/Plugins/Experimental/MeshModelingToolset/Source/MeshModelingTools/Public/AddPrimitiveTool.h b/Engine/Plugins/Experimental/MeshModelingToolset/Source/MeshModelingTools/Public/AddPrimitiveTool.h
--- a/Engine/Plugins/Experimental/MeshModelingToolset/Source/MeshModelingTools/Public/AddPrimitiveTool.h
+++ b/Engine/Plugins/Experimental/MeshModelingToolset/Source/MeshModelingTools/Public/AddPrimitiveTool.h
@@ -35,18 +35,20 @@ public:
 
 /** Shape Types */
 UENUM()
-enum class EMakeMeshShapeType : uint8
+enum class EMakeMeshShapeType : uint32
 {
-	None			 = 0x00 UMETA(DisplayName = "None", Hidden),
-	All				 = 0xff UMETA(DisplayName = "All", Hidden), 
-	Box				 = 0x01 UMETA(DisplayName = "Box"),
-	Cylinder		 = 0x02 UMETA(DisplayName = "Cylinder"),
-	Cone			 = 0x04 UMETA(DisplayName = "Cone"),
-	Arrow			 = 0x08 UMETA(DisplayName = "Arrow"),
-	Rectangle		 = 0x10 UMETA(DisplayName = "Rectangle"),
-	RoundedRectangle = 0x20 UMETA(DisplayName = "Rounded Rectangle"),
-	Sphere			 = 0x40 UMETA(DisplayName = "Sphere"),
-	SphericalBox	 = 0x80 UMETA(DisplayName = "Spherical Box")
+	None			 = 0x000 UMETA(DisplayName = "None", Hidden),
+	All				 = 0xfff UMETA(DisplayName = "All", Hidden), 
+	Box				 = 0x001 UMETA(DisplayName = "Box"),
+	Cylinder		 = 0x002 UMETA(DisplayName = "Cylinder"),
+	Cone			 = 0x004 UMETA(DisplayName = "Cone"),
+	Arrow			 = 0x008 UMETA(DisplayName = "Arrow"),
+	Rectangle		 = 0x010 UMETA(DisplayName = "Rectangle"),
+	RoundedRectangle = 0x020 UMETA(DisplayName = "Rounded Rectangle"),
+	Disc			 = 0x040 UMETA(DisplayName = "Disc"),
+	PuncturedDisc	 = 0x080 UMETA(DisplayName = "Punctured Disc"),
+	Sphere			 = 0x100 UMETA(DisplayName = "Sphere"),
+	SphericalBox	 = 0x200 UMETA(DisplayName = "Spherical Box")
 	
 };
 ENUM_CLASS_FLAGS(EMakeMeshShapeType);
@@ -192,6 +194,8 @@ protected:
 	void GenerateBox(FDynamicMesh3* OutMesh);
 	void GenerateRectangle(FDynamicMesh3* OutMesh);
 	void GenerateRoundedRectangle(FDynamicMesh3* OutMesh);
+	void GenerateDisc(FDynamicMesh3* OutMesh);
+	void GeneratePuncturedDisc(FDynamicMesh3* OutMesh);
 	void GenerateSphere(FDynamicMesh3* OutMesh);
 	void GenerateSphericalBox(FDynamicMesh3* OutMesh);
 	void GenerateArrow(FDynamicMesh3* OutMesh);
