add random sampling on triangles to jacketing, to fix cases where centroid/vertices are occluded but other points on triangle are not

#rnx

diff --git a/Engine/Plugins/Experimental/GeometryProcessing/Source/DynamicMesh/Public/Operations/RemoveOccludedTriangles.h b/Engine/Plugins/Experimental/GeometryProcessing/Source/DynamicMesh/Public/Operations/RemoveOccludedTriangles.h
--- a/Engine/Plugins/Experimental/GeometryProcessing/Source/DynamicMesh/Public/Operations/RemoveOccludedTriangles.h
+++ b/Engine/Plugins/Experimental/GeometryProcessing/Source/DynamicMesh/Public/Operations/RemoveOccludedTriangles.h
@@ -70,6 +70,7 @@ public:
 		// ray directions
 		TArray<FVector3d> RayDirs; int NR = 0;
 
+		FRandomStream RaysRandomStream(2123123);
 		if (InsideMode == EOcclusionCalculationMode::SimpleOcclusionTest)
 		{
 			RayDirs.Add(FVector3d::UnitX()); RayDirs.Add(-FVector3d::UnitX());
@@ -78,11 +79,25 @@ public:
 
 			for (int AddRayIdx = 0; AddRayIdx < AddRandomRays; AddRayIdx++)
 			{
-				RayDirs.Add(FVector3d(FMath::VRand()));
+				RayDirs.Add(FVector3d(RaysRandomStream.VRand()));
 			}
 			NR = RayDirs.Num();
 		}
 
+		// triangle samples get their own random stream to make behavior slightly more predictable (e.g. moving the ray samples up shouldn't change all the triangle sample locations)
+		FRandomStream TrisRandomStream(124233);
+		TArray<FVector3d> TriangleBaryCoordSamples;
+		for (int AddSampleIdx = 0; AddSampleIdx < AddTriangleSamples; AddSampleIdx++)
+		{
+			FVector3d BaryCoords(TrisRandomStream.FRand() * .999 + .001, TrisRandomStream.FRand() * .999 + .001, TrisRandomStream.FRand() * .999 + .001);
+			BaryCoords /= (BaryCoords.X + BaryCoords.Y + BaryCoords.Z);
+			TriangleBaryCoordSamples.Add(BaryCoords);
+		}
+		if (TriangleSamplingMethod == EOcclusionTriangleSampling::Centroids || TriangleSamplingMethod == EOcclusionTriangleSampling::VerticesAndCentroids)
+		{
+			TriangleBaryCoordSamples.Add(FVector3d(1 / 3.0, 1 / 3.0, 1 / 3.0));
+		}
+
 		auto IsOccludedFWN = [this, &FastWindingTree](const FVector3d& Pt)
 		{
 			return FastWindingTree->FastWindingNumber(Pt) > WindingIsoValue;
@@ -139,7 +154,7 @@ public:
 
 		RemovedT.Empty();
 		FCriticalSection RemoveTMutex;
-		ParallelFor(Mesh->MaxTriangleID(), [this, &VertexOccluded, &IsOccludedF, &RemoveTMutex, &MeshLocalToOccluderSpace](int32 TID)
+		ParallelFor(Mesh->MaxTriangleID(), [this, &VertexOccluded, &IsOccludedF, &RemoveTMutex, &MeshLocalToOccluderSpace, &TriangleBaryCoordSamples](int32 TID)
 		{
 			if (!Mesh->IsTriangle(TID))
 			{
@@ -152,14 +167,17 @@ public:
 				FIndex3i Tri = Mesh->GetTriangle(TID);
 				bInside = VertexOccluded[Tri.A] && VertexOccluded[Tri.B] && VertexOccluded[Tri.C];
 			}
-			if (bInside && (TriangleSamplingMethod == EOcclusionTriangleSampling::Centroids || TriangleSamplingMethod == EOcclusionTriangleSampling::VerticesAndCentroids))
+			if (bInside && TriangleBaryCoordSamples.Num() > 0)
 			{
-				FVector3d Centroid, Normal;
-				double Area;
-				Mesh->GetTriInfo(TID, Normal, Area, Centroid);
-
-				FVector3d SamplePos = Centroid + Normal * NormalOffset;
-				bInside = IsOccludedF(MeshLocalToOccluderSpace.TransformPosition(SamplePos));
+				FVector3d Normal = Mesh->GetTriNormal(TID);
+				FVector3d V0, V1, V2;
+				Mesh->GetTriVertices(TID, V0, V1, V2);
+				for (int32 SampleIdx = 0, NumSamples = TriangleBaryCoordSamples.Num(); bInside && SampleIdx < NumSamples; SampleIdx++)
+				{
+					FVector3d BaryCoords = TriangleBaryCoordSamples[SampleIdx];
+					FVector3d SamplePos = V0 * BaryCoords.X + V1 * BaryCoords.Y + V2 * BaryCoords.Z + Normal * NormalOffset;
+					bInside = IsOccludedF(MeshLocalToOccluderSpace.TransformPosition(SamplePos));
+				}
 			}
 			if (bInside)
 			{
@@ -231,6 +249,9 @@ public:
 	/** Number of additional ray directions to add to raycast-based occlusion checks, beyond the default +/- major axis directions */
 	int AddRandomRays = 0;
 
+	/** Number of additional samples to add per triangle */
+	int AddTriangleSamples = 0;
+
 	/**
 	 * Set this to be able to cancel running operation
 	 */
diff --git a/Engine/Plugins/Experimental/MeshModelingToolset/Source/MeshModelingTools/Private/RemoveOccludedTrianglesTool.cpp b/Engine/Plugins/Experimental/MeshModelingToolset/Source/MeshModelingTools/Private/RemoveOccludedTrianglesTool.cpp
--- a/Engine/Plugins/Experimental/MeshModelingToolset/Source/MeshModelingTools/Private/RemoveOccludedTrianglesTool.cpp
+++ b/Engine/Plugins/Experimental/MeshModelingToolset/Source/MeshModelingTools/Private/RemoveOccludedTrianglesTool.cpp
@@ -232,6 +232,8 @@ TUniquePtr<FDynamicMeshOperator> URemoveOccludedTrianglesOperatorFactory::MakeNe
 	Op->bOnlySelfOcclude = Tool->BasicProperties->bOnlySelfOcclude;
 
 	Op->AddRandomRays = Tool->BasicProperties->AddRandomRays;
+
+	Op->AddTriangleSamples = Tool->BasicProperties->AddTriangleSamples;
 	
 	Op->SetTransform(LocalToWorld);
 
diff --git a/Engine/Plugins/Experimental/MeshModelingToolset/Source/MeshModelingTools/Public/RemoveOccludedTrianglesTool.h b/Engine/Plugins/Experimental/MeshModelingToolset/Source/MeshModelingTools/Public/RemoveOccludedTrianglesTool.h
--- a/Engine/Plugins/Experimental/MeshModelingToolset/Source/MeshModelingTools/Public/RemoveOccludedTrianglesTool.h
+++ b/Engine/Plugins/Experimental/MeshModelingToolset/Source/MeshModelingTools/Public/RemoveOccludedTrianglesTool.h
@@ -81,9 +81,13 @@ public:
 	double WindingIsoValue = 0.5;
 
 	/** For raycast-based occlusion tests, optionally add random ray direction to increase the accuracy of the visibility sampling */
-	UPROPERTY(EditAnywhere, Category = OcclusionCalculation, meta = (EditCondition = "OcclusionTestMethod==EOcclusionCalculationUIMode::RaycastOcclusionSamples"))
+	UPROPERTY(EditAnywhere, Category = OcclusionCalculation, meta = (UIMin = "0", UIMax = "100", ClampMin = "0", ClampMax = "1000", EditCondition = "OcclusionTestMethod==EOcclusionCalculationUIMode::RaycastOcclusionSamples"))
 	int AddRandomRays = 0;
 
+	/** Optionally add random samples to each triangle (in addition to those from TriangleSampling) to increase the accuracy of the visibility sampling */
+	UPROPERTY(EditAnywhere, Category = OcclusionCalculation, meta = (UIMin = "0", UIMax = "100", ClampMin = "0", ClampMax = "1000"))
+	int AddTriangleSamples = 0;
+
 	/** If false, when multiple meshes are selected the meshes can occlude each other.  When true, we process each selected mesh independently and only consider self-occlusions. */
 	UPROPERTY(EditAnywhere, Category = OcclusionCalculation)
 	bool bOnlySelfOcclude = false;
diff --git a/Engine/Plugins/Experimental/MeshModelingToolset/Source/ModelingOperators/Private/CleaningOps/RemoveOccludedTrianglesOp.cpp b/Engine/Plugins/Experimental/MeshModelingToolset/Source/ModelingOperators/Private/CleaningOps/RemoveOccludedTrianglesOp.cpp
--- a/Engine/Plugins/Experimental/MeshModelingToolset/Source/ModelingOperators/Private/CleaningOps/RemoveOccludedTrianglesOp.cpp
+++ b/Engine/Plugins/Experimental/MeshModelingToolset/Source/ModelingOperators/Private/CleaningOps/RemoveOccludedTrianglesOp.cpp
@@ -69,6 +69,7 @@ void FRemoveOccludedTrianglesOp::CalculateResult(FProgressCancel* Progress)
 		SelfJacket.WindingIsoValue = WindingIsoValue;
 		SelfJacket.NormalOffset = NormalOffset;
 		SelfJacket.AddRandomRays = AddRandomRays;
+		SelfJacket.AddTriangleSamples = AddTriangleSamples;
 		SelfJacket.Apply(FTransform3d::Identity(), &SelfAABB, &SelfFWTree);
 	}
 	else
@@ -85,6 +86,7 @@ void FRemoveOccludedTrianglesOp::CalculateResult(FProgressCancel* Progress)
 		Jacket.WindingIsoValue = WindingIsoValue;
 		Jacket.NormalOffset = NormalOffset;
 		Jacket.AddRandomRays = AddRandomRays;
+		Jacket.AddTriangleSamples = AddTriangleSamples;
 		Jacket.Apply(ResultTransform, &CombinedMeshTrees->AABB, &CombinedMeshTrees->FastWinding);
 	}
 }
\ No newline at end of file
diff --git a/Engine/Plugins/Experimental/MeshModelingToolset/Source/ModelingOperators/Public/CleaningOps/RemoveOccludedTrianglesOp.h b/Engine/Plugins/Experimental/MeshModelingToolset/Source/ModelingOperators/Public/CleaningOps/RemoveOccludedTrianglesOp.h
--- a/Engine/Plugins/Experimental/MeshModelingToolset/Source/ModelingOperators/Public/CleaningOps/RemoveOccludedTrianglesOp.h
+++ b/Engine/Plugins/Experimental/MeshModelingToolset/Source/ModelingOperators/Public/CleaningOps/RemoveOccludedTrianglesOp.h
@@ -54,6 +54,9 @@ public:
 	// random rays to add beyond +/- major axes, for raycast sampling
 	int AddRandomRays = 0;
 
+	// add triangle samples per triangle (in addition to TriangleSamplingMethod)
+	int AddTriangleSamples = 0;
+
 	// if true, *ignore* the combined selected occluder data in CombinedMeshTrees, and only consider self-occlusion of this mesh
 	bool bOnlySelfOcclude = false;
 
