add torus to addprimitivetool
add support for looping to sweep gen

#rnx

diff --git a/Engine/Plugins/Experimental/GeometryProcessing/Source/GeometricObjects/Public/Generators/SweepGenerator.h b/Engine/Plugins/Experimental/GeometryProcessing/Source/GeometricObjects/Public/Generators/SweepGenerator.h
--- a/Engine/Plugins/Experimental/GeometryProcessing/Source/GeometricObjects/Public/Generators/SweepGenerator.h
+++ b/Engine/Plugins/Experimental/GeometryProcessing/Source/GeometricObjects/Public/Generators/SweepGenerator.h
@@ -43,6 +43,7 @@ protected:
 							   const TArrayView<const int32>& NormalSections,
 							   const TArrayView<const int32>& SharpNormalsAlongLength,
 							   int32 NumCrossSections,
+							   bool bLoop,
 							   const ECapType Caps[2],
 							   FVector2d UVScale, FVector2d UVOffset)
 	{
@@ -51,8 +52,8 @@ protected:
 		int32 XNormals = XVerts + NormalSections.Num();
 		int32 XUVs = XVerts + UVSections.Num() + 1;
 
-		int32 NumVerts = XVerts * NumCrossSections;
-		int32 NumNormals = NumCrossSections > 1 ? XNormals * NumCrossSections : 0;
+		int32 NumVerts = XVerts * NumCrossSections - (bLoop ? XVerts : 0);
+		int32 NumNormals = NumCrossSections > 1 ? (XNormals * NumCrossSections - (bLoop ? XNormals : 0)) : 0;
 		NumNormals += XNormals * SharpNormalsAlongLength.Num();
 		int32 NumUVs = NumCrossSections > 1 ? XUVs * NumCrossSections : 0;
 		int32 NumPolygons = (NumCrossSections - 1) * XVerts;
@@ -60,88 +61,105 @@ protected:
 
 		TArray<FIndex3i> OutTriangles;
 
-		for (int32 CapIdx = 0; CapIdx < 2; CapIdx++)
+		// doesn't make sense to have cap types if the sweep is a loop
+		ensure(!bLoop || (Caps[0] == ECapType::None && Caps[1] == ECapType::None));
+		
+		if (!bLoop)
 		{
-			CapVertStart[CapIdx] = NumVerts;
-			CapNormalStart[CapIdx] = NumNormals;
-			CapUVStart[CapIdx] = NumUVs;
-			CapTriangleStart[CapIdx] = NumTriangles;
-			CapPolygonStart[CapIdx] = NumPolygons;
-
-			if (Caps[CapIdx] == ECapType::FlatTriangulation)
+			for (int32 CapIdx = 0; !bLoop && CapIdx < 2; CapIdx++)
 			{
-				NumTriangles += XVerts - 2;
-				NumPolygons++;
-				NumUVs += XVerts;
-				NumNormals += XVerts;
-			}
+				CapVertStart[CapIdx] = NumVerts;
+				CapNormalStart[CapIdx] = NumNormals;
+				CapUVStart[CapIdx] = NumUVs;
+				CapTriangleStart[CapIdx] = NumTriangles;
+				CapPolygonStart[CapIdx] = NumPolygons;
+
+				if (Caps[CapIdx] == ECapType::FlatTriangulation)
+				{
+					NumTriangles += XVerts - 2;
+					NumPolygons++;
+					NumUVs += XVerts;
+					NumNormals += XVerts;
+				}
 
-			// TODO: support more cap type; e.g.:
-			//else if (Caps[CapIdx] == ECapType::FlatMidpointFan)
-			//{
-			//	NumTriangles += XVerts;
-			//	NumPolygons += XVerts;
-			//	NumUVs += XVerts + 1;
-			//	NumNormals += XVerts + 1;
-			//	NumVerts += 1;
-			//}
-			//else if (Caps[CapIdx] == ECapType::Cone)
-			//{
-			//	NumTriangles += XVerts;
-			//	NumPolygons += XVerts;
-			//	NumUVs += XVerts + 1;
-			//	NumNormals += XVerts * 2;
-			//	NumVerts += 1;
-			//}
+				// TODO: support more cap type; e.g.:
+				//else if (Caps[CapIdx] == ECapType::FlatMidpointFan)
+				//{
+				//	NumTriangles += XVerts;
+				//	NumPolygons += XVerts;
+				//	NumUVs += XVerts + 1;
+				//	NumNormals += XVerts + 1;
+				//	NumVerts += 1;
+				//}
+				//else if (Caps[CapIdx] == ECapType::Cone)
+				//{
+				//	NumTriangles += XVerts;
+				//	NumPolygons += XVerts;
+				//	NumUVs += XVerts + 1;
+				//	NumNormals += XVerts * 2;
+				//	NumVerts += 1;
+				//}
+			}
 		}
 
 		SetBufferSizes(NumVerts, NumTriangles, NumUVs, NumNormals);
 
-		for (int32 CapIdx = 0; CapIdx < 2; CapIdx++)
+		if (!bLoop)
 		{
-
-			if (Caps[CapIdx] == ECapType::FlatTriangulation)
+			for (int32 CapIdx = 0; CapIdx < 2; CapIdx++)
 			{
-				int32 VertOffset = CapIdx * (XVerts * (NumCrossSections - 1));
 
-				PolygonTriangulation::TriangulateSimplePolygon(CrossSection.GetVertices(), OutTriangles);
-				int32 TriIdx = CapTriangleStart[CapIdx];
-				int32 PolyIdx = CapPolygonStart[CapIdx];
-				for (const FIndex3i& Triangle : OutTriangles)
+				if (Caps[CapIdx] == ECapType::FlatTriangulation)
 				{
-					bool Flipped = CapIdx == 0;
-					Flipped = !Flipped;
-					SetTriangle(TriIdx,
-								Triangle.A + VertOffset, Triangle.B + VertOffset, Triangle.C + VertOffset,
-								Flipped);
-					SetTriangleUVs(TriIdx,
-								   Triangle.A + CapUVStart[CapIdx],
-								   Triangle.B + CapUVStart[CapIdx],
-								   Triangle.C + CapUVStart[CapIdx],
-								   Flipped);
-					SetTriangleNormals(TriIdx,
-									   Triangle.A + CapNormalStart[CapIdx],
-									   Triangle.B + CapNormalStart[CapIdx],
-									   Triangle.C + CapNormalStart[CapIdx],
-									   Flipped);
-					SetTrianglePolygon(TriIdx, PolyIdx);
-					TriIdx++;
-				}
-				float SideScale = 2 * CapIdx - 1;
-				for (int32 Idx = 0; Idx < XVerts; Idx++)
-				{
-					FVector2d CenteredVert = CrossSection.GetVertices()[Idx] * UVScale + UVOffset;
-					SetUV(CapUVStart[CapIdx] + Idx, FVector2f(CenteredVert.X * SideScale, CenteredVert.Y), VertOffset + Idx);
-					SetNormal(CapNormalStart[CapIdx] + Idx, FVector3f::Zero(), VertOffset + Idx);
+					int32 VertOffset = CapIdx * (XVerts * (NumCrossSections - 1));
+
+					PolygonTriangulation::TriangulateSimplePolygon(CrossSection.GetVertices(), OutTriangles);
+					int32 TriIdx = CapTriangleStart[CapIdx];
+					int32 PolyIdx = CapPolygonStart[CapIdx];
+					for (const FIndex3i& Triangle : OutTriangles)
+					{
+						bool Flipped = CapIdx == 0;
+						Flipped = !Flipped;
+						SetTriangle(TriIdx,
+							Triangle.A + VertOffset, Triangle.B + VertOffset, Triangle.C + VertOffset,
+							Flipped);
+						SetTriangleUVs(TriIdx,
+							Triangle.A + CapUVStart[CapIdx],
+							Triangle.B + CapUVStart[CapIdx],
+							Triangle.C + CapUVStart[CapIdx],
+							Flipped);
+						SetTriangleNormals(TriIdx,
+							Triangle.A + CapNormalStart[CapIdx],
+							Triangle.B + CapNormalStart[CapIdx],
+							Triangle.C + CapNormalStart[CapIdx],
+							Flipped);
+						SetTrianglePolygon(TriIdx, PolyIdx);
+						TriIdx++;
+					}
+					float SideScale = 2 * CapIdx - 1;
+					for (int32 Idx = 0; Idx < XVerts; Idx++)
+					{
+						FVector2d CenteredVert = CrossSection.GetVertices()[Idx] * UVScale + UVOffset;
+						SetUV(CapUVStart[CapIdx] + Idx, FVector2f(CenteredVert.X * SideScale, CenteredVert.Y), VertOffset + Idx);
+						SetNormal(CapNormalStart[CapIdx] + Idx, FVector3f::Zero(), VertOffset + Idx);
+					}
 				}
 			}
 		}
 
 		// fill in UVs and triangles along length
-		if (NumCrossSections > 1)
+		int MinValidCrossSections = bLoop ? 3 : 2;
+		if (NumCrossSections >= MinValidCrossSections)
 		{
 			int32 UVSection = 0, UVSubIdx = 0;
 
+			int CrossSectionsMod = NumCrossSections;
+			if (bLoop)
+			{
+				CrossSectionsMod--; // last cross section becomes the first
+			}
+			int NormalCrossSectionsMod = CrossSectionsMod + SharpNormalsAlongLength.Num();
+
 			int32 NumSections = UVSections.Num();
 			int32 NextDupVertIdx = UVSection < NumSections ? UVSections[UVSection] : -1;
 			for (int32 VertSubIdx = 0; VertSubIdx < XVerts; UVSubIdx++)
@@ -150,7 +168,7 @@ protected:
 				for (int32 XIdx = 0; XIdx < NumCrossSections; XIdx++)
 				{
 					float UVY = XIdx / float(NumCrossSections - 1);
-					SetUV(XIdx * XUVs + UVSubIdx, FVector2f(1-UVX, 1-UVY), XIdx * XVerts + VertSubIdx);
+					SetUV(XIdx * XUVs + UVSubIdx, FVector2f(1-UVX, 1-UVY), (XIdx % CrossSectionsMod) * XVerts + VertSubIdx);
 				}
 
 				if (VertSubIdx == NextDupVertIdx)
@@ -182,7 +200,7 @@ protected:
 				for (int32 XIdx = 0; XIdx < NumCrossSections; XIdx++)
 				{
 					float UVY = XIdx / float(NumCrossSections - 1);
-					SetUV(XIdx * XUVs + UVSubIdx, FVector2f(1-UVX, 1-UVY), XIdx * XVerts + VertSubIdx);
+					SetUV(XIdx * XUVs + UVSubIdx, FVector2f(1-UVX, 1-UVY), (XIdx % CrossSectionsMod) * XVerts + VertSubIdx);
 				}
 			}
 			NumSections = NormalSections.Num();
@@ -195,12 +213,12 @@ protected:
 				for (int32 XIdx = 0, NormalXIdx = 0; XIdx < NumCrossSections; XIdx++, NormalXIdx++)
 				{
 					// just set the normal parent; don't compute normal yet
-					SetNormal(NormalXIdx * XNormals + NormalSubIdx, FVector3f(0, 0, 0), XIdx * XVerts + VertSubIdx);
+					SetNormal((NormalXIdx % NormalCrossSectionsMod) * XNormals + NormalSubIdx, FVector3f(0, 0, 0), (XIdx % CrossSectionsMod) * XVerts + VertSubIdx);
 					// duplicate normals for cross sections that are 'sharp'
 					if (SharpNormalIdx < SharpNormalsAlongLength.Num() && XIdx == SharpNormalsAlongLength[SharpNormalIdx])
 					{
 						NormalXIdx++;
-						SetNormal(NormalXIdx * XNormals + NormalSubIdx, FVector3f(0, 0, 0), XIdx * XVerts + VertSubIdx);
+						SetNormal((NormalXIdx % NormalCrossSectionsMod) * XNormals + NormalSubIdx, FVector3f(0, 0, 0), (XIdx % CrossSectionsMod) * XVerts + VertSubIdx);
 						SharpNormalIdx++;
 					}
 				}
@@ -220,25 +238,27 @@ protected:
 						int32 T0Idx = XVerts * 2 * XIdx + 2 * VertSubIdx;
 						int32 T1Idx = T0Idx + 1;
 						int32 PIdx = XVerts * XIdx + VertSubIdx;
+						int32 NextXIdx = (XIdx + 1) % CrossSectionsMod;
+						int32 NextNXIdx = (NXIdx + 1) % NormalCrossSectionsMod;
 						SetTrianglePolygon(T0Idx, PIdx);
 						SetTrianglePolygon(T1Idx, PIdx);
 						SetTriangle(T0Idx,
 									XIdx * XVerts + VertSubIdx,
 									XIdx * XVerts + WrappedNextVertexSubIdx,
-									(XIdx + 1) * XVerts + VertSubIdx, true);
+									NextXIdx * XVerts + VertSubIdx, true);
 						SetTriangle(T1Idx,
-									(XIdx + 1) * XVerts + WrappedNextVertexSubIdx,
-									(XIdx + 1) * XVerts + VertSubIdx,
+									NextXIdx * XVerts + WrappedNextVertexSubIdx,
+									NextXIdx * XVerts + VertSubIdx,
 									XIdx * XVerts + WrappedNextVertexSubIdx, true);
 						SetTriangleNormals(
 							T0Idx,
 							NXIdx * XNormals + NormalSubIdx,
 							NXIdx * XNormals + WrappedNextNormalSubIdx,
-							(NXIdx + 1) * XNormals + NormalSubIdx, true);
+							NextNXIdx * XNormals + NormalSubIdx, true);
 						SetTriangleNormals(
 							T1Idx,
-							(NXIdx + 1) * XNormals + WrappedNextNormalSubIdx,
-							(NXIdx + 1) * XNormals + NormalSubIdx,
+							NextNXIdx * XNormals + WrappedNextNormalSubIdx,
+							NextNXIdx * XNormals + NormalSubIdx,
 							NXIdx * XNormals + WrappedNextNormalSubIdx, true);
 						if (SharpNormalIdx < SharpNormalsAlongLength.Num() && XIdx+1 == SharpNormalsAlongLength[SharpNormalIdx])
 						{
@@ -308,7 +328,7 @@ public:
 		TArray<float> AlongPercents;
 		float LenAlong = ComputeSegLengths(Radii, Heights, AlongPercents);
 
-		ConstructMeshTopology(X, {}, {}, SharpNormalsAlongLength, NumX, Caps, FVector2d(.5, .5), FVector2d(.5, .5));
+		ConstructMeshTopology(X, {}, {}, SharpNormalsAlongLength, NumX, false, Caps, FVector2d(.5, .5), FVector2d(.5, .5));
 
 		TArray<FVector2d> NormalSides; NormalSides.SetNum(NumX - 1);
 		for (int XIdx = 0; XIdx+1 < NumX; XIdx++)
@@ -494,12 +514,11 @@ public:
 };
 
 
-
 /**
  * Sweep a 2D Profile Polygon along a 3D Path.
  * 
  * TODO: 
- *  - Loop path support
+ *  - a custom variant for toruses specifically (would be faster)
  *  - Mitering cross sections support?
  */
 class /*DYNAMICMESH_API*/ FGeneralizedCylinderGenerator : public FSweepGeneratorBase
@@ -511,23 +530,22 @@ public:
 	FFrame3d InitialFrame;
 
 	bool bCapped = false;
+	bool bLoop = false;
 
 public:
 	/** Generate the mesh */
 	virtual FMeshShapeGenerator& Generate() override
 	{
-		const bool bLoop = false; // TODO: support loops -- note this requires updating FSweepGeneratorBase to set up the mesh topology correctly!
-		
 		const TArray<FVector2d>& XVerts = CrossSection.GetVertices();
 		ECapType Caps[2] = {ECapType::None, ECapType::None};
 
-		if (bCapped)
+		if (bCapped && !bLoop)
 		{
 			Caps[0] = ECapType::FlatTriangulation;
 			Caps[1] = ECapType::FlatTriangulation;
 		}
 		int PathNum = Path.Num();
-		ConstructMeshTopology(CrossSection, {}, {}, {}, PathNum, Caps, FVector2d(.5, .5), FVector2d(.5, .5));
+		ConstructMeshTopology(CrossSection, {}, {}, {}, PathNum + (bLoop ? 1 : 0), bLoop, Caps, FVector2d(.5, .5), FVector2d(.5, .5));
 
 		int XNum = CrossSection.VertexCount();
 		TArray<FVector2d> XNormals; XNormals.SetNum(XNum);
diff --git a/Engine/Plugins/Experimental/MeshModelingToolset/Source/MeshModelingTools/Private/AddPrimitiveTool.cpp b/Engine/Plugins/Experimental/MeshModelingToolset/Source/MeshModelingTools/Private/AddPrimitiveTool.cpp
--- a/Engine/Plugins/Experimental/MeshModelingToolset/Source/MeshModelingTools/Private/AddPrimitiveTool.cpp
+++ b/Engine/Plugins/Experimental/MeshModelingToolset/Source/MeshModelingTools/Private/AddPrimitiveTool.cpp
@@ -95,14 +95,14 @@ namespace
 	  { TEXT("Shape"),         EMakeMeshShapeType::All },
 	  { TEXT("Width"),         EMakeMeshShapeType::All },
 	  { TEXT("Height"),        EMakeMeshShapeType::Box | EMakeMeshShapeType::Cylinder | EMakeMeshShapeType::Cone | EMakeMeshShapeType::Arrow | EMakeMeshShapeType::Rectangle | EMakeMeshShapeType::RoundedRectangle },
-	  { TEXT("FeatureRadius"),		   EMakeMeshShapeType::Arrow | EMakeMeshShapeType::RoundedRectangle | EMakeMeshShapeType::PuncturedDisc },
+	  { TEXT("FeatureRadius"),		   EMakeMeshShapeType::Arrow | EMakeMeshShapeType::RoundedRectangle | EMakeMeshShapeType::PuncturedDisc | EMakeMeshShapeType::Torus },
 	  { TEXT("Rotation"),      EMakeMeshShapeType::All },
 	  { TEXT("PlaceMode"),     EMakeMeshShapeType::All },
 	  { TEXT("PivotLocation"), EMakeMeshShapeType::All },
 	  { TEXT("bAlignShapeToPlacementSurface"), EMakeMeshShapeType::All },
-	  { TEXT("Slices"),        EMakeMeshShapeType::Cylinder | EMakeMeshShapeType::Cone | EMakeMeshShapeType::Arrow | EMakeMeshShapeType::RoundedRectangle | EMakeMeshShapeType::Disc | EMakeMeshShapeType::PuncturedDisc | EMakeMeshShapeType::Sphere },
+	  { TEXT("Slices"),        EMakeMeshShapeType::Cylinder | EMakeMeshShapeType::Cone | EMakeMeshShapeType::Arrow | EMakeMeshShapeType::RoundedRectangle | EMakeMeshShapeType::Disc | EMakeMeshShapeType::PuncturedDisc | EMakeMeshShapeType::Sphere | EMakeMeshShapeType::Torus },
 	  { TEXT("Subdivisions"),  EMakeMeshShapeType::Box | EMakeMeshShapeType::Rectangle | EMakeMeshShapeType::RoundedRectangle | EMakeMeshShapeType::Disc | EMakeMeshShapeType::PuncturedDisc | EMakeMeshShapeType::Cylinder |
-							   EMakeMeshShapeType::Cone | EMakeMeshShapeType::Arrow | EMakeMeshShapeType::SphericalBox }
+							   EMakeMeshShapeType::Cone | EMakeMeshShapeType::Arrow | EMakeMeshShapeType::SphericalBox | EMakeMeshShapeType::Torus }
 	};
 };
 
@@ -300,6 +300,10 @@ void UAddPrimitiveTool::UpdatePreviewMesh()
 		GenerateArrow(&NewMesh);
 		break;
 
+	case EMakeMeshShapeType::Torus:
+		GenerateTorus(&NewMesh);
+		break;
+
 	case EMakeMeshShapeType::Sphere:
 		GenerateSphere(&NewMesh);
 		break;
@@ -352,7 +356,7 @@ void UAddPrimitiveTool::OnClicked(const FInputDeviceRay& DeviceClickPos)
 	GetToolManager()->BeginUndoTransaction(LOCTEXT("AddPrimitiveToolTransactionName", "Add Primitive Mesh"));
 
 	const UEnum* const MakeMeshShapeTypeEnum = StaticEnum<EMakeMeshShapeType>();		
-	FString ShapeTypeName = MakeMeshShapeTypeEnum->GetNameStringByValue((uint8)ShapeSettings->Shape);
+	FString ShapeTypeName = MakeMeshShapeTypeEnum->GetNameStringByValue((uint32)ShapeSettings->Shape);
 
 	AActor* NewActor = AssetGenerationUtil::GenerateStaticMeshActor(
 		AssetAPI, TargetWorld,
@@ -429,6 +433,23 @@ void UAddPrimitiveTool::GeneratePuncturedDisc(FDynamicMesh3* OutMesh)
 }
 
 
+void UAddPrimitiveTool::GenerateTorus(FDynamicMesh3* OutMesh)
+{
+	FGeneralizedCylinderGenerator Gen;
+	Gen.CrossSection = FPolygon2d::MakeCircle(ShapeSettings->FeatureRadius, ShapeSettings->Slices);
+	FPolygon2d PathCircle = FPolygon2d::MakeCircle(ShapeSettings->Width*.5, ShapeSettings->Subdivisions+4);
+	for (int Idx = 0; Idx < PathCircle.VertexCount(); Idx++)
+	{
+		Gen.Path.Add(PathCircle[Idx]);
+	}
+	Gen.bLoop = true;
+	Gen.bCapped = false;
+	Gen.InitialFrame = FFrame3d(Gen.Path[0]);
+	Gen.Generate();
+	OutMesh->Copy(&Gen);
+}
+
+
 void UAddPrimitiveTool::GenerateCylinder(FDynamicMesh3* OutMesh)
 {
 	FCylinderGenerator CylGen;
diff --git a/Engine/Plugins/Experimental/MeshModelingToolset/Source/MeshModelingTools/Public/AddPrimitiveTool.h b/Engine/Plugins/Experimental/MeshModelingToolset/Source/MeshModelingTools/Public/AddPrimitiveTool.h
--- a/Engine/Plugins/Experimental/MeshModelingToolset/Source/MeshModelingTools/Public/AddPrimitiveTool.h
+++ b/Engine/Plugins/Experimental/MeshModelingToolset/Source/MeshModelingTools/Public/AddPrimitiveTool.h
@@ -47,8 +47,9 @@ enum class EMakeMeshShapeType : uint32
 	RoundedRectangle = 0x020 UMETA(DisplayName = "Rounded Rectangle"),
 	Disc			 = 0x040 UMETA(DisplayName = "Disc"),
 	PuncturedDisc	 = 0x080 UMETA(DisplayName = "Punctured Disc"),
-	Sphere			 = 0x100 UMETA(DisplayName = "Sphere"),
-	SphericalBox	 = 0x200 UMETA(DisplayName = "Spherical Box")
+	Torus			 = 0x100 UMETA(DisplayName = "Torus"),
+	Sphere			 = 0x200 UMETA(DisplayName = "Sphere"),
+	SphericalBox	 = 0x400 UMETA(DisplayName = "Spherical Box")
 	
 };
 ENUM_CLASS_FLAGS(EMakeMeshShapeType);
@@ -205,6 +206,7 @@ protected:
 	void GenerateRoundedRectangle(FDynamicMesh3* OutMesh);
 	void GenerateDisc(FDynamicMesh3* OutMesh);
 	void GeneratePuncturedDisc(FDynamicMesh3* OutMesh);
+	void GenerateTorus(FDynamicMesh3* OutMesh);
 	void GenerateSphere(FDynamicMesh3* OutMesh);
 	void GenerateSphericalBox(FDynamicMesh3* OutMesh);
 	void GenerateArrow(FDynamicMesh3* OutMesh);
