attempt CIS fix for clang -- directly include header

#rnx

diff --git a/Engine/Plugins/Experimental/GeometryProcessing/Source/GeometricObjects/Private/Generators/DiscMeshGenerator.cpp b/Engine/Plugins/Experimental/GeometryProcessing/Source/GeometricObjects/Private/Generators/DiscMeshGenerator.cpp
--- a/Engine/Plugins/Experimental/GeometryProcessing/Source/GeometricObjects/Private/Generators/DiscMeshGenerator.cpp
+++ b/Engine/Plugins/Experimental/GeometryProcessing/Source/GeometricObjects/Private/Generators/DiscMeshGenerator.cpp
@@ -1,6 +1,7 @@
 // Copyright 1998-2019 Epic Games, Inc. All Rights Reserved.
 
 #include "Generators/DiscMeshGenerator.h"
+#include "MathUtil.h"
 
 FDiscMeshGenerator::FDiscMeshGenerator()
 {
diff --git a/Engine/Plugins/Experimental/GeometryProcessing/Source/GeometricObjects/Private/Generators/RectangleMeshGenerator.cpp b/Engine/Plugins/Experimental/GeometryProcessing/Source/GeometricObjects/Private/Generators/RectangleMeshGenerator.cpp
--- a/Engine/Plugins/Experimental/GeometryProcessing/Source/GeometricObjects/Private/Generators/RectangleMeshGenerator.cpp
+++ b/Engine/Plugins/Experimental/GeometryProcessing/Source/GeometricObjects/Private/Generators/RectangleMeshGenerator.cpp
@@ -1,7 +1,7 @@
 // Copyright 1998-2019 Epic Games, Inc. All Rights Reserved.
 
 #include "Generators/RectangleMeshGenerator.h"
-
+#include "MathUtil.h"
 
 FRectangleMeshGenerator::FRectangleMeshGenerator()
 {
