Rework MathUtils constants to fix CIS build

diff --git a/Engine/Plugins/Experimental/GeometryProcessing/Source/GeometricObjects/Private/MathUtil.cpp b/Engine/Plugins/Experimental/GeometryProcessing/Source/GeometricObjects/Private/MathUtil.cpp
--- a/Engine/Plugins/Experimental/GeometryProcessing/Source/GeometricObjects/Private/MathUtil.cpp
+++ b/Engine/Plugins/Experimental/GeometryProcessing/Source/GeometricObjects/Private/MathUtil.cpp
@@ -2,49 +2,36 @@
 
 #include "MathUtil.h"
 
-#include <cfloat>
-#include <cmath>
-
 #include "Math/UnrealMath.h"
 
-template<> GEOMETRICOBJECTS_API const float TMathUtil<float>::Epsilon = FLT_EPSILON;
-template<> GEOMETRICOBJECTS_API const float TMathUtil<float>::ZeroTolerance = 1e-06f;
-template<> GEOMETRICOBJECTS_API const float TMathUtil<float>::MaxReal = FLT_MAX;
-template<> GEOMETRICOBJECTS_API const float TMathUtil<float>::Pi = (float)(4.0*atan(1.0));
-template<> GEOMETRICOBJECTS_API const float TMathUtil<float>::FourPi = 4.0f*TMathUtil<float>::Pi;
-template<> GEOMETRICOBJECTS_API const float TMathUtil<float>::TwoPi = 2.0f*TMathUtil<float>::Pi;
-template<> GEOMETRICOBJECTS_API const float TMathUtil<float>::HalfPi = 0.5f*TMathUtil<float>::Pi;
-template<> GEOMETRICOBJECTS_API const float TMathUtil<float>::InvPi = 1.0f / TMathUtil<float>::Pi;
-template<> GEOMETRICOBJECTS_API const float TMathUtil<float>::InvTwoPi = 1.0f / TMathUtil<float>::TwoPi;
-template<> GEOMETRICOBJECTS_API const float TMathUtil<float>::DegToRad = TMathUtil<float>::Pi / 180.0f;
-template<> GEOMETRICOBJECTS_API const float TMathUtil<float>::RadToDeg = 180.0f / TMathUtil<float>::Pi;
-//template<> const float TMathUtil<float>::LN_2 = TMathUtil<float>::Log(2.0f);
-//template<> const float TMathUtil<float>::LN_10 = TMathUtil<float>::Log(10.0f);
-//template<> const float TMathUtil<float>::INV_LN_2 = 1.0f / TMathUtil<float>::LN_2;
-//template<> const float TMathUtil<float>::INV_LN_10 = 1.0f / TMathUtil<float>::LN_10;
-template<> GEOMETRICOBJECTS_API const float TMathUtil<float>::Sqrt2 = (float)(sqrt(2.0));
-template<> GEOMETRICOBJECTS_API const float TMathUtil<float>::InvSqrt2 = 1.0f / TMathUtil<float>::Sqrt2;
-template<> GEOMETRICOBJECTS_API const float TMathUtil<float>::Sqrt3 = (float)(sqrt(3.0));
-template<> GEOMETRICOBJECTS_API const float TMathUtil<float>::InvSqrt3 = 1.0f / TMathUtil<float>::Sqrt3;
-
-template<> GEOMETRICOBJECTS_API const double TMathUtil<double>::Epsilon = DBL_EPSILON;
-template<> GEOMETRICOBJECTS_API const double TMathUtil<double>::ZeroTolerance = 1e-08;
-template<> GEOMETRICOBJECTS_API const double TMathUtil<double>::MaxReal = DBL_MAX;
-template<> GEOMETRICOBJECTS_API const double TMathUtil<double>::Pi = 4.0*atan(1.0);
-template<> GEOMETRICOBJECTS_API const double TMathUtil<double>::FourPi = 4.0*TMathUtil<double>::Pi;
-template<> GEOMETRICOBJECTS_API const double TMathUtil<double>::TwoPi = 2.0*TMathUtil<double>::Pi;
-template<> GEOMETRICOBJECTS_API const double TMathUtil<double>::HalfPi = 0.5*TMathUtil<double>::Pi;
-template<> GEOMETRICOBJECTS_API const double TMathUtil<double>::InvPi = 1.0 / TMathUtil<double>::Pi;
-template<> GEOMETRICOBJECTS_API const double TMathUtil<double>::InvTwoPi = 1.0 / TMathUtil<double>::TwoPi;
-template<> GEOMETRICOBJECTS_API const double TMathUtil<double>::DegToRad = TMathUtil<double>::Pi / 180.0;
-template<> GEOMETRICOBJECTS_API const double TMathUtil<double>::RadToDeg = 180.0 / TMathUtil<double>::Pi;
-//template<> const double TMathUtil<double>::LN_2 = TMathUtil<double>::Log(2.0);
-//template<> const double TMathUtil<double>::LN_10 = TMathUtil<double>::Log(10.0);
-//template<> const double TMathUtil<double>::INV_LN_2 = 1.0 / TMathUtil<double>::LN_2;
-//template<> const double TMathUtil<double>::INV_LN_10 = 1.0 / TMathUtil<double>::LN_10;
-template<> GEOMETRICOBJECTS_API const double TMathUtil<double>::Sqrt2 = sqrt(2.0);
-template<> GEOMETRICOBJECTS_API const double TMathUtil<double>::InvSqrt2 = 1.0f / TMathUtil<float>::Sqrt2;
-template<> GEOMETRICOBJECTS_API const double TMathUtil<double>::Sqrt3 = sqrt(3.0);
-template<> GEOMETRICOBJECTS_API const double TMathUtil<double>::InvSqrt3 = 1.0f / TMathUtil<float>::Sqrt3;
-
+constexpr float TMathUtilConstants<float>::Epsilon;
+constexpr float TMathUtilConstants<float>::ZeroTolerance;
+constexpr float TMathUtilConstants<float>::MaxReal;
+constexpr float TMathUtilConstants<float>::Pi;
+constexpr float TMathUtilConstants<float>::FourPi;
+constexpr float TMathUtilConstants<float>::TwoPi;
+constexpr float TMathUtilConstants<float>::HalfPi;
+constexpr float TMathUtilConstants<float>::InvPi;
+constexpr float TMathUtilConstants<float>::InvTwoPi;
+constexpr float TMathUtilConstants<float>::DegToRad;
+constexpr float TMathUtilConstants<float>::RadToDeg;
+constexpr float TMathUtilConstants<float>::Sqrt2;
+constexpr float TMathUtilConstants<float>::InvSqrt2;
+constexpr float TMathUtilConstants<float>::Sqrt3;
+constexpr float TMathUtilConstants<float>::InvSqrt3;
 
+constexpr double TMathUtilConstants<double>::Epsilon;
+constexpr double TMathUtilConstants<double>::ZeroTolerance;
+constexpr double TMathUtilConstants<double>::MaxReal;
+constexpr double TMathUtilConstants<double>::Pi;
+constexpr double TMathUtilConstants<double>::FourPi;
+constexpr double TMathUtilConstants<double>::TwoPi;
+constexpr double TMathUtilConstants<double>::HalfPi;
+constexpr double TMathUtilConstants<double>::InvPi;
+constexpr double TMathUtilConstants<double>::InvTwoPi;
+constexpr double TMathUtilConstants<double>::DegToRad;
+constexpr double TMathUtilConstants<double>::RadToDeg;
+constexpr double TMathUtilConstants<double>::Sqrt2;
+constexpr double TMathUtilConstants<double>::InvSqrt2;
+constexpr double TMathUtilConstants<double>::Sqrt3;
+constexpr double TMathUtilConstants<double>::InvSqrt3;
diff --git a/Engine/Plugins/Experimental/GeometryProcessing/Source/GeometricObjects/Public/MathUtil.h b/Engine/Plugins/Experimental/GeometryProcessing/Source/GeometricObjects/Public/MathUtil.h
--- a/Engine/Plugins/Experimental/GeometryProcessing/Source/GeometricObjects/Public/MathUtil.h
+++ b/Engine/Plugins/Experimental/GeometryProcessing/Source/GeometricObjects/Public/MathUtil.h
@@ -4,48 +4,95 @@
 
 #include "CoreTypes.h"   // required for GEOMETRICOBJECTS_API macro
 #include <cmath>
+#include <cfloat>
+
 
 /**
  * Math constants and utility functions, templated on float/double type
  */
 template<typename RealType>
-class TMathUtil
+struct TMathUtilConstants;
+
+template<>
+struct TMathUtilConstants<float>
+{
+	/** Machine Epsilon - float approx 1e-7, double approx 2e-16 */
+	GEOMETRICOBJECTS_API static constexpr float Epsilon = FLT_EPSILON;
+	/** Zero tolerance for math operations (eg like parallel tests) - float 1e-6, double 1e-8 */
+	GEOMETRICOBJECTS_API static constexpr float ZeroTolerance = 1e-06f;
+
+	/** largest possible number for type */
+	GEOMETRICOBJECTS_API static constexpr float MaxReal = FLT_MAX;
+
+	/** 3.14159... */
+	GEOMETRICOBJECTS_API static constexpr float Pi = 3.1415926535897932384626433832795f;
+	GEOMETRICOBJECTS_API static constexpr float FourPi = 4.0f * TMathUtilConstants::Pi;
+	GEOMETRICOBJECTS_API static constexpr float TwoPi = 2.0f*TMathUtilConstants::Pi;
+	GEOMETRICOBJECTS_API static constexpr float HalfPi = 0.5f*TMathUtilConstants::Pi;
+
+	/** 1.0 / Pi */
+	GEOMETRICOBJECTS_API static constexpr float InvPi = 1.0f / TMathUtilConstants::Pi;
+	/** 1.0 / (2*Pi) */
+	GEOMETRICOBJECTS_API static constexpr float InvTwoPi = 1.0f / TMathUtilConstants::TwoPi;
+
+	/** pi / 180 */
+	GEOMETRICOBJECTS_API static constexpr float DegToRad = TMathUtilConstants::Pi / 180.0f;
+	/** 180 / pi */
+	GEOMETRICOBJECTS_API static constexpr float RadToDeg = 180.0f / TMathUtilConstants::Pi;
+
+	//static constexpr float LN_2;
+	//static constexpr float LN_10;
+	//static constexpr float INV_LN_2;
+	//static constexpr float INV_LN_10;
+
+	GEOMETRICOBJECTS_API static constexpr float Sqrt2 = 1.4142135623730950488016887242097f;
+	GEOMETRICOBJECTS_API static constexpr float InvSqrt2 = 1.0f / TMathUtilConstants::Sqrt2;
+	GEOMETRICOBJECTS_API static constexpr float Sqrt3 = 1.7320508075688772935274463415059f;
+	GEOMETRICOBJECTS_API static constexpr float InvSqrt3 = 1.0f / TMathUtilConstants::Sqrt3;
+};
+
+template<>
+struct TMathUtilConstants<double>
 {
-public:
 	/** Machine Epsilon - float approx 1e-7, double approx 2e-16 */
-	GEOMETRICOBJECTS_API static const RealType Epsilon;
+	GEOMETRICOBJECTS_API static constexpr double Epsilon = DBL_EPSILON;
 	/** Zero tolerance for math operations (eg like parallel tests) - float 1e-6, double 1e-8 */
-	GEOMETRICOBJECTS_API static const RealType ZeroTolerance;
+	GEOMETRICOBJECTS_API static constexpr double ZeroTolerance = 1e-08;
 
 	/** largest possible number for type */
-	GEOMETRICOBJECTS_API static const RealType MaxReal;
+	GEOMETRICOBJECTS_API static constexpr double MaxReal = DBL_MAX;
 
 	/** 3.14159... */
-	GEOMETRICOBJECTS_API static const RealType Pi;
-	GEOMETRICOBJECTS_API static const RealType FourPi;
-	GEOMETRICOBJECTS_API static const RealType TwoPi;
-	GEOMETRICOBJECTS_API static const RealType HalfPi;
+	GEOMETRICOBJECTS_API static constexpr double Pi = 3.1415926535897932384626433832795;
+	GEOMETRICOBJECTS_API static constexpr double FourPi = 4.0 * Pi;
+	GEOMETRICOBJECTS_API static constexpr double TwoPi = 2.0*TMathUtilConstants::Pi;
+	GEOMETRICOBJECTS_API static constexpr double HalfPi = 0.5*TMathUtilConstants::Pi;
 
 	/** 1.0 / Pi */
-	GEOMETRICOBJECTS_API static const RealType InvPi;
+	GEOMETRICOBJECTS_API static constexpr double InvPi = 1.0 / TMathUtilConstants::Pi;
 	/** 1.0 / (2*Pi) */
-	GEOMETRICOBJECTS_API static const RealType InvTwoPi;
+	GEOMETRICOBJECTS_API static constexpr double InvTwoPi = 1.0 / TMathUtilConstants::TwoPi;
 
 	/** pi / 180 */
-	GEOMETRICOBJECTS_API static const RealType DegToRad;
+	GEOMETRICOBJECTS_API static constexpr double DegToRad = TMathUtilConstants::Pi / 180.0;
 	/** 180 / pi */
-	GEOMETRICOBJECTS_API static const RealType RadToDeg;
+	GEOMETRICOBJECTS_API static constexpr double RadToDeg = 180.0 / TMathUtilConstants::Pi;
 
-	//static const RealType LN_2;
-	//static const RealType LN_10;
-	//static const RealType INV_LN_2;
-	//static const RealType INV_LN_10;
+	//static constexpr double LN_2;
+	//static constexpr double LN_10;
+	//static constexpr double INV_LN_2;
+	//static constexpr double INV_LN_10;
 
-	GEOMETRICOBJECTS_API static const RealType Sqrt2;
-	GEOMETRICOBJECTS_API static const RealType InvSqrt2;
-	GEOMETRICOBJECTS_API static const RealType Sqrt3;
-	GEOMETRICOBJECTS_API static const RealType InvSqrt3;
+	GEOMETRICOBJECTS_API static constexpr double Sqrt2 = 1.4142135623730950488016887242097;
+	GEOMETRICOBJECTS_API static constexpr double InvSqrt2 = 1.0 / TMathUtilConstants::Sqrt2;
+	GEOMETRICOBJECTS_API static constexpr double Sqrt3 = 1.7320508075688772935274463415059;
+	GEOMETRICOBJECTS_API static constexpr double InvSqrt3 = 1.0 / TMathUtilConstants::Sqrt3;
+};
 
+template<typename RealType>
+class TMathUtil : public TMathUtilConstants<RealType>
+{
+public:
 	static inline bool IsNaN(const RealType Value);
 	static inline bool IsFinite(const RealType Value);
 	static inline RealType Abs(const RealType Value);
