Revert constexpr math constants back to static const to fix non-strict VS2017 build
#rnx

diff --git a/Engine/Plugins/Experimental/GeometryProcessing/Source/GeometricObjects/Private/MathUtil.cpp b/Engine/Plugins/Experimental/GeometryProcessing/Source/GeometricObjects/Private/MathUtil.cpp
--- a/Engine/Plugins/Experimental/GeometryProcessing/Source/GeometricObjects/Private/MathUtil.cpp
+++ b/Engine/Plugins/Experimental/GeometryProcessing/Source/GeometricObjects/Private/MathUtil.cpp
@@ -4,34 +4,36 @@
 
 #include "Math/UnrealMath.h"
 
-constexpr float TMathUtilConstants<float>::Epsilon;
-constexpr float TMathUtilConstants<float>::ZeroTolerance;
-constexpr float TMathUtilConstants<float>::MaxReal;
-constexpr float TMathUtilConstants<float>::Pi;
-constexpr float TMathUtilConstants<float>::FourPi;
-constexpr float TMathUtilConstants<float>::TwoPi;
-constexpr float TMathUtilConstants<float>::HalfPi;
-constexpr float TMathUtilConstants<float>::InvPi;
-constexpr float TMathUtilConstants<float>::InvTwoPi;
-constexpr float TMathUtilConstants<float>::DegToRad;
-constexpr float TMathUtilConstants<float>::RadToDeg;
-constexpr float TMathUtilConstants<float>::Sqrt2;
-constexpr float TMathUtilConstants<float>::InvSqrt2;
-constexpr float TMathUtilConstants<float>::Sqrt3;
-constexpr float TMathUtilConstants<float>::InvSqrt3;
+#include <cfloat>
 
-constexpr double TMathUtilConstants<double>::Epsilon;
-constexpr double TMathUtilConstants<double>::ZeroTolerance;
-constexpr double TMathUtilConstants<double>::MaxReal;
-constexpr double TMathUtilConstants<double>::Pi;
-constexpr double TMathUtilConstants<double>::FourPi;
-constexpr double TMathUtilConstants<double>::TwoPi;
-constexpr double TMathUtilConstants<double>::HalfPi;
-constexpr double TMathUtilConstants<double>::InvPi;
-constexpr double TMathUtilConstants<double>::InvTwoPi;
-constexpr double TMathUtilConstants<double>::DegToRad;
-constexpr double TMathUtilConstants<double>::RadToDeg;
-constexpr double TMathUtilConstants<double>::Sqrt2;
-constexpr double TMathUtilConstants<double>::InvSqrt2;
-constexpr double TMathUtilConstants<double>::Sqrt3;
-constexpr double TMathUtilConstants<double>::InvSqrt3;
+const float TMathUtilConstants<float>::Epsilon = FLT_EPSILON;
+const float TMathUtilConstants<float>::ZeroTolerance = 1e-06f;
+const float TMathUtilConstants<float>::MaxReal = FLT_MAX;
+const float TMathUtilConstants<float>::Pi = 3.1415926535897932384626433832795f;
+const float TMathUtilConstants<float>::FourPi = 4.0f * TMathUtilConstants::Pi;
+const float TMathUtilConstants<float>::TwoPi = 2.0f * TMathUtilConstants::Pi;
+const float TMathUtilConstants<float>::HalfPi = 0.5f * TMathUtilConstants::Pi;
+const float TMathUtilConstants<float>::InvPi = 1.0f / TMathUtilConstants::Pi;
+const float TMathUtilConstants<float>::InvTwoPi = 1.0f / TMathUtilConstants::TwoPi;
+const float TMathUtilConstants<float>::DegToRad = TMathUtilConstants::Pi / 180.0f;
+const float TMathUtilConstants<float>::RadToDeg = 180.0f / TMathUtilConstants::Pi;
+const float TMathUtilConstants<float>::Sqrt2 = 1.4142135623730950488016887242097f;
+const float TMathUtilConstants<float>::InvSqrt2 = 1.0f / TMathUtilConstants::Sqrt2;
+const float TMathUtilConstants<float>::Sqrt3 = 1.7320508075688772935274463415059f;
+const float TMathUtilConstants<float>::InvSqrt3 = 1.0f / TMathUtilConstants::Sqrt3;
+
+const double TMathUtilConstants<double>::Epsilon = DBL_EPSILON;
+const double TMathUtilConstants<double>::ZeroTolerance = 1e-08;
+const double TMathUtilConstants<double>::MaxReal = DBL_MAX;
+const double TMathUtilConstants<double>::Pi = 3.1415926535897932384626433832795;
+const double TMathUtilConstants<double>::FourPi = 4.0 * Pi;
+const double TMathUtilConstants<double>::TwoPi = 2.0 * TMathUtilConstants::Pi;
+const double TMathUtilConstants<double>::HalfPi = 0.5 * TMathUtilConstants::Pi;
+const double TMathUtilConstants<double>::InvPi = 1.0 / TMathUtilConstants::Pi;
+const double TMathUtilConstants<double>::InvTwoPi = 1.0 / TMathUtilConstants::TwoPi;
+const double TMathUtilConstants<double>::DegToRad = TMathUtilConstants::Pi / 180.0;
+const double TMathUtilConstants<double>::RadToDeg = 180.0 / TMathUtilConstants::Pi;
+const double TMathUtilConstants<double>::Sqrt2 = 1.4142135623730950488016887242097;
+const double TMathUtilConstants<double>::InvSqrt2 = 1.0 / TMathUtilConstants::Sqrt2;
+const double TMathUtilConstants<double>::Sqrt3 = 1.7320508075688772935274463415059;
+const double TMathUtilConstants<double>::InvSqrt3 = 1.0 / TMathUtilConstants::Sqrt3;
diff --git a/Engine/Plugins/Experimental/GeometryProcessing/Source/GeometricObjects/Public/MathUtil.h b/Engine/Plugins/Experimental/GeometryProcessing/Source/GeometricObjects/Public/MathUtil.h
--- a/Engine/Plugins/Experimental/GeometryProcessing/Source/GeometricObjects/Public/MathUtil.h
+++ b/Engine/Plugins/Experimental/GeometryProcessing/Source/GeometricObjects/Public/MathUtil.h
@@ -4,7 +4,6 @@
 
 #include "CoreTypes.h"   // required for GEOMETRICOBJECTS_API macro
 #include <cmath>
-#include <cfloat>
 
 
 /**
@@ -17,76 +16,66 @@ template<>
 struct TMathUtilConstants<float>
 {
 	/** Machine Epsilon - float approx 1e-7, double approx 2e-16 */
-	GEOMETRICOBJECTS_API static constexpr float Epsilon = FLT_EPSILON;
+	GEOMETRICOBJECTS_API static const float Epsilon;
 	/** Zero tolerance for math operations (eg like parallel tests) - float 1e-6, double 1e-8 */
-	GEOMETRICOBJECTS_API static constexpr float ZeroTolerance = 1e-06f;
+	GEOMETRICOBJECTS_API static const float ZeroTolerance;
 
 	/** largest possible number for type */
-	GEOMETRICOBJECTS_API static constexpr float MaxReal = FLT_MAX;
+	GEOMETRICOBJECTS_API static const float MaxReal;
 
 	/** 3.14159... */
-	GEOMETRICOBJECTS_API static constexpr float Pi = 3.1415926535897932384626433832795f;
-	GEOMETRICOBJECTS_API static constexpr float FourPi = 4.0f * TMathUtilConstants::Pi;
-	GEOMETRICOBJECTS_API static constexpr float TwoPi = 2.0f*TMathUtilConstants::Pi;
-	GEOMETRICOBJECTS_API static constexpr float HalfPi = 0.5f*TMathUtilConstants::Pi;
+	GEOMETRICOBJECTS_API static const float Pi;
+	GEOMETRICOBJECTS_API static const float FourPi;
+	GEOMETRICOBJECTS_API static const float TwoPi;
+	GEOMETRICOBJECTS_API static const float HalfPi;
 
 	/** 1.0 / Pi */
-	GEOMETRICOBJECTS_API static constexpr float InvPi = 1.0f / TMathUtilConstants::Pi;
+	GEOMETRICOBJECTS_API static const float InvPi;
 	/** 1.0 / (2*Pi) */
-	GEOMETRICOBJECTS_API static constexpr float InvTwoPi = 1.0f / TMathUtilConstants::TwoPi;
+	GEOMETRICOBJECTS_API static const float InvTwoPi;
 
 	/** pi / 180 */
-	GEOMETRICOBJECTS_API static constexpr float DegToRad = TMathUtilConstants::Pi / 180.0f;
+	GEOMETRICOBJECTS_API static const float DegToRad;
 	/** 180 / pi */
-	GEOMETRICOBJECTS_API static constexpr float RadToDeg = 180.0f / TMathUtilConstants::Pi;
+	GEOMETRICOBJECTS_API static const float RadToDeg;
 
-	//static constexpr float LN_2;
-	//static constexpr float LN_10;
-	//static constexpr float INV_LN_2;
-	//static constexpr float INV_LN_10;
-
-	GEOMETRICOBJECTS_API static constexpr float Sqrt2 = 1.4142135623730950488016887242097f;
-	GEOMETRICOBJECTS_API static constexpr float InvSqrt2 = 1.0f / TMathUtilConstants::Sqrt2;
-	GEOMETRICOBJECTS_API static constexpr float Sqrt3 = 1.7320508075688772935274463415059f;
-	GEOMETRICOBJECTS_API static constexpr float InvSqrt3 = 1.0f / TMathUtilConstants::Sqrt3;
+	GEOMETRICOBJECTS_API static const float Sqrt2;
+	GEOMETRICOBJECTS_API static const float InvSqrt2;
+	GEOMETRICOBJECTS_API static const float Sqrt3;
+	GEOMETRICOBJECTS_API static const float InvSqrt3;
 };
 
 template<>
 struct TMathUtilConstants<double>
 {
 	/** Machine Epsilon - float approx 1e-7, double approx 2e-16 */
-	GEOMETRICOBJECTS_API static constexpr double Epsilon = DBL_EPSILON;
+	GEOMETRICOBJECTS_API static const double Epsilon;
 	/** Zero tolerance for math operations (eg like parallel tests) - float 1e-6, double 1e-8 */
-	GEOMETRICOBJECTS_API static constexpr double ZeroTolerance = 1e-08;
+	GEOMETRICOBJECTS_API static const double ZeroTolerance;
 
 	/** largest possible number for type */
-	GEOMETRICOBJECTS_API static constexpr double MaxReal = DBL_MAX;
+	GEOMETRICOBJECTS_API static const double MaxReal;
 
 	/** 3.14159... */
-	GEOMETRICOBJECTS_API static constexpr double Pi = 3.1415926535897932384626433832795;
-	GEOMETRICOBJECTS_API static constexpr double FourPi = 4.0 * Pi;
-	GEOMETRICOBJECTS_API static constexpr double TwoPi = 2.0*TMathUtilConstants::Pi;
-	GEOMETRICOBJECTS_API static constexpr double HalfPi = 0.5*TMathUtilConstants::Pi;
+	GEOMETRICOBJECTS_API static const double Pi;
+	GEOMETRICOBJECTS_API static const double FourPi;
+	GEOMETRICOBJECTS_API static const double TwoPi;
+	GEOMETRICOBJECTS_API static const double HalfPi;
 
 	/** 1.0 / Pi */
-	GEOMETRICOBJECTS_API static constexpr double InvPi = 1.0 / TMathUtilConstants::Pi;
+	GEOMETRICOBJECTS_API static const double InvPi;
 	/** 1.0 / (2*Pi) */
-	GEOMETRICOBJECTS_API static constexpr double InvTwoPi = 1.0 / TMathUtilConstants::TwoPi;
+	GEOMETRICOBJECTS_API static const double InvTwoPi;
 
 	/** pi / 180 */
-	GEOMETRICOBJECTS_API static constexpr double DegToRad = TMathUtilConstants::Pi / 180.0;
+	GEOMETRICOBJECTS_API static const double DegToRad;
 	/** 180 / pi */
-	GEOMETRICOBJECTS_API static constexpr double RadToDeg = 180.0 / TMathUtilConstants::Pi;
-
-	//static constexpr double LN_2;
-	//static constexpr double LN_10;
-	//static constexpr double INV_LN_2;
-	//static constexpr double INV_LN_10;
+	GEOMETRICOBJECTS_API static const double RadToDeg;
 
-	GEOMETRICOBJECTS_API static constexpr double Sqrt2 = 1.4142135623730950488016887242097;
-	GEOMETRICOBJECTS_API static constexpr double InvSqrt2 = 1.0 / TMathUtilConstants::Sqrt2;
-	GEOMETRICOBJECTS_API static constexpr double Sqrt3 = 1.7320508075688772935274463415059;
-	GEOMETRICOBJECTS_API static constexpr double InvSqrt3 = 1.0 / TMathUtilConstants::Sqrt3;
+	GEOMETRICOBJECTS_API static const double Sqrt2;
+	GEOMETRICOBJECTS_API static const double InvSqrt2;
+	GEOMETRICOBJECTS_API static const double Sqrt3;
+	GEOMETRICOBJECTS_API static const double InvSqrt3;
 };
 
 template<typename RealType>
