cis fix -- missed a file in holey circle and rounded rect shelf

#rnx

diff --git a/Engine/Plugins/Experimental/GeometryProcessing/Source/GeometryAlgorithms/Public/ConstrainedDelaunay2.h b/Engine/Plugins/Experimental/GeometryProcessing/Source/GeometryAlgorithms/Public/ConstrainedDelaunay2.h
--- a/Engine/Plugins/Experimental/GeometryProcessing/Source/GeometryAlgorithms/Public/ConstrainedDelaunay2.h
+++ b/Engine/Plugins/Experimental/GeometryProcessing/Source/GeometryAlgorithms/Public/ConstrainedDelaunay2.h
@@ -23,6 +23,9 @@ struct TConstrainedDelaunay2
 	TArray<FIndex2i> Edges;	// Edges can be boundaries or not based on the EFillRule setting
 	TArray<FIndex2i> HoleEdges; // Any triangles inside 'hole' edges *must* be cut out
 
+	//TODO: also support FeatureEdges?
+	//TArray<FIndex2i> FeatureEdges; // Edges that should be preserved in mesh but do not correspond to a boundary and should not affect inside/outside classification at all
+
 	bool bOrientedEdges = true;
 	bool bOutputCCW = false;
 
@@ -79,3 +82,16 @@ TArray<FIndex3i> GEOMETRYALGORITHMS_API ConstrainedDelaunayTriangulate(const TGe
 typedef TConstrainedDelaunay2<float> FConstrainedDelaunay2f;
 typedef TConstrainedDelaunay2<double> FConstrainedDelaunay2d;
 
+
+template void GEOMETRYALGORITHMS_API TConstrainedDelaunay2<double>::Add(const FDynamicGraph2<double>& Graph);
+template void GEOMETRYALGORITHMS_API TConstrainedDelaunay2<double>::Add(const TGeneralPolygon2<double>& Polygon);
+template void GEOMETRYALGORITHMS_API TConstrainedDelaunay2<double>::Add(const TGeneralPolygon2<float>& Polygon);
+template void GEOMETRYALGORITHMS_API TConstrainedDelaunay2<double>::Add(const TPolygon2<double>& Polygon, bool bIsHole);
+template void GEOMETRYALGORITHMS_API TConstrainedDelaunay2<double>::Add(const TPolygon2<float>& Polygon, bool bIsHole);
+
+template void GEOMETRYALGORITHMS_API TConstrainedDelaunay2<float>::Add(const FDynamicGraph2<double>& Graph);
+template void GEOMETRYALGORITHMS_API TConstrainedDelaunay2<float>::Add(const TGeneralPolygon2<double>& Polygon);
+template void GEOMETRYALGORITHMS_API TConstrainedDelaunay2<float>::Add(const TGeneralPolygon2<float>& Polygon);
+template void GEOMETRYALGORITHMS_API TConstrainedDelaunay2<float>::Add(const TPolygon2<double>& Polygon, bool bIsHole);
+template void GEOMETRYALGORITHMS_API TConstrainedDelaunay2<float>::Add(const TPolygon2<float>& Polygon, bool bIsHole);
+
