GeometryProcessing: add mesh octree FindNearestHitObject variant that takes a predicate
#rnx

diff --git a/Engine/Plugins/Experimental/GeometryProcessing/Source/DynamicMesh/Public/DynamicMeshOctree3.h b/Engine/Plugins/Experimental/GeometryProcessing/Source/DynamicMesh/Public/DynamicMeshOctree3.h
--- a/Engine/Plugins/Experimental/GeometryProcessing/Source/DynamicMesh/Public/DynamicMeshOctree3.h
+++ b/Engine/Plugins/Experimental/GeometryProcessing/Source/DynamicMesh/Public/DynamicMeshOctree3.h
@@ -181,6 +181,29 @@ public:
 	}
 
 
+	/**
+	 * Find the nearest triangle of the mesh that is hit by the ray
+	 * @param IncludeTriangleIDFunc predicate function that must return true for given TriangleID for it to be considered
+	 */
+	int32 FindNearestHitObject(const FRay3d& Ray,
+		TFunctionRef<bool(int)> IncludeTriangleIDFunc,
+		double MaxDistance = TNumericLimits<double>::Max()) const
+	{
+		return FSparseDynamicOctree3::FindNearestHitObject(Ray,
+			[&](int tid) {return Mesh->GetTriBounds(tid); },
+			[&](int tid, const FRay3d& Ray) {
+				if (IncludeTriangleIDFunc(tid) == false)
+				{
+					return TNumericLimits<double>::Max();
+				}
+				FIntrRay3Triangle3d Intr = TMeshQueries<FDynamicMesh3>::TriangleIntersection(*Mesh, tid, Ray);
+				return (Intr.IntersectionType == EIntersectionType::Point) ?
+					Intr.RayParameter : TNumericLimits<double>::Max();
+			}, 
+			MaxDistance);
+	}
+
+
 	/**
 	 * Check that the Octree is internally valid
 	 */
