Clean up FDynamicMesh3, making it a regular type and enabling move semantics

#rnx

diff --git a/Engine/Plugins/Experimental/GeometryProcessing/Source/DynamicMesh/Private/DynamicMesh3.cpp b/Engine/Plugins/Experimental/GeometryProcessing/Source/DynamicMesh/Private/DynamicMesh3.cpp
--- a/Engine/Plugins/Experimental/GeometryProcessing/Source/DynamicMesh/Private/DynamicMesh3.cpp
+++ b/Engine/Plugins/Experimental/GeometryProcessing/Source/DynamicMesh/Private/DynamicMesh3.cpp
@@ -3,105 +3,143 @@
 #include "DynamicMesh3.h"
 #include "DynamicMeshAttributeSet.h"
 #include "Generators/MeshShapeGenerator.h"
+#include "Templates/UniquePtr.h"
 
+// NB: These have to be here until C++17 allows inline variables
+constexpr int       FDynamicMesh3::InvalidID;
+constexpr int       FDynamicMesh3::NonManifoldID;
+constexpr int       FDynamicMesh3::InvalidGroupID;
+constexpr FVector3d FDynamicMesh3::InvalidVertex;
+constexpr FIndex3i  FDynamicMesh3::InvalidTriangle;
+constexpr FIndex2i  FDynamicMesh3::InvalidEdge;
 
-const int FDynamicMesh3::InvalidID = IndexConstants::InvalidID; // -1;
-const int FDynamicMesh3::NonManifoldID = -2;
-const int FDynamicMesh3::InvalidGroupID = IndexConstants::InvalidID; // -1;
+FDynamicMesh3::FDynamicMesh3(bool bWantNormals, bool bWantColors, bool bWantUVs, bool bWantTriGroups)
+{
+	if ( bWantNormals )   { VertexNormals = TDynamicVector<float>{}; }
+	if ( bWantColors )    { VertexColors = TDynamicVector<float>{}; }
+	if ( bWantUVs )       { VertexUVs = TDynamicVector<float>{}; }
+	if ( bWantTriGroups ) { TriangleGroups = TDynamicVector<int>{}; }
+}
 
+// normals/colors/uvs will only be copied if they exist
+FDynamicMesh3::FDynamicMesh3(const FDynamicMesh3& Other)
+	:
+	Vertices{ Other.Vertices },
+	VertexRefCounts{ Other.VertexRefCounts },
+	VertexNormals{ Other.VertexNormals },
+	VertexColors{ Other.VertexColors },
+	VertexUVs{ Other.VertexUVs },
+	VertexEdgeLists{ Other.VertexEdgeLists },
 
-FDynamicMesh3::~FDynamicMesh3() 
-{ 
-	if (VertexNormals != nullptr)
-	{
-		delete VertexNormals;
-	}
-	if (VertexColors != nullptr)
-	{
-		delete VertexColors;
-	}
-	if (VertexUVs != nullptr)
-	{
-		delete VertexUVs;
-	}
-	if (TriangleGroups != nullptr)
-	{
-		delete TriangleGroups;
-	}
-	if (AttributeSet != nullptr)
-	{
-		delete AttributeSet;
-	}
-}
+	Triangles{ Other.Triangles },
+	TriangleRefCounts{ Other.TriangleRefCounts },
+	TriangleEdges{ Other.TriangleEdges },
+	TriangleGroups{ Other.TriangleGroups },
+	GroupIDCounter{ Other.GroupIDCounter },
 
+	Edges{ Other.Edges },
+	EdgeRefCounts{ Other.EdgeRefCounts },
 
+	Timestamp{ Other.Timestamp },
+	ShapeTimestamp{ Other.ShapeTimestamp },
+	TopologyTimestamp{ Other.TopologyTimestamp },
 
+	CachedBoundingBox{ Other.CachedBoundingBox },
+	CachedBoundingBoxTimestamp{ Other.CachedBoundingBoxTimestamp },
 
-FDynamicMesh3::FDynamicMesh3(bool bWantNormals, bool bWantColors, bool bWantUVs, bool bWantTriGroups)
+	bIsClosedCached{ Other.bIsClosedCached },
+	CachedIsClosedTimestamp{ Other.CachedIsClosedTimestamp }
 {
-	Vertices = TDynamicVector<double>();
-	VertexNormals = (bWantNormals) ? new TDynamicVector<float>() : nullptr;
-	VertexColors = (bWantColors) ? new TDynamicVector<float>() : nullptr;
-	VertexUVs = (bWantUVs) ? new TDynamicVector<float>() : nullptr;
+	if (Other.HasAttributes())
+	{
+		EnableAttributes();
+		AttributeSet->Copy(*Other.AttributeSet);
+	}
+}
 
-	VertexEdgeLists = FSmallListSet();
+FDynamicMesh3::FDynamicMesh3(FDynamicMesh3&& Other)
+	:
+	Vertices{ MoveTemp(Other.Vertices) },
+	VertexRefCounts{ MoveTemp(Other.VertexRefCounts) },
+	VertexNormals{ MoveTemp(Other.VertexNormals) },
+	VertexColors{ MoveTemp(Other.VertexColors) },
+	VertexUVs{ MoveTemp( Other.VertexUVs ) },
+	VertexEdgeLists{ MoveTemp( Other.VertexEdgeLists ) },
 
-	VertexRefCounts = FRefCountVector();
+	Triangles{ MoveTemp( Other.Triangles ) },
+	TriangleRefCounts{ MoveTemp( Other.TriangleRefCounts ) },
+	TriangleEdges{ MoveTemp( Other.TriangleEdges ) },
+	TriangleGroups{ MoveTemp( Other.TriangleGroups ) },
+	GroupIDCounter{ MoveTemp( Other.GroupIDCounter ) },
 
-	Triangles = TDynamicVector<int>();
-	TriangleEdges = TDynamicVector<int>();
-	TriangleRefCounts = FRefCountVector();
-	TriangleGroups = (bWantTriGroups) ? new TDynamicVector<int>() : nullptr;
-	GroupIDCounter = 0;
+	Edges{ MoveTemp( Other.Edges ) },
+	EdgeRefCounts{ MoveTemp( Other.EdgeRefCounts ) },
 
-	Edges = TDynamicVector<int>();
-	EdgeRefCounts = FRefCountVector();
+	AttributeSet{ MoveTemp( Other.AttributeSet ) },
+	Timestamp{ MoveTemp( Other.Timestamp ) },
+	ShapeTimestamp{ MoveTemp( Other.ShapeTimestamp ) },
+	TopologyTimestamp{ MoveTemp( Other.TopologyTimestamp ) },
 
-	AttributeSet = nullptr;
+	CachedBoundingBox{ MoveTemp( Other.CachedBoundingBox ) },
+	CachedBoundingBoxTimestamp{ MoveTemp( Other.CachedBoundingBoxTimestamp ) },
+
+	bIsClosedCached{ Other.bIsClosedCached },
+	CachedIsClosedTimestamp{ MoveTemp( Other.CachedIsClosedTimestamp ) }
+{
+	if (AttributeSet)
+	{
+		AttributeSet->Reparent(this);
+	}
 }
 
+FDynamicMesh3::~FDynamicMesh3() = default;
 
-FDynamicMesh3::FDynamicMesh3(const FDynamicMesh3& CopyMesh, bool bCompact, EMeshComponents flags)
-	: FDynamicMesh3(CopyMesh, bCompact, ((int)flags & (int)EMeshComponents::VertexNormals) != 0, ((int)flags & (int)EMeshComponents::VertexColors) != 0,
-	((int)flags & (int)EMeshComponents::VertexUVs) != 0)
+const FDynamicMesh3& FDynamicMesh3::operator=(const FDynamicMesh3& CopyMesh)
 {
+	Copy(CopyMesh);
+	return *this;
 }
 
+const FDynamicMesh3 & FDynamicMesh3::operator=(FDynamicMesh3 && Other)
+{
+	Vertices = MoveTemp(Other.Vertices);
+	VertexRefCounts = MoveTemp(Other.VertexRefCounts);
+	VertexNormals = MoveTemp(Other.VertexNormals);
+	VertexColors = MoveTemp(Other.VertexColors);
+	VertexUVs = MoveTemp(Other.VertexUVs);
+	VertexEdgeLists = MoveTemp(Other.VertexEdgeLists);
 
+	Triangles = MoveTemp(Other.Triangles);
+	TriangleRefCounts = MoveTemp(Other.TriangleRefCounts);
+	TriangleEdges = MoveTemp(Other.TriangleEdges);
+	TriangleGroups = MoveTemp(Other.TriangleGroups);
+	GroupIDCounter = MoveTemp(Other.GroupIDCounter);
 
+	Edges = MoveTemp(Other.Edges);
+	EdgeRefCounts = MoveTemp(Other.EdgeRefCounts);
 
-// normals/colors/uvs will only be copied if they exist
-FDynamicMesh3::FDynamicMesh3(const FDynamicMesh3& CopyMesh, bool bCompact, bool bWantNormals, bool bWantColors, bool bWantUVs, bool bWantAttributes)
-{
-	VertexNormals = nullptr;
-	VertexColors = nullptr;
-	VertexUVs = nullptr;
-	TriangleGroups = nullptr;
-	AttributeSet = nullptr;
-
-	if (bCompact)
-	{
-		CompactCopy(CopyMesh, bWantNormals, bWantColors, bWantUVs, bWantAttributes);
-	}
-	else
+	AttributeSet = MoveTemp(Other.AttributeSet);
+	if (AttributeSet)
 	{
-		Copy(CopyMesh, bWantNormals, bWantColors, bWantUVs, bWantAttributes);
+		AttributeSet->Reparent(this);
 	}
-}
+	Timestamp = MoveTemp(Other.Timestamp);
+	ShapeTimestamp = MoveTemp(Other.ShapeTimestamp);
+	TopologyTimestamp = MoveTemp(Other.TopologyTimestamp);
 
+	CachedBoundingBox = MoveTemp(Other.CachedBoundingBox);
+	CachedBoundingBoxTimestamp = MoveTemp(Other.CachedBoundingBoxTimestamp);
 
-
-const FDynamicMesh3& FDynamicMesh3::operator=(const FDynamicMesh3& CopyMesh)
-{
-	Copy(CopyMesh);
+	bIsClosedCached = MoveTemp(Other.bIsClosedCached);
+	CachedIsClosedTimestamp = MoveTemp( Other.CachedIsClosedTimestamp);
 	return *this;
 }
 
-
 FDynamicMesh3::FDynamicMesh3(const FMeshShapeGenerator* Generator)
 {
 	Copy(Generator);
 }
+
 void FDynamicMesh3::Copy(const FMeshShapeGenerator* Generator)
 {
 	Clear();
@@ -112,7 +150,7 @@ void FDynamicMesh3::Copy(const FMeshShapeGenerator* Generator)
 	Triangles = TDynamicVector<int>();
 	TriangleEdges = TDynamicVector<int>();
 	TriangleRefCounts = FRefCountVector();
-	TriangleGroups = new TDynamicVector<int>();
+	TriangleGroups = TDynamicVector<int>();
 	GroupIDCounter = 0;
 	Edges = TDynamicVector<int>();
 	EdgeRefCounts = FRefCountVector();
@@ -147,52 +185,35 @@ void FDynamicMesh3::Copy(const FMeshShapeGenerator* Generator)
 	}
 }
 
-
-
-
-
 void FDynamicMesh3::Copy(const FDynamicMesh3& copy, bool bNormals, bool bColors, bool bUVs, bool bAttributes)
 {
-	// currently cannot re-use existing attribute buffers
-	DiscardVertexNormals();
-	DiscardVertexColors();
-	DiscardVertexUVs();
-	DiscardTriangleGroups();
-	DiscardAttributes();
-
-	Vertices = TDynamicVector<double>(copy.Vertices);
-
-	VertexNormals = (bNormals && copy.HasVertexNormals()) ? new TDynamicVector<float>(*copy.VertexNormals) : nullptr;
-	VertexColors = (bColors && copy.HasVertexColors()) ? new TDynamicVector<float>(*copy.VertexColors) : nullptr;
-	VertexUVs = (bUVs && copy.HasVertexUVs()) ? new TDynamicVector<float>(*copy.VertexUVs) : nullptr;
-
-	VertexRefCounts = FRefCountVector(copy.VertexRefCounts);
-
-	VertexEdgeLists = FSmallListSet(copy.VertexEdgeLists);
+	Vertices        = copy.Vertices;
+	VertexNormals   = bNormals ? copy.VertexNormals : TOptional<TDynamicVector<float>>{};
+	VertexColors    = bColors  ? copy.VertexColors : TOptional<TDynamicVector<float>>{};
+	VertexUVs       = bUVs     ? copy.VertexUVs : TOptional<TDynamicVector<float>>{};
+	VertexRefCounts = copy.VertexRefCounts;
+	VertexEdgeLists = copy.VertexEdgeLists;
 
-	Triangles = TDynamicVector<int>(copy.Triangles);
-	TriangleEdges = TDynamicVector<int>(copy.TriangleEdges);
-	TriangleRefCounts = FRefCountVector(copy.TriangleRefCounts);
-	TriangleGroups = (copy.HasTriangleGroups()) ? new TDynamicVector<int>(*copy.TriangleGroups) : nullptr;
-	GroupIDCounter = copy.GroupIDCounter;
+	Triangles         = copy.Triangles;
+	TriangleEdges     = copy.TriangleEdges;
+	TriangleRefCounts = copy.TriangleRefCounts;
+	TriangleGroups    = copy.TriangleGroups;
+	GroupIDCounter    = copy.GroupIDCounter;
 
-	Edges = TDynamicVector<int>(copy.Edges);
-	EdgeRefCounts = FRefCountVector(copy.EdgeRefCounts);
+	Edges         = copy.Edges;
+	EdgeRefCounts = copy.EdgeRefCounts;
 
+	DiscardAttributes();
 	if (bAttributes && copy.HasAttributes())
 	{
 		EnableAttributes();
-		AttributeSet->Copy(*copy.Attributes());
+		AttributeSet->Copy(*copy.AttributeSet);
 	}
 
 	Timestamp = FMath::Max(Timestamp + 1, copy.Timestamp);
 	ShapeTimestamp = TopologyTimestamp = Timestamp;
 }
 
-
-
-
-
 void FDynamicMesh3::CompactCopy(const FDynamicMesh3& copy, bool bNormals, bool bColors, bool bUVs, bool bAttributes, FCompactMaps* CompactInfo)
 {
 	// @todo these are not compacted and hence code cannot work. fix!
@@ -207,39 +228,24 @@ void FDynamicMesh3::CompactCopy(const FDynamicMesh3& copy, bool bNormals, bool b
 	//}
 
 	// currently cannot re-use existing attribute buffers
-	DiscardVertexNormals();
-	DiscardVertexColors();
-	DiscardVertexUVs();
-	DiscardTriangleGroups();
-
-	Vertices = TDynamicVector<double>();
-	VertexEdgeLists = FSmallListSet();
-	VertexRefCounts = FRefCountVector();
-	Triangles = TDynamicVector<int>();
-	TriangleEdges = TDynamicVector<int>();
-	TriangleRefCounts = FRefCountVector();
-	Edges = TDynamicVector<int>();
-	EdgeRefCounts = FRefCountVector();
-	GroupIDCounter = 0;
-
-	TriangleGroups = (copy.HasTriangleGroups()) ? new TDynamicVector<int>(*copy.TriangleGroups) : nullptr;
+	Clear();
 
 	// [TODO] if we knew some of these were dense we could copy directly...
 
 	FVertexInfo vinfo;
 	TArray<int> mapV; mapV.SetNum(copy.MaxVertexID());
-	for (int vid : copy.VertexIndicesItr()) 
+	for (int vid : copy.VertexIndicesItr())
 	{
 		copy.GetVertex(vid, vinfo, bNormals, bColors, bUVs);
 		mapV[vid] = AppendVertex(vinfo);
 	}
 
 	// [TODO] would be much faster to explicitly copy triangle & edge data structures!!
-	if (copy.HasTriangleGroups()) 
+	if (copy.HasTriangleGroups())
 	{
 		EnableTriangleGroups(0);
 	}
-	for (int tid : copy.TriangleIndicesItr()) 
+	for (int tid : copy.TriangleIndicesItr())
 	{
 		FIndex3i t = copy.GetTriangle(tid);
 		t = FIndex3i(mapV[t.A], mapV[t.B], mapV[t.C]);
@@ -249,9 +255,9 @@ void FDynamicMesh3::CompactCopy(const FDynamicMesh3& copy, bool bNormals, bool b
 	}
 
 	// generate compact info...?
-	if (CompactInfo != nullptr) 
+	if (CompactInfo != nullptr)
 	{
-		for (int vid : copy.VertexIndicesItr()) 
+		for (int vid : copy.VertexIndicesItr())
 		{
 			CompactInfo->MapV.Add(vid, mapV[vid]);
 		}
@@ -261,35 +267,11 @@ void FDynamicMesh3::CompactCopy(const FDynamicMesh3& copy, bool bNormals, bool b
 	ShapeTimestamp = TopologyTimestamp = Timestamp;
 }
 
-
-
-
 void FDynamicMesh3::Clear()
 {
-	Vertices.Clear();
-	VertexEdgeLists = FSmallListSet();
-	VertexRefCounts = FRefCountVector();
-	Triangles.Clear();
-	TriangleEdges.Clear();
-	TriangleRefCounts = FRefCountVector();
-	Edges.Clear();
-	EdgeRefCounts = FRefCountVector();
-
-	DiscardTriangleGroups();
-	DiscardVertexColors();
-	DiscardVertexUVs();
-	DiscardVertexNormals();
-	DiscardAttributes();
-
-	Timestamp = ShapeTimestamp = TopologyTimestamp = 0;
-	GroupIDCounter = 0;
-	CachedBoundingBoxTimestamp = -1;
-	bIsClosedCached = false;
-	CachedIsClosedTimestamp = -1;
+	*this = FDynamicMesh3();
 }
 
-
-
 int FDynamicMesh3::GetComponentsFlags() const
 {
 	int c = 0;
@@ -312,158 +294,111 @@ int FDynamicMesh3::GetComponentsFlags() const
 	return c;
 }
 
-
-
-
-
-
-
-
-
-
 void FDynamicMesh3::EnableVertexNormals(const FVector3f& InitialNormal)
 {
 	if (HasVertexNormals())
 	{
 		return;
 	}
-	VertexNormals = new TDynamicVector<float>();
+	VertexNormals = TDynamicVector<float>();
 	int NV = MaxVertexID();
 	VertexNormals->Resize(3 * NV);
 	for (int i = 0; i < NV; ++i)
 	{
 		int vi = 3 * i;
-		(*VertexNormals)[vi] = InitialNormal.X;
-		(*VertexNormals)[vi + 1] = InitialNormal.Y;
-		(*VertexNormals)[vi + 2] = InitialNormal.Z;
+		VertexNormals.GetValue()[vi] = InitialNormal.X;
+		VertexNormals.GetValue()[vi + 1] = InitialNormal.Y;
+		VertexNormals.GetValue()[vi + 2] = InitialNormal.Z;
 	}
 }
 
 void FDynamicMesh3::DiscardVertexNormals()
 {
-	if (VertexNormals != nullptr)
-	{
-		delete VertexNormals;
-		VertexNormals = nullptr;
-	}
+	VertexNormals.Reset();
 }
 
-
-
-
-
 void FDynamicMesh3::EnableVertexColors(const FVector3f& InitialColor)
 {
 	if (HasVertexColors())
 	{
 		return;
 	}
-	VertexColors = new TDynamicVector<float>();
+	VertexColors = TDynamicVector<float>();
 	int NV = MaxVertexID();
 	VertexColors->Resize(3 * NV);
 	for (int i = 0; i < NV; ++i)
 	{
 		int vi = 3 * i;
-		(*VertexColors)[vi] = InitialColor.X;
-		(*VertexColors)[vi + 1] = InitialColor.Y;
-		(*VertexColors)[vi + 2] = InitialColor.Z;
+		VertexColors.GetValue()[vi] = InitialColor.X;
+		VertexColors.GetValue()[vi + 1] = InitialColor.Y;
+		VertexColors.GetValue()[vi + 2] = InitialColor.Z;
 	}
 }
 
 void FDynamicMesh3::DiscardVertexColors()
 {
-	if (VertexColors != nullptr)
-	{
-		delete VertexColors;
-		VertexColors = nullptr;
-	}
+	VertexColors.Reset();
 }
 
-
-
-
 void FDynamicMesh3::EnableVertexUVs(const FVector2f& InitialUV)
 {
 	if (HasVertexUVs())
 	{
 		return;
 	}
-	VertexUVs = new TDynamicVector<float>();
+	VertexUVs = TDynamicVector<float>();
 	int NV = MaxVertexID();
 	VertexUVs->Resize(2 * NV);
 	for (int i = 0; i < NV; ++i)
 	{
 		int vi = 2 * i;
-		(*VertexUVs)[vi] = InitialUV.X;
-		(*VertexUVs)[vi + 1] = InitialUV.Y;
+		VertexUVs.GetValue()[vi] = InitialUV.X;
+		VertexUVs.GetValue()[vi + 1] = InitialUV.Y;
 	}
 }
 
 void FDynamicMesh3::DiscardVertexUVs()
 {
-	if (VertexUVs != nullptr)
-	{
-		delete VertexUVs;
-		VertexUVs = nullptr;
-	}
+		VertexUVs.Reset();
 }
 
-
-
 void FDynamicMesh3::EnableTriangleGroups(int InitialGroup)
 {
 	if (HasTriangleGroups())
 	{
 		return;
 	}
-	TriangleGroups = new TDynamicVector<int>();
+	TriangleGroups = TDynamicVector<int>();
 	int NT = MaxTriangleID();
 	TriangleGroups->Resize(NT);
 	for (int i = 0; i < NT; ++i)
 	{
-		(*TriangleGroups)[i] = InitialGroup;
+		TriangleGroups.GetValue()[i] = InitialGroup;
 	}
 	GroupIDCounter = 0;
 }
 
 void FDynamicMesh3::DiscardTriangleGroups()
 {
-	if (TriangleGroups != nullptr)
-	{
-		delete TriangleGroups;
-		TriangleGroups = nullptr;
-	}
+	TriangleGroups.Reset();
 	GroupIDCounter = 0;
 }
 
-
-
-
 void FDynamicMesh3::EnableAttributes()
 {
 	if (HasAttributes())
 	{
 		return;
 	}
-	AttributeSet = new FDynamicMeshAttributeSet(this);
+	AttributeSet = MakeUnique<FDynamicMeshAttributeSet>(this);
 	AttributeSet->Initialize(MaxVertexID(), MaxTriangleID());
 }
 
-
 void FDynamicMesh3::DiscardAttributes()
 {
-	if (AttributeSet != nullptr) 
-	{
-		delete AttributeSet;
-		AttributeSet = nullptr;
-	}
+	AttributeSet = nullptr;
 }
 
-
-
-
-
-
 bool FDynamicMesh3::GetVertex(int vID, FVertexInfo& vinfo, bool bWantNormals, bool bWantColors, bool bWantUVs) const
 {
 	if (VertexRefCounts.IsValid(vID) == false)
@@ -473,26 +408,28 @@ bool FDynamicMesh3::GetVertex(int vID, FVertexInfo& vinfo, bool bWantNormals, bo
 	int vi = 3 * vID;
 	vinfo.Position = FVector3d(Vertices[vi], Vertices[vi + 1], Vertices[vi + 2]);
 	vinfo.bHaveN = vinfo.bHaveUV = vinfo.bHaveC = false;
-	if (HasVertexNormals() && bWantNormals) 
+	if (HasVertexNormals() && bWantNormals)
 	{
 		vinfo.bHaveN = true;
-		vinfo.Normal = FVector3f((*VertexNormals)[vi], (*VertexNormals)[vi + 1], (*VertexNormals)[vi + 2]);
+		const TDynamicVector<float>& NormalVec = VertexNormals.GetValue();
+		vinfo.Normal = FVector3f(NormalVec[vi], NormalVec[vi + 1], NormalVec[vi + 2]);
 	}
-	if (HasVertexColors() && bWantColors) 
+	if (HasVertexColors() && bWantColors)
 	{
 		vinfo.bHaveC = true;
-		vinfo.Color = FVector3f((*VertexColors)[vi], (*VertexColors)[vi + 1], (*VertexColors)[vi + 2]);
+		const TDynamicVector<float>& ColorVec = VertexColors.GetValue();
+		vinfo.Color = FVector3f(ColorVec[vi], ColorVec[vi + 1], ColorVec[vi + 2]);
 	}
-	if (HasVertexUVs() && bWantUVs) 
+	if (HasVertexUVs() && bWantUVs)
 	{
 		vinfo.bHaveUV = true;
-		vinfo.UV = FVector2f((*VertexUVs)[2*vID], (*VertexUVs)[2*vID + 1]);
+		const TDynamicVector<float>& UVVec = VertexUVs.GetValue();
+		vinfo.UV = FVector2f(UVVec[2*vID], UVVec[2*vID + 1]);
 	}
 	return true;
 }
 
-
-int FDynamicMesh3::GetMaxVtxEdgeCount() const 
+int FDynamicMesh3::GetMaxVtxEdgeCount() const
 {
 	int max = 0;
 	for (int vid : VertexIndicesItr())
@@ -502,8 +439,6 @@ int FDynamicMesh3::GetMaxVtxEdgeCount() const
 	return max;
 }
 
-
-
 FVertexInfo FDynamicMesh3::GetVertexInfo(int i) const
 {
 	FVertexInfo vi = FVertexInfo();
@@ -527,8 +462,6 @@ FVertexInfo FDynamicMesh3::GetVertexInfo(int i) const
 	return vi;
 }
 
-
-
 FIndex3i FDynamicMesh3::GetTriNeighbourTris(int tID) const
 {
 	if (TriangleRefCounts.IsValid(tID))
@@ -544,7 +477,7 @@ FIndex3i FDynamicMesh3::GetTriNeighbourTris(int tID) const
 	}
 	else
 	{
-		return InvalidTriangle();
+		return InvalidTriangle;
 	}
 }
 
@@ -566,10 +499,6 @@ void FDynamicMesh3::GetVertexOneRingTriangles(int VertexID, TArray<int>& Triangl
 	}
 }
 
-
-
-
-
 FString FDynamicMesh3::MeshInfoString()
 {
 	FString VtxString = FString::Printf(TEXT("Vertices count %d max %d  %s  VtxEdges %s"),
@@ -586,9 +515,6 @@ FString FDynamicMesh3::MeshInfoString()
 	return VtxString + "\n" + TriString + "\n" + EdgeString + "\n" + AttribString + "\n" + InfoString;
 }
 
-
-
-
 bool FDynamicMesh3::IsSameMesh(const FDynamicMesh3& m2, bool bCheckConnectivity, bool bCheckEdgeIDs,
 	bool bCheckNormals, bool bCheckColors, bool bCheckUVs, bool bCheckGroups,
 	float Epsilon)
@@ -717,12 +643,6 @@ bool FDynamicMesh3::IsSameMesh(const FDynamicMesh3& m2, bool bCheckConnectivity,
 	return true;
 }
 
-
-
-
-
-
-
 bool FDynamicMesh3::CheckValidity(bool bAllowNonManifoldVertices, EValidityCheckFailMode FailMode) const
 {
 
@@ -734,6 +654,10 @@ bool FDynamicMesh3::CheckValidity(bool bAllowNonManifoldVertices, EValidityCheck
 	TFunction<void(bool)> CheckOrFailF = [&](bool b)
 	{
 		is_ok = is_ok && b;
+		if (!is_ok)
+		{
+			is_ok = false;
+		}
 	};
 	if (FailMode == EValidityCheckFailMode::Check)
 	{
@@ -1091,4 +1015,4 @@ int FDynamicMesh3::FindEdgeFromTri(int vA, int vB, int tID) const
 		return TriangleEdges[i + 2];
 	}
 	return InvalidID;
-}
\ No newline at end of file
+}
diff --git a/Engine/Plugins/Experimental/GeometryProcessing/Source/DynamicMesh/Private/DynamicMesh3_Edits.cpp b/Engine/Plugins/Experimental/GeometryProcessing/Source/DynamicMesh/Private/DynamicMesh3_Edits.cpp
--- a/Engine/Plugins/Experimental/GeometryProcessing/Source/DynamicMesh/Private/DynamicMesh3_Edits.cpp
+++ b/Engine/Plugins/Experimental/GeometryProcessing/Source/DynamicMesh/Private/DynamicMesh3_Edits.cpp
@@ -62,11 +62,12 @@ int FDynamicMesh3::AppendVertex(const FDynamicMesh3& from, int fromVID)
 	Vertices.InsertAt(from.Vertices[bi], i);
 	if (HasVertexNormals()) 
 	{
-		if (from.HasVertexNormals()) 
+		if (from.HasVertexNormals())
 		{
-			VertexNormals->InsertAt((*from.VertexNormals)[bi + 2], i + 2);
-			VertexNormals->InsertAt((*from.VertexNormals)[bi + 1], i + 1);
-			VertexNormals->InsertAt((*from.VertexNormals)[bi], i);
+			const TDynamicVector<float>& FromNormals = from.VertexNormals.GetValue();
+			VertexNormals->InsertAt(FromNormals[bi + 2], i + 2);
+			VertexNormals->InsertAt(FromNormals[bi + 1], i + 1);
+			VertexNormals->InsertAt(FromNormals[bi], i);
 		}
 		else 
 		{
@@ -80,9 +81,10 @@ int FDynamicMesh3::AppendVertex(const FDynamicMesh3& from, int fromVID)
 	{
 		if (from.HasVertexColors()) 
 		{
-			VertexColors->InsertAt((*from.VertexColors)[bi + 2], i + 2);
-			VertexColors->InsertAt((*from.VertexColors)[bi + 1], i + 1);
-			VertexColors->InsertAt((*from.VertexColors)[bi], i);
+			const TDynamicVector<float>& FromColors = from.VertexColors.GetValue();
+			VertexColors->InsertAt(FromColors[bi + 2], i + 2);
+			VertexColors->InsertAt(FromColors[bi + 1], i + 1);
+			VertexColors->InsertAt(FromColors[bi], i);
 		}
 		else 
 		{
@@ -97,9 +99,10 @@ int FDynamicMesh3::AppendVertex(const FDynamicMesh3& from, int fromVID)
 		int j = 2 * vid;
 		if (from.HasVertexUVs()) 
 		{
+			const TDynamicVector<float>& FromUVs = from.VertexUVs.GetValue();
 			int bj = 2 * fromVID;
-			VertexUVs->InsertAt((*from.VertexUVs)[bj + 1], j + 1);
-			VertexUVs->InsertAt((*from.VertexUVs)[bj], j);
+			VertexUVs->InsertAt(FromUVs[bj + 1], j + 1);
+			VertexUVs->InsertAt(FromUVs[bj], j);
 		}
 		else 
 		{
@@ -338,21 +341,24 @@ void FDynamicMesh3::CompactInPlace(FCompactMaps* CompactInfo)
 		Vertices[kc] = Vertices[kl];  Vertices[kc + 1] = Vertices[kl + 1];  Vertices[kc + 2] = Vertices[kl + 2];
 		if (HasVertexNormals())
 		{
-			(*VertexNormals)[kc] = (*VertexNormals)[kl];  
-			(*VertexNormals)[kc + 1] = (*VertexNormals)[kl + 1];  
-			(*VertexNormals)[kc + 2] = (*VertexNormals)[kl + 2];
+			TDynamicVector<float>& Normals = VertexNormals.GetValue();
+			Normals[kc]     = Normals[kl];
+			Normals[kc + 1] = Normals[kl + 1];
+			Normals[kc + 2] = Normals[kl + 2];
 		}
 		if (HasVertexColors())
 		{
-			(*VertexColors)[kc] = (*VertexColors)[kl];  
-			(*VertexColors)[kc + 1] = (*VertexColors)[kl + 1];  
-			(*VertexColors)[kc + 2] = (*VertexColors)[kl + 2];
+			TDynamicVector<float>& Colors = VertexColors.GetValue();
+			Colors[kc] =     Colors[kl];
+			Colors[kc + 1] = Colors[kl + 1];
+			Colors[kc + 2] = Colors[kl + 2];
 		}
 		if (HasVertexUVs())
 		{
+			TDynamicVector<float>& UVs = VertexUVs.GetValue();
 			int ukc = iCurV * 2, ukl = iLastV * 2;
-			(*VertexUVs)[ukc] = (*VertexUVs)[ukl]; 
-			(*VertexUVs)[ukc + 1] = (*VertexUVs)[ukl + 1];
+			UVs[ukc] =     UVs[ukl];
+			UVs[ukc + 1] = UVs[ukl + 1];
 		}
 
 		for (int eid : VertexEdgeLists.Values(iLastV))
@@ -439,7 +445,7 @@ void FDynamicMesh3::CompactInPlace(FCompactMaps* CompactInfo)
 		}
 		if (HasTriangleGroups())
 		{
-			(*TriangleGroups)[iCurT] = (*TriangleGroups)[iLastT];
+			TriangleGroups.GetValue()[iCurT] = TriangleGroups.GetValue()[iLastT];
 		}
 
 		// update edges
@@ -563,28 +569,24 @@ void FDynamicMesh3::ReverseTriOrientationInternal(int tID)
 
 void FDynamicMesh3::ReverseOrientation(bool bFlipNormals)
 {
-	for (int tid : TriangleIndicesItr()) 
+	for (int tid : TriangleIndicesItr())
 	{
 		ReverseTriOrientationInternal(tid);
 	}
-	if (bFlipNormals && HasVertexNormals()) 
+	if (bFlipNormals && HasVertexNormals())
 	{
-		for (int vid : VertexIndicesItr()) 
+		for (int vid : VertexIndicesItr())
 		{
 			int i = 3 * vid;
-			(*VertexNormals)[i] = -(*VertexNormals)[i];
-			(*VertexNormals)[i + 1] = -(*VertexNormals)[i + 1];
-			(*VertexNormals)[i + 2] = -(*VertexNormals)[i + 2];
+			TDynamicVector<float>& Normals = VertexNormals.GetValue();
+			Normals[i]     = -Normals[i];
+			Normals[i + 1] = -Normals[i + 1];
+			Normals[i + 2] = -Normals[i + 2];
 		}
 	}
 	UpdateTimeStamp(true, true);
 }
 
-
-
-
-
-
 EMeshResult FDynamicMesh3::RemoveVertex(int vID, bool bRemoveAllTriangles, bool bPreserveManifold)
 {
 	if (VertexRefCounts.IsValid(vID) == false)
@@ -916,7 +918,7 @@ EMeshResult FDynamicMesh3::SplitEdge(int eab, FEdgeSplitInfo& SplitInfo, double
 		int t2 = AddTriangleInternal(f, b, c, InvalidID, InvalidID, InvalidID);
 		if (HasTriangleGroups())
 		{
-			int group0 = (*TriangleGroups)[t0];
+			int group0 = TriangleGroups.GetValue()[t0];
 			TriangleGroups->InsertAt(group0, t2);
 		}
 
@@ -998,9 +1000,9 @@ EMeshResult FDynamicMesh3::SplitEdge(int eab, FEdgeSplitInfo& SplitInfo, double
 		int t3 = AddTriangleInternal(f, d, b, InvalidID, InvalidID, InvalidID);
 		if (HasTriangleGroups()) 
 		{
-			int group0 = (*TriangleGroups)[t0];
+			int group0 = TriangleGroups.GetValue()[t0];
 			TriangleGroups->InsertAt(group0, t2);
-			int group1 = (*TriangleGroups)[t1];
+			int group1 = TriangleGroups.GetValue()[t1];
 			TriangleGroups->InsertAt(group1, t3);
 		}
 
@@ -1829,7 +1831,7 @@ EMeshResult FDynamicMesh3::PokeTriangle(int TriangleID, const FVector3d& BaryCoo
 	// transfer groups
 	if (HasTriangleGroups()) 
 	{
-		int g = (*TriangleGroups)[TriangleID];
+		int g = TriangleGroups.GetValue()[TriangleID];
 		TriangleGroups->InsertAt(g, t1);
 		TriangleGroups->InsertAt(g, t2);
 	}
diff --git a/Engine/Plugins/Experimental/GeometryProcessing/Source/DynamicMesh/Private/DynamicMesh3_Queries.cpp b/Engine/Plugins/Experimental/GeometryProcessing/Source/DynamicMesh/Private/DynamicMesh3_Queries.cpp
--- a/Engine/Plugins/Experimental/GeometryProcessing/Source/DynamicMesh/Private/DynamicMesh3_Queries.cpp
+++ b/Engine/Plugins/Experimental/GeometryProcessing/Source/DynamicMesh/Private/DynamicMesh3_Queries.cpp
@@ -284,8 +284,6 @@ EMeshResult FDynamicMesh3::GetVtxContiguousTriangles(int VertexID, TArray<int>&
 	return EMeshResult::Ok;
 }
 
-
-
 bool FDynamicMesh3::IsBoundaryVertex(int vID) const 
 {
 	check(IsVertex(vID));
@@ -299,8 +297,6 @@ bool FDynamicMesh3::IsBoundaryVertex(int vID) const
 	return false;
 }
 
-
-
 bool FDynamicMesh3::IsBoundaryTriangle(int tID) const
 {
 	check(IsTriangle(tID));
@@ -308,9 +304,6 @@ bool FDynamicMesh3::IsBoundaryTriangle(int tID) const
 	return IsBoundaryEdge(TriangleEdges[i]) || IsBoundaryEdge(TriangleEdges[i + 1]) || IsBoundaryEdge(TriangleEdges[i + 2]);
 }
 
-
-
-
 FIndex2i FDynamicMesh3::GetOrientedBoundaryEdgeV(int eID) const
 {
 	if (EdgeRefCounts.IsValid(eID)) 
@@ -326,7 +319,7 @@ FIndex2i FDynamicMesh3::GetOrientedBoundaryEdgeV(int eID) const
 		}
 	}
 	check(false);
-	return InvalidEdge();
+	return InvalidEdge;
 }
 
 
@@ -340,14 +333,12 @@ bool FDynamicMesh3::IsGroupBoundaryEdge(int eID) const
 	{
 		return false;
 	}
-	int g1 = (*TriangleGroups)[et1];
+	int g1 = TriangleGroups.GetValue()[et1];
 	int et0 = Edges[4 * eID + 2];
-	int g0 = (*TriangleGroups)[et0];
+	int g0 = TriangleGroups.GetValue()[et0];
 	return g1 != g0;
 }
 
-
-
 bool FDynamicMesh3::IsGroupBoundaryVertex(int vID) const
 {
 	check(IsVertex(vID));
@@ -357,7 +348,7 @@ bool FDynamicMesh3::IsGroupBoundaryVertex(int vID) const
 	for (int eID : VertexEdgeLists.Values(vID)) 
 	{
 		int et0 = Edges[4 * eID + 2];
-		int g0 = (*TriangleGroups)[et0];
+		int g0 = TriangleGroups.GetValue()[et0];
 		if (group_id != g0) 
 		{
 			if (group_id == InvalidGroupID)
@@ -372,7 +363,7 @@ bool FDynamicMesh3::IsGroupBoundaryVertex(int vID) const
 		int et1 = Edges[4 * eID + 3];
 		if (et1 != InvalidID) 
 		{
-			int g1 = (*TriangleGroups)[et1];
+			int g1 = TriangleGroups.GetValue()[et1];
 			if (group_id != g1)
 			{
 				return true;        // saw multiple group IDs
@@ -399,7 +390,7 @@ bool FDynamicMesh3::IsGroupJunctionVertex(int vID) const
 			{
 				continue;
 			}
-			int g0 = (*TriangleGroups)[et[k]];
+			int g0 = TriangleGroups.GetValue()[et[k]];
 			if (g0 != groups[0] && g0 != groups[1]) 
 			{
 				if (groups[0] != InvalidGroupID && groups[1] != InvalidGroupID)
@@ -432,7 +423,7 @@ bool FDynamicMesh3::GetVertexGroups(int vID, FIndex4i& groups) const
 	for (int eID : VertexEdgeLists.Values(vID)) 
 	{
 		int et0 = Edges[4 * eID + 2];
-		int g0 = (*TriangleGroups)[et0];
+		int g0 = TriangleGroups.GetValue()[et0];
 		if (groups.Contains(g0) == false)
 		{
 			groups[ng++] = g0;
@@ -444,7 +435,7 @@ bool FDynamicMesh3::GetVertexGroups(int vID, FIndex4i& groups) const
 		int et1 = Edges[4 * eID + 3];
 		if (et1 != InvalidID) 
 		{
-			int g1 = (*TriangleGroups)[et1];
+			int g1 = TriangleGroups.GetValue()[et1];
 			if (groups.Contains(g1) == false)
 			{
 				groups[ng++] = g1;
@@ -468,13 +459,13 @@ bool FDynamicMesh3::GetAllVertexGroups(int vID, TArray<int>& GroupsOut) const
 	for (int eID : VertexEdgeLists.Values(vID)) 
 	{
 		int et0 = Edges[4 * eID + 2];
-		int g0 = (*TriangleGroups)[et0];
+		int g0 = TriangleGroups.GetValue()[et0];
 		GroupsOut.AddUnique(g0);
 
 		int et1 = Edges[4 * eID + 3];
 		if (et1 != InvalidID) 
 		{
-			int g1 = (*TriangleGroups)[et1];
+			int g1 = TriangleGroups.GetValue()[et1];
 			GroupsOut.AddUnique(g1);
 		}
 	}
@@ -722,7 +713,8 @@ FFrame3d FDynamicMesh3::GetVertexFrame(int vID, bool bFrameNormalY) const
 
 	int vi = 3 * vID;
 	FVector3d v(Vertices[vi], Vertices[vi + 1], Vertices[vi + 2]);
-	FVector3d normal((*VertexNormals)[vi], (*VertexNormals)[vi + 1], (*VertexNormals)[vi + 2]);
+	const TDynamicVector<float>& Normals = VertexNormals.GetValue();
+	FVector3d normal(Normals[vi], Normals[vi + 1], Normals[vi + 2]);
 	int eid = VertexEdgeLists.First(vID);
 	int ovi = 3 * GetOtherEdgeVertex(eid, vID);
 	FVector3d ov(Vertices[ovi], Vertices[ovi + 1], Vertices[ovi + 2]);
@@ -786,7 +778,7 @@ FVector3d FDynamicMesh3::GetTriBaryNormal(int tID, double bary0, double bary1, d
 	int ai = 3 * Triangles[3 * tID],
 		bi = 3 * Triangles[3 * tID + 1],
 		ci = 3 * Triangles[3 * tID + 2];
-	const TDynamicVector<float>& normalsR = *(this->VertexNormals);
+	const TDynamicVector<float>& normalsR = VertexNormals.GetValue();
 	FVector3d n = FVector3d(
 		(bary0*normalsR[ai] + bary1 * normalsR[bi] + bary2 * normalsR[ci]),
 		(bary0*normalsR[ai + 1] + bary1 * normalsR[bi + 1] + bary2 * normalsR[ci + 1]),
@@ -819,9 +811,9 @@ void FDynamicMesh3::GetTriBaryPoint(int tID, double bary0, double bary1, double
 		(bary0 * Vertices[ai + 1] + bary1 * Vertices[bi + 1] + bary2 * Vertices[ci + 1]),
 		(bary0 * Vertices[ai + 2] + bary1 * Vertices[bi + 2] + bary2 * Vertices[ci + 2]));
 	vinfo.bHaveN = HasVertexNormals();
-	if (vinfo.bHaveN) 
+	if (vinfo.bHaveN)
 	{
-		TDynamicVector<float>& normalsR = *(this->VertexNormals);
+		const TDynamicVector<float>& normalsR = this->VertexNormals.GetValue();
 		vinfo.Normal = FVector3f(
 			(float)(bary0 * normalsR[ai] + bary1 * normalsR[bi] + bary2 * normalsR[ci]),
 			(float)(bary0 * normalsR[ai + 1] + bary1 * normalsR[bi + 1] + bary2 * normalsR[ci + 1]),
@@ -831,7 +823,7 @@ void FDynamicMesh3::GetTriBaryPoint(int tID, double bary0, double bary1, double
 	vinfo.bHaveC = HasVertexColors();
 	if (vinfo.bHaveC) 
 	{
-		TDynamicVector<float>& colorsR = *(this->VertexColors);
+		const TDynamicVector<float>& colorsR = this->VertexColors.GetValue();
 		vinfo.Color = FVector3f(
 			(float)(bary0 * colorsR[ai] + bary1 * colorsR[bi] + bary2 * colorsR[ci]),
 			(float)(bary0 * colorsR[ai + 1] + bary1 * colorsR[bi + 1] + bary2 * colorsR[ci + 1]),
@@ -840,7 +832,7 @@ void FDynamicMesh3::GetTriBaryPoint(int tID, double bary0, double bary1, double
 	vinfo.bHaveUV = HasVertexUVs();
 	if (vinfo.bHaveUV) 
 	{
-		TDynamicVector<float>& uvR = *(this->VertexUVs);
+		const TDynamicVector<float>& uvR = this->VertexUVs.GetValue();
 		ai = 2 * Triangles[3 * tID];
 		bi = 2 * Triangles[3 * tID + 1];
 		ci = 2 * Triangles[3 * tID + 2];
@@ -962,4 +954,4 @@ double FDynamicMesh3::CalculateWindingNumber(const FVector3d& QueryPoint) const
 		sum += GetTriSolidAngle(tid, QueryPoint);
 	}
 	return sum / FMathd::FourPi;
-}
\ No newline at end of file
+}
diff --git a/Engine/Plugins/Experimental/GeometryProcessing/Source/DynamicMesh/Private/MeshSimplification.cpp b/Engine/Plugins/Experimental/GeometryProcessing/Source/DynamicMesh/Private/MeshSimplification.cpp
--- a/Engine/Plugins/Experimental/GeometryProcessing/Source/DynamicMesh/Private/MeshSimplification.cpp
+++ b/Engine/Plugins/Experimental/GeometryProcessing/Source/DynamicMesh/Private/MeshSimplification.cpp
@@ -628,7 +628,7 @@ ESimplificationResult TMeshSimplification<QuadricErrorType>::CollapseEdge(int ed
 	// look up 'other' verts c (from t0) and d (from t1, if it exists)
 	FIndex3i T0tv = Mesh->GetTriangle(t0);
 	int c = IndexUtil::FindTriOtherVtx(a, b, T0tv);
-	FIndex3i T1tv = (bIsBoundaryEdge) ? FDynamicMesh3::InvalidTriangle() : Mesh->GetTriangle(t1);
+	FIndex3i T1tv = (bIsBoundaryEdge) ? FDynamicMesh3::InvalidTriangle : Mesh->GetTriangle(t1);
 	int d = (bIsBoundaryEdge) ? FDynamicMesh3::InvalidID : IndexUtil::FindTriOtherVtx(a, b, T1tv);
 
 	FVector3d vA = Mesh->GetVertex(a);
diff --git a/Engine/Plugins/Experimental/GeometryProcessing/Source/DynamicMesh/Public/DynamicAttribute.h b/Engine/Plugins/Experimental/GeometryProcessing/Source/DynamicMesh/Public/DynamicAttribute.h
--- a/Engine/Plugins/Experimental/GeometryProcessing/Source/DynamicMesh/Public/DynamicAttribute.h
+++ b/Engine/Plugins/Experimental/GeometryProcessing/Source/DynamicMesh/Public/DynamicAttribute.h
@@ -61,7 +61,7 @@ public:
 
 	/** Allocate a new copy of the attribute layer, optionally with a different parent mesh */
 	virtual TDynamicAttributeBase* MakeCopy(ParentType* ParentIn) const = 0;
-
+	virtual void Reparent(ParentType* NewParent) = 0;
 
 	virtual void OnNewVertex(int VertexID, bool bInserted)
 	{
@@ -101,27 +101,27 @@ public:
 
 
 	/** Update to reflect an edge split in the parent mesh */
-	virtual void OnSplitEdge(const FDynamicMesh3::FEdgeSplitInfo& SplitInfo)
+	virtual void OnSplitEdge(const DynamicMeshInfo::FEdgeSplitInfo& SplitInfo)
 	{
 	}
 
 	/** Update to reflect an edge flip in the parent mesh */
-	virtual void OnFlipEdge(const FDynamicMesh3::FEdgeFlipInfo& FlipInfo)
+	virtual void OnFlipEdge(const DynamicMeshInfo::FEdgeFlipInfo& FlipInfo)
 	{
 	}
 
 	/** Update to reflect an edge collapse in the parent mesh */
-	virtual void OnCollapseEdge(const FDynamicMesh3::FEdgeCollapseInfo& CollapseInfo)
+	virtual void OnCollapseEdge(const DynamicMeshInfo::FEdgeCollapseInfo& CollapseInfo)
 	{
 	}
 
 	/** Update to reflect a face poke in the parent mesh */
-	virtual void OnPokeTriangle(const FDynamicMesh3::FPokeTriangleInfo& PokeInfo)
+	virtual void OnPokeTriangle(const DynamicMeshInfo::FPokeTriangleInfo& PokeInfo)
 	{
 	}
 
 	/** Update to reflect an edge merge in the parent mesh */
-	virtual void OnMergeEdges(const FDynamicMesh3::FMergeEdgesInfo& MergeInfo)
+	virtual void OnMergeEdges(const DynamicMeshInfo::FMergeEdgesInfo& MergeInfo)
 	{
 	}
 
@@ -227,35 +227,35 @@ public:
 	// mesh-specific on* functions; may be split out
 public:
 
-	virtual void OnSplitEdge(const FDynamicMesh3::FEdgeSplitInfo& SplitInfo)
+	virtual void OnSplitEdge(const DynamicMeshInfo::FEdgeSplitInfo& SplitInfo)
 	{
 		for (TDynamicAttributeBase<ParentType>* A : RegisteredAttributes)
 		{
 			A->OnSplitEdge(SplitInfo);
 		}
 	}
-	virtual void OnFlipEdge(const FDynamicMesh3::FEdgeFlipInfo& FlipInfo)
+	virtual void OnFlipEdge(const DynamicMeshInfo::FEdgeFlipInfo& FlipInfo)
 	{
 		for (TDynamicAttributeBase<ParentType>* A : RegisteredAttributes)
 		{
 			A->OnFlipEdge(FlipInfo);
 		}
 	}
-	virtual void OnCollapseEdge(const FDynamicMesh3::FEdgeCollapseInfo& CollapseInfo)
+	virtual void OnCollapseEdge(const DynamicMeshInfo::FEdgeCollapseInfo& CollapseInfo)
 	{
 		for (TDynamicAttributeBase<ParentType>* A : RegisteredAttributes)
 		{
 			A->OnCollapseEdge(CollapseInfo);
 		}
 	}
-	virtual void OnPokeTriangle(const FDynamicMesh3::FPokeTriangleInfo& PokeInfo)
+	virtual void OnPokeTriangle(const DynamicMeshInfo::FPokeTriangleInfo& PokeInfo)
 	{
 		for (TDynamicAttributeBase<ParentType>* A : RegisteredAttributes)
 		{
 			A->OnPokeTriangle(PokeInfo);
 		}
 	}
-	virtual void OnMergeEdges(const FDynamicMesh3::FMergeEdgesInfo& MergeInfo)
+	virtual void OnMergeEdges(const DynamicMeshInfo::FMergeEdgesInfo& MergeInfo)
 	{
 		for (TDynamicAttributeBase<ParentType>* A : RegisteredAttributes)
 		{
diff --git a/Engine/Plugins/Experimental/GeometryProcessing/Source/DynamicMesh/Public/DynamicMesh3.h b/Engine/Plugins/Experimental/GeometryProcessing/Source/DynamicMesh/Public/DynamicMesh3.h
--- a/Engine/Plugins/Experimental/GeometryProcessing/Source/DynamicMesh/Public/DynamicMesh3.h
+++ b/Engine/Plugins/Experimental/GeometryProcessing/Source/DynamicMesh/Public/DynamicMesh3.h
@@ -5,8 +5,12 @@
 #pragma once
 
 #include "BoxTypes.h"
+#include "DynamicMesh3.h"
+#include "DynamicMeshAttributeSet.h"
+#include "InfoTypes.h"
 #include "FrameTypes.h"
 #include "GeometryTypes.h"
+#include "HAL/Platform.h"
 #include "MathUtil.h"
 #include "Quaternion.h"
 #include "Util/DynamicVector.h"
@@ -17,11 +21,9 @@
 #include "VectorTypes.h"
 #include "VectorUtil.h"
 
-
-class FDynamicMeshAttributeSet;
 class FMeshShapeGenerator;
 
-enum class EMeshComponents
+enum class EMeshComponents : uint8
 {
 	None = 0,
 	VertexNormals = 1,
@@ -30,59 +32,6 @@ enum class EMeshComponents
 	FaceGroups = 8
 };
 
-/**
- * FVertexInfo stores information about vertex attributes - position, normal, color, UV
- */
-struct FVertexInfo
-{
-	FVector3d Position;
-	FVector3f Normal;
-	FVector3f Color;
-	FVector2f UV;
-	bool bHaveN, bHaveUV, bHaveC;
-
-	FVertexInfo()
-	{
-		Position = FVector3d::Zero();
-		Normal = Color = FVector3f::Zero();
-		UV = FVector2f::Zero();
-		bHaveN = bHaveC = bHaveUV = false;
-	}
-	FVertexInfo(const FVector3d& PositionIn)
-	{
-		Position = PositionIn;
-		Normal = Color = FVector3f::Zero();
-		UV = FVector2f::Zero();
-		bHaveN = bHaveC = bHaveUV = false;
-	}
-	FVertexInfo(const FVector3d& PositionIn, const FVector3f& NormalIn)
-	{
-		Position = PositionIn;
-		Normal = NormalIn;
-		Color = FVector3f::Zero();
-		UV = FVector2f::Zero();
-		bHaveN = true;
-		bHaveC = bHaveUV = false;
-	}
-	FVertexInfo(const FVector3d& PositionIn, const FVector3f& NormalIn, const FVector3f& ColorIn)
-	{
-		Position = PositionIn;
-		Normal = NormalIn;
-		Color = ColorIn;
-		UV = FVector2f::Zero();
-		bHaveN = bHaveC = true;
-		bHaveUV = false;
-	}
-	FVertexInfo(const FVector3d& PositionIn, const FVector3f& NormalIn, const FVector3f& ColorIn, const FVector2f& UVIn)
-	{
-		Position = PositionIn;
-		Normal = NormalIn;
-		Color = ColorIn;
-		UV = UVIn;
-		bHaveN = bHaveC = bHaveUV = true;
-	}
-};
-
 /**
 * FDynamicMesh3 is a dynamic triangle mesh class. The mesh has has connectivity,
 * is an indexed mesh, and allows for gaps in the index space.
@@ -136,56 +85,54 @@ struct FVertexInfo
 class DYNAMICMESH_API FDynamicMesh3
 {
 public:
+using FEdgeFlipInfo = DynamicMeshInfo::FEdgeFlipInfo;
+using FEdgeSplitInfo = DynamicMeshInfo::FEdgeSplitInfo;
+using FEdgeCollapseInfo = DynamicMeshInfo::FEdgeCollapseInfo;
+using FMergeEdgesInfo = DynamicMeshInfo::FMergeEdgesInfo;
+using FPokeTriangleInfo = DynamicMeshInfo::FPokeTriangleInfo;
+
 	/** InvalidID indicates that a vertex/edge/triangle ID is invalid */
-	static const int InvalidID; // = IndexConstants::InvalidID = -1
+	constexpr static int InvalidID = IndexConstants::InvalidID;
 	/** NonManifoldID is returned by AppendTriangle() to indicate that the added triangle would result in nonmanifold geometry and hence was ignored */
-	static const int NonManifoldID; // = -2
+	constexpr static int NonManifoldID = -2;
 	/** InvalidGroupID indicates that a group ID is invalid */
-	static const int InvalidGroupID; // = IndexConstants::InvalidID = -1
+	constexpr static int InvalidGroupID = IndexConstants::InvalidID;
 
-	static FVector3d InvalidVertex()
-	{
-		return FVector3d(TNumericLimits<double>::Max(), 0, 0);
-	}
-	static FIndex3i InvalidTriangle()
-	{
-		return FIndex3i(InvalidID, InvalidID, InvalidID);
-	}
-	static FIndex2i InvalidEdge()
-	{
-		return FIndex2i(InvalidID, InvalidID);
-	}
+	constexpr static FVector3d InvalidVertex{ TNumericLimits<double>::Max(), 0.0, 0.0 };
+	constexpr static FIndex3i  InvalidTriangle{ InvalidID, InvalidID, InvalidID };
+	constexpr static FIndex2i  InvalidEdge{ InvalidID, InvalidID };
 
 protected:
-	/** Reference counts of vertex indices. Iterate over this to find out which vertex indices are valid. */
-	FRefCountVector VertexRefCounts;
 	/** List of vertex positions */
-	TDynamicVector<double> Vertices;
+	TDynamicVector<double> Vertices{};
+	/** Reference counts of vertex indices. Iterate over this to find out which vertex indices are valid. */
+	FRefCountVector VertexRefCounts{};
 	/** (optional) List of per-vertex normals */
-	TDynamicVector<float>* VertexNormals = nullptr;
+	TOptional<TDynamicVector<float>> VertexNormals{};
 	/** (optional) List of per-vertex colors */
-	TDynamicVector<float>* VertexColors = nullptr;
-	/** (optional) List of per-vertex uvs */
-	TDynamicVector<float>* VertexUVs = nullptr;
-
+	TOptional<TDynamicVector<float>> VertexColors{};
+	/** (optional) List of per-vertex uv's */
+	TOptional<TDynamicVector<float>> VertexUVs{};
 	/** List of per-vertex edge one-rings */
 	FSmallListSet VertexEdgeLists;
 
-	/** Reference counts of triangle indices. Iterate over this to find out which triangle indices are valid. */
-	FRefCountVector TriangleRefCounts;
 	/** List of triangle vertex-index triplets [Vert0 Vert1 Vert2]*/
 	TDynamicVector<int> Triangles;
+	/** Reference counts of triangle indices. Iterate over this to find out which triangle indices are valid. */
+	FRefCountVector TriangleRefCounts;
 	/** List of triangle edge triplets [Edge0 Edge1 Edge2] */
 	TDynamicVector<int> TriangleEdges;
 	/** (optional) List of per-triangle group identifiers */
-	TDynamicVector<int>* TriangleGroups = nullptr;
+	TOptional<TDynamicVector<int>> TriangleGroups{};
+	/** Upper bound on the triangle group IDs used in the mesh (may be larger than the actual maximum if triangles have been deleted) */
+	int GroupIDCounter = 0;
 
-	/** Reference counts of edge indices. Iterate over this to find out which edge indices are valid. */
-	FRefCountVector EdgeRefCounts;
 	/** List of edge elements. An edge is four elements [VertA, VertB, Tri0, Tri1], where VertA < VertB, and Tri1 may be InvalidID (if the edge is a boundary edge) */
 	TDynamicVector<int> Edges;
+	/** Reference counts of edge indices. Iterate over this to find out which edge indices are valid. */
+	FRefCountVector EdgeRefCounts;
 
-	FDynamicMeshAttributeSet* AttributeSet = nullptr;
+	TUniquePtr<FDynamicMeshAttributeSet> AttributeSet{};
 
 	/** The mesh timestamp is incremented any time a function that modifies the mesh is called */
 	int Timestamp = 0;
@@ -194,10 +141,8 @@ protected:
 	/** The topology timestamp is incremented any time a function that modifies the mesh topology is called */
 	int TopologyTimestamp = 0;
 
-	/** Upper bound on the triangle group IDs used in the mesh (may be larger than the actual maximum if triangles have been deleted) */
-	int GroupIDCounter = 0;
 
-	/** Cached vertex bounding box (includes un-referenced vertices) */
+	/** Cached vertex bounding box (includes unreferenced vertices) */
 	FAxisAlignedBox3d CachedBoundingBox;
 	/** timestamp for CachedBoundingBox, if less than current timestamp, cache is invalid */
 	int CachedBoundingBoxTimestamp = -1;
@@ -207,48 +152,36 @@ protected:
 	int CachedIsClosedTimestamp = -1;
 
 public:
-	virtual ~FDynamicMesh3();
-
+	/** Default constructor */
 	FDynamicMesh3() 
 		: FDynamicMesh3(true, false, false, false)
 	{
 	}
 
-	explicit FDynamicMesh3(bool bWantNormals, bool bWantColors, bool bWantUVs, bool bWantTriGroups);
+	/** Copy/Move construction */
+	FDynamicMesh3(const FDynamicMesh3& CopyMesh );
+	FDynamicMesh3(FDynamicMesh3&& MoveMesh);
+
+	/** Copy and move assignment */
+	const FDynamicMesh3& operator=(const FDynamicMesh3& CopyMesh);
+	const FDynamicMesh3& operator=(FDynamicMesh3&& MoveMesh);
+
+	/** Destructor */
+	virtual ~FDynamicMesh3();
+
 
+	/** Construct an empty mesh with specified attributes */
+	explicit FDynamicMesh3(bool bWantNormals, bool bWantColors, bool bWantUVs, bool bWantTriGroups);
 	explicit FDynamicMesh3(EMeshComponents flags)
 		: FDynamicMesh3(((int)flags & (int)EMeshComponents::VertexNormals) != 0, ((int)flags & (int)EMeshComponents::VertexColors) != 0,
 						((int)flags & (int)EMeshComponents::VertexUVs) != 0, ((int)flags & (int)EMeshComponents::FaceGroups) != 0)
 	{
 	}
 
-	/**
-	 * @param CopyMesh mesh to copy
-	 * @param bCompact if true, compact CopyMesh on the fly
-	 * @param bWantNormals should we copy per-vertex normals, if they exist
-	 * @param bWantColors should we copy per-vertex colors, if they exist
-	 * @param bWantUVs should we copy per-vertex uvs, if they exist
-	 */
-	FDynamicMesh3(const FDynamicMesh3& CopyMesh, bool bCompact = false, bool bWantNormals = true, bool bWantColors = true, bool bWantUVs = true, bool bWantAttributes = true);
-
-	/**
-	 * @param CopyMesh mesh to copy
-	 * @param bCompact if true, compact CopyMesh on the fly
-	 * @param Flags which components of CopyMesh to copy, if they exist
-	 */
-	FDynamicMesh3(const FDynamicMesh3& CopyMesh, bool bCompact, EMeshComponents Flags);
-
-	/** Initialize mesh from the output of a MeshShapeGenerator (assumes Generate() was already called) */
+	/** Construction from Mesh Generator */
 	FDynamicMesh3(const FMeshShapeGenerator* Generator);
 
-	/** copy assignment operator */
-	const FDynamicMesh3& operator=(const FDynamicMesh3& CopyMesh);
-
-	//
-	// Copy functions to construct a mesh from an input mesh
-	//
-public:
-	/** Set internal data structures to be a copy of input mesh */
+	/** Set internal data structures to be a copy of input mesh using the specified attributes*/
 	void Copy(const FDynamicMesh3& CopyMesh, bool bNormals = true, bool bColors = true, bool bUVs = true, bool bAttributes = true);
 
 	/** Initialize mesh from the output of a MeshShapeGenerator (assumes Generate() was already called) */
@@ -260,7 +193,7 @@ public:
 		TMap<int, int> MapV;
 	};
 
-	/** Copy input mesh while compacting, ie removing unused vertices/triangles/edges */
+	/** Copy input mesh while compacting, i.e. removing unused vertices/triangles/edges */
 	void CompactCopy(const FDynamicMesh3& CopyMesh, bool bNormals = true, bool bColors = true, bool bUVs = true, bool bAttributes = true, FCompactMaps* CompactInfo = nullptr);
 
 	/** Discard all data */
@@ -285,22 +218,22 @@ public:
 		return (int)EdgeRefCounts.GetCount();
 	}
 
-	/** @return upper bound on vertex IDs used in the mesh, ie all vertex IDs in use are < MaxVertexID */
+	/** @return upper bound on vertex IDs used in the mesh, i.e. all vertex IDs in use are < MaxVertexID */
 	int MaxVertexID() const
 	{
 		return (int)VertexRefCounts.GetMaxIndex();
 	}
-	/** @return upper bound on triangle IDs used in the mesh, ie all triangle IDs in use are < MaxTriangleID */
+	/** @return upper bound on triangle IDs used in the mesh, i.e. all triangle IDs in use are < MaxTriangleID */
 	int MaxTriangleID() const
 	{
 		return (int)TriangleRefCounts.GetMaxIndex();
 	}
-	/** @return upper bound on edge IDs used in the mesh, ie all edge IDs in use are < MaxEdgeID */
+	/** @return upper bound on edge IDs used in the mesh, i.e. all edge IDs in use are < MaxEdgeID */
 	int MaxEdgeID() const
 	{
 		return (int)EdgeRefCounts.GetMaxIndex();
 	}
-	/** @return upper bound on group IDs used in the mesh, ie all group IDs in use are < MaxGroupID */
+	/** @return upper bound on group IDs used in the mesh, i.e. all group IDs in use are < MaxGroupID */
 	int MaxGroupID() const
 	{
 		return GroupIDCounter;
@@ -310,29 +243,27 @@ public:
 	/** @return true if this mesh has per-vertex normals */
 	bool HasVertexNormals() const
 	{
-		return VertexNormals != nullptr;
+		return VertexNormals.IsSet();
 	}
 	/** @return true if this mesh has per-vertex colors */
 	bool HasVertexColors() const
 	{
-		return VertexColors != nullptr;
+		return VertexColors.IsSet();
 	}
 	/** @return true if this mesh has per-vertex UVs */
 	bool HasVertexUVs() const
 	{
-		return VertexUVs != nullptr;
+		return VertexUVs.IsSet();
 	}
 	/** @return true if this mesh has per-triangle groups */
 	bool HasTriangleGroups() const
 	{
-		return TriangleGroups != nullptr;
+		return TriangleGroups.IsSet();
 	}
-
-
 	/** @return true if this mesh has attribute layers */
 	bool HasAttributes() const
 	{
-		return AttributeSet != nullptr;
+		return AttributeSet.IsValid();
 	}
 
 	/** @return bitwise-or of EMeshComponents flags specifying which extra data this mesh has */
@@ -665,7 +596,8 @@ public:
 		}
 		check(IsVertex(vID));
 		int i = 3 * vID;
-		return FVector3f((*VertexNormals)[i], (*VertexNormals)[i + 1], (*VertexNormals)[i + 2]);
+		const TDynamicVector<float>& Normals = VertexNormals.GetValue();
+		return FVector3f( Normals[i], Normals[i + 1], Normals[i + 2] );
 	}
 
 	void SetVertexNormal(int vID, const FVector3f& vNewNormal)
@@ -674,9 +606,10 @@ public:
 		{
 			check(IsVertex(vID));
 			int i = 3 * vID;
-			(*VertexNormals)[i] = vNewNormal.X;
-			(*VertexNormals)[i + 1] = vNewNormal.Y;
-			(*VertexNormals)[i + 2] = vNewNormal.Z;
+			TDynamicVector<float>& Normals = VertexNormals.GetValue();
+			Normals[i]     = vNewNormal.X;
+			Normals[i + 1] = vNewNormal.Y;
+			Normals[i + 2] = vNewNormal.Z;
 			UpdateTimeStamp(false, false);
 		}
 	}
@@ -693,7 +626,9 @@ public:
 		}
 		check(IsVertex(vID));
 		int i = 3 * vID;
-		return FVector3f((*VertexColors)[i], (*VertexColors)[i + 1], (*VertexColors)[i + 2]);
+
+		const TDynamicVector<float>& Colors = VertexColors.GetValue();
+		return FVector3f(Colors[i], Colors[i + 1], Colors[i + 2]);
 	}
 
 	void SetVertexColor(int vID, const FVector3f& vNewColor)
@@ -702,9 +637,10 @@ public:
 		{
 			check(IsVertex(vID));
 			int i = 3 * vID;
-			(*VertexColors)[i] = vNewColor.X;
-			(*VertexColors)[i + 1] = vNewColor.Y;
-			(*VertexColors)[i + 2] = vNewColor.Z;
+			TDynamicVector<float>& Colors = VertexColors.GetValue();
+			Colors[i] = vNewColor.X;
+			Colors[i + 1] = vNewColor.Y;
+			Colors[i + 2] = vNewColor.Z;
 			UpdateTimeStamp(false, false);
 		}
 	}
@@ -720,7 +656,8 @@ public:
 		}
 		check(IsVertex(vID));
 		int i = 2 * vID;
-		return FVector2f((*VertexUVs)[i], (*VertexUVs)[i + 1]);
+		const TDynamicVector<float> UVs = VertexUVs.GetValue();
+		return FVector2f(UVs[i], UVs[i + 1]);
 	}
 
 	void SetVertexUV(int vID, const FVector2f& vNewUV)
@@ -729,8 +666,9 @@ public:
 		{
 			check(IsVertex(vID));
 			int i = 2 * vID;
-			(*VertexUVs)[i] = vNewUV.X;
-			(*VertexUVs)[i + 1] = vNewUV.Y;
+			TDynamicVector<float> UVs = VertexUVs.GetValue();
+			UVs[i] = vNewUV.X;
+			UVs[i + 1] = vNewUV.Y;
 			UpdateTimeStamp(false, false);
 		}
 	}
@@ -743,7 +681,7 @@ public:
 	int GetTriangleGroup(int tID) const
 	{
 		return (HasTriangleGroups() == false) ? -1
-			: (TriangleRefCounts.IsValid(tID) ? (*TriangleGroups)[tID] : 0);
+			: (TriangleRefCounts.IsValid(tID) ? TriangleGroups.GetValue()[tID] : 0);
 	}
 
 	void SetTriangleGroup(int tid, int group_id)
@@ -751,7 +689,7 @@ public:
 		if (HasTriangleGroups())
 		{
 			check(IsTriangle(tid));
-			(*TriangleGroups)[tid] = group_id;
+			TriangleGroups.GetValue()[tid] = group_id;
 			GroupIDCounter = FMath::Max(GroupIDCounter, group_id + 1);
 			UpdateTimeStamp(false, false);
 		}
@@ -759,15 +697,14 @@ public:
 
 	FDynamicMeshAttributeSet* Attributes()
 	{
-		return AttributeSet;
+		return HasAttributes() ? AttributeSet.Get() : nullptr;
 	}
 	const FDynamicMeshAttributeSet* Attributes() const
 	{
-		return AttributeSet;
+		return HasAttributes() ? AttributeSet.Get() : nullptr;
 	}
 
 	void EnableAttributes();
-
 	void DiscardAttributes();
 
 
@@ -983,7 +920,7 @@ public:
 	void GetVtxOneRingCentroid(int VertexID, FVector3d& CentroidOut) const;
 
 	/**
-	 * Compute mesh winding number, from Jacobson et al, Robust Inside-Outside Segmentation using Generalized Winding Numbers
+	 * Compute mesh winding number, from Jacobson et. al., Robust Inside-Outside Segmentation using Generalized Winding Numbers
 	 * http://igl.ethz.ch/projects/winding-number/
 	 * returns ~0 for points outside a closed, consistently oriented mesh, and a positive or negative integer
 	 * for points inside, with value > 1 depending on how many "times" the point inside the mesh (like in 2D polygon winding)
@@ -1004,17 +941,16 @@ public:
 	}
 	const TDynamicVector<float>* GetNormalsBuffer()
 	{
-		return VertexNormals;
+		return HasVertexNormals() ? &VertexNormals.GetValue() : nullptr;
 	}
 	const TDynamicVector<float>* GetColorsBuffer()
 	{
-		return VertexColors;
+		return HasVertexColors() ? &VertexColors.GetValue() : nullptr;
 	}
 	const TDynamicVector<float>* GetUVBuffer()
 	{
-		return VertexUVs;
+		return HasVertexUVs() ? &VertexUVs.GetValue() : nullptr;
 	}
-
 	const TDynamicVector<int>& GetTrianglesBuffer()
 	{
 		return Triangles;
@@ -1025,9 +961,8 @@ public:
 	}
 	const TDynamicVector<int>* GetTriangleGroupsBuffer()
 	{
-		return TriangleGroups;
+		return HasTriangleGroups() ? &TriangleGroups.GetValue() : nullptr;
 	}
-
 	const TDynamicVector<int>& GetEdgesBuffer()
 	{
 		return Edges;
@@ -1069,16 +1004,16 @@ public:
 	/**
 	 * Remove vertex vID, and all connected triangles if bRemoveAllTriangles = true
 	 * Returns Failed_VertexStillReferenced if vertex is still referenced by triangles.
-	 * if bPreserveManifold, checks that we will not create a bowtie vertex first
+	 * if bPreserveManifold, checks that we will not create a bow tie vertex first
 	 */
 	EMeshResult RemoveVertex(int VertexID, bool bRemoveAllTriangles = true, bool bPreserveManifold = false);
 
 	/**
 	* Remove a triangle from the mesh. Also removes any unreferenced edges after tri is removed.
 	* If bRemoveIsolatedVertices is false, then if you remove all tris from a vert, that vert is also removed.
-	* If bPreserveManifold, we check that you will not create a bowtie vertex (and return false).
-	* If this check is not done, you have to make sure you don't create a bowtie, because other
-	* code assumes we don't have bowties, and will not handle it properly
+	* If bPreserveManifold, we check that you will not create a bow tie vertex (and return false).
+	* If this check is not done, you have to make sure you don't create a bow tie, because other
+	* code assumes we don't have bow ties, and will not handle it properly
 	*/
 	EMeshResult RemoveTriangle(int TriangleID, bool bRemoveIsolatedVertices = true, bool bPreserveManifold = false);
 
@@ -1089,24 +1024,6 @@ public:
 	 */
 	virtual EMeshResult SetTriangle(int TriangleID, const FIndex3i& NewVertices, bool bRemoveIsolatedVertices = true);
 
-
-
-	/** Information about the mesh elements created by a call to SplitEdge() */
-	struct FEdgeSplitInfo
-	{
-		int OriginalEdge;					// the edge that was split
-		FIndex2i OriginalVertices;			// original edge vertices [a,b]
-		FIndex2i OtherVertices;				// original opposing vertices [c,d] - d is InvalidID for boundary edges
-		FIndex2i OriginalTriangles;			// original edge triangles [t0,t1]
-		bool bIsBoundary;					// was the split edge a boundary edge?  (redundant)
-
-		int NewVertex;						// new vertex f that was created
-		FIndex2i NewTriangles;				// new triangles [t2,t3], oriented as explained below
-		FIndex3i NewEdges;					// new edges are [f,b], [f,c] and [f,d] if this is not a boundary edge
-		
-		double SplitT;						// parameter value for NewVertex along original edge
-	};
-
 	/**
 	 * Split an edge of the mesh by inserting a vertex. This creates a new triangle on either side of the edge (ie a 2-4 split).
 	 * If the original edge had vertices [a,b], with triangles t0=[a,b,c] and t1=[b,a,d],  then the split inserts new vertex f.
@@ -1128,19 +1045,7 @@ public:
 	 */
 	EMeshResult SplitEdge(int EdgeVertA, int EdgeVertB, FEdgeSplitInfo& SplitInfo);
 
-
-
-	/** Information about the mesh elements modified by a call to FlipEdge() */
-	struct FEdgeFlipInfo
-	{
-		int EdgeID;						// the edge that was flipped
-		FIndex2i OriginalVerts;			// original verts of the flipped edge, that are no longer connected
-		FIndex2i OpposingVerts;			// the opposing verts of the flipped edge, that are now connected
-		FIndex2i Triangles;				// the two triangle IDs. Original tris vert [Vert0,Vert1,OtherVert0] and [Vert1,Vert0,OtherVert1].
-										// New triangles are [OtherVert0, OtherVert1, Vert1] and [OtherVert1, OtherVert0, Vert0]
-	};
-
-	/** 
+	/**
 	 * Flip/Rotate an edge of the mesh. This does not change the number of edges, vertices, or triangles.
 	 * Boundary edges of the mesh cannot be flipped.
 	 * @param EdgeAB index of edge to be flipped
@@ -1159,22 +1064,6 @@ public:
 
 
 
-	/** Information about mesh elements modified/removed by CollapseEdge() */
-	struct FEdgeCollapseInfo
-	{
-		int KeptVertex;					// the vertex that was kept (ie collapsed "to")
-		int RemovedVertex;				// the vertex that was removed
-		FIndex2i OpposingVerts;			// the opposing vertices [c,d]. If the edge was a boundary edge, d is InvalidID
-		bool bIsBoundary;				// was the edge a boundary edge
-
-		int CollapsedEdge;				// the edge that was collapsed/removed
-		FIndex2i RemovedTris;			// the triangles that were removed in the collapse (second is InvalidID for boundary edge)
-		FIndex2i RemovedEdges;			// the edges that were removed (second is InvalidID for boundary edge)
-		FIndex2i KeptEdges;				// the edges that were kept (second is InvalidID for boundary edge)
-
-		double CollapseT;				// interpolation parameter along edge for new vertex in range [0,1]
-	};
-
 	/**
 	 * Collapse the edge between the two vertices, if topologically possible.
 	 * @param KeepVertID index of the vertex that should be kept
@@ -1191,19 +1080,6 @@ public:
 
 
 
-	/** Information about mesh elements modified by MergeEdges() */
-	struct FMergeEdgesInfo
-	{
-		int KeptEdge;				// the edge that was kept
-		int RemovedEdge;			// the edge that was removed
-
-		FIndex2i KeptVerts;			// The two vertices that were kept (redundant w/ KeptEdge?)
-		FIndex2i RemovedVerts;		// The removed vertices of RemovedEdge. Either may be InvalidID if it was same as the paired KeptVert
-
-		FIndex2i ExtraRemovedEdges; // extra removed edges, see description below. Either may be or InvalidID
-		FIndex2i ExtraKeptEdges;	// extra kept edges, paired with ExtraRemovedEdges
-	};
-
 	/**
 	 * Given two edges of the mesh, weld both their vertices, so that one edge is removed.
 	 * This could result in one neighbour edge-pair attached to each vertex also collapsing,
@@ -1221,19 +1097,6 @@ public:
 
 
 
-	/** Information about mesh elements modified/created by PokeTriangle() */
-	struct FPokeTriangleInfo
-	{
-		int OriginalTriangle;				// the triangle that was poked
-		FIndex3i TriVertices;				// vertices of the original triangle
-			
-		int NewVertex;						// the new vertex that was inserted
-		FIndex2i NewTriangles;				// the two new triangles that were added (OriginalTriangle is re-used, see code for vertex orders)
-		FIndex3i NewEdges;					// the three new edges connected to NewVertex
-
-		FVector3d BaryCoords;				// barycentric coords that NewVertex was inserted at
-	};
-
 	/**
 	 * Insert a new vertex inside a triangle, ie do a 1 to 3 triangle split
 	 * @param TriangleID index of triangle to poke
diff --git a/Engine/Plugins/Experimental/GeometryProcessing/Source/DynamicMesh/Public/DynamicMeshAttributeSet.h b/Engine/Plugins/Experimental/GeometryProcessing/Source/DynamicMesh/Public/DynamicMeshAttributeSet.h
--- a/Engine/Plugins/Experimental/GeometryProcessing/Source/DynamicMesh/Public/DynamicMeshAttributeSet.h
+++ b/Engine/Plugins/Experimental/GeometryProcessing/Source/DynamicMesh/Public/DynamicMeshAttributeSet.h
@@ -7,6 +7,7 @@
 #include "DynamicMeshTriangleAttribute.h"
 #include "DynamicAttribute.h"
 #include "GeometryTypes.h"
+#include "InfoTypes.h"
 
 /** Standard UV overlay type - 2-element float */
 typedef TDynamicMeshVectorOverlay<float, 2, FVector2f> FDynamicMeshUVOverlay;
@@ -70,7 +71,28 @@ public:
 	/** @return the parent mesh for this overlay */
 	FDynamicMesh3* GetParentMesh() { return ParentMesh; }
 
+private:
+	/** @set the parent mesh for this overlay.  Only safe for use during FDynamicMesh move */
+	void Reparent(FDynamicMesh3* NewParent)
+	{
+		for (int UVIdx = 0; UVIdx < NumUVLayers(); UVIdx++)
+		{
+			UVLayers[UVIdx].Reparent( NewParent );
+		}
+		Normals0.Reparent( NewParent );
+
+		if (MaterialIDAttrib)
+		{
+			MaterialIDAttrib->Reparent( NewParent );
+		}
+
+		for (int Idx = 0; Idx < GenericAttributes.Num(); Idx++)
+		{
+			GenericAttributes[Idx]->Reparent( NewParent );
+		}
+	}
 
+public:
 	/** @return number of UV layers */
 	virtual int NumUVLayers() const
 	{
@@ -251,11 +273,11 @@ protected:
 	virtual void OnRemoveTriangle(int TriangleID);
 	virtual void OnRemoveVertex(int VertexID);
 	virtual void OnReverseTriOrientation(int TriangleID);
-	virtual void OnSplitEdge(const FDynamicMesh3::FEdgeSplitInfo & splitInfo);
-	virtual void OnFlipEdge(const FDynamicMesh3::FEdgeFlipInfo & flipInfo);
-	virtual void OnCollapseEdge(const FDynamicMesh3::FEdgeCollapseInfo & collapseInfo);
-	virtual void OnPokeTriangle(const FDynamicMesh3::FPokeTriangleInfo & pokeInfo);
-	virtual void OnMergeEdges(const FDynamicMesh3::FMergeEdgesInfo & mergeInfo);
+	virtual void OnSplitEdge(const DynamicMeshInfo::FEdgeSplitInfo & splitInfo);
+	virtual void OnFlipEdge(const DynamicMeshInfo::FEdgeFlipInfo & flipInfo);
+	virtual void OnCollapseEdge(const DynamicMeshInfo::FEdgeCollapseInfo & collapseInfo);
+	virtual void OnPokeTriangle(const DynamicMeshInfo::FPokeTriangleInfo & pokeInfo);
+	virtual void OnMergeEdges(const DynamicMeshInfo::FMergeEdgesInfo & mergeInfo);
 
 	/**
 	 * Check validity of attributes
diff --git a/Engine/Plugins/Experimental/GeometryProcessing/Source/DynamicMesh/Public/DynamicMeshOverlay.h b/Engine/Plugins/Experimental/GeometryProcessing/Source/DynamicMesh/Public/DynamicMeshOverlay.h
--- a/Engine/Plugins/Experimental/GeometryProcessing/Source/DynamicMesh/Public/DynamicMeshOverlay.h
+++ b/Engine/Plugins/Experimental/GeometryProcessing/Source/DynamicMesh/Public/DynamicMeshOverlay.h
@@ -7,8 +7,9 @@
 #include "Util/RefCountVector.h"
 #include "VectorTypes.h"
 #include "GeometryTypes.h"
-#include "DynamicMesh3.h"
+#include "InfoTypes.h"
 
+class FDynamicMesh3;
 
 /**
  * TDynamicMeshOverlay is an add-on to a FDynamicMesh3 that allows for per-triangle storage
@@ -60,7 +61,13 @@ public:
 	{
 		ParentMesh = ParentMeshIn;
 	}
-
+private:
+	/** @set the parent mesh for this overlay.  Only safe for use during FDynamicMesh move */
+	void Reparent(FDynamicMesh3* ParentMeshIn)
+	{
+		ParentMesh = ParentMeshIn;
+	}
+public:
 	/** @return the parent mesh for this overlay */
 	const FDynamicMesh3* GetParentMesh() const { return ParentMesh; }
 	/** @return the parent mesh for this overlay */
@@ -254,15 +261,15 @@ public:
 	/** Reverse the orientation of a triangle's elements */
 	void OnReverseTriOrientation(int TriangleID);
 	/** Update the overlay to reflect an edge split in the parent mesh */
-	void OnSplitEdge(const FDynamicMesh3::FEdgeSplitInfo& SplitInfo);
+	void OnSplitEdge(const DynamicMeshInfo::FEdgeSplitInfo& SplitInfo);
 	/** Update the overlay to reflect an edge flip in the parent mesh */
-	void OnFlipEdge(const FDynamicMesh3::FEdgeFlipInfo& FlipInfo);
+	void OnFlipEdge(const DynamicMeshInfo::FEdgeFlipInfo& FlipInfo);
 	/** Update the overlay to reflect an edge collapse in the parent mesh */
-	void OnCollapseEdge(const FDynamicMesh3::FEdgeCollapseInfo& CollapseInfo);
+	void OnCollapseEdge(const DynamicMeshInfo::FEdgeCollapseInfo& CollapseInfo);
 	/** Update the overlay to reflect a face poke in the parent mesh */
-	void OnPokeTriangle(const FDynamicMesh3::FPokeTriangleInfo& PokeInfo);
+	void OnPokeTriangle(const DynamicMeshInfo::FPokeTriangleInfo& PokeInfo);
 	/** Update the overlay to reflect an edge merge in the parent mesh */
-	void OnMergeEdges(const FDynamicMesh3::FMergeEdgesInfo& MergeInfo);
+	void OnMergeEdges(const DynamicMeshInfo::FMergeEdgesInfo& MergeInfo);
 
 protected:
 	/** Set the value at an Element to be a linear interpolation of two other Elements */
diff --git a/Engine/Plugins/Experimental/GeometryProcessing/Source/DynamicMesh/Public/DynamicMeshTriangleAttribute.h b/Engine/Plugins/Experimental/GeometryProcessing/Source/DynamicMesh/Public/DynamicMeshTriangleAttribute.h
--- a/Engine/Plugins/Experimental/GeometryProcessing/Source/DynamicMesh/Public/DynamicMeshTriangleAttribute.h
+++ b/Engine/Plugins/Experimental/GeometryProcessing/Source/DynamicMesh/Public/DynamicMeshTriangleAttribute.h
@@ -72,6 +72,14 @@ public:
 		ParentMesh = ParentMeshIn;
 	}
 
+private:
+	/** @set the parent mesh for this overlay.  Only safe for use during FDynamicMesh move */
+	void Reparent(FDynamicMesh3* ParentMeshIn)
+	{
+		ParentMesh = ParentMeshIn;
+	}
+
+public:
 	/** @return the parent mesh for this overlay */
 	const FDynamicMesh3* GetParentMesh() const { return ParentMesh; }
 	/** @return the parent mesh for this overlay */
@@ -199,7 +207,7 @@ public:
 public:
 
 	/** Update the overlay to reflect an edge split in the parent mesh */
-	void OnSplitEdge(const FDynamicMesh3::FEdgeSplitInfo& SplitInfo)
+	void OnSplitEdge(const DynamicMeshInfo::FEdgeSplitInfo& SplitInfo)
 	{
 		CopyValue(SplitInfo.OriginalTriangles.A, SplitInfo.NewTriangles.A);
 		if (SplitInfo.OriginalTriangles.B != FDynamicMesh3::InvalidID)
@@ -209,26 +217,26 @@ public:
 	}
 
 	/** Update the overlay to reflect an edge flip in the parent mesh */
-	void OnFlipEdge(const FDynamicMesh3::FEdgeFlipInfo& FlipInfo)
+	void OnFlipEdge(const DynamicMeshInfo::FEdgeFlipInfo& FlipInfo)
 	{
 		// yikes! triangles did not actually change so we will leave attrib unmodified
 	}
 
 	/** Update the overlay to reflect an edge collapse in the parent mesh */
-	void OnCollapseEdge(const FDynamicMesh3::FEdgeCollapseInfo& CollapseInfo)
+	void OnCollapseEdge(const DynamicMeshInfo::FEdgeCollapseInfo& CollapseInfo)
 	{
 		// nothing to do here, triangles were only deleted
 	}
 
 	/** Update the overlay to reflect a face poke in the parent mesh */
-	void OnPokeTriangle(const FDynamicMesh3::FPokeTriangleInfo& PokeInfo)
+	void OnPokeTriangle(const DynamicMeshInfo::FPokeTriangleInfo& PokeInfo)
 	{
 		CopyValue(PokeInfo.OriginalTriangle, PokeInfo.NewTriangles.A);
 		CopyValue(PokeInfo.OriginalTriangle, PokeInfo.NewTriangles.B);
 	}
 
 	/** Update the overlay to reflect an edge merge in the parent mesh */
-	void OnMergeEdges(const FDynamicMesh3::FMergeEdgesInfo& MergeInfo)
+	void OnMergeEdges(const DynamicMeshInfo::FMergeEdgesInfo& MergeInfo)
 	{
 		// nothing to do here because triangles did not change
 	}
diff --git a/Engine/Plugins/Experimental/GeometryProcessing/Source/DynamicMesh/Public/DynamicVertexAttribute.h b/Engine/Plugins/Experimental/GeometryProcessing/Source/DynamicMesh/Public/DynamicVertexAttribute.h
--- a/Engine/Plugins/Experimental/GeometryProcessing/Source/DynamicMesh/Public/DynamicVertexAttribute.h
+++ b/Engine/Plugins/Experimental/GeometryProcessing/Source/DynamicMesh/Public/DynamicVertexAttribute.h
@@ -74,7 +74,9 @@ public:
 	const ParentType* GetParent() const { return Parent; }
 	/** @return the parent for this attribute */
 	ParentType* GetParent() { return Parent; }
-
+private:
+	void Reparent(ParentType* NewParent) override { Parent = NewParent;  }
+public:
 	virtual TDynamicAttributeBase<ParentType>* MakeCopy(ParentType* ParentIn) const override
 	{
 		TDynamicVertexAttribute<AttribValueType, AttribDimension, ParentType>* ToFill = new TDynamicVertexAttribute<AttribValueType, AttribDimension, ParentType>(ParentIn);
diff --git a/Engine/Plugins/Experimental/GeometryProcessing/Source/GeometricObjects/Public/IndexTypes.h b/Engine/Plugins/Experimental/GeometryProcessing/Source/GeometricObjects/Public/IndexTypes.h
--- a/Engine/Plugins/Experimental/GeometryProcessing/Source/GeometricObjects/Public/IndexTypes.h
+++ b/Engine/Plugins/Experimental/GeometryProcessing/Source/GeometricObjects/Public/IndexTypes.h
@@ -18,26 +18,23 @@ namespace IndexConstants
  */
 struct FIndex2i
 {
-	int A, B;
+	int A{ IndexConstants::InvalidID };
+	int B{ IndexConstants::InvalidID };
 
-	FIndex2i()
-	{
-	}
-	FIndex2i(int ValA, int ValB)
-	{
-		this->A = ValA;
-		this->B = ValB;
-	}
+	constexpr FIndex2i() = default;
+	constexpr FIndex2i(int ValA, int ValB)
+		: A(ValA), B(ValB)
+	{}
 
-	static FIndex2i Zero()
+	constexpr static FIndex2i Zero()
 	{
 		return FIndex2i(0, 0);
 	}
-	static FIndex2i Max()
+	constexpr static FIndex2i Max()
 	{
 		return FIndex2i(TNumericLimits<int>::Max(), TNumericLimits<int>::Max());
 	}
-	static FIndex2i Invalid()
+	constexpr static FIndex2i Invalid()
 	{
 		return FIndex2i(IndexConstants::InvalidID, IndexConstants::InvalidID);
 	}
@@ -82,7 +79,7 @@ struct FIndex2i
 FORCEINLINE uint32 GetTypeHash(const FIndex2i& Index)
 {
 	// (this is how FIntVector and all the other FVectors do their hash functions)
-	// Note: this assumes there's no padding that could contain uncompared data.
+	// Note: this assumes there's no padding that could contain non compared data.
 	return FCrc::MemCrc_DEPRECATED(&Index, sizeof(FIndex2i));
 }
 
@@ -97,27 +94,24 @@ FORCEINLINE uint32 GetTypeHash(const FIndex2i& Index)
  */
 struct FIndex3i
 {
-	int A, B, C;
+	int A{ IndexConstants::InvalidID };
+	int B{ IndexConstants::InvalidID };
+	int C{ IndexConstants::InvalidID };
 
-	FIndex3i()
-	{
-	}
-	FIndex3i(int ValA, int ValB, int ValC)
-	{
-		this->A = ValA;
-		this->B = ValB;
-		this->C = ValC;
-	}
+	constexpr FIndex3i() = default;
+	constexpr FIndex3i(int ValA, int ValB, int ValC)
+		: A(ValA), B(ValB), C(ValC)
+	{}
 
-	static FIndex3i Zero()
+	constexpr static FIndex3i Zero()
 	{
 		return FIndex3i(0, 0, 0);
 	}
-	static FIndex3i Max()
+	constexpr static FIndex3i Max()
 	{
 		return FIndex3i(TNumericLimits<int>::Max(), TNumericLimits<int>::Max(), TNumericLimits<int>::Max());
 	}
-	static FIndex3i Invalid()
+	constexpr static FIndex3i Invalid()
 	{
 		return FIndex3i(IndexConstants::InvalidID, IndexConstants::InvalidID, IndexConstants::InvalidID);
 	}
diff --git a/Engine/Plugins/Experimental/GeometryProcessing/Source/GeometricObjects/Public/Util/DynamicVector.h b/Engine/Plugins/Experimental/GeometryProcessing/Source/GeometricObjects/Public/Util/DynamicVector.h
--- a/Engine/Plugins/Experimental/GeometryProcessing/Source/GeometricObjects/Public/Util/DynamicVector.h
+++ b/Engine/Plugins/Experimental/GeometryProcessing/Source/GeometricObjects/Public/Util/DynamicVector.h
@@ -13,7 +13,7 @@
 
 /*
  * Blocked array with fixed, power-of-two sized blocks.
- * 
+ *
  * Iterator functions suitable for use with range-based for are provided
  */
 template <class Type>
@@ -22,19 +22,29 @@ class TDynamicVector
 public:
 	TDynamicVector()
 	{
-		CurBlock = 0;
-		CurBlockUsed = 0;
 		Blocks.Add(new BlockType());
 	}
 
-	TDynamicVector(const TDynamicVector& Copy)
-	{
-		*this = Copy;
-	}
-
+	TDynamicVector(const TDynamicVector& Copy) = default;
 	TDynamicVector(TDynamicVector&& Moved)
 	{
-		*this = MoveTemp(Moved);
+		Blocks = MoveTemp(Moved.Blocks);
+		CurBlock = MoveTemp(Moved.CurBlock);
+		CurBlockUsed = MoveTemp(Moved.CurBlockUsed);
+		Moved.Blocks.Add(new BlockType());
+		Moved.CurBlock = 0;
+		Moved.CurBlockUsed = 0;
+	}
+	TDynamicVector& operator=(const TDynamicVector& Copy) = default;
+	TDynamicVector& operator=(TDynamicVector&& Moved)
+	{
+		Blocks = MoveTemp(Moved.Blocks);
+		CurBlock = MoveTemp(Moved.CurBlock);
+		CurBlockUsed = MoveTemp(Moved.CurBlockUsed);
+		Moved.Blocks.Add(new BlockType());
+		Moved.CurBlock = 0;
+		Moved.CurBlockUsed = 0;
+		return *this;
 	}
 
 	TDynamicVector(const TArray<Type>& Array)
@@ -54,11 +64,6 @@ public:
 		}
 	}
 
-	~TDynamicVector() {}
-
-	const TDynamicVector& operator=(const TDynamicVector& Copy);
-	const TDynamicVector& operator=(TDynamicVector&& Moved);
-
 	inline void Clear();
 	inline void Fill(const Type& Value);
 	inline void Resize(size_t Count);
@@ -72,7 +77,7 @@ public:
 	inline size_t GetByteCount() const { return Blocks.Num() * BlockSize * sizeof(Type); }
 
 	inline void Add(const Type& Data);
-	inline void Add(const TDynamicVector<Type>& Data);
+	inline void Add(const TDynamicVector& Data);
 	inline void PopBack();
 
 	inline void InsertAt(const Type& Data, unsigned int Index);
@@ -93,22 +98,6 @@ public:
 	template <typename Func>
 	void Apply(const Func& f);
 
-
-protected:
-	// [RMS] BlockSize must be a power-of-two, so we can use bit-shifts in operator[]
-	static constexpr int BlockSize = 2048; // (1 << 11)
-	static constexpr int nShiftBits = 11;
-	static constexpr int BlockIndexBitmask = 2047; // low 11 bits
-
-	unsigned int CurBlock;
-	unsigned int CurBlockUsed;
-
-	using BlockType = TStaticArray<Type, BlockSize>;
-	TIndirectArray<BlockType> Blocks;
-
-	friend class FIterator;
-
-
 public:
 	/*
 	 * FIterator class iterates over values of vector
@@ -144,32 +133,25 @@ public:
 			return DVector != Itr2.DVector || Idx != Itr2.Idx;
 		}
 
-	protected:
-		TDynamicVector<Type>* DVector;
-		int Idx;
-		inline FIterator(TDynamicVector<Type>* Parent, int ICur)
-		{
-			DVector = Parent;
-			Idx = ICur;
-		}
-		friend class TDynamicVector<Type>;
+	private:
+		friend class TDynamicVector;
+		FIterator(TDynamicVector* DVectorIn, int IdxIn)
+			: DVector(DVectorIn), Idx(IdxIn){}
+		TDynamicVector* DVector{};
+		int Idx{0};
 	};
 
 	/** @return iterator at beginning of vector */
 	FIterator begin()
 	{
-		return IsEmpty() ? end() : FIterator(this, 0);
+		return IsEmpty() ? end() : FIterator{this};
 	}
 	/** @return iterator at end of vector */
 	FIterator end()
 	{
-		return FIterator(this, (int)GetLength());
+		return FIterator{this, (int)GetLength()};
 	}
 
-
-
-
-
 	/*
 	 * FConstIterator class iterates over values of vector
 	 */
@@ -200,66 +182,47 @@ public:
 			return DVector != Itr2.DVector || Idx != Itr2.Idx;
 		}
 
-	protected:
-		const TDynamicVector<Type>* DVector;
-		int Idx;
-		inline FConstIterator(const TDynamicVector<Type>* Parent, int ICur)
-		{
-			DVector = Parent;
-			Idx = ICur;
-		}
-		friend class TDynamicVector<Type>;
+	private:
+		friend class TDynamicVector;
+		FConstIterator(const TDynamicVector* DVectorIn, int IdxIn)
+			: DVector(DVectorIn), Idx(IdxIn){}
+		const TDynamicVector* DVector{};
+		int Idx{0};
 	};
 
 	/** @return iterator at beginning of vector */
 	FConstIterator begin() const
 	{
-		return IsEmpty() ? end() : FConstIterator(this, 0);
+		return IsEmpty() ? end() : FConstIterator{this, 0};
 	}
 	/** @return iterator at end of vector */
 	FConstIterator end() const
 	{
-		return FConstIterator(this, (int)GetLength());
+		return FConstIterator{this, (int)GetLength()};
 	}
 
+private:
+	static constexpr int nShiftBits = 11;
+	static constexpr int BlockSize = 1 << nShiftBits;
+	static constexpr int BlockIndexBitmask = BlockSize - 1; // low 11 bits
+	static_assert( BlockSize && ((BlockSize & (BlockSize - 1)) == 0), "DynamicVector: BlockSize must be a power of two");
 
+	unsigned int CurBlock{0};
+	unsigned int CurBlockUsed{0};
 
+	using BlockType = TStaticArray<Type, BlockSize>;
+	TIndirectArray<BlockType> Blocks;
 };
 
-
-
-
-
-
 template <class Type, int N>
 class TDynamicVectorN
 {
 public:
-	TDynamicVectorN()
-	{
-	}
-	TDynamicVectorN(const TDynamicVectorN& Copy)
-		: Data(Copy.Data)
-	{
-	}
-	TDynamicVectorN(TDynamicVectorN&& Moved)
-		: Data(MoveTemp(Moved.Data))
-	{
-	}
-	~TDynamicVectorN()
-	{
-	}
-
-	const TDynamicVectorN& operator=(const TDynamicVectorN& Copy)
-	{
-		Data = Copy.Data;
-		return *this;
-	}
-	const TDynamicVectorN& operator=(TDynamicVectorN&& Moved)
-	{
-		Data = MoveTemp(Moved.Data);
-		return *this;
-	}
+	TDynamicVectorN() = default;
+	TDynamicVectorN(const TDynamicVectorN& Copy) = default;
+	TDynamicVectorN(TDynamicVectorN&& Moved) = default;
+	TDynamicVectorN& operator=(const TDynamicVectorN& Copy) = default;
+	TDynamicVectorN& operator=(TDynamicVectorN&& Moved) = default;
 
 	inline void Clear()
 	{
@@ -277,7 +240,6 @@ public:
 	{
 		Data.Resize(Count * N, InitValue);
 	}
-
 	inline bool IsEmpty() const
 	{
 		return Data.IsEmpty();
@@ -389,7 +351,7 @@ public:
 			Data[TopIndex * N + 3]);
 	}
 
-protected:
+private:
 	TDynamicVector<Type> Data;
 
 	friend class FIterator;
@@ -397,34 +359,12 @@ protected:
 
 template class TDynamicVectorN<double, 2>;
 
-typedef TDynamicVectorN<float, 3> TDynamicVector3f;
-typedef TDynamicVectorN<float, 2> TDynamicVector2f;
-typedef TDynamicVectorN<double, 3> TDynamicVector3d;
-typedef TDynamicVectorN<double, 2> TDynamicVector2d;
-typedef TDynamicVectorN<int, 3> TDynamicVector3i;
-typedef TDynamicVectorN<int, 2> TDynamicVector2i;
-
-
-
-
-
-template <class Type>
-const TDynamicVector<Type>& TDynamicVector<Type>::operator=(const TDynamicVector& Copy)
-{
-	Blocks = Copy.Blocks;
-	CurBlock = Copy.CurBlock;
-	CurBlockUsed = Copy.CurBlockUsed;
-	return *this;
-}
-
-template <class Type>
-const TDynamicVector<Type>& TDynamicVector<Type>::operator=(TDynamicVector&& Moved)
-{
-	Blocks = MoveTemp(Moved.Blocks);
-	CurBlock = Moved.CurBlock;
-	CurBlockUsed = Moved.CurBlockUsed;
-	return *this;
-}
+using  TDynamicVector3f = TDynamicVectorN<float, 3>;
+using  TDynamicVector2f = TDynamicVectorN<float, 2>;
+using  TDynamicVector3d = TDynamicVectorN<double, 3>;
+using  TDynamicVector2d = TDynamicVectorN<double, 2>;
+using  TDynamicVector3i = TDynamicVectorN<int, 3>;
+using  TDynamicVector2i = TDynamicVectorN<int, 2>;
 
 template <class Type>
 void TDynamicVector<Type>::Clear()
diff --git a/Engine/Plugins/Experimental/GeometryProcessing/Source/GeometricObjects/Public/Util/RefCountVector.h b/Engine/Plugins/Experimental/GeometryProcessing/Source/GeometricObjects/Public/Util/RefCountVector.h
--- a/Engine/Plugins/Experimental/GeometryProcessing/Source/GeometricObjects/Public/Util/RefCountVector.h
+++ b/Engine/Plugins/Experimental/GeometryProcessing/Source/GeometricObjects/Public/Util/RefCountVector.h
@@ -22,92 +22,76 @@ class FRefCountVector
 public:
 	static constexpr short INVALID_REF_COUNT = -1;
 
-	TDynamicVector<short> RefCounts;
-	TDynamicVector<int> FreeIndices;
-	int UsedCount;
+	FRefCountVector() = default;
+	FRefCountVector(const FRefCountVector&) = default;
+	FRefCountVector(FRefCountVector&&) = default;
+	FRefCountVector& operator=(const FRefCountVector&) = default;
+	FRefCountVector& operator=(FRefCountVector&&) = default;
 
-	FRefCountVector()
-	{
-		RefCounts = TDynamicVector<short>();
-		FreeIndices = TDynamicVector<int>();
-		UsedCount = 0;
-	}
-
-	FRefCountVector(const FRefCountVector& Copy)
-	{
-		RefCounts = TDynamicVector<short>(Copy.RefCounts);
-		FreeIndices = TDynamicVector<int>(Copy.FreeIndices);
-		UsedCount = Copy.UsedCount;
-	}
-
-
-	bool IsEmpty() const 
+	bool IsEmpty() const
 	{
 		return UsedCount == 0;
 	}
 
-	size_t GetCount() const 
+	size_t GetCount() const
 	{
 		return UsedCount;
 	}
 
-	size_t GetMaxIndex() const 
+	size_t GetMaxIndex() const
 	{
 		return RefCounts.GetLength();
 	}
 
-	bool IsDense() const 
+	bool IsDense() const
 	{
 		return FreeIndices.GetLength() == 0;
 	}
 
-
-	bool IsValid(int Index) const 
+	bool IsValid(int Index) const
 	{
 		return (Index >= 0 && Index < (int)RefCounts.GetLength() && RefCounts[Index] > 0);
 	}
 
-	bool IsValidUnsafe(int Index) const 
+	bool IsValidUnsafe(int Index) const
 	{
 		return RefCounts[Index] > 0;
 	}
 
-
-	int GetRefCount(int Index) const 
+	int GetRefCount(int Index) const
 	{
 		int n = RefCounts[Index];
 		return (n == INVALID_REF_COUNT) ? 0 : n;
 	}
 
-	int GetRawRefCount(int Index) const 
+	int GetRawRefCount(int Index) const
 	{
 		return RefCounts[Index];
 	}
 
-
-	int Allocate() 
+	int Allocate()
 	{
 		UsedCount++;
-		if (FreeIndices.IsEmpty()) 
+		if (FreeIndices.IsEmpty())
 		{
-			// [RMS] do we need this branch anymore? 
+			// [RMS] do we need this branch anymore?
 			RefCounts.Add(1);
 			return (int)RefCounts.GetLength() - 1;
 		}
-		else 
+		else
 		{
 			int iFree = INVALID_REF_COUNT;
-			while (iFree == INVALID_REF_COUNT && FreeIndices.IsEmpty() == false) 
+			while (iFree == INVALID_REF_COUNT && FreeIndices.IsEmpty() == false)
 			{
 				iFree = FreeIndices.Back();
 				FreeIndices.PopBack();
 			}
-			if (iFree != INVALID_REF_COUNT) 
+			if (iFree != INVALID_REF_COUNT)
 			{
 				RefCounts[iFree] = 1;
 				return iFree;
 			}
-			else 
+			else
 			{
 				RefCounts.Add(1);
 				return (int)RefCounts.GetLength() - 1;
@@ -115,9 +99,7 @@ public:
 		}
 	}
 
-
-
-	int Increment(int Index, short IncrementCount = 1) 
+	int Increment(int Index, short IncrementCount = 1)
 	{
 		check(IsValid(Index));
 		// debug check for overflow...
@@ -126,12 +108,12 @@ public:
 		return RefCounts[Index];
 	}
 
-	void Decrement(int Index, short DecrementCount = 1) 
+	void Decrement(int Index, short DecrementCount = 1)
 	{
 		check(IsValid(Index));
 		RefCounts[Index] -= DecrementCount;
 		check(RefCounts[Index] >= 0);
-		if (RefCounts[Index] == 0) 
+		if (RefCounts[Index] == 0)
 		{
 			FreeIndices.Add(Index);
 			RefCounts[Index] = INVALID_REF_COUNT;
@@ -139,21 +121,19 @@ public:
 		}
 	}
 
-
-
 	/**
 	 * allocate at specific Index, which must either be larger than current max Index,
 	 * or on the free list. If larger, all elements up to this one will be pushed onto
 	 * free list. otherwise we have to do a linear search through free list.
-	 * If you are doing many of these, it is likely faster to use 
+	 * If you are doing many of these, it is likely faster to use
 	 * AllocateAtUnsafe(), and then RebuildFreeList() after you are done.
 	 */
 	bool AllocateAt(int Index)
 	{
-		if (Index >= (int)RefCounts.GetLength()) 
+		if (Index >= (int)RefCounts.GetLength())
 		{
 			int j = (int)RefCounts.GetLength();
-			while (j < Index) 
+			while (j < Index)
 			{
 				int InvalidCount = INVALID_REF_COUNT;	// required on older clang because a constexpr can't be passed by ref
 				RefCounts.Add(InvalidCount);
@@ -163,19 +143,17 @@ public:
 			RefCounts.Add(1);
 			UsedCount++;
 			return true;
-
 		}
-		else 
+		else
 		{
 			if (RefCounts[Index] > 0)
 			{
 				return false;
 			}
-
 			int N = (int)FreeIndices.GetLength();
-			for (int i = 0; i < N; ++i) 
+			for (int i = 0; i < N; ++i)
 			{
-				if (FreeIndices[i] == Index) 
+				if (FreeIndices[i] == Index)
 				{
 					FreeIndices[i] = INVALID_REF_COUNT;
 					RefCounts[Index] = 1;
@@ -187,17 +165,16 @@ public:
 		}
 	}
 
-
 	/**
 	 * allocate at specific Index, which must be free or larger than current max Index.
 	 * However, we do not update free list. So, you probably need to do RebuildFreeList() after calling this.
 	 */
 	bool AllocateAtUnsafe(int Index)
 	{
-		if (Index >= (int)RefCounts.GetLength()) 
+		if (Index >= (int)RefCounts.GetLength())
 		{
 			int j = (int)RefCounts.GetLength();
-			while (j < Index) 
+			while (j < Index)
 			{
 				int InvalidCount = INVALID_REF_COUNT;	// required on older clang because a constexpr can't be passed by ref
 				RefCounts.Add(InvalidCount);
@@ -208,7 +185,7 @@ public:
 			return true;
 
 		}
-		else 
+		else
 		{
 			if (RefCounts[Index] > 0)
 			{
@@ -220,8 +197,6 @@ public:
 		}
 	}
 
-
-
 	const TDynamicVector<short>& GetRawRefCounts() const
 	{
 		return RefCounts;
@@ -247,7 +222,6 @@ public:
 	//   remove
 	//   clear
 
-
 	void RebuildFreeList()
 	{
 		FreeIndices = TDynamicVector<int>();
@@ -267,8 +241,6 @@ public:
 		}
 	}
 
-
-
 	void Trim(int maxIndex)
 	{
 		FreeIndices = TDynamicVector<int>();
@@ -276,13 +248,9 @@ public:
 		UsedCount = maxIndex;
 	}
 
-
-
 	//
 	// Iterators
-	// 
-
-
+	//
 
 	/**
 	 * base iterator for indices with valid refcount (skips zero-refcount indices)
@@ -291,17 +259,17 @@ public:
 	{
 	public:
 		inline BaseIterator()
-		{ 
-			Vector = nullptr; 
-			Index = 0; 
-			LastIndex = 0; 
+		{
+			Vector = nullptr;
+			Index = 0;
+			LastIndex = 0;
 		}
 
-		inline bool operator==(const BaseIterator& Other) const 
+		inline bool operator==(const BaseIterator& Other) const
 		{
 			return Index == Other.Index;
 		}
-		inline bool operator!=(const BaseIterator& Other) const 
+		inline bool operator!=(const BaseIterator& Other) const
 		{
 			return Index != Other.Index;
 		}
@@ -335,7 +303,6 @@ public:
 		friend class FRefCountVector;
 	};
 
-
 	/*
 	 *  iterator over valid indices (ie non-zero refcount)
 	 */
@@ -344,7 +311,7 @@ public:
 	public:
 		inline IndexIterator() : BaseIterator() {}
 
-		inline int operator*() const 
+		inline int operator*() const
 		{
 			return this->Index;
 		}
@@ -367,19 +334,16 @@ public:
 		friend class FRefCountVector;
 	};
 
-
-	inline IndexIterator BeginIndices() const 
+	inline IndexIterator BeginIndices() const
 	{
 		return IndexIterator(this, (int)0, (int)RefCounts.GetLength());
 	}
 
-	inline IndexIterator EndIndices() const 
+	inline IndexIterator EndIndices() const
 	{
 		return IndexIterator(this, (int)RefCounts.GetLength(), (int)RefCounts.GetLength());
 	}
 
-
-
 	/**
 	 * enumerable object that provides begin()/end() semantics, so
 	 * you can iterate over valid indices using range-based for loop
@@ -398,7 +362,7 @@ public:
 	 * returns iteration object over valid indices
 	 * usage: for (int idx : indices()) { ... }
 	 */
-	inline IndexEnumerable Indices() const 
+	inline IndexEnumerable Indices() const
 	{
 		return IndexEnumerable(this);
 	}
@@ -414,18 +378,18 @@ public:
 		TFunction<ToType(int)> MapFunc;
 		IndexEnumerable enumerable;
 
-		MappedEnumerable(const IndexEnumerable& enumerable, TFunction<ToType(int)> MapFunc) 
+		MappedEnumerable(const IndexEnumerable& enumerable, TFunction<ToType(int)> MapFunc)
 		{
 			this->enumerable = enumerable;
 			this->MapFunc = MapFunc;
 		}
 
-		MappedIterator<int, ToType, IndexIterator> begin() 
+		MappedIterator<int, ToType, IndexIterator> begin()
 		{
 			return MappedIterator<int, ToType, IndexIterator>(enumerable.begin(), MapFunc);
 		}
 
-		MappedIterator<int, ToType, IndexIterator> end() 
+		MappedIterator<int, ToType, IndexIterator> end()
 		{
 			return MappedIterator<int, ToType, IndexIterator>(enumerable.end(), MapFunc);
 		}
@@ -437,14 +401,11 @@ public:
 	* eg usage: for (FVector3d v : mapped_indices(fn_that_looks_up_mesh_vtx_from_id)) { ... }
 	*/
 	template<typename ToType>
-	inline MappedEnumerable<ToType> MappedIndices(TFunction<ToType(int)> MapFunc) const 
+	inline MappedEnumerable<ToType> MappedIndices(TFunction<ToType(int)> MapFunc) const
 	{
 		return MappedEnumerable<ToType>(Indices(), MapFunc);
 	}
 
-
-
-
 	/*
 	* iteration object that maps indices output by Index_iteration to a second type
 	*/
@@ -453,35 +414,38 @@ public:
 	public:
 		TFunction<bool(int)> FilterFunc;
 		IndexEnumerable enumerable;
-
-		FilteredEnumerable(const IndexEnumerable& enumerable, TFunction<bool(int)> FilterFuncIn) 
+		FilteredEnumerable(const IndexEnumerable& enumerable, TFunction<bool(int)> FilterFuncIn)
 		{
 			this->enumerable = enumerable;
 			this->FilterFunc = FilterFuncIn;
 		}
 
-		FilteredIterator<int, IndexIterator> begin() 
+		FilteredIterator<int, IndexIterator> begin()
 		{
 			return FilteredIterator<int, IndexIterator>(enumerable.begin(), enumerable.end(), FilterFunc);
 		}
 
-		FilteredIterator<int, IndexIterator> end() 
+		FilteredIterator<int, IndexIterator> end()
 		{
 			return FilteredIterator<int, IndexIterator>(enumerable.end(), enumerable.end(), FilterFunc);
 		}
 	};
 
-	inline FilteredEnumerable FilteredIndices(TFunction<bool(int)> FilterFunc) const 
+	inline FilteredEnumerable FilteredIndices(TFunction<bool(int)> FilterFunc) const
 	{
 		return FilteredEnumerable(Indices(), FilterFunc);
 	}
 
-
-	FString UsageStats() 
+	FString UsageStats()
 	{
 		return FString::Printf(TEXT("RefCountSize %d  FreeSize %d  FreeMem %dkb"),
 			RefCounts.GetLength(), FreeIndices.GetLength(), (FreeIndices.GetByteCount() / 1024));
 	}
 
+private:
+	TDynamicVector<short> RefCounts{};
+	TDynamicVector<int> FreeIndices{};
+	int UsedCount{0};
+
 
-};
\ No newline at end of file
+};
diff --git a/Engine/Plugins/Experimental/GeometryProcessing/Source/GeometricObjects/Public/VectorTypes.h b/Engine/Plugins/Experimental/GeometryProcessing/Source/GeometricObjects/Public/VectorTypes.h
--- a/Engine/Plugins/Experimental/GeometryProcessing/Source/GeometricObjects/Public/VectorTypes.h
+++ b/Engine/Plugins/Experimental/GeometryProcessing/Source/GeometricObjects/Public/VectorTypes.h
@@ -26,24 +26,32 @@ struct FVector2
 {
 	T X, Y;
 
-	FVector2()
+	constexpr FVector2()
 	{
 	}
-	FVector2(T ValX, T ValY)
+	constexpr FVector2(T ValX, T ValY)
 		: X(ValX), Y(ValY)
 	{
 	}
-	FVector2(const T* Data)
+
+	constexpr FVector2(const T* Data)
 	{
 		X = Data[0];
 		Y = Data[1];
 	}
 
-	inline operator const T*() const
+	constexpr FVector2(const FVector2D& Vec)
+	{
+		X = (T)Vec.X;
+		Y = (T)Vec.Y;
+	}
+
+	constexpr operator const T*() const
 	{
 		return &X;
 	};
-	inline operator T*()
+
+	constexpr  operator T*()
 	{
 		return &X;
 	}
@@ -52,46 +60,40 @@ struct FVector2
 	{
 		return FVector2D((float)X, (float)Y);
 	}
-	FVector2(const FVector2D& Vec)
-	{
-		X = (T)Vec.X;
-		Y = (T)Vec.Y;
-	}
-
-	static FVector2<T> Zero()
+	constexpr static FVector2<T> Zero()
 	{
 		return FVector2<T>((T)0, (T)0);
 	}
-	static FVector2<T> One()
+	constexpr static FVector2<T> One()
 	{
 		return FVector2<T>((T)1, (T)1);
 	}
-	static FVector2<T> UnitX()
+	constexpr static FVector2<T> UnitX()
 	{
 		return FVector2<T>((T)1, (T)0);
 	}
-	static FVector2<T> UnitY()
+	constexpr static FVector2<T> UnitY()
 	{
 		return FVector2<T>((T)0, (T)1);
 	}
 
-	FVector2<T>& operator=(const FVector2<T>& V2)
+	constexpr FVector2<T>& operator=(const FVector2<T>& V2)
 	{
 		X = V2.X;
 		Y = V2.Y;
 		return *this;
 	}
 
-	T& operator[](int Idx)
+	constexpr T& operator[](int Idx)
 	{
 		return (&X)[Idx];
 	}
-	const T& operator[](int Idx) const
+	constexpr const T& operator[](int Idx) const
 	{
 		return (&X)[Idx];
 	}
 
-	T Length() const
+	constexpr T Length() const
 	{
 		return TMathUtil<T>::Sqrt(X * X + Y * Y);
 	}
@@ -100,113 +102,113 @@ struct FVector2
 		return X * X + Y * Y;
 	}
 
-	inline T Distance(const FVector2<T>& V2) const
+	constexpr T Distance(const FVector2<T>& V2) const
 	{
 		T dx = V2.X - X;
 		T dy = V2.Y - Y;
 		return TMathUtil<T>::Sqrt(dx * dx + dy * dy);
 	}
-	inline T DistanceSquared(const FVector2<T>& V2) const
+	constexpr T DistanceSquared(const FVector2<T>& V2) const
 	{
 		T dx = V2.X - X;
 		T dy = V2.Y - Y;
 		return dx * dx + dy * dy;
 	}
 
-	T Dot(const FVector2<T>& V2) const
+	constexpr T Dot(const FVector2<T>& V2) const
 	{
 		return X * V2.X + Y * V2.Y;
 	}
 
 	// Note: DotPerp and FVector2's version of Cross are the same function
-	inline T DotPerp(const FVector2<T>& V2) const
+	constexpr T DotPerp(const FVector2<T>& V2) const
 	{
 		return X * V2.Y - Y * V2.X;
 	}
-	inline T Cross(const FVector2<T>& V2) const
+	constexpr T Cross(const FVector2<T>& V2) const
 	{
 		return X * V2.Y - Y * V2.X;
 	}
 
-	FVector2<T> Perp() const
+	constexpr FVector2<T> Perp() const
 	{
 		return FVector2<T>(Y, -X);
 	}
 
-	inline FVector2 operator-() const
+	constexpr  FVector2 operator-() const
 	{
 		return FVector2(-X, -Y);
 	}
 
-	inline FVector2 operator+(const FVector2& V2) const
+	constexpr FVector2 operator+(const FVector2& V2) const
 	{
 		return FVector2(X + V2.X, Y + V2.Y);
 	}
 
-	inline FVector2 operator-(const FVector2& V2) const
+	constexpr FVector2 operator-(const FVector2& V2) const
 	{
 		return FVector2(X - V2.X, Y - V2.Y);
 	}
 
-	inline FVector2<T> operator+(const T& Scalar) const
+	constexpr FVector2<T> operator+(const T& Scalar) const
 	{
 		return FVector2<T>(X + Scalar, Y + Scalar);
 	}
 
-	inline FVector2<T> operator-(const T& Scalar) const
+	constexpr FVector2<T> operator-(const T& Scalar) const
 	{
 		return FVector2<T>(X - Scalar, Y - Scalar);
 	}
 
-	inline FVector2 operator*(const T& Scalar) const
+	constexpr FVector2 operator*(const T& Scalar) const
 	{
 		return FVector2(X * Scalar, Y * Scalar);
 	}
 
-	inline FVector2<T> operator*(const FVector2<T>& V2) const // component-wise
+	constexpr FVector2<T> operator*(const FVector2<T>& V2) const // component-wise
 	{
 		return FVector2<T>(X * V2.X, Y * V2.Y);
 	}
 
-	inline FVector2 operator/(const T& Scalar) const
+	constexpr FVector2 operator/(const T& Scalar) const
 	{
 		return FVector2(X / Scalar, Y / Scalar);
 	}
 
-	inline FVector2<T> operator/(const FVector2<T>& V2) const // component-wise
+	constexpr FVector2<T> operator/(const FVector2<T>& V2) const // component-wise
 	{
 		return FVector2<T>(X / V2.X, Y / V2.Y);
 	}
 
-	inline FVector2<T>& operator+=(const FVector2<T>& V2)
+	constexpr FVector2<T>& operator+=(const FVector2<T>& V2)
 	{
 		X += V2.X;
 		Y += V2.Y;
 		return *this;
 	}
 
-	inline FVector2<T>& operator-=(const FVector2<T>& V2)
+	constexpr FVector2<T>& operator-=(const FVector2<T>& V2)
 	{
 		X -= V2.X;
 		Y -= V2.Y;
 		return *this;
 	}
 
-	inline FVector2<T>& operator*=(const T& Scalar)
+	constexpr FVector2<T>& operator*=(const T& Scalar)
 	{
 		X *= Scalar;
 		Y *= Scalar;
 		return *this;
 	}
 
-	inline FVector2<T>& operator/=(const T& Scalar)
+	constexpr FVector2<T>& operator/=(const T& Scalar)
 	{
 		X /= Scalar;
 		Y /= Scalar;
 		return *this;
 	}
 
-	inline bool IsNormalized()
+	constexpr bool IsNormalized()
 	{
 		return TMathUtil<T>::Abs((X * X + Y * Y) - 1) < TMathUtil<T>::ZeroTolerance;
 	}
@@ -258,7 +260,7 @@ struct FVector2
 		return 0;
 	}
 
-	inline FVector2<T> Normalized(const T Epsilon = 0) const
+	constexpr FVector2<T> Normalized(const T Epsilon = 0) const
 	{
 		T length = Length();
 		if (length > Epsilon)
@@ -310,40 +312,40 @@ struct FVector3
 {
 	T X, Y, Z;
 
-	FVector3()
+	constexpr FVector3()
 	{
 	}
-	FVector3(T ValX, T ValY, T ValZ)
+	constexpr FVector3(T ValX, T ValY, T ValZ)
 		: X(ValX), Y(ValY), Z(ValZ)
 	{
 	}
-	FVector3(const T* Data)
+	constexpr FVector3(const T* Data)
 	{
 		X = Data[0];
 		Y = Data[1];
 		Z = Data[2];
 	}
-	FVector3(const FVector2<T>& Data)
+	constexpr FVector3(const FVector2<T>& Data)
 	{
 		X = Data.X;
 		Y = Data.Y;
 		Z = (T)0;
 	}
 
-	inline operator const T*() const
+	constexpr operator const T*() const
 	{
 		return &X;
 	};
-	inline operator T*()
+	constexpr operator T*()
 	{
 		return &X;
 	}
 
-	inline operator FVector() const
+	constexpr operator FVector() const
 	{
 		return FVector((float)X, (float)Y, (float)Z);
 	}
-	inline FVector3(const FVector& Vec)
+	constexpr FVector3(const FVector& Vec)
 	{
 		X = (T)Vec.X;
 		Y = (T)Vec.Y;
@@ -351,18 +353,18 @@ struct FVector3
 	}
 
 
-	explicit inline operator FLinearColor() const
+	explicit constexpr operator FLinearColor() const
 	{
 		return FLinearColor((float)X, (float)Y, (float)Z);
 	}
-	inline FVector3(const FLinearColor& Color)
+	constexpr FVector3(const FLinearColor& Color)
 	{
 		X = (T)Color.R;
 		Y = (T)Color.G;
 		Z = (T)Color.B;
 	}
 
-	explicit inline operator FColor() const
+	explicit constexpr operator FColor() const
 	{
 		return FColor(
 			FMathf::Clamp((int)((float)X*255.0f), 0, 255),
@@ -422,14 +424,14 @@ struct FVector3
 		return X * X + Y * Y + Z * Z;
 	}
 
-	inline T Distance(const FVector3<T>& V2) const
+	constexpr T Distance(const FVector3<T>& V2) const
 	{
 		T dx = V2.X - X;
 		T dy = V2.Y - Y;
 		T dz = V2.Z - Z;
 		return TMathUtil<T>::Sqrt(dx * dx + dy * dy + dz * dz);
 	}
-	inline T DistanceSquared(const FVector3<T>& V2) const
+	constexpr T DistanceSquared(const FVector3<T>& V2) const
 	{
 		T dx = V2.X - X;
 		T dy = V2.Y - Y;
@@ -437,52 +439,52 @@ struct FVector3
 		return dx * dx + dy * dy + dz * dz;
 	}
 
-	inline FVector3<T> operator-() const
+	constexpr FVector3<T> operator-() const
 	{
 		return FVector3<T>(-X, -Y, -Z);
 	}
 
-	inline FVector3<T> operator+(const FVector3<T>& V2) const
+	constexpr FVector3<T> operator+(const FVector3<T>& V2) const
 	{
 		return FVector3<T>(X + V2.X, Y + V2.Y, Z + V2.Z);
 	}
 
-	inline FVector3<T> operator-(const FVector3<T>& V2) const
+	constexpr FVector3<T> operator-(const FVector3<T>& V2) const
 	{
 		return FVector3<T>(X - V2.X, Y - V2.Y, Z - V2.Z);
 	}
 
-	inline FVector3<T> operator+(const T& Scalar) const
+	constexpr FVector3<T> operator+(const T& Scalar) const
 	{
 		return FVector3<T>(X + Scalar, Y + Scalar, Z + Scalar);
 	}
 
-	inline FVector3<T> operator-(const T& Scalar) const
+	constexpr FVector3<T> operator-(const T& Scalar) const
 	{
 		return FVector3<T>(X - Scalar, Y - Scalar, Z - Scalar);
 	}
 
-	inline FVector3<T> operator*(const T& Scalar) const
+	constexpr FVector3<T> operator*(const T& Scalar) const
 	{
 		return FVector3<T>(X * Scalar, Y * Scalar, Z * Scalar);
 	}
 
-	inline FVector3<T> operator*(const FVector3<T>& V2) const // component-wise
+	constexpr FVector3<T> operator*(const FVector3<T>& V2) const // component-wise
 	{
 		return FVector3<T>(X * V2.X, Y * V2.Y, Z * V2.Z);
 	}
 
-	inline FVector3<T> operator/(const T& Scalar) const
+	constexpr FVector3<T> operator/(const T& Scalar) const
 	{
 		return FVector3<T>(X / Scalar, Y / Scalar, Z / Scalar);
 	}
 
-	inline FVector3<T> operator/(const FVector3<T>& V2) const // component-wise
+	constexpr FVector3<T> operator/(const FVector3<T>& V2) const // component-wise
 	{
 		return FVector3<T>(X / V2.X, Y / V2.Y, Z / V2.Z);
 	}
 
-	inline FVector3<T>& operator+=(const FVector3<T>& V2)
+	constexpr FVector3<T>& operator+=(const FVector3<T>& V2)
 	{
 		X += V2.X;
 		Y += V2.Y;
@@ -490,7 +492,7 @@ struct FVector3
 		return *this;
 	}
 
-	inline FVector3<T>& operator-=(const FVector3<T>& V2)
+	constexpr FVector3<T>& operator-=(const FVector3<T>& V2)
 	{
 		X -= V2.X;
 		Y -= V2.Y;
@@ -498,7 +500,7 @@ struct FVector3
 		return *this;
 	}
 
-	inline FVector3<T>& operator*=(const T& Scalar)
+	constexpr FVector3<T>& operator*=(const T& Scalar)
 	{
 		X *= Scalar;
 		Y *= Scalar;
@@ -506,7 +508,7 @@ struct FVector3
 		return *this;
 	}
 
-	inline FVector3<T>& operator/=(const T& Scalar)
+	constexpr FVector3<T>& operator/=(const T& Scalar)
 	{
 		X /= Scalar;
 		Y /= Scalar;
@@ -543,7 +545,7 @@ struct FVector3
 		return (T)acos(ClampedDot);
 	}
 
-	inline bool IsNormalized()
+	constexpr bool IsNormalized()
 	{
 		return TMathUtil<T>::Abs((X * X + Y * Y + Z * Z) - 1) < TMathUtil<T>::ZeroTolerance;
 	}
@@ -563,7 +565,7 @@ struct FVector3
 		return 0;
 	}
 
-	inline FVector3<T> Normalized(const T Epsilon = 0) const
+	constexpr FVector3<T> Normalized(const T Epsilon = 0) const
 	{
 		T length = Length();
 		if (length > Epsilon)
@@ -574,12 +576,12 @@ struct FVector3
 		return FVector3<T>((T)0, (T)0, (T)0);
 	}
 
-	inline T MaxAbs() const
+	constexpr T MaxAbs() const
 	{
 		return TMathUtil<T>::Max3(TMathUtil<T>::Abs(X), TMathUtil<T>::Abs(Y), TMathUtil<T>::Abs(Z));
 	}
 
-	inline T MinAbs() const
+	constexpr T MinAbs() const
 	{
 		return TMathUtil<T>::Min3(TMathUtil<T>::Abs(X), TMathUtil<T>::Abs(Y), TMathUtil<T>::Abs(Z));
 	}
@@ -592,7 +594,7 @@ struct FVector3
 						   OneMinusAlpha * A.Z + Alpha * B.Z);
 	}
 
-	inline bool operator==(const FVector3<T>& Other) const
+	constexpr bool operator==(const FVector3<T>& Other) const
 	{
 		return X == Other.X && Y == Other.Y && Z == Other.Z;
 	}
