Remove testing code accidentally left in

#rnx

diff --git a/Engine/Plugins/Experimental/GeometryProcessing/Source/DynamicMesh/Private/DynamicMesh3.cpp b/Engine/Plugins/Experimental/GeometryProcessing/Source/DynamicMesh/Private/DynamicMesh3.cpp
--- a/Engine/Plugins/Experimental/GeometryProcessing/Source/DynamicMesh/Private/DynamicMesh3.cpp
+++ b/Engine/Plugins/Experimental/GeometryProcessing/Source/DynamicMesh/Private/DynamicMesh3.cpp
@@ -654,10 +654,6 @@ bool FDynamicMesh3::CheckValidity(bool bAllowNonManifoldVertices, EValidityCheck
 	TFunction<void(bool)> CheckOrFailF = [&](bool b)
 	{
 		is_ok = is_ok && b;
-		if (!is_ok)
-		{
-			is_ok = false;
-		}
 	};
 	if (FailMode == EValidityCheckFailMode::Check)
 	{
