Add forgotten file InfoTypes.h

#rnx

diff --git a/Engine/Plugins/Experimental/GeometryProcessing/Source/DynamicMesh/Public/InfoTypes.h b/Engine/Plugins/Experimental/GeometryProcessing/Source/DynamicMesh/Public/InfoTypes.h
new file mode 100644
--- /dev/null
+++ b/Engine/Plugins/Experimental/GeometryProcessing/Source/DynamicMesh/Public/InfoTypes.h
@@ -0,0 +1,100 @@
+// Copyright 1998-2019 Epic Games, Inc. All Rights Reserved.
+
+#pragma once
+
+#include "IndexTypes.h"
+#include "VectorTypes.h"
+
+/**
+ * FVertexInfo stores information about vertex attributes - position, normal, color, UV
+ */
+struct FVertexInfo
+{
+	FVector3d Position{ FVector3d::Zero() };
+	FVector3f Normal{ FVector3f::Zero() };
+	FVector3f Color{ FVector3f::Zero() };
+	FVector2f UV{ FVector2f::Zero() };
+	bool bHaveN{}, bHaveC{}, bHaveUV{};
+
+	FVertexInfo() = default;
+	FVertexInfo(const FVector3d& PositionIn)
+		: Position{ PositionIn }{}
+	FVertexInfo(const FVector3d& PositionIn, const FVector3f& NormalIn)
+		: Position{ PositionIn }, Normal{ NormalIn }, bHaveN{ true }{}
+	FVertexInfo(const FVector3d& PositionIn, const FVector3f& NormalIn, const FVector3f& ColorIn)
+		: Position{ PositionIn }, Normal{ NormalIn }, Color{ ColorIn }, bHaveN{ true }, bHaveC{true}{}
+	FVertexInfo(const FVector3d& PositionIn, const FVector3f& NormalIn, const FVector3f& ColorIn, const FVector2f& UVIn)
+		: Position{ PositionIn }, Normal{ NormalIn }, Color{ ColorIn }, UV{ UVIn }, bHaveN{ true }, bHaveC{true}, bHaveUV{true}{}
+};
+
+namespace DynamicMeshInfo
+{
+
+/** Information about the mesh elements created by a call to SplitEdge() */
+struct FEdgeSplitInfo
+{
+	int OriginalEdge;					// the edge that was split
+	FIndex2i OriginalVertices;			// original edge vertices [a,b]
+	FIndex2i OtherVertices;				// original opposing vertices [c,d] - d is InvalidID for boundary edges
+	FIndex2i OriginalTriangles;			// original edge triangles [t0,t1]
+	bool bIsBoundary;					// was the split edge a boundary edge?  (redundant)
+
+	int NewVertex;						// new vertex f that was created
+	FIndex2i NewTriangles;				// new triangles [t2,t3], oriented as explained below
+	FIndex3i NewEdges;					// new edges are [f,b], [f,c] and [f,d] if this is not a boundary edge
+
+	double SplitT;						// parameter value for NewVertex along original edge
+};
+
+/** Information about the mesh elements modified by a call to FlipEdge() */
+struct FEdgeFlipInfo
+{
+	int EdgeID;						// the edge that was flipped
+	FIndex2i OriginalVerts;			// original verts of the flipped edge, that are no longer connected
+	FIndex2i OpposingVerts;			// the opposing verts of the flipped edge, that are now connected
+	FIndex2i Triangles;				// the two triangle IDs. Original tris vert [Vert0,Vert1,OtherVert0] and [Vert1,Vert0,OtherVert1].
+										// New triangles are [OtherVert0, OtherVert1, Vert1] and [OtherVert1, OtherVert0, Vert0]
+};
+
+/** Information about mesh elements modified/removed by CollapseEdge() */
+struct FEdgeCollapseInfo
+{
+	int KeptVertex;					// the vertex that was kept (ie collapsed "to")
+	int RemovedVertex;				// the vertex that was removed
+	FIndex2i OpposingVerts;			// the opposing vertices [c,d]. If the edge was a boundary edge, d is InvalidID
+	bool bIsBoundary;				// was the edge a boundary edge
+
+	int CollapsedEdge;				// the edge that was collapsed/removed
+	FIndex2i RemovedTris;			// the triangles that were removed in the collapse (second is InvalidID for boundary edge)
+	FIndex2i RemovedEdges;			// the edges that were removed (second is InvalidID for boundary edge)
+	FIndex2i KeptEdges;				// the edges that were kept (second is InvalidID for boundary edge)
+
+	double CollapseT;				// interpolation parameter along edge for new vertex in range [0,1]
+};
+
+/** Information about mesh elements modified by MergeEdges() */
+struct FMergeEdgesInfo
+{
+	int KeptEdge;				// the edge that was kept
+	int RemovedEdge;			// the edge that was removed
+
+	FIndex2i KeptVerts;			// The two vertices that were kept (redundant w/ KeptEdge?)
+	FIndex2i RemovedVerts;		// The removed vertices of RemovedEdge. Either may be InvalidID if it was same as the paired KeptVert
+
+	FIndex2i ExtraRemovedEdges; // extra removed edges, see description below. Either may be or InvalidID
+	FIndex2i ExtraKeptEdges;	// extra kept edges, paired with ExtraRemovedEdges
+};
+
+/** Information about mesh elements modified/created by PokeTriangle() */
+struct FPokeTriangleInfo
+{
+	int OriginalTriangle;				// the triangle that was poked
+	FIndex3i TriVertices;				// vertices of the original triangle
+
+	int NewVertex;						// the new vertex that was inserted
+	FIndex2i NewTriangles;				// the two new triangles that were added (OriginalTriangle is re-used, see code for vertex orders)
+	FIndex3i NewEdges;					// the three new edges connected to NewVertex
+
+	FVector3d BaryCoords;				// barycentric coords that NewVertex was inserted at
+};
+}
