Remove constexpr from a function at clang's request

#rnx

diff --git a/Engine/Plugins/Experimental/GeometryProcessing/Source/GeometricObjects/Public/VectorTypes.h b/Engine/Plugins/Experimental/GeometryProcessing/Source/GeometricObjects/Public/VectorTypes.h
--- a/Engine/Plugins/Experimental/GeometryProcessing/Source/GeometricObjects/Public/VectorTypes.h
+++ b/Engine/Plugins/Experimental/GeometryProcessing/Source/GeometricObjects/Public/VectorTypes.h
@@ -341,11 +341,11 @@ struct FVector3
 		return &X;
 	}
 
-	constexpr operator FVector() const
+	operator FVector() const
 	{
 		return FVector((float)X, (float)Y, (float)Z);
 	}
-	constexpr FVector3(const FVector& Vec)
+	FVector3(const FVector& Vec)
 	{
 		X = (T)Vec.X;
 		Y = (T)Vec.Y;
