Revert the externing of the CDT::Add functions to fix Win64

#rnx

diff --git a/Engine/Plugins/Experimental/GeometryProcessing/Source/GeometryAlgorithms/Public/ConstrainedDelaunay2.h b/Engine/Plugins/Experimental/GeometryProcessing/Source/GeometryAlgorithms/Public/ConstrainedDelaunay2.h
--- a/Engine/Plugins/Experimental/GeometryProcessing/Source/GeometryAlgorithms/Public/ConstrainedDelaunay2.h
+++ b/Engine/Plugins/Experimental/GeometryProcessing/Source/GeometryAlgorithms/Public/ConstrainedDelaunay2.h
@@ -83,15 +83,15 @@ typedef TConstrainedDelaunay2<float> FConstrainedDelaunay2f;
 typedef TConstrainedDelaunay2<double> FConstrainedDelaunay2d;
 
 
-extern template void GEOMETRYALGORITHMS_API TConstrainedDelaunay2<double>::Add(const FDynamicGraph2<double>& Graph);
-extern template void GEOMETRYALGORITHMS_API TConstrainedDelaunay2<double>::Add(const TGeneralPolygon2<double>& Polygon);
-extern template void GEOMETRYALGORITHMS_API TConstrainedDelaunay2<double>::Add(const TGeneralPolygon2<float>& Polygon);
-extern template void GEOMETRYALGORITHMS_API TConstrainedDelaunay2<double>::Add(const TPolygon2<double>& Polygon, bool bIsHole);
-extern template void GEOMETRYALGORITHMS_API TConstrainedDelaunay2<double>::Add(const TPolygon2<float>& Polygon, bool bIsHole);
-
-extern template void GEOMETRYALGORITHMS_API TConstrainedDelaunay2<float>::Add(const FDynamicGraph2<double>& Graph);
-extern template void GEOMETRYALGORITHMS_API TConstrainedDelaunay2<float>::Add(const TGeneralPolygon2<double>& Polygon);
-extern template void GEOMETRYALGORITHMS_API TConstrainedDelaunay2<float>::Add(const TGeneralPolygon2<float>& Polygon);
-extern template void GEOMETRYALGORITHMS_API TConstrainedDelaunay2<float>::Add(const TPolygon2<double>& Polygon, bool bIsHole);
-extern template void GEOMETRYALGORITHMS_API TConstrainedDelaunay2<float>::Add(const TPolygon2<float>& Polygon, bool bIsHole);
+template void GEOMETRYALGORITHMS_API TConstrainedDelaunay2<double>::Add(const FDynamicGraph2<double>& Graph);
+template void GEOMETRYALGORITHMS_API TConstrainedDelaunay2<double>::Add(const TGeneralPolygon2<double>& Polygon);
+template void GEOMETRYALGORITHMS_API TConstrainedDelaunay2<double>::Add(const TGeneralPolygon2<float>& Polygon);
+template void GEOMETRYALGORITHMS_API TConstrainedDelaunay2<double>::Add(const TPolygon2<double>& Polygon, bool bIsHole);
+template void GEOMETRYALGORITHMS_API TConstrainedDelaunay2<double>::Add(const TPolygon2<float>& Polygon, bool bIsHole);
+
+template void GEOMETRYALGORITHMS_API TConstrainedDelaunay2<float>::Add(const FDynamicGraph2<double>& Graph);
+template void GEOMETRYALGORITHMS_API TConstrainedDelaunay2<float>::Add(const TGeneralPolygon2<double>& Polygon);
+template void GEOMETRYALGORITHMS_API TConstrainedDelaunay2<float>::Add(const TGeneralPolygon2<float>& Polygon);
+template void GEOMETRYALGORITHMS_API TConstrainedDelaunay2<float>::Add(const TPolygon2<double>& Polygon, bool bIsHole);
+template void GEOMETRYALGORITHMS_API TConstrainedDelaunay2<float>::Add(const TPolygon2<float>& Polygon, bool bIsHole);
 
