GeometryProcessing: new FMeshConnectedComponents calculates connected regions of mesh in various ways
#rnx

diff --git a/Engine/Plugins/Experimental/GeometryProcessing/Source/DynamicMesh/Private/Selections/MeshConnectedComponents.cpp b/Engine/Plugins/Experimental/GeometryProcessing/Source/DynamicMesh/Private/Selections/MeshConnectedComponents.cpp
new file mode 100644
--- /dev/null
+++ b/Engine/Plugins/Experimental/GeometryProcessing/Source/DynamicMesh/Private/Selections/MeshConnectedComponents.cpp
@@ -0,0 +1,245 @@
+// Copyright 1998-2019 Epic Games, Inc. All Rights Reserved.
+
+#include "Selections/MeshConnectedComponents.h"
+#include "Algo/Sort.h"
+
+
+
+void FMeshConnectedComponents::FindConnectedTriangles(TFunction<bool(int32, int32)> TrisConnectedPredicate)
+{
+	// initial active set contains all valid triangles
+	// active values are as follows:   0: unprocessed,  1: in queue,  2: done,  3: invalid
+	TArray<uint8> ActiveSet;
+	int NumTriangles = Mesh->MaxTriangleID();
+	ActiveSet.Init(255, NumTriangles);
+	FInterval1i ActiveRange = FInterval1i::Empty();
+	for (int tid = 0; tid < NumTriangles; ++tid)
+	{
+		if (Mesh->IsTriangle(tid))
+		{
+			ActiveSet[tid] = 0;
+			ActiveRange.Contain(tid);
+		}
+	}
+
+	FindTriComponents(ActiveRange, ActiveSet, TrisConnectedPredicate);
+}
+
+
+
+void FMeshConnectedComponents::FindConnectedTriangles(TFunctionRef<bool(int)> IndexFilterFunc, TFunction<bool(int32, int32)> TrisConnectedPredicate)
+{
+	// initial active set contains all valid triangles
+	// active values are as follows:   0: unprocessed,  1: in queue,  2: done,  3: invalid
+	TArray<uint8> ActiveSet;
+	int NumTriangles = Mesh->MaxTriangleID();
+	ActiveSet.Init(255, NumTriangles);
+	FInterval1i ActiveRange = FInterval1i::Empty();
+	for (int tid = 0; tid < NumTriangles; ++tid)
+	{
+		if (Mesh->IsTriangle(tid) && IndexFilterFunc(tid))
+		{
+			ActiveSet[tid] = 0;
+			ActiveRange.Contain(tid);
+		}
+	}
+
+	FindTriComponents(ActiveRange, ActiveSet, TrisConnectedPredicate);
+}
+
+
+
+void FMeshConnectedComponents::FindConnectedTriangles(const TArray<int>& TriangleROI, TFunction<bool(int32, int32)> TrisConnectedPredicate)
+{
+	// initial active set contains all valid triangles
+	// active values are as follows:   0: unprocessed,  1: in queue,  2: done,  3: invalid
+	TArray<uint8> ActiveSet;
+	int NumTriangles = Mesh->MaxTriangleID();
+	ActiveSet.Init(255, NumTriangles);
+	FInterval1i ActiveRange = FInterval1i::Empty();
+	for (int tid : TriangleROI)
+	{
+		if (Mesh->IsTriangle(tid))
+		{
+			ActiveSet[tid] = 0;
+			ActiveRange.Contain(tid);
+		}
+	}
+
+	FindTriComponents(ActiveRange, ActiveSet, TrisConnectedPredicate);
+}
+
+
+
+void FMeshConnectedComponents::FindTriComponents(FInterval1i ActiveRange, TArray<uint8>& ActiveSet, TFunction<bool(int32, int32)> TrisConnectedPredicate)
+{
+	Components.Empty();
+
+	// temporary queue
+	TArray<int32> ComponentQueue;
+	ComponentQueue.Reserve(256);
+
+	// keep finding valid seed triangles and growing connected components
+	// until we are done
+	for (int i = ActiveRange.Min; i <= ActiveRange.Max; i++)
+	{
+		if (ActiveSet[i] != 255)
+		{
+			int SeedTri = i;
+			ComponentQueue.Add(SeedTri);
+			ActiveSet[SeedTri] = 1;      // in ComponentQueue
+
+			FComponent* Component = new FComponent();
+			if (TrisConnectedPredicate)
+			{
+				FindTriComponent(Component, ComponentQueue, ActiveSet, TrisConnectedPredicate);
+			}
+			else
+			{
+				FindTriComponent(Component, ComponentQueue, ActiveSet);
+			}
+			Components.Add(Component);
+
+			RemoveFromActiveSet(Component, ActiveSet);
+
+			ComponentQueue.Reset(0);
+		}
+	}
+}
+
+
+
+void FMeshConnectedComponents::FindTriComponent(FComponent* Component, TArray<int32>& ComponentQueue, TArray<uint8>& ActiveSet)
+{
+	while (ComponentQueue.Num() > 0)
+	{
+		int32 CurTriangle = ComponentQueue.Pop(false);
+
+		ActiveSet[CurTriangle] = 2;   // tri has been processed
+		Component->Indices.Add(CurTriangle);
+
+		FIndex3i TriNbrTris = Mesh->GetTriNeighbourTris(CurTriangle);
+		for (int j = 0; j < 3; ++j)
+		{
+			int NbrTri = TriNbrTris[j];
+			if (NbrTri != FDynamicMesh3::InvalidID && ActiveSet[NbrTri] == 0)
+			{
+				ComponentQueue.Add(NbrTri);
+				ActiveSet[NbrTri] = 1;           // in ComponentQueue
+			}
+		}
+	}
+}
+
+
+void FMeshConnectedComponents::FindTriComponent(FComponent* Component, TArray<int32>& ComponentQueue, TArray<uint8>& ActiveSet, 
+	TFunctionRef<bool(int32, int32)> TriConnectedPredicate)
+{
+	while (ComponentQueue.Num() > 0)
+	{
+		int32 CurTriangle = ComponentQueue.Pop(false);
+
+		ActiveSet[CurTriangle] = 2;   // tri has been processed
+		Component->Indices.Add(CurTriangle);
+
+		FIndex3i TriNbrTris = Mesh->GetTriNeighbourTris(CurTriangle);
+		for (int j = 0; j < 3; ++j)
+		{
+			int NbrTri = TriNbrTris[j];
+			if (NbrTri != FDynamicMesh3::InvalidID && ActiveSet[NbrTri] == 0 && TriConnectedPredicate(CurTriangle, NbrTri))
+			{
+				ComponentQueue.Add(NbrTri);
+				ActiveSet[NbrTri] = 1;           // in ComponentQueue
+			}
+		}
+	}
+}
+
+
+
+void FMeshConnectedComponents::RemoveFromActiveSet(const FComponent* Component, TArray<uint8>& ActiveSet)
+{
+	int32 ComponentSize = Component->Indices.Num();
+	for (int32 j = 0; j < ComponentSize; ++j)
+	{
+		ActiveSet[Component->Indices[j]] = 255;
+	}
+}
+
+
+
+
+int32 FMeshConnectedComponents::GetLargestIndexByCount() const
+{
+	if (Components.Num() == 0)
+	{
+		return -1;
+	}
+
+	int LargestIdx = 0;
+	int LargestCount = Components[LargestIdx].Indices.Num();
+	int NumComponents = Components.Num();
+	for (int i = 1; i < NumComponents; ++i) {
+		if (Components[i].Indices.Num() > LargestCount) {
+			LargestCount = Components[i].Indices.Num();
+			LargestIdx = i;
+		}
+	}
+	return LargestIdx;
+}
+
+
+
+void FMeshConnectedComponents::SortByCount(bool bLargestFirst)
+{
+	TArrayView<FComponent*> View( Components.GetData(), Components.Num() );
+	if (bLargestFirst)
+	{
+		View.StableSort(
+			[](const FComponent& A, const FComponent& B) { return (A).Indices.Num() > (B).Indices.Num(); }
+		);
+	}
+	else
+	{
+		View.StableSort(
+			[](const FComponent& A, const FComponent& B) { return (A).Indices.Num() < (B).Indices.Num(); }
+		);
+	}
+}
+
+
+
+void FMeshConnectedComponents::GrowToConnectedTriangles(const FDynamicMesh3* Mesh,
+	const TArray<int>& InputROI, TArray<int>& ResultROI,
+	TArray<int32>* QueueBuffer, 
+	TSet<int32>* DoneBuffer,
+	TFunctionRef<bool(int32, int32)> CanGrowPredicate
+)
+{
+	TArray<int32> LocalQueue;
+	QueueBuffer = (QueueBuffer == nullptr) ? &LocalQueue : QueueBuffer;
+	QueueBuffer->Reset(); 
+	QueueBuffer->Insert(InputROI, 0);
+
+	TSet<int32> LocalDone;
+	DoneBuffer = (DoneBuffer == nullptr) ? &LocalDone : DoneBuffer;
+	DoneBuffer->Reset(); 
+	DoneBuffer->Append(InputROI);
+
+	while (QueueBuffer->Num() > 0)
+	{
+		int32 CurTri = QueueBuffer->Pop(false);
+		ResultROI.Add(CurTri);
+
+		FIndex3i NbrTris = Mesh->GetTriNeighbourTris(CurTri);
+		for (int j = 0; j < 3; ++j)
+		{
+			int32 tid = NbrTris[j];
+			if (tid != FDynamicMesh3::InvalidID && DoneBuffer->Contains(tid) == false && CanGrowPredicate(CurTri, tid))
+			{
+				QueueBuffer->Add(tid);
+				DoneBuffer->Add(tid);
+			}
+		}
+	}
+}
\ No newline at end of file
diff --git a/Engine/Plugins/Experimental/GeometryProcessing/Source/DynamicMesh/Public/Selections/MeshConnectedComponents.h b/Engine/Plugins/Experimental/GeometryProcessing/Source/DynamicMesh/Public/Selections/MeshConnectedComponents.h
new file mode 100644
--- /dev/null
+++ b/Engine/Plugins/Experimental/GeometryProcessing/Source/DynamicMesh/Public/Selections/MeshConnectedComponents.h
@@ -0,0 +1,143 @@
+// Copyright 1998-2019 Epic Games, Inc. All Rights Reserved.
+
+// Port of geometry3sharp MeshConnectedComponents
+
+#pragma once
+
+#include "CoreMinimal.h"
+#include "DynamicMesh3.h"
+
+/**
+ * FMeshConnectedComponents calculates Connected Components of a Mesh, or sub-regions of a Mesh.
+ * By default the actual mesh connectivity is used, but an optional connectivity predicate
+ * can be provided to specify when two elements should be considered connected.
+ */
+class DYNAMICMESH_API FMeshConnectedComponents
+{
+public:
+	const FDynamicMesh3* Mesh;
+
+	/**
+	 * Connected Component found by one of the calculation functions
+	 */
+	struct FComponent
+	{
+		/** List of indices contained in component */
+		TArray<int> Indices;
+	};
+
+	/**
+	 * List of Connected Components that have been found by one of the calculation functions
+	 */
+	TIndirectArray<FComponent> Components;
+
+public:
+
+	FMeshConnectedComponents(const FDynamicMesh3* MeshIn)
+		: Mesh(MeshIn)
+	{
+	}
+
+	//
+	// Calculation functions. Call these to calculate different types of Components.
+	//
+
+	/**
+	 * Find all connected triangle components of the Mesh and store in Components array.
+	 * Triangle connectivity is based on edge connectivity, ie bowtie-vertices are not connections between triangles.
+	 * @param TrisConnectedPredicate optional function that specifies whether two edge-connected triangles should be considered connected by the search
+	 */
+	void FindConnectedTriangles(TFunction<bool(int32, int32)> TrisConnectedPredicate = nullptr);
+
+	/**
+	 * Find all connected triangle components of a subset of triangles of the Mesh and store in Components array.
+	 * Triangle connectivity is based on edge connectivity, ie bowtie-vertices are not connections between triangles.
+	 * @param TriangleROI list of triangles to search across
+	 * @param TrisConnectedPredicate optional function that specifies whether two edge-connected triangles should be considered connected by the search
+	 */
+	void FindConnectedTriangles(const TArray<int>& TriangleROI, TFunction<bool(int32, int32)> TrisConnectedPredicate = nullptr);
+
+	/**
+	 * Find all connected triangle components of a subset of triangles of the Mesh and store in Components array.
+	 * Triangle connectivity is based on edge connectivity, ie bowtie-vertices are not connections between triangles.
+	 * @param IndexFilterFunc defines set of triangles to search across, return true for triangle IDs that are to be considered
+	 * @param TrisConnectedPredicate optional function that specifies whether two edge-connected triangles should be considered connected by the search
+	 */
+	void FindConnectedTriangles(TFunctionRef<bool(int)> IndexFilterFunc, TFunction<bool(int32, int32)> TrisConnectedPredicate = nullptr);
+
+
+
+	//
+	// Query functions. Only valid to call after a Calculation function has been called.
+	//
+
+	/** @return Number of Components that were found */
+	int32 Num() const
+	{
+		return Components.Num();
+	}
+
+	/** @return element of Components array at given Index */
+	const FComponent& GetComponent(int32 Index) const { return Components[Index]; }
+
+	/** @return element of Components array at given Index */
+	FComponent& GetComponent(int32 Index) { return Components[Index]; }
+
+
+	/** 
+	 * @return index of largest component by element count 
+	 */
+	int32 GetLargestIndexByCount() const;
+
+	/**
+	 * Sort the Components array by component element count
+	 * @param bLargestFirst if true, sort by decreasing count, otherwise by increasing count
+	 */
+	void SortByCount(bool bLargestFirst = true);
+
+
+
+public:
+	/**
+	 * DO NOT USE DIRECTLY
+	 * STL-like iterators to enable ranged-based for loop support (forwarding TIndirectArray declarations)
+	 */
+	auto begin() { return Components.begin(); }
+	auto begin() const { return Components.begin(); }
+	auto end() { return Components.end();  }
+	auto end() const { return Components.end(); }
+
+
+
+protected:
+	//
+	// Internal functions to calculate ROI
+	//
+	void FindTriComponents(FInterval1i ActiveRange, TArray<uint8>& ActiveSet, TFunction<bool(int32, int32)> TriConnectedPredicate);
+	void FindTriComponent(FComponent* Component, TArray<int32>& ComponentQueue, TArray<uint8>& ActiveSet);
+	void FindTriComponent(FComponent* Component, TArray<int32>& ComponentQueue, TArray<uint8>& ActiveSet, TFunctionRef<bool(int32, int32)> TriConnectedPredicate);
+	void RemoveFromActiveSet(const FComponent* Component, TArray<uint8>& ActiveSet);
+
+
+
+public:
+	/**
+	 * Utility function to expand a triangle selection to all triangles considered "connected".
+	 * More efficient than using full FMeshConnectedComponents instance if ROI is small relative to Mesh size (or if temp buffers can be re-used)
+	 * @param Mesh Mesh to calculate on
+	 * @param InputROI input set of triangles
+	 * @param ResultROI output set of triangles connected to InputROI
+	 * @param QueueBuffer optional buffer used as internal Queue. If passed as nullptr, a TArray will be locally allocated
+	 * @param DoneBuffer optional set used to track which triangles have already been processed. If passed as nullptr, an TSet will be locally allocated
+	 * @param CanGrowPredicate determines whether two connected mesh triangles should be considered connected while growing
+	 */
+	static void GrowToConnectedTriangles(const FDynamicMesh3* Mesh, 
+		const TArray<int>& InputROI, 
+		TArray<int>& ResultROI,
+		TArray<int32>* QueueBuffer = nullptr, 
+		TSet<int32>* DoneBuffer = nullptr,
+		TFunctionRef<bool(int32, int32)> CanGrowPredicate = [](int32, int32) { return true; }
+	);
+
+
+};
