Geometry: add MeshTransforms::Translate, fix hardcoded double in template, add check for invalid input in MeshEditor::Append
#rnx

diff --git a/Engine/Plugins/Experimental/GeometryProcessing/Source/DynamicMesh/Private/DynamicMeshEditor.cpp b/Engine/Plugins/Experimental/GeometryProcessing/Source/DynamicMesh/Private/DynamicMeshEditor.cpp
--- a/Engine/Plugins/Experimental/GeometryProcessing/Source/DynamicMesh/Private/DynamicMeshEditor.cpp
+++ b/Engine/Plugins/Experimental/GeometryProcessing/Source/DynamicMesh/Private/DynamicMeshEditor.cpp
@@ -808,6 +808,9 @@ void FDynamicMeshEditor::AppendMesh(const FDynamicMesh3* AppendMesh,
 	TFunction<FVector3d(int, const FVector3d&)> PositionTransform,
 	TFunction<FVector3d(int, const FVector3d&)> NormalTransform)
 {
+	// todo: handle this case by making a copy?
+	check(AppendMesh != Mesh);
+
 	IndexMapsOut.Reset();
 	IndexMapsOut.Initialize(Mesh);
 
diff --git a/Engine/Plugins/Experimental/GeometryProcessing/Source/DynamicMesh/Private/MeshTransforms.cpp b/Engine/Plugins/Experimental/GeometryProcessing/Source/DynamicMesh/Private/MeshTransforms.cpp
--- a/Engine/Plugins/Experimental/GeometryProcessing/Source/DynamicMesh/Private/MeshTransforms.cpp
+++ b/Engine/Plugins/Experimental/GeometryProcessing/Source/DynamicMesh/Private/MeshTransforms.cpp
@@ -4,6 +4,22 @@
 #include "MeshTransforms.h"
 #include "DynamicMeshAttributeSet.h"
 #include "DynamicMeshOverlay.h"
+#include "Async/ParallelFor.h"
+
+
+
+void MeshTransforms::Translate(FDynamicMesh3& Mesh, const FVector3d& Translation)
+{
+	int NumVertices = Mesh.MaxVertexID();
+	ParallelFor(NumVertices, [&](int vid) 
+	{
+		if (Mesh.IsVertex(vid))
+		{
+			Mesh.SetVertex(vid, Mesh.GetVertex(vid) + Translation);
+		}
+	});
+}
+
 
 
 void MeshTransforms::ApplyTransform(FDynamicMesh3& Mesh, const FTransform3d& Transform)
diff --git a/Engine/Plugins/Experimental/GeometryProcessing/Source/DynamicMesh/Public/MeshTransforms.h b/Engine/Plugins/Experimental/GeometryProcessing/Source/DynamicMesh/Public/MeshTransforms.h
--- a/Engine/Plugins/Experimental/GeometryProcessing/Source/DynamicMesh/Public/MeshTransforms.h
+++ b/Engine/Plugins/Experimental/GeometryProcessing/Source/DynamicMesh/Public/MeshTransforms.h
@@ -13,6 +13,11 @@
 namespace MeshTransforms
 {
 	
+	/**
+	 * Apply Translation to vertex positions of Mesh. Does not modify any other attributes.
+	 */
+	DYNAMICMESH_API void Translate(FDynamicMesh3& Mesh, const FVector3d& Translation);
+
 	/**
 	 * Apply given Transform to a Mesh.
 	 * Modifies Vertex Positions and Normals, and any Per-Triangle Normal Overlays
diff --git a/Engine/Plugins/Experimental/GeometryProcessing/Source/GeometricObjects/Public/BoxTypes.h b/Engine/Plugins/Experimental/GeometryProcessing/Source/GeometricObjects/Public/BoxTypes.h
--- a/Engine/Plugins/Experimental/GeometryProcessing/Source/GeometricObjects/Public/BoxTypes.h
+++ b/Engine/Plugins/Experimental/GeometryProcessing/Source/GeometricObjects/Public/BoxTypes.h
@@ -244,9 +244,9 @@ struct TAxisAlignedBox3
 	FVector3<RealType> GetCorner(int Index) const
 	{
 		check(Index >= 0 && Index <= 7);
-		double X = (((Index & 1) != 0) ^ ((Index & 2) != 0)) ? (Max.X) : (Min.X);
-		double Y = ((Index / 2) % 2 == 0) ? (Min.Y) : (Max.Y);
-		double Z = (Index < 4) ? (Min.Z) : (Max.Z);
+		RealType X = (((Index & 1) != 0) ^ ((Index & 2) != 0)) ? (Max.X) : (Min.X);
+		RealType Y = ((Index / 2) % 2 == 0) ? (Min.Y) : (Max.Y);
+		RealType Z = (Index < 4) ? (Min.Z) : (Max.Z);
 		return FVector3<RealType>(X, Y, Z);
 	}
 
