Copy up from Dev-Editor [REDACTED_HANDLE]

diff --git a/Engine/Plugins/Enterprise/DatasmithImporter/Source/DatasmithImporter/Private/DatasmithImporter.cpp b/Engine/Plugins/Enterprise/DatasmithImporter/Source/DatasmithImporter/Private/DatasmithImporter.cpp
--- a/Engine/Plugins/Enterprise/DatasmithImporter/Source/DatasmithImporter/Private/DatasmithImporter.cpp
+++ b/Engine/Plugins/Enterprise/DatasmithImporter/Source/DatasmithImporter/Private/DatasmithImporter.cpp
@@ -2362,6 +2362,7 @@ void FDatasmithImporter::FinalizeImport(FDatasmithImportContext& ImportContext,
 		DatasmithImporterImpl::CheckAssetPersistenceValidity(ExistingMaterialFunctionPtr.Get(), ImportContext);
 	}
 
+	TArray<UMaterial*> MaterialsToRefreshAfterVirtualTextureConversion;
 	for (const TPair< TSharedRef< IDatasmithBaseMaterialElement >, UMaterialInterface* >& ImportedMaterialPair : ImportContext.ImportedMaterials)
 	{
 		if (ImportContext.bUserCancelled)
@@ -2389,6 +2390,17 @@ void FDatasmithImporter::FinalizeImport(FDatasmithImportContext& ImportContext,
 		if (UMaterial* SourceMaterial = Cast< UMaterial >(SourceMaterialInterface))
 		{
 			SourceMaterial->RebuildExpressionTextureReferences();
+			
+			TArray<UObject*> ReferencedTextures;
+			SourceMaterial->AppendReferencedTextures(ReferencedTextures);
+			for (UTexture2D* VirtualTexture : ImportContext.AssetsContext.VirtualTexturesToConvert)
+			{
+				if (ReferencedTextures.Contains(VirtualTexture))
+				{
+					MaterialsToRefreshAfterVirtualTextureConversion.Add(SourceMaterial);
+					break;
+				}
+			}
 
 			for (FMaterialFunctionInfo& MaterialFunctionInfo : SourceMaterial->MaterialFunctionInfos)
 			{
@@ -2551,5 +2563,31 @@ void FDatasmithImporter::FinalizeImport(FDatasmithImportContext& ImportContext,
 	GEditor->ResetTransaction(LOCTEXT("Reset Transaction Buffer", "Datasmith Import Finalization"));
 }
 
+void FDatasmithImporter::ConvertUnsupportedVirtualTexture(FDatasmithImportContext& ImportContext, TArray<UTexture2D*>& VirtualTexturesToConvert, TArray<UMaterial*>& MaterialsToRefreshAfterVirtualTextureConversion, const TMap<UObject*, UObject*>& ReferencesToRemap)
+{
+	if (VirtualTexturesToConvert.Num() != 0)
+	{
+		for (int32 TextureIndex = 0; TextureIndex < VirtualTexturesToConvert.Num(); ++TextureIndex)
+		{
+			UTexture2D*& TextureToConvert = VirtualTexturesToConvert[TextureIndex];
+			if (UObject* const* RemappedTexture = ReferencesToRemap.Find(TextureToConvert))
+			{
+				TextureToConvert = Cast<UTexture2D>(*RemappedTexture);
+			}
+
+			ImportContext.LogWarning(FText::Format(LOCTEXT("DatasmithVirtualTextureConverted", "The imported texture {0} could not be imported as virtual texture as it is not supported in all the materials using it."), FText::FromString(TextureToConvert->GetName())));
+		}
+		for (int32 MaterialIndex = 0; MaterialIndex < MaterialsToRefreshAfterVirtualTextureConversion.Num(); ++MaterialIndex)
+		{
+			if (UObject* const* RemappedMaterial = ReferencesToRemap.Find(MaterialsToRefreshAfterVirtualTextureConversion[MaterialIndex]))
+			{
+				MaterialsToRefreshAfterVirtualTextureConversion[MaterialIndex] = Cast<UMaterial>(*RemappedMaterial);
+			}
+		}
+
+		IAssetTools& AssetTools = FModuleManager::LoadModuleChecked<FAssetToolsModule>("AssetTools").Get();
+		AssetTools.ConvertVirtualTextures(VirtualTexturesToConvert, true, &MaterialsToRefreshAfterVirtualTextureConversion);
+	}
+}
 
 #undef LOCTEXT_NAMESPACE
diff --git a/Engine/Plugins/Experimental/GeometryProcessing/Source/DynamicMesh/Private/DynamicMeshEditor.cpp b/Engine/Plugins/Experimental/GeometryProcessing/Source/DynamicMesh/Private/DynamicMeshEditor.cpp
--- a/Engine/Plugins/Experimental/GeometryProcessing/Source/DynamicMesh/Private/DynamicMeshEditor.cpp
+++ b/Engine/Plugins/Experimental/GeometryProcessing/Source/DynamicMesh/Private/DynamicMeshEditor.cpp
@@ -5,7 +5,7 @@
 #include "DynamicMeshAttributeSet.h"
 #include "Util/BufferUtil.h"
 #include "MeshRegionBoundaryLoops.h"
-
+#include "DynamicSubmesh3.h"
 
 
 void FMeshIndexMappings::Initialize(FDynamicMesh3* Mesh)
@@ -306,8 +306,7 @@ bool FDynamicMeshEditor::DisconnectTriangles(const TArray<int>& Triangles, TArra
 	// find the region boundary loops
 	FMeshRegionBoundaryLoops RegionLoops(Mesh, Triangles, false);
 	bool bOK = RegionLoops.Compute();
-	check(bOK);
-	if (!bOK)
+	if (!ensure(bOK))
 	{
 		return false;
 	}
@@ -379,6 +378,220 @@ bool FDynamicMeshEditor::DisconnectTriangles(const TArray<int>& Triangles, TArra
 }
 
 
+void FDynamicMeshEditor::DisconnectTriangles(const TArray<int>& Triangles, bool bPreventBowties)
+{
+	TSet<int> TriSet, BoundaryVerts;
+	TArray<int> NewVerts, OldVertsThatSplit;
+	TArray<int> FilteredTriangles;
+	DynamicMeshInfo::FVertexSplitInfo SplitInfo;
+
+	TriSet.Append(Triangles);
+	for (int TID : Triangles)
+	{
+		FIndex3i Nbrs = Mesh->GetTriNeighbourTris(TID);
+		FIndex3i Tri = Mesh->GetTriangle(TID);
+		for (int SubIdx = 0; SubIdx < 3; SubIdx++)
+		{
+			int NeighborTID = Nbrs[SubIdx];
+			if (!TriSet.Contains(NeighborTID))
+			{
+				BoundaryVerts.Add(Tri[SubIdx]);
+				BoundaryVerts.Add(Tri[(SubIdx + 1) % 3]);
+			}
+		}
+	}
+	for (int VID : BoundaryVerts)
+	{
+		FilteredTriangles.Reset();
+		int TriRingCount = 0;
+		for (int RingTID : Mesh->VtxTrianglesItr(VID))
+		{
+			if (TriSet.Contains(RingTID))
+			{
+				FilteredTriangles.Add(RingTID);
+			}
+			TriRingCount++;
+		}
+
+		if (FilteredTriangles.Num() < TriRingCount)
+		{
+			checkSlow(!Mesh->SplitVertexWouldLeaveIsolated(VID, FilteredTriangles));
+			ensure(EMeshResult::Ok == Mesh->SplitVertex(VID, FilteredTriangles, SplitInfo));
+			NewVerts.Add(SplitInfo.NewVertex);
+			OldVertsThatSplit.Add(SplitInfo.OriginalVertex);
+		}
+	}
+	if (bPreventBowties)
+	{
+		FDynamicMeshEditResult Result;
+		for (int VID : OldVertsThatSplit)
+		{
+			SplitBowties(VID, Result);
+			Result.Reset(); // don't actually keep results; they are not used in this fn
+		}
+		for (int VID : NewVerts)
+		{
+			SplitBowties(VID, Result);
+			Result.Reset(); // don't actually keep results; they are not used in this fn
+		}
+	}
+}
+
+
+
+
+void FDynamicMeshEditor::SplitBowties(FDynamicMeshEditResult& ResultOut)
+{
+	ResultOut.Reset();
+	TSet<int> AddedVerticesWithIDLessThanMax; // added vertices that we can't filter just by checking against original max id; this will be empty for compact meshes
+	for (int VertexID = 0, OriginalMaxID = Mesh->MaxVertexID(); VertexID < OriginalMaxID; VertexID++)
+	{
+		if (!Mesh->IsVertex(VertexID) || AddedVerticesWithIDLessThanMax.Contains(VertexID))
+		{
+			continue;
+		}
+		int32 NumVertsBefore = ResultOut.NewVertices.Num();
+		// TODO: may be faster to inline this call to reuse the contiguous triangle arrays?
+		SplitBowties(VertexID, ResultOut);
+		for (int Idx = NumVertsBefore; Idx < ResultOut.NewVertices.Num(); Idx++)
+		{
+			if (ResultOut.NewVertices[Idx] < OriginalMaxID)
+			{
+				AddedVerticesWithIDLessThanMax.Add(ResultOut.NewVertices[Idx]);
+			}
+		}
+	}
+}
+
+
+
+void FDynamicMeshEditor::SplitBowties(int VertexID, FDynamicMeshEditResult& ResultOut)
+{
+	TArray<int> TrianglesOut, ContiguousGroupLengths;
+	TArray<bool> GroupIsLoop;
+	DynamicMeshInfo::FVertexSplitInfo SplitInfo;
+	check(Mesh->IsVertex(VertexID));
+	if (ensure(EMeshResult::Ok == Mesh->GetVtxContiguousTriangles(VertexID, TrianglesOut, ContiguousGroupLengths, GroupIsLoop)))
+	{
+		if (ContiguousGroupLengths.Num() > 1)
+		{
+			// is bowtie
+			for (int GroupIdx = 1, GroupStartIdx = ContiguousGroupLengths[0]; GroupIdx < ContiguousGroupLengths.Num(); GroupStartIdx += ContiguousGroupLengths[GroupIdx++])
+			{
+				ensure(EMeshResult::Ok == Mesh->SplitVertex(VertexID, TArrayView<const int>(TrianglesOut.GetData() + GroupStartIdx, ContiguousGroupLengths[GroupIdx]), SplitInfo));
+				ResultOut.NewVertices.Add(SplitInfo.NewVertex);
+			}
+		}
+	}
+}
+
+
+bool FDynamicMeshEditor::ReinsertSubmesh(const FDynamicSubmesh3& Region, FOptionallySparseIndexMap& SubToNewV, TArray<int>* new_tris, EDuplicateTriBehavior DuplicateBehavior)
+{
+	check(Region.GetBaseMesh() == Mesh);
+	const FDynamicMesh3& Sub = Region.GetSubmesh();
+	bool bAllOK = true;
+
+	FIndexFlagSet done_v(Sub.MaxVertexID(), Sub.TriangleCount()/2);
+	SubToNewV.Initialize(Sub.MaxVertexID(), Sub.VertexCount());
+
+	int NT = Sub.MaxTriangleID();
+	for (int ti = 0; ti < NT; ++ti )
+	{
+		if (Sub.IsTriangle(ti) == false)
+		{
+			continue;
+		}
+
+		FIndex3i sub_t = Sub.GetTriangle(ti);
+		int gid = Sub.GetTriangleGroup(ti);
+
+		FIndex3i new_t = FIndex3i::Zero();
+		for ( int j = 0; j < 3; ++j )
+		{
+			int sub_v = sub_t[j];
+			int new_v = -1;
+			if (done_v[sub_v] == false)
+			{
+				// first check if this is a boundary vtx on submesh and maps to a bdry vtx on base mesh
+				if (Sub.IsBoundaryVertex(sub_v))
+				{
+					int base_v = Region.MapVertexToBaseMesh(sub_v);
+					if (base_v >= 0 && Mesh->IsVertex(base_v) && Region.InBaseBorderVertices(base_v) == true)
+					{
+						// [RMS] this should always be true, but assert in tests to find out
+						if (ensure(Mesh->IsBoundaryVertex(base_v)))
+						{
+							new_v = base_v;
+						}
+					}
+				}
+
+				// if that didn't happen, append new vtx
+				if (new_v == -1)
+				{
+					new_v = Mesh->AppendVertex(Sub, sub_v);
+				}
+
+				SubToNewV.Set(sub_v, new_v);
+				done_v.Add(sub_v);
+
+			}
+			else
+			{
+				new_v = SubToNewV[sub_v];
+			}
+
+			new_t[j] = new_v;
+		}
+
+		// try to handle duplicate-tri case
+		if (DuplicateBehavior == EDuplicateTriBehavior::EnsureContinue)
+		{
+			ensure(Mesh->FindTriangle(new_t.A, new_t.B, new_t.C) == FDynamicMesh3::InvalidID);
+		}
+		else
+		{
+			int existing_tid = Mesh->FindTriangle(new_t.A, new_t.B, new_t.C);
+			if (existing_tid != FDynamicMesh3::InvalidID)
+			{
+				if (DuplicateBehavior == EDuplicateTriBehavior::EnsureAbort)
+				{
+					ensure(false);
+					return false;
+				}
+				else if (DuplicateBehavior == EDuplicateTriBehavior::UseExisting)
+				{
+					if (new_tris)
+					{
+						new_tris->Add(existing_tid);
+					}
+					continue;
+				}
+				else if (DuplicateBehavior == EDuplicateTriBehavior::Replace)
+				{
+					Mesh->RemoveTriangle(existing_tid, false);
+				}
+			}
+		}
+
+
+		int new_tid = Mesh->AppendTriangle(new_t, gid);
+		ensure(new_tid != FDynamicMesh3::InvalidID && new_tid != FDynamicMesh3::NonManifoldID);
+		if (!Mesh->IsTriangle(new_tid))
+		{
+			bAllOK = false;
+		}
+
+		if (new_tris)
+		{
+			new_tris->Add(new_tid);
+		}
+	}
+
+	return bAllOK;
+}
+
 
 
 FVector3f FDynamicMeshEditor::ComputeAndSetQuadNormal(const FIndex2i& QuadTris, bool bIsPlanar)
@@ -725,6 +938,13 @@ void FDynamicMeshEditor::CopyAttributes(int FromTriangleID, int ToTriangleID, FM
 		NormalOverlay->SetTriangle(ToTriangleID, ToElemTri);
 		NormalLayerIndex++;
 	}
+
+	if (Mesh->Attributes()->HasMaterialID())
+	{
+		FDynamicMeshMaterialAttribute* MaterialIDs = Mesh->Attributes()->GetMaterialID();
+		MaterialIDs->SetValue(ToTriangleID, MaterialIDs->GetValue(FromTriangleID));
+	}
+
 }
 
 
@@ -808,6 +1028,9 @@ void FDynamicMeshEditor::AppendMesh(const FDynamicMesh3* AppendMesh,
 	TFunction<FVector3d(int, const FVector3d&)> PositionTransform,
 	TFunction<FVector3d(int, const FVector3d&)> NormalTransform)
 {
+	// todo: handle this case by making a copy?
+	check(AppendMesh != Mesh);
+
 	IndexMapsOut.Reset();
 	IndexMapsOut.Initialize(Mesh);
 
@@ -895,12 +1118,22 @@ void FDynamicMeshEditor::AppendMesh(const FDynamicMesh3* AppendMesh,
 			FDynamicMeshUVOverlay* ToUVs = Mesh->Attributes()->GetUVLayer(UVLayerIndex);
 			if (FromUVs != nullptr && ToUVs != nullptr)
 			{
-				FIndexMapi& UVMap = IndexMapsOut.GetUVMap(0);
+				FIndexMapi& UVMap = IndexMapsOut.GetUVMap(UVLayerIndex);
 				UVMap.Reserve(FromUVs->ElementCount());
 				AppendUVs(AppendMesh, FromUVs, ToUVs,
 					VertexMap, TriangleMap, UVMap);
 			}
 		}
+
+		if (AppendMesh->Attributes()->HasMaterialID() && Mesh->Attributes()->HasMaterialID())
+		{
+			const FDynamicMeshMaterialAttribute* FromMaterialIDs = AppendMesh->Attributes()->GetMaterialID();
+			FDynamicMeshMaterialAttribute* ToMaterialIDs = Mesh->Attributes()->GetMaterialID();
+			for (int TriID : AppendMesh->TriangleIndicesItr())
+			{
+				ToMaterialIDs->SetValue(TriangleMap.GetTo(TriID), FromMaterialIDs->GetValue(TriID));
+			}
+		}
 	}
 }
 
@@ -1048,7 +1281,7 @@ static void AppendAttributes(const FDynamicMesh3* FromMesh, int FromTriangleID,
 		{
 			if (FromElemTri[j] != FDynamicMesh3::InvalidID)
 			{
-				int NewElemID = AppendTriangleUVAttribute(FromMesh, FromElemTri[j], ToMesh, 0, IndexMaps);
+				int NewElemID = AppendTriangleUVAttribute(FromMesh, FromElemTri[j], ToMesh, UVLayerIndex, IndexMaps);
 				ToElemTri[j] = NewElemID;
 			}
 		}
@@ -1072,6 +1305,13 @@ static void AppendAttributes(const FDynamicMesh3* FromMesh, int FromTriangleID,
 		}
 		ToNormalOverlay->SetTriangle(ToTriangleID, ToElemTri);
 	}
+
+	if (FromMesh->Attributes()->HasMaterialID() && ToMesh->Attributes()->HasMaterialID())
+	{
+		const FDynamicMeshMaterialAttribute* FromMaterialIDs = FromMesh->Attributes()->GetMaterialID();
+		FDynamicMeshMaterialAttribute* ToMaterialIDs = ToMesh->Attributes()->GetMaterialID();
+		ToMaterialIDs->SetValue(ToTriangleID, FromMaterialIDs->GetValue(FromTriangleID));
+	}
 }
 
 
@@ -1079,7 +1319,7 @@ static void AppendAttributes(const FDynamicMesh3* FromMesh, int FromTriangleID,
 
 
 
-void FDynamicMeshEditor::AppendTriangles(const FDynamicMesh3* SourceMesh, const TArray<int>& SourceTriangles, FMeshIndexMappings& IndexMaps, FDynamicMeshEditResult& ResultOut)
+void FDynamicMeshEditor::AppendTriangles(const FDynamicMesh3* SourceMesh, const TArrayView<const int>& SourceTriangles, FMeshIndexMappings& IndexMaps, FDynamicMeshEditResult& ResultOut, bool bComputeTriangleMap)
 {
 	ResultOut.Reset();
 	IndexMaps.Initialize(Mesh);
@@ -1120,7 +1360,10 @@ void FDynamicMeshEditor::AppendTriangles(const FDynamicMesh3* SourceMesh, const
 		}
 
 		int NewTriangleID = Mesh->AppendTriangle(NewTri, NewGroupID);
-		IndexMaps.SetTriangle(SourceTriangleID, NewTriangleID);
+		if (bComputeTriangleMap)
+		{
+			IndexMaps.SetTriangle(SourceTriangleID, NewTriangleID);
+		}
 		ResultOut.NewTriangles.Add(NewTriangleID);
 
 		AppendAttributes(SourceMesh, SourceTriangleID, Mesh, NewTriangleID, IndexMaps, ResultOut);
@@ -1128,3 +1371,91 @@ void FDynamicMeshEditor::AppendTriangles(const FDynamicMesh3* SourceMesh, const
 		//Mesh->CheckValidity(true);
 	}
 }
+
+
+bool FDynamicMeshEditor::SplitMesh(const FDynamicMesh3* SourceMesh, TArray<FDynamicMesh3>& SplitMeshes, TFunctionRef<int(int)> TriIDToMeshID, int DeleteMeshID)
+{
+	TMap<int, int> MeshIDToIndex;
+	int NumMeshes = 0;
+	bool bAlsoDelete = false;
+	for (int TID : SourceMesh->TriangleIndicesItr())
+	{
+		int MeshID = TriIDToMeshID(TID);
+		if (MeshID == DeleteMeshID)
+		{
+			bAlsoDelete = true;
+			continue;
+		}
+		if (!MeshIDToIndex.Contains(MeshID))
+		{
+			MeshIDToIndex.Add(MeshID, NumMeshes++);
+		}
+	}
+
+	if (!bAlsoDelete && NumMeshes < 2)
+	{
+		return false; // nothing to do, so don't bother filling the split meshes array
+	}
+
+	SplitMeshes.Reset();
+	SplitMeshes.SetNum(NumMeshes);
+	// enable matching attributes
+	for (FDynamicMesh3& M : SplitMeshes)
+	{
+		if (SourceMesh->HasAttributes())
+		{
+			M.EnableAttributes();
+			M.Attributes()->EnableMatchingAttributes(*SourceMesh->Attributes());
+		}
+	}
+
+	if (NumMeshes == 0) // full delete case, just leave the empty mesh
+	{
+		return true;
+	}
+
+	TArray<FMeshIndexMappings> Mappings; Mappings.Reserve(NumMeshes);
+	FDynamicMeshEditResult UnusedInvalidResultAccumulator; // only here because some functions require it
+	for (int Idx = 0; Idx < NumMeshes; Idx++)
+	{
+		FMeshIndexMappings& Map = Mappings.Emplace_GetRef();
+		Map.Initialize(&SplitMeshes[Idx]);
+	}
+
+	for (int SourceTID : SourceMesh->TriangleIndicesItr())
+	{
+		int MeshID = TriIDToMeshID(SourceTID);
+		if (MeshID == DeleteMeshID)
+		{
+			continue; // just skip triangles w/ the Delete Mesh ID
+		}
+		int MeshIndex = MeshIDToIndex[MeshID];
+		FDynamicMesh3& Mesh = SplitMeshes[MeshIndex];
+		FMeshIndexMappings& IndexMaps = Mappings[MeshIndex];
+
+		FIndex3i Tri = SourceMesh->GetTriangle(SourceTID);
+
+		// FindOrCreateDuplicateGroup
+		int SourceGID = SourceMesh->GetTriangleGroup(SourceTID);
+		int NewGID = IndexMaps.GetNewGroup(SourceGID);
+
+		FIndex3i NewTri;
+		for (int j = 0; j < 3; ++j)
+		{
+			int SourceVID = Tri[j];
+			int NewVID = IndexMaps.GetNewVertex(SourceVID);
+			if (NewVID == IndexMaps.InvalidID())
+			{
+				NewVID = Mesh.AppendVertex(*SourceMesh, SourceVID);
+				IndexMaps.SetVertex(SourceVID, NewVID);
+			}
+			NewTri[j] = NewVID;
+		}
+
+		int NewTID = Mesh.AppendTriangle(NewTri, NewGID);
+		IndexMaps.SetTriangle(SourceTID, NewTID);
+		AppendAttributes(SourceMesh, SourceTID, &Mesh, NewTID, IndexMaps, UnusedInvalidResultAccumulator);
+	}
+	
+	return true;
+}
\ No newline at end of file
diff --git a/Engine/Plugins/Experimental/GeometryProcessing/Source/DynamicMesh/Private/GroupTopology.cpp b/Engine/Plugins/Experimental/GeometryProcessing/Source/DynamicMesh/Private/GroupTopology.cpp
--- a/Engine/Plugins/Experimental/GeometryProcessing/Source/DynamicMesh/Private/GroupTopology.cpp
+++ b/Engine/Plugins/Experimental/GeometryProcessing/Source/DynamicMesh/Private/GroupTopology.cpp
@@ -55,7 +55,7 @@ void FGroupTopology::RebuildTopology()
 	}
 	for (FGroup& Group : Groups)
 	{
-		Group.Faces.Reserve(GroupFaceCounts[Group.GroupID]);
+		Group.Triangles.Reserve(GroupFaceCounts[Group.GroupID]);
 	}
 
 
@@ -63,7 +63,7 @@ void FGroupTopology::RebuildTopology()
 	for (int tid : Mesh->TriangleIndicesItr())
 	{
 		int GroupID = GetGroupID(tid);
-		Groups[GroupIDToGroupIndexMap[GroupID]].Faces.Add(tid);
+		Groups[GroupIDToGroupIndexMap[GroupID]].Triangles.Add(tid);
 	}
 
 	// precompute junction vertices set
@@ -171,7 +171,7 @@ const TArray<int>& FGroupTopology::GetGroupFaces(int GroupID) const
 {
 	const FGroup* Found = FindGroupByID(GroupID);
 	ensure(Found != nullptr);
-	return (Found != nullptr) ? Found->Faces : EmptyArray;
+	return (Found != nullptr) ? Found->Triangles : EmptyArray;
 }
 
 const TArray<int>& FGroupTopology::GetGroupNbrGroups(int GroupID) const
@@ -281,7 +281,7 @@ void FGroupTopology::CollectGroupVertices(int GroupID, TSet<int>& Vertices) cons
 	ensure(Found != nullptr);
 	if (Found != nullptr)
 	{
-		for (int TriID : Found->Faces)
+		for (int TriID : Found->Triangles)
 		{
 			FIndex3i TriVerts = Mesh->GetTriangle(TriID);
 			Vertices.Add(TriVerts.A);
@@ -367,7 +367,7 @@ void FGroupTopology::ForGroupSetEdges(const TArray<int>& GroupIDs,
 
 void FGroupTopology::ExtractGroupEdges(FGroup& Group)
 {
-	FMeshRegionBoundaryLoops BdryLoops(Mesh, Group.Faces, true);
+	FMeshRegionBoundaryLoops BdryLoops(Mesh, Group.Triangles, true);
 	int NumLoops = BdryLoops.Loops.Num();
 
 	Group.Boundaries.SetNum(NumLoops);
@@ -494,4 +494,49 @@ bool FGroupTopology::GetGroupEdgeTangent(int GroupEdgeID, FVector3d& TangentOut)
 		TangentOut = FVector3d::UnitX();
 		return false;
 	}
+}
+
+
+
+FFrame3d FGroupTopology::GetSelectionFrame(const FGroupTopologySelection& Selection)
+{
+	int32 NumCorners = Selection.SelectedCornerIDs.Num();
+	int32 NumEdges = Selection.SelectedEdgeIDs.Num();
+	int32 NumFaces = Selection.SelectedGroupIDs.Num();
+
+	FFrame3d Accumulated;
+	int AccumCount = 0;
+	for (int32 CornerID : Selection.SelectedCornerIDs)
+	{
+		Accumulated.Origin += Mesh->GetVertex(GetCornerVertexID(CornerID));
+		AccumCount++;
+	}
+
+	for (int32 EdgeID : Selection.SelectedEdgeIDs)
+	{
+		const FGroupEdge& Edge = Edges[EdgeID];
+		FVector3d StartPos = Mesh->GetVertex(Edge.Span.Vertices[0]);
+		FVector3d EndPos = Mesh->GetVertex(Edge.Span.Vertices[Edge.Span.Vertices.Num() - 1]);
+		Accumulated.Origin +=  0.5*(StartPos + EndPos);
+		AccumCount++;
+	}
+
+	for (int32 GroupID : Selection.SelectedGroupIDs)
+	{
+		FVector3d Centroid = FVector3d::Zero();
+		const FGroup& Face = Groups[GroupIDToGroupIndexMap[GroupID]];
+		for (int tid : Face.Triangles)
+		{
+			Centroid += Mesh->GetTriCentroid(tid);
+		}
+		Accumulated.Origin += Centroid / (double)Face.Triangles.Num();
+		AccumCount++;
+	}
+
+	if (AccumCount > 0)
+	{
+		Accumulated.Origin /= (double)AccumCount;
+	}
+
+	return Accumulated;
 }
\ No newline at end of file
diff --git a/Engine/Plugins/Experimental/GeometryProcessing/Source/DynamicMesh/Private/MeshAdapterUtil.cpp b/Engine/Plugins/Experimental/GeometryProcessing/Source/DynamicMesh/Private/MeshAdapterUtil.cpp
--- a/Engine/Plugins/Experimental/GeometryProcessing/Source/DynamicMesh/Private/MeshAdapterUtil.cpp
+++ b/Engine/Plugins/Experimental/GeometryProcessing/Source/DynamicMesh/Private/MeshAdapterUtil.cpp
@@ -2,6 +2,23 @@
 
 #include "MeshAdapterUtil.h"
 
+FPointSetAdapterd MeshAdapterUtil::MakePointsAdapter(const FDynamicPointSet3d* PointSet)
+{
+	FPointSetAdapterd Adapter;
+	Adapter.MaxPointID = [PointSet]() { return PointSet->MaxVertexID(); };
+	Adapter.PointCount = [PointSet]() { return PointSet->VertexCount(); };
+	Adapter.IsPoint = [PointSet](int Idx) { return PointSet->IsVertex(Idx); };
+	Adapter.GetPoint = [PointSet](int Idx) { return PointSet->GetVertex(Idx); };
+	Adapter.Timestamp = [PointSet] { return PointSet->GetTimestamp(); };
+
+	Adapter.HasNormals = [PointSet] { return false; };
+	Adapter.GetPointNormal = [PointSet](int Idx) {return FVector3f(0,0,1); };
+
+	return Adapter;
+}
+
+
+
 FPointSetAdapterd MeshAdapterUtil::MakeVerticesAdapter(const FDynamicMesh3* Mesh)
 {
 	FPointSetAdapterd Adapter;
@@ -78,9 +95,10 @@ FPointSetAdapterd MeshAdapterUtil::MakeBoundaryEdgeMidpointsAdapter(const FDynam
 FTriangleMeshAdapterd MeshAdapterUtil::MakeTransformedDynamicMeshAdapter(const FDynamicMesh3* Mesh, FTransform Transform)
 {
 	FTriangleMeshAdapterd Adapter = MakeDynamicMeshAdapter(Mesh);
-	Adapter.GetVertex = [Mesh, Transform](int32 Idx)
+	FTransform3d Transformd(Transform);
+	Adapter.GetVertex = [Mesh, Transformd](int32 Idx)
 	{
-		return FVector3d(Transform.TransformPosition(Mesh->GetVertex(Idx)));
+		return FVector3d(Transformd.TransformPosition(Mesh->GetVertex(Idx)));
 	};
 	return Adapter;
 }
diff --git a/Engine/Plugins/Experimental/GeometryProcessing/Source/DynamicMesh/Private/MeshNormals.cpp b/Engine/Plugins/Experimental/GeometryProcessing/Source/DynamicMesh/Private/MeshNormals.cpp
--- a/Engine/Plugins/Experimental/GeometryProcessing/Source/DynamicMesh/Private/MeshNormals.cpp
+++ b/Engine/Plugins/Experimental/GeometryProcessing/Source/DynamicMesh/Private/MeshNormals.cpp
@@ -3,13 +3,14 @@
 // Port of geometry3Sharp MeshNormals
 
 #include "MeshNormals.h"
+#include "Async/ParallelFor.h"
 
 
 void FMeshNormals::SetCount(int Count, bool bClearToZero)
 {
 	if (Normals.Num() < Count) 
 	{
-		Normals.SetNum(Count);
+		Normals.SetNumUninitialized(Count);
 	}
 	if (bClearToZero)
 	{
@@ -105,12 +106,15 @@ void FMeshNormals::Compute_FaceAvg(bool bWeightByArea, bool bWeightByAngle)
 
 void FMeshNormals::Compute_Triangle()
 {
-	SetCount(Mesh->MaxTriangleID(), false);
-
-	for (int TriIdx : Mesh->TriangleIndicesItr())
+	int NumTriangles = Mesh->MaxTriangleID();
+	SetCount(NumTriangles, false);
+	ParallelFor(NumTriangles, [&](int32 Index)
 	{
-		Normals[TriIdx] = Mesh->GetTriNormal(TriIdx);
-	}
+		if (Mesh->IsTriangle(Index))
+		{
+			Normals[Index] = Mesh->GetTriNormal(Index);
+		}
+	});
 }
 
 
@@ -118,7 +122,7 @@ void FMeshNormals::Compute_Triangle()
 
 void FMeshNormals::Compute_Overlay_FaceAvg(const FDynamicMeshNormalOverlay* NormalOverlay, bool bWeightByArea, bool bWeightByAngle)
 {
-	if (!bWeightByAngle && bWeightByArea)
+	if ((!bWeightByAngle) && bWeightByArea)
 	{
 		Compute_Overlay_FaceAvg_AreaWeighted(NormalOverlay); // faster case
 		return;
@@ -214,4 +218,36 @@ FVector3d FMeshNormals::ComputeOverlayNormal(const FDynamicMesh3& Mesh, FDynamic
 	}
 
 	return (Count > 0) ? SumNormal.Normalized() : FVector3d::Zero();
+}
+
+
+void FMeshNormals::InitializeOverlayToPerVertexNormals(FDynamicMeshNormalOverlay* NormalOverlay, bool bUseMeshVertexNormalsIfAvailable)
+{
+	const FDynamicMesh3* Mesh = NormalOverlay->GetParentMesh();
+	bool bUseMeshNormals = bUseMeshVertexNormalsIfAvailable && Mesh->HasVertexNormals();
+	FMeshNormals Normals(Mesh);
+	if (bUseMeshNormals == false)
+	{
+		Normals.ComputeVertexNormals();
+	}
+
+	NormalOverlay->ClearElements();
+
+	TArray<int> VertToNormalMap;
+	VertToNormalMap.SetNumUninitialized(Mesh->MaxVertexID());
+	for (int vid : Mesh->VertexIndicesItr())
+	{
+		FVector3f Normal = (bUseMeshNormals) ? Mesh->GetVertexNormal(vid) : (FVector3f)Normals[vid];
+		int nid = NormalOverlay->AppendElement(Normal, vid);
+		VertToNormalMap[vid] = nid;
+	}
+
+	for (int tid : Mesh->TriangleIndicesItr())
+	{
+		FIndex3i Tri = Mesh->GetTriangle(tid);
+		Tri.A = VertToNormalMap[Tri.A];
+		Tri.B = VertToNormalMap[Tri.B];
+		Tri.C = VertToNormalMap[Tri.C];
+		NormalOverlay->SetTriangle(tid, Tri);
+	}
 }
\ No newline at end of file
diff --git a/Engine/Plugins/Experimental/GeometryProcessing/Source/DynamicMesh/Private/MeshTransforms.cpp b/Engine/Plugins/Experimental/GeometryProcessing/Source/DynamicMesh/Private/MeshTransforms.cpp
--- a/Engine/Plugins/Experimental/GeometryProcessing/Source/DynamicMesh/Private/MeshTransforms.cpp
+++ b/Engine/Plugins/Experimental/GeometryProcessing/Source/DynamicMesh/Private/MeshTransforms.cpp
@@ -4,6 +4,99 @@
 #include "MeshTransforms.h"
 #include "DynamicMeshAttributeSet.h"
 #include "DynamicMeshOverlay.h"
+#include "Async/ParallelFor.h"
+
+
+
+void MeshTransforms::Translate(FDynamicMesh3& Mesh, const FVector3d& Translation)
+{
+	int NumVertices = Mesh.MaxVertexID();
+	ParallelFor(NumVertices, [&](int vid) 
+	{
+		if (Mesh.IsVertex(vid))
+		{
+			Mesh.SetVertex(vid, Mesh.GetVertex(vid) + Translation);
+		}
+	});
+}
+
+
+
+void MeshTransforms::WorldToFrameCoords(FDynamicMesh3& Mesh, const FFrame3d& Frame)
+{
+	bool bVertexNormals = Mesh.HasVertexNormals();
+
+	int NumVertices = Mesh.MaxVertexID();
+	ParallelFor(NumVertices, [&](int vid)
+	{
+		if (Mesh.IsVertex(vid))
+		{
+			FVector3d Position = Mesh.GetVertex(vid);
+			Mesh.SetVertex(vid, Frame.ToFramePoint(Position));
+
+			if (bVertexNormals)
+			{
+				FVector3f Normal = Mesh.GetVertexNormal(vid);
+				Mesh.SetVertexNormal(vid, (FVector3f)Frame.ToFrameVector((FVector3d)Normal));
+			}
+		}
+	});
+
+	if (Mesh.HasAttributes())
+	{
+		FDynamicMeshNormalOverlay* Normals = Mesh.Attributes()->PrimaryNormals();
+		int NumNormals = Normals->MaxElementID();
+		ParallelFor(NumNormals, [&](int elemid)
+		{
+			if (Normals->IsElement(elemid))
+			{
+				FVector3f Normal = Normals->GetElement(elemid);
+				Normals->SetElement(elemid, (FVector3f)Frame.ToFrameVector((FVector3d)Normal));
+			}
+		});
+	}
+}
+
+
+
+
+
+void MeshTransforms::FrameCoordsToWorld(FDynamicMesh3& Mesh, const FFrame3d& Frame)
+{
+	bool bVertexNormals = Mesh.HasVertexNormals();
+
+	int NumVertices = Mesh.MaxVertexID();
+	ParallelFor(NumVertices, [&](int vid)
+	{
+		if (Mesh.IsVertex(vid))
+		{
+			FVector3d Position = Mesh.GetVertex(vid);
+			Mesh.SetVertex(vid, Frame.FromFramePoint(Position));
+
+			if (bVertexNormals)
+			{
+				FVector3f Normal = Mesh.GetVertexNormal(vid);
+				Mesh.SetVertexNormal(vid, (FVector3f)Frame.FromFrameVector((FVector3d)Normal));
+			}
+		}
+	});
+
+	if (Mesh.HasAttributes())
+	{
+		FDynamicMeshNormalOverlay* Normals = Mesh.Attributes()->PrimaryNormals();
+		int NumNormals = Normals->MaxElementID();
+		ParallelFor(NumNormals, [&](int elemid)
+		{
+			if (Normals->IsElement(elemid))
+			{
+				FVector3f Normal = Normals->GetElement(elemid);
+				Normals->SetElement(elemid, (FVector3f)Frame.FromFrameVector((FVector3d)Normal));
+			}
+		});
+	}
+}
+
+
 
 
 void MeshTransforms::ApplyTransform(FDynamicMesh3& Mesh, const FTransform3d& Transform)
@@ -41,27 +134,84 @@ void MeshTransforms::ApplyTransformInverse(FDynamicMesh3& Mesh, const FTransform
 {
 	bool bVertexNormals = Mesh.HasVertexNormals();
 
-	for (int vid : Mesh.VertexIndicesItr())
+	int NumVertices = Mesh.MaxVertexID();
+	ParallelFor(NumVertices, [&](int vid)
 	{
-		FVector3d Position = Mesh.GetVertex(vid);
-		Position = Transform.InverseTransformPosition(Position);
-		Mesh.SetVertex(vid, Position);
+		if (Mesh.IsVertex(vid)) 
+		{
+			FVector3d Position = Mesh.GetVertex(vid);
+			Position = Transform.InverseTransformPosition(Position);
+			Mesh.SetVertex(vid, Position);
 
-		if (bVertexNormals)
+			if (bVertexNormals)
+			{
+				FVector3f Normal = Mesh.GetVertexNormal(vid);
+				Normal = (FVector3f)Transform.InverseTransformNormal((FVector3d)Normal);
+				Mesh.SetVertexNormal(vid, Normal.Normalized());
+			}
+		}
+	});
+
+	if (Mesh.HasAttributes())
+	{
+		FDynamicMeshNormalOverlay* Normals = Mesh.Attributes()->PrimaryNormals();
+		if (Normals)
 		{
-			FVector3f Normal = Mesh.GetVertexNormal(vid);
-			Normal = (FVector3f)Transform.InverseTransformNormal((FVector3d)Normal);
-			Mesh.SetVertexNormal(vid, Normal.Normalized());
+			int NumNormals = Normals->MaxElementID();
+			ParallelFor(NumNormals, [&](int elemid)
+			{
+				if (Normals->IsElement(elemid))
+				{
+					FVector3f Normal = Normals->GetElement(elemid);
+					Normal = (FVector3f)Transform.InverseTransformNormal((FVector3d)Normal);
+					Normals->SetElement(elemid, Normal.Normalized());
+				}
+			});
 		}
 	}
+}
+
+
+
+void MeshTransforms::ApplyTransform(FDynamicMesh3& Mesh,
+	TFunctionRef<FVector3d(const FVector3d&)> PositionTransform,
+	TFunctionRef<FVector3f(const FVector3f&)> NormalTransform)
+{
+	bool bVertexNormals = Mesh.HasVertexNormals();
+
+	int NumVertices = Mesh.MaxVertexID();
+	ParallelFor(NumVertices, [&](int vid)
+	{
+		if (Mesh.IsVertex(vid))
+		{
+			FVector3d Position = Mesh.GetVertex(vid);
+			Position = PositionTransform(Position);
+			Mesh.SetVertex(vid, Position);
+
+			if (bVertexNormals)
+			{
+				FVector3f Normal = Mesh.GetVertexNormal(vid);
+				Normal = NormalTransform(Normal);
+				Mesh.SetVertexNormal(vid, Normal.Normalized());
+			}
+		}
+	});
+
 	if (Mesh.HasAttributes())
 	{
 		FDynamicMeshNormalOverlay* Normals = Mesh.Attributes()->PrimaryNormals();
-		for (int elemid : Normals->ElementIndicesItr())
+		if (Normals)
 		{
-			FVector3f Normal = Normals->GetElement(elemid);
-			Normal = (FVector3f)Transform.InverseTransformNormal((FVector3d)Normal);
-			Normals->SetElement(elemid, Normal.Normalized());
+			int NumNormals = Normals->MaxElementID();
+			ParallelFor(NumNormals, [&](int elemid)
+			{
+				if (Normals->IsElement(elemid))
+				{
+					FVector3f Normal = Normals->GetElement(elemid);
+					Normal = NormalTransform(Normal);
+					Normals->SetElement(elemid, Normal.Normalized());
+				}
+			});
 		}
 	}
-}
\ No newline at end of file
+}
diff --git a/Engine/Plugins/Experimental/GeometryProcessing/Source/DynamicMesh/Private/Operations/ExtrudeMesh.cpp b/Engine/Plugins/Experimental/GeometryProcessing/Source/DynamicMesh/Private/Operations/ExtrudeMesh.cpp
--- a/Engine/Plugins/Experimental/GeometryProcessing/Source/DynamicMesh/Private/Operations/ExtrudeMesh.cpp
+++ b/Engine/Plugins/Experimental/GeometryProcessing/Source/DynamicMesh/Private/Operations/ExtrudeMesh.cpp
@@ -9,7 +9,7 @@ FExtrudeMesh::FExtrudeMesh(FDynamicMesh3* mesh) : Mesh(mesh)
 {
 	ExtrudedPositionFunc = [this](const FVector3d& Position, const FVector3f& Normal, int VertexID) 
 	{
-		return Position + this->DefaultExtrudeDistance * Normal;
+		return Position + this->DefaultExtrudeDistance * (FVector3d)Normal;
 	};
 }
 
@@ -47,6 +47,10 @@ bool FExtrudeMesh::Apply()
 	// set vertices to new positions
 	for (int vid : InitialVertices)
 	{
+		if (!InitialToOffsetMapV.Contains(vid))
+		{
+			continue;
+		}
 		int newvid = InitialToOffsetMapV[vid];
 		if ( ! Mesh->IsVertex(newvid) )
 		{
diff --git a/Engine/Plugins/Experimental/GeometryProcessing/Source/DynamicMesh/Private/Operations/GroupTopologyDeformer.cpp b/Engine/Plugins/Experimental/GeometryProcessing/Source/DynamicMesh/Private/Operations/GroupTopologyDeformer.cpp
--- a/Engine/Plugins/Experimental/GeometryProcessing/Source/DynamicMesh/Private/Operations/GroupTopologyDeformer.cpp
+++ b/Engine/Plugins/Experimental/GeometryProcessing/Source/DynamicMesh/Private/Operations/GroupTopologyDeformer.cpp
@@ -35,18 +35,29 @@ void FGroupTopologyDeformer::SetActiveHandleFaces(const TArray<int>& FaceGroupID
 {
 	Reset();
 
-	check(FaceGroupIDs.Num() == 1);   // multi-face not supported yet
-	int GroupID = FaceGroupIDs[0];
-
-	// find set of vertices in handle 
-	Topology->CollectGroupVertices(GroupID, HandleVertices);
-	Topology->CollectGroupBoundaryVertices(GroupID, HandleBoundaryVertices);
+	// is this right? some HandleBoundaryVertices may also be interior if we do it per-face...
+	for (int GroupID : FaceGroupIDs)
+	{
+		Topology->CollectGroupVertices(GroupID, HandleVertices);
+		Topology->CollectGroupBoundaryVertices(GroupID, HandleBoundaryVertices);
+	}
 	ModifiedVertices = HandleVertices;
 
 	// find neighbour group set
 	TArray<int> HandleGroups = FaceGroupIDs;
-	const TArray<int>& GroupNbrGroups = Topology->GetGroupNbrGroups(GroupID);
-	CalculateROI(HandleGroups, GroupNbrGroups);
+	TArray<int> AllGroupNbrGroups;
+	for (int GroupID : FaceGroupIDs)
+	{
+		const TArray<int>& CurGroupNbrGroups = Topology->GetGroupNbrGroups(GroupID);
+		for (int NbrGroupID : CurGroupNbrGroups)
+		{
+			if (HandleGroups.Contains(NbrGroupID) == false)
+			{
+				AllGroupNbrGroups.AddUnique(NbrGroupID);
+			}
+		}
+	}
+	CalculateROI(HandleGroups, AllGroupNbrGroups);
 
 	SaveInitialPositions();
 
diff --git a/Engine/Plugins/Experimental/GeometryProcessing/Source/DynamicMesh/Public/DynamicMeshAttributeSet.h b/Engine/Plugins/Experimental/GeometryProcessing/Source/DynamicMesh/Public/DynamicMeshAttributeSet.h
--- a/Engine/Plugins/Experimental/GeometryProcessing/Source/DynamicMesh/Public/DynamicMeshAttributeSet.h
+++ b/Engine/Plugins/Experimental/GeometryProcessing/Source/DynamicMesh/Public/DynamicMeshAttributeSet.h
@@ -5,6 +5,9 @@
 
 #include "DynamicMeshOverlay.h"
 #include "DynamicMeshTriangleAttribute.h"
+#include "DynamicAttribute.h"
+#include "GeometryTypes.h"
+#include "InfoTypes.h"
 
 /** Standard UV overlay type - 2-element float */
 typedef TDynamicMeshVectorOverlay<float, 2, FVector2f> FDynamicMeshUVOverlay;
@@ -12,7 +15,7 @@ typedef TDynamicMeshVectorOverlay<float, 2, FVector2f> FDynamicMeshUVOverlay;
 typedef TDynamicMeshVectorOverlay<float, 3, FVector3f> FDynamicMeshNormalOverlay;
 
 /** Standard per-triangle integer material ID */
-typedef TDynamicMeshTriangleAttribute<int32, 1> FDynamicMeshMaterialAttribute;
+typedef TDynamicMeshScalarTriangleAttribute<int32> FDynamicMeshMaterialAttribute;
 
 /**
  * FDynamicMeshAttributeSet manages a set of extended attributes for a FDynamicMesh3.
@@ -22,11 +25,11 @@ typedef TDynamicMeshTriangleAttribute<int32, 1> FDynamicMeshMaterialAttribute;
  * 
  * @todo current internal structure is a work-in-progress
  */
-class DYNAMICMESH_API FDynamicMeshAttributeSet
+class DYNAMICMESH_API FDynamicMeshAttributeSet : public FDynamicMeshAttributeSetBase
 {
 public:
 
-	FDynamicMeshAttributeSet(FDynamicMesh3* Mesh) 
+	FDynamicMeshAttributeSet(FDynamicMesh3* Mesh)
 		: ParentMesh(Mesh), Normals0(Mesh)
 	{
 		SetNumUVLayers(1);
@@ -49,20 +52,166 @@ public:
 		if (Copy.MaterialIDAttrib)
 		{
 			EnableMaterialID();
-			MaterialIDAttrib->Copy( *(Copy.MaterialIDAttrib) );
+			MaterialIDAttrib->Copy(*(Copy.MaterialIDAttrib));
+		}
+		else
+		{
+			DisableMaterialID();
+		}
+
+		GenericAttributes.Reset();
+		ResetRegisteredAttributes();
+		for (int Idx = 0; Idx < Copy.GenericAttributes.Num(); Idx++)
+		{
+			const FDynamicMeshAttributeBase* SourceAttrib = Copy.GenericAttributes[Idx].Get();
+			AttachAttribute(SourceAttrib->MakeCopy(ParentMesh));
+		}
+
+		// parent mesh is *not* copied!
+	}
+
+	/** returns true if the attached overlays/attributes are compact */
+	bool IsCompact()
+	{
+		for (int UVIdx = 0; UVIdx < NumUVLayers(); UVIdx++)
+		{
+			if (!UVLayers[UVIdx].IsCompact())
+			{
+				return false;
+			}
+		}
+		if (!Normals0.IsCompact())
+		{
+			return false;
+		}
+		// material ID and generic per-element attributes currently cannot be non-compact
+		return true;
+	}
+
+	/**
+	 * Performs a CompactCopy of the attached overlays/attributes.
+	 * Called by the parent mesh CompactCopy function.
+	 *
+	 * @param CompactMaps Maps indicating how vertices and triangles were changes in the parent
+	 * @param Copy The attribute set to be copied
+	 */
+	void CompactCopy(const FCompactMaps& CompactMaps, const FDynamicMeshAttributeSet& Copy)
+	{
+		SetNumUVLayers(Copy.NumUVLayers());
+		for (int UVIdx = 0; UVIdx < NumUVLayers(); UVIdx++)
+		{
+			UVLayers[UVIdx].CompactCopy(CompactMaps, Copy.UVLayers[UVIdx]);
+		}
+		Normals0.CompactCopy(CompactMaps, Copy.Normals0);
+
+		if (Copy.MaterialIDAttrib)
+		{
+			EnableMaterialID();
+			MaterialIDAttrib->CompactCopy(CompactMaps, *(Copy.MaterialIDAttrib));
+		}
+		else
+		{
+			DisableMaterialID();
+		}
+
+		GenericAttributes.Reset();
+		ResetRegisteredAttributes();
+		for (int Idx = 0; Idx < Copy.GenericAttributes.Num(); Idx++)
+		{
+			const FDynamicMeshAttributeBase* SourceAttrib = Copy.GenericAttributes[Idx].Get();
+			AttachAttribute(SourceAttrib->MakeCompactCopy(CompactMaps, ParentMesh));
 		}
 
 		// parent mesh is *not* copied!
 	}
 
+	/**
+	 * Compacts the attribute set in place
+	 * Called by the parent mesh CompactInPlace function
+	 *
+	 * @param CompactMaps Maps of how the vertices and triangles were compacted in the parent
+	 */
+	void CompactInPlace(const FCompactMaps& CompactMaps)
+	{
+		for (int UVIdx = 0; UVIdx < NumUVLayers(); UVIdx++)
+		{
+			UVLayers[UVIdx].CompactInPlace(CompactMaps);
+		}
+		Normals0.CompactInPlace(CompactMaps);
+
+		if (MaterialIDAttrib.IsValid())
+		{
+			MaterialIDAttrib->CompactInPlace(CompactMaps);
+		}
+
+		for (int Idx = 0; Idx < GenericAttributes.Num(); Idx++)
+		{
+			GenericAttributes[Idx]->CompactInPlace(CompactMaps);
+		}
+	}
+
+
+	/**
+	 * Enable the matching attributes and overlay layers as the reference Copy set, but do not copy any data across
+	 */
+	void EnableMatchingAttributes(const FDynamicMeshAttributeSet& ToMatch)
+	{
+		SetNumUVLayers(ToMatch.NumUVLayers());
+		for (int UVIdx = 0; UVIdx < NumUVLayers(); UVIdx++)
+		{
+			UVLayers[UVIdx].ClearElements();
+		}
+		Normals0.ClearElements();
+
+		if (ToMatch.MaterialIDAttrib)
+		{
+			EnableMaterialID();
+		}
+		else
+		{
+			DisableMaterialID();
+		}
+
+		GenericAttributes.Reset();
+		ResetRegisteredAttributes();
+		for (int Idx = 0; Idx < ToMatch.GenericAttributes.Num(); Idx++)
+		{
+			const FDynamicMeshAttributeBase* SourceAttrib = ToMatch.GenericAttributes[Idx].Get();
+			AttachAttribute(SourceAttrib->MakeNew(ParentMesh));
+		}
+	}
+
 	/** @return the parent mesh for this overlay */
 	const FDynamicMesh3* GetParentMesh() const { return ParentMesh; }
 	/** @return the parent mesh for this overlay */
 	FDynamicMesh3* GetParentMesh() { return ParentMesh; }
 
+private:
+	/** @set the parent mesh for this overlay.  Only safe for use during FDynamicMesh move */
+	void Reparent(FDynamicMesh3* NewParent)
+	{
+		ParentMesh = NewParent;
 
+		for (int UVIdx = 0; UVIdx < NumUVLayers(); UVIdx++)
+		{
+			UVLayers[UVIdx].Reparent( NewParent );
+		}
+		Normals0.Reparent( NewParent );
+
+		if (MaterialIDAttrib)
+		{
+			MaterialIDAttrib->Reparent( NewParent );
+		}
+
+		for (int Idx = 0; Idx < GenericAttributes.Num(); Idx++)
+		{
+			GenericAttributes[Idx]->Reparent( NewParent );
+		}
+	}
+
+public:
 	/** @return number of UV layers */
-	virtual int NumUVLayers() const 
+	virtual int NumUVLayers() const
 	{
 		return UVLayers.Num();
 	}
@@ -100,7 +249,7 @@ public:
 	/** @return true if the given vertex is a seam vertex in any overlay */
 	virtual bool IsSeamVertex(int VertexID, bool bBoundaryIsSeam = true) const;
 
-	
+
 	//
 	// UV Layers 
 	//
@@ -118,7 +267,7 @@ public:
 	}
 
 	/** @return the primary UV layer (layer 0) */
-	FDynamicMeshUVOverlay* PrimaryUV() 
+	FDynamicMeshUVOverlay* PrimaryUV()
 	{
 		return &UVLayers[0];
 	}
@@ -176,6 +325,8 @@ public:
 
 	void EnableMaterialID();
 
+	void DisableMaterialID();
+
 	FDynamicMeshMaterialAttribute* GetMaterialID()
 	{
 		return MaterialIDAttrib.Get();
@@ -186,6 +337,24 @@ public:
 		return MaterialIDAttrib.Get();
 	}
 
+	// Attach a new attribute (and transfer ownership of it to the attribute set)
+	int AttachAttribute(FDynamicMeshAttributeBase* Attribute)
+	{
+		int AttributeID = GenericAttributes.Add(TUniquePtr<FDynamicMeshAttributeBase>(Attribute));
+		RegisterExternalAttribute(Attribute);
+		return AttributeID;
+	}
+
+	FDynamicMeshAttributeBase* GetAttachedAttribute(int AttributeID)
+	{
+		return GenericAttributes[AttributeID].Get();
+	}
+
+	int NumAttachedAttributes()
+	{
+		return GenericAttributes.Num();
+	}
+
 protected:
 	/** Parent mesh of this attribute set */
 	FDynamicMesh3* ParentMesh;
@@ -197,7 +366,8 @@ protected:
 	TArray<FDynamicMeshNormalOverlay*> NormalLayers;
 
 	TUniquePtr<FDynamicMeshMaterialAttribute> MaterialIDAttrib;
-	
+
+	TArray<TUniquePtr<FDynamicMeshAttributeBase>> GenericAttributes;
 
 protected:
 	friend class FDynamicMesh3;
@@ -216,13 +386,34 @@ protected:
 
 	// These functions are called by the FDynamicMesh3 to update the various
 	// attributes when the parent mesh topology has been modified.
+	// TODO: would it be better to register all the overlays and attributes with the base set and not overload these?  maybe!
 	virtual void OnNewTriangle(int TriangleID, bool bInserted);
+	virtual void OnNewVertex(int VertexID, bool bInserted);
 	virtual void OnRemoveTriangle(int TriangleID);
+	virtual void OnRemoveVertex(int VertexID);
 	virtual void OnReverseTriOrientation(int TriangleID);
-	virtual void OnSplitEdge(const FDynamicMesh3::FEdgeSplitInfo & splitInfo);
-	virtual void OnFlipEdge(const FDynamicMesh3::FEdgeFlipInfo & flipInfo);
-	virtual void OnCollapseEdge(const FDynamicMesh3::FEdgeCollapseInfo & collapseInfo);
-	virtual void OnPokeTriangle(const FDynamicMesh3::FPokeTriangleInfo & pokeInfo);
-	virtual void OnMergeEdges(const FDynamicMesh3::FMergeEdgesInfo & mergeInfo);
+	virtual void OnSplitEdge(const DynamicMeshInfo::FEdgeSplitInfo & splitInfo);
+	virtual void OnFlipEdge(const DynamicMeshInfo::FEdgeFlipInfo & flipInfo);
+	virtual void OnCollapseEdge(const DynamicMeshInfo::FEdgeCollapseInfo & collapseInfo);
+	virtual void OnPokeTriangle(const DynamicMeshInfo::FPokeTriangleInfo & pokeInfo);
+	virtual void OnMergeEdges(const DynamicMeshInfo::FMergeEdgesInfo & mergeInfo);
+	virtual void OnSplitVertex(const DynamicMeshInfo::FVertexSplitInfo& SplitInfo, const TArrayView<const int>& TrianglesToUpdate);
+
+	/**
+	 * Check validity of attributes
+	 * 
+	 * @param bAllowNonmanifold Accept non-manifold topology as valid. Note that this should almost always be true for attributes; non-manifold overlays are generally valid.
+	 * @param FailMode Desired behavior if mesh is found invalid
+	 */
+	virtual bool CheckValidity(bool bAllowNonmanifold, EValidityCheckFailMode FailMode) const
+	{
+		bool bValid = FDynamicMeshAttributeSetBase::CheckValidity(bAllowNonmanifold, FailMode);
+		for (int UVLayerIndex = 0; UVLayerIndex < NumUVLayers(); UVLayerIndex++)
+		{
+			bValid = GetUVLayer(UVLayerIndex)->CheckValidity(bAllowNonmanifold, FailMode) && bValid;
+		}
+		bValid = PrimaryNormals()->CheckValidity(bAllowNonmanifold, FailMode) && bValid;
+		return bValid;
+	}
 };
 
diff --git a/Engine/Plugins/Experimental/GeometryProcessing/Source/DynamicMesh/Public/DynamicMeshEditor.h b/Engine/Plugins/Experimental/GeometryProcessing/Source/DynamicMesh/Public/DynamicMeshEditor.h
--- a/Engine/Plugins/Experimental/GeometryProcessing/Source/DynamicMesh/Public/DynamicMeshEditor.h
+++ b/Engine/Plugins/Experimental/GeometryProcessing/Source/DynamicMesh/Public/DynamicMeshEditor.h
@@ -7,6 +7,9 @@
 #include "DynamicMesh3.h"
 #include "DynamicMeshAttributeSet.h"
 #include "EdgeLoop.h"
+#include "Util/SparseIndexCollectionTypes.h"
+
+struct FDynamicSubmesh3;
 
 /**
  * FMeshIndexMappings stores a set of integer IndexMaps for a mesh
@@ -28,7 +31,7 @@ public:
 	void Initialize(FDynamicMesh3* Mesh);
 
 	/** @return the value used to indicate "invalid" in the mapping */
-	int InvalidID() { return VertexMap.GetInvalidID(); }
+	constexpr int InvalidID() const { return VertexMap.UnmappedID(); }
 
 	void Reset()
 	{
@@ -45,12 +48,19 @@ public:
 		}
 	}
 
+	void ResetTriangleMap()
+	{
+		TriangleMap.Reset();
+	}
+
 	FIndexMapi& GetVertexMap() { return VertexMap; }
+	const FIndexMapi& GetVertexMap() const { return VertexMap; }
 	inline void SetVertex(int FromID, int ToID) { VertexMap.Add(FromID, ToID); }
 	inline int GetNewVertex(int FromID) const { return VertexMap.GetTo(FromID); }
 	inline bool ContainsVertex(int FromID) const { return VertexMap.ContainsFrom(FromID); }
 
 	FIndexMapi& GetTriangleMap() { return TriangleMap; }
+	const FIndexMapi& GetTriangleMap() const { return TriangleMap; }
 	void SetTriangle(int FromID, int ToID) { TriangleMap.Add(FromID, ToID); }
 	int GetNewTriangle(int FromID) const { return TriangleMap.GetTo(FromID); }
 	inline bool ContainsTriangle(int FromID) const { return TriangleMap.ContainsFrom(FromID); }
@@ -186,12 +196,64 @@ public:
 	/**
 	 * Finds boundary loops of connected components of a set of triangles, and duplicates the vertices
 	 * along the boundary, such that the triangles become disconnected.
+	 * Note: This can create a mesh with bowties if you pass in a subset of triangles that would have bowtie connectivity.
+	 *		 (it is impossible to create the paired LoopSetOut with 1:1 vertex pairs without leaving in the bowties?)
 	 * @param Triangles set of triangles
 	 * @param LoopSetOut set of boundary loops. LoopA is original loop which remains with "outer" triangles, and LoopB is new boundary loop of triangle set
 	 * @return true on success
 	 */
 	bool DisconnectTriangles(const TArray<int>& Triangles, TArray<FLoopPairSet>& LoopSetOut);
 
+	/**
+	* Disconnects triangles (without constructing boundary loops) such that the input Triangles are not connected to any other triangles in the mesh
+	* @param Triangles set of triangles
+	* @param bPreventBowties do some additional processing and vertex splitting as needed to prevent the creation of any new bowties
+	*/
+	void DisconnectTriangles(const TArray<int>& Triangles, bool bPreventBowties = true);
+
+
+	/**
+	 * Splits all bowties across the whole mesh
+	 */
+	void SplitBowties(FDynamicMeshEditResult& ResultOut);
+
+	/**
+	 * Splits any bowties specifically on the given vertex, and updates (does not reset!) ResultOut with any added vertices
+	 */
+	void SplitBowties(int VertexID, FDynamicMeshEditResult& ResultOut);
+
+	/**
+	 * In ReinsertSubmesh, a problem can arise where the mesh we are 
+	 * inserting has duplicate triangles of the base mesh.
+	 *
+	 * This can lead to problematic behavior later. We can do various things,
+	 * like delete and replace that existing triangle, or just use it instead
+	 * of adding a new one. Or fail, or ignore it.
+	 * 
+	 * This enum/argument controls the behavior. 
+	 * However, fundamentally this kind of problem should be handled upstream!!
+	 * For example by not trying to remesh areas that contain nonmanifold geometry...
+	 */
+	enum class EDuplicateTriBehavior : uint8
+	{
+		EnsureContinue,         
+		EnsureAbort, UseExisting, Replace
+	};
+
+	/**
+	 * Update a Base Mesh from a Submesh; See FMeshRegionOperator::BackPropropagate for a usage example.
+	 * 
+	 * Assumes that Submesh has been modified, but boundary loop has been preserved, and that old submesh has already been removed from this mesh.
+	 * Just appends new vertices and rewrites triangles.
+	 *
+	 * @param Submesh The mesh to insert back into its BaseMesh.  The original submesh triangles should have already been removed when this function is called
+	 * @param SubToNewV Mapping from submesh to vertices in the updated base mesh
+	 * @param NewTris If not null, will be filled with IDs of triangles added to the base mesh
+	 * @param DuplicateBehavior Choice of what to do if inserting a triangle from the submesh would duplicate a triangle in the base mesh
+	 * @return true if submesh successfully inserted, false if any triangles failed (which happens if triangle would result in non-manifold mesh)
+	 */
+	bool ReinsertSubmesh(const FDynamicSubmesh3& Submesh, FOptionallySparseIndexMap& SubToNewV, TArray<int>* NewTris = nullptr,
+						 EDuplicateTriBehavior DuplicateBehavior = EDuplicateTriBehavior::EnsureAbort);
 
 	/**
 	 * Remove a list of triangles from the mesh, and optionally any vertices that are now orphaned
@@ -403,9 +465,20 @@ public:
 	 * @param SourceTriangles the triangles to copy
 	 * @param IndexMaps returned mappings from old to new triangles/vertices/etc (you may initialize to optimize memory usage, etc)
 	 * @param ResultOut lists of newly created triangles/vertices/etc
+	 * @param bComputeTriangleMap if true, computes the triangle map section of IndexMaps (which is not needed for the append to work, so is optional)
 	 */
-	void AppendTriangles(const FDynamicMesh3* SourceMesh, const TArray<int>& SourceTriangles, FMeshIndexMappings& IndexMaps, FDynamicMeshEditResult& ResultOut);
+	void AppendTriangles(const FDynamicMesh3* SourceMesh, const TArrayView<const int>& SourceTriangles, FMeshIndexMappings& IndexMaps, FDynamicMeshEditResult& ResultOut, bool bComputeTriangleMap = true);
 
+	/**
+	 * Create multiple meshes out of the source mesh by splitting triangles out.
+	 * Static because it creates multiple output meshes, so doesn't quite fit in the FDynamicMeshEditor model of operating on a single mesh
+	 *
+	 * @param SourceMesh
+	 * @param SplitMeshes
+	 * @param TriIDToMeshID
+	 * @return true if needed split, false if there were not multiple mesh ids so no split was needed
+	 */
+	static bool SplitMesh(const FDynamicMesh3* SourceMesh, TArray<FDynamicMesh3>& SplitMeshes, TFunctionRef<int(int)> TriIDToMeshID, int DeleteMeshID = -1);
 
 };
 
diff --git a/Engine/Plugins/Experimental/GeometryProcessing/Source/DynamicMesh/Public/DynamicMeshOverlay.h b/Engine/Plugins/Experimental/GeometryProcessing/Source/DynamicMesh/Public/DynamicMeshOverlay.h
--- a/Engine/Plugins/Experimental/GeometryProcessing/Source/DynamicMesh/Public/DynamicMeshOverlay.h
+++ b/Engine/Plugins/Experimental/GeometryProcessing/Source/DynamicMesh/Public/DynamicMeshOverlay.h
@@ -5,10 +5,12 @@
 
 #include "Util/DynamicVector.h"
 #include "Util/RefCountVector.h"
+#include "Util/CompactMaps.h"
 #include "VectorTypes.h"
 #include "GeometryTypes.h"
-#include "DynamicMesh3.h"
+#include "InfoTypes.h"
 
+class FDynamicMesh3;
 
 /**
  * TDynamicMeshOverlay is an add-on to a FDynamicMesh3 that allows for per-triangle storage
@@ -60,7 +62,13 @@ public:
 	{
 		ParentMesh = ParentMeshIn;
 	}
-
+private:
+	/** @set the parent mesh for this overlay.  Only safe for use during FDynamicMesh move */
+	void Reparent(FDynamicMesh3* ParentMeshIn)
+	{
+		ParentMesh = ParentMeshIn;
+	}
+public:
 	/** @return the parent mesh for this overlay */
 	const FDynamicMesh3* GetParentMesh() const { return ParentMesh; }
 	/** @return the parent mesh for this overlay */
@@ -75,6 +83,110 @@ public:
 		ElementTriangles = Copy.ElementTriangles;
 	}
 
+	/** Copy the Copy overlay to a compact rep, also updating parent references based on the CompactMaps */
+	void CompactCopy(const FCompactMaps& CompactMaps, const TDynamicMeshOverlay<RealType, ElementSize>& Copy)
+	{
+		ClearElements();
+
+		// map of element IDs
+		TArray<int> MapE; MapE.SetNumUninitialized(Copy.MaxElementID());
+
+		// copy elements across
+		RealType Data[ElementSize];
+		for (int EID = 0; EID < Copy.MaxElementID(); EID++)
+		{
+			if (Copy.IsElement(EID))
+			{
+				Copy.GetElement(EID, Data);
+				MapE[EID] = AppendElement(Data, CompactMaps.MapV[Copy.ParentVertices[EID]]);
+			}
+			else
+			{
+				MapE[EID] = -1;
+			}
+		}
+
+		// copy triangles across
+		check(CompactMaps.bKeepTriangleMap && CompactMaps.MapT.Num() == Copy.GetParentMesh()->MaxTriangleID()); // must have valid triangle map
+		for (int FromTID : Copy.GetParentMesh()->TriangleIndicesItr())
+		{
+			if (!Copy.IsSetTriangle(FromTID))
+			{
+				continue;
+			}
+			int ToTID = CompactMaps.MapT[FromTID];
+			FIndex3i FromTriElements = Copy.GetTriangle(FromTID);
+			SetTriangle(ToTID, FIndex3i(MapE[FromTriElements.A], MapE[FromTriElements.B], MapE[FromTriElements.C]));
+		}
+	}
+
+	/** Compact overlay and update links to parent based on CompactMaps */
+	void CompactInPlace(const FCompactMaps& CompactMaps)
+	{
+		int iLastE = MaxElementID() - 1, iCurE = 0;
+		while (iLastE >= 0 && ElementsRefCounts.IsValidUnsafe(iLastE) == false)
+		{
+			iLastE--;
+		}
+		while (iCurE < iLastE && ElementsRefCounts.IsValidUnsafe(iCurE) == false)
+		{
+			iCurE++;
+		}
+
+		// make a map to track element index changes, to use to update element triangles later
+		// TODO: it may be faster to not construct this and to do the remapping per element as we go (by iterating the one ring of each parent vertex for each element)
+		TArray<int> MapE; MapE.SetNumUninitialized(MaxElementID());
+		for (int ID = 0; ID < MapE.Num(); ID++)
+		{
+			// mapping is 1:1 by default; sparsely re-mapped below
+			MapE[ID] = ID;
+		}
+
+		TDynamicVector<short>& ERef = ElementsRefCounts.GetRawRefCountsUnsafe();
+		RealType Data[ElementSize];
+		while (iCurE < iLastE)
+		{
+			// remap the element data
+			GetElement(iLastE, Data);
+			SetElement(iCurE, Data);
+			int OrigParent = ParentVertices[iLastE];
+			ParentVertices[iCurE] = CompactMaps.GetVertex(OrigParent);
+			ERef[iCurE] = ERef[iLastE];
+			ERef[iLastE] = FRefCountVector::INVALID_REF_COUNT;
+			MapE[iLastE] = iCurE;
+
+			// move cur forward one, last back one, and  then search for next valid
+			iLastE--; iCurE++;
+			while (iLastE >= 0 && ElementsRefCounts.IsValidUnsafe(iLastE) == false)
+			{
+				iLastE--;
+			}
+			while (iCurE < iLastE && ElementsRefCounts.IsValidUnsafe(iCurE))
+			{
+				iCurE++;
+			}
+		}
+		ElementsRefCounts.Trim(ElementCount());
+		Elements.Resize(ElementCount() * ElementSize);
+		ParentVertices.Resize(ElementCount());
+
+		for (int TID = 0, OldMaxTID = ElementTriangles.Num() / 3; TID < OldMaxTID; TID++)
+		{
+			int OldStart = TID * 3;
+			if (ElementTriangles[OldStart] == -1)
+			{
+				continue; // triangle was not set; skip it
+			}
+			int NewStart = CompactMaps.GetTriangle(TID) * 3;
+			for (int SubIdx = 0; SubIdx < 3; SubIdx++)
+			{
+				ElementTriangles[NewStart + SubIdx] = MapE[ElementTriangles[OldStart + SubIdx]];
+			}
+		}
+
+		checkSlow(IsCompact());
+	}
+
 	/** Discard current set of elements, but keep triangles */
 	void ClearElements();
 
@@ -85,6 +197,9 @@ public:
 	/** @return true if this element index is in use */
 	inline bool IsElement(int vID) const { return ElementsRefCounts.IsValid(vID); }
 
+	/** @return true if the elements are compact */
+	bool IsCompact() const { return ElementsRefCounts.IsDense(); }
+
 
 	typedef typename FRefCountVector::IndexEnumerable element_iterator;
 
@@ -125,6 +240,33 @@ public:
 	void SplitVerticesWithPredicate(TFunctionRef<bool(int ElementIdx, int TriID)> ShouldSplitOutVertex, TFunctionRef<void(int ElementIdx, int TriID, RealType* FillVect)> GetNewElementValue);
 
 
+	/**
+	 * Create a new copy of ElementID, and update connected triangles in the TrianglesToUpdate array to reference the copy of ElementID where they used to reference ElementID
+	 * (Note: This just calls "SplitElementWithNewParent" with the existing element's parent id.)
+	 *
+	 * @param ElementID the element to copy
+	 * @param TrianglesToUpdate the triangles that should now reference the new element
+	 * @return the ID of the newly created element
+	 */
+	int SplitElement(int ElementID, const TArrayView<const int>& TrianglesToUpdate);
+
+	/**
+	* Create a new copy of ElementID, and update connected triangles in the TrianglesToUpdate array to reference the copy of ElementID where they used to reference ElementID.  The new element will have the given parent vertex ID.
+	*
+	* @param ElementID the element to copy
+	* @param SplitParentVertexID the new parent vertex for copied elements
+	* @param TrianglesToUpdate the triangles that should now reference the new element.  Note: this is allowed to include triangles that do not have the element at all; sometimes you may want to do so to avoid creating a new array for each call.
+	* @return the ID of the newly created element
+	*/
+	int SplitElementWithNewParent(int ElementID, int SplitParentVertexID, const TArrayView<const int>& TrianglesToUpdate);
+
+
+	/**
+	* Refine an existing overlay topology by splitting any bow ties
+	*/
+	void SplitBowties();
+
+
 	//
 	// Support for inserting element at specific ID. This is a bit tricky
 	// because we likely will need to update the free list in the RefCountVector, which
@@ -229,6 +371,9 @@ public:
 	/** Returns true if the parent-mesh vertex is connected to any seam edges */
 	bool IsSeamVertex(int VertexID, bool bBoundaryIsSeam = true) const;
 
+	/** @return true if the two triangles are connected, ie shared edge exists and is not a seam edge */
+	bool AreTrianglesConnected(int TriangleID0, int TriangleID1) const;
+
 	/** find the elements associated with a given parent-mesh vertex */
 	void GetVertexElements(int VertexID, TArray<int>& OutElements) const;
 	/** Count the number of unique elements for a given parent-mesh vertex */
@@ -254,15 +399,17 @@ public:
 	/** Reverse the orientation of a triangle's elements */
 	void OnReverseTriOrientation(int TriangleID);
 	/** Update the overlay to reflect an edge split in the parent mesh */
-	void OnSplitEdge(const FDynamicMesh3::FEdgeSplitInfo& SplitInfo);
+	void OnSplitEdge(const DynamicMeshInfo::FEdgeSplitInfo& SplitInfo);
 	/** Update the overlay to reflect an edge flip in the parent mesh */
-	void OnFlipEdge(const FDynamicMesh3::FEdgeFlipInfo& FlipInfo);
+	void OnFlipEdge(const DynamicMeshInfo::FEdgeFlipInfo& FlipInfo);
 	/** Update the overlay to reflect an edge collapse in the parent mesh */
-	void OnCollapseEdge(const FDynamicMesh3::FEdgeCollapseInfo& CollapseInfo);
+	void OnCollapseEdge(const DynamicMeshInfo::FEdgeCollapseInfo& CollapseInfo);
 	/** Update the overlay to reflect a face poke in the parent mesh */
-	void OnPokeTriangle(const FDynamicMesh3::FPokeTriangleInfo& PokeInfo);
+	void OnPokeTriangle(const DynamicMeshInfo::FPokeTriangleInfo& PokeInfo);
 	/** Update the overlay to reflect an edge merge in the parent mesh */
-	void OnMergeEdges(const FDynamicMesh3::FMergeEdgesInfo& MergeInfo);
+	void OnMergeEdges(const DynamicMeshInfo::FMergeEdgesInfo& MergeInfo);
+	/** Update the overlay to reflect a vertex split in the parent mesh */
+	void OnSplitVertex(const DynamicMeshInfo::FVertexSplitInfo& SplitInfo, const TArrayView<const int>& TrianglesToUpdate);
 
 protected:
 	/** Set the value at an Element to be a linear interpolation of two other Elements */
diff --git a/Engine/Plugins/Experimental/GeometryProcessing/Source/DynamicMesh/Public/FaceGroupUtil.h b/Engine/Plugins/Experimental/GeometryProcessing/Source/DynamicMesh/Public/FaceGroupUtil.h
new file mode 100644
--- /dev/null
+++ b/Engine/Plugins/Experimental/GeometryProcessing/Source/DynamicMesh/Public/FaceGroupUtil.h
@@ -0,0 +1,223 @@
+// Copyright 1998-2019 Epic Games, Inc. All Rights Reserved.
+
+// Port of geometry3sharp FaceGroupUtil
+
+#pragma once
+
+#include "DynamicMesh3.h"
+#include "DynamicMeshEditor.h"
+
+/**
+ * Utility functions for dealing with Mesh indices
+ */
+namespace FaceGroupUtil
+{
+
+	/**
+	 * Set group ID of all triangles in Mesh
+	 */
+	DYNAMICMESH_API void SetGroupID(FDynamicMesh3& Mesh, int to)
+	{
+		if (Mesh.HasTriangleGroups() == false)
+		{
+			return;
+		}
+		for (int tid : Mesh.TriangleIndicesItr())
+		{
+			Mesh.SetTriangleGroup(tid, to);
+		}
+	}
+
+
+	/**
+	 * Set group id of subset of triangles in Mesh
+	 */
+	DYNAMICMESH_API void SetGroupID(FDynamicMesh3& Mesh, const TArrayView<const int>& triangles, int to)
+	{
+		if (Mesh.HasTriangleGroups() == false)
+		{
+			return;
+		}
+		for (int tid : triangles)
+		{
+			Mesh.SetTriangleGroup(tid, to);
+		}
+	}
+
+
+	/**
+	 * replace group id in Mesh
+	 */
+	DYNAMICMESH_API void SetGroupToGroup(FDynamicMesh3& Mesh, int from, int to)
+	{
+		if (Mesh.HasTriangleGroups() == false)
+		{
+			return;
+		}
+
+		int NT = Mesh.MaxTriangleID();
+		for ( int tid = 0; tid < NT; ++tid)
+		{
+			if (Mesh.IsTriangle(tid))
+			{
+				int gid = Mesh.GetTriangleGroup(tid);
+				if (gid == from)
+				{
+					Mesh.SetTriangleGroup(tid, to);
+				}
+			}
+		}
+	}
+
+
+	/**
+	 * find the set of group ids used in Mesh
+	 */
+	DYNAMICMESH_API TSet<int> FindAllGroups(const FDynamicMesh3& Mesh)
+	{
+		TSet<int> Groups;
+
+		if (Mesh.HasTriangleGroups())
+		{
+			int NT = Mesh.MaxTriangleID();
+			for (int tid = 0; tid < NT; ++tid)
+			{
+				if (Mesh.IsTriangle(tid))
+				{
+					int gid = Mesh.GetTriangleGroup(tid);
+					Groups.Add(gid);
+				}
+			}
+		}
+		return Groups;
+	}
+
+
+	/**
+	 * count number of tris in each group in Mesh; TODO: does this need sparse storage?
+	 */
+	DYNAMICMESH_API TArray<int> CountAllGroups(const FDynamicMesh3& Mesh)
+	{
+		TArray<int> GroupCounts; GroupCounts.SetNum(Mesh.MaxGroupID());
+
+		if (Mesh.HasTriangleGroups())
+		{
+			int NT = Mesh.MaxTriangleID();
+			for (int tid = 0; tid < NT; ++tid)
+			{
+				if (Mesh.IsTriangle(tid))
+				{
+					int gid = Mesh.GetTriangleGroup(tid);
+					GroupCounts[gid]++;
+				}
+			}
+		}
+		return GroupCounts;
+	}
+
+
+	/**
+	 * collect triangles by group id. Returns array of triangle lists (stored as arrays).
+	 * This requires 2 passes over Mesh, but each pass is linear
+	 */
+	DYNAMICMESH_API TArray<TArray<int>> FindTriangleSetsByGroup(const FDynamicMesh3& Mesh, int IgnoreGID = -1)
+	{
+		TArray<TArray<int>> Sets;
+		if (!Mesh.HasTriangleGroups())
+		{
+			return Sets;
+		}
+
+		// find # of groups and triangle count for each
+		TArray<int> Counts = CountAllGroups(Mesh);
+		TArray<int> GroupIDs;
+		for (int CountIdx = 0; CountIdx < Counts.Num(); CountIdx++) {
+			int Count = Counts[CountIdx];
+			if (CountIdx != IgnoreGID && Count > 0)
+			{
+				GroupIDs.Add(CountIdx);
+			}
+		}
+		TArray<int> GroupMap; GroupMap.SetNum(Mesh.MaxGroupID());
+
+		// allocate sets
+		Sets.SetNum(GroupIDs.Num());
+		for (int i = 0; i < GroupIDs.Num(); ++i)
+		{
+			int GID = GroupIDs[i];
+			Sets[i].Reserve(Counts[GID]);
+			GroupMap[GID] = i;
+		}
+
+		// accumulate triangles
+		int NT = Mesh.MaxTriangleID();
+		for (int tid = 0; tid < NT; ++tid)
+		{
+			if (Mesh.IsTriangle(tid))
+			{
+				int GID = Mesh.GetTriangleGroup(tid);
+				int i = GroupMap[GID];
+				if (i >= 0)
+				{
+					Sets[i].Add(tid);
+				}
+			}
+		}
+
+		return Sets;
+	}
+
+
+
+	/**
+	 * find list of triangles in Mesh with specific group id
+	 */
+	DYNAMICMESH_API TArray<int> FindTrianglesByGroup(FDynamicMesh3& Mesh, int FindGroupID)
+	{
+		TArray<int> tris;
+		if (Mesh.HasTriangleGroups() == false)
+		{
+			return tris;
+		}
+		for (int tid : Mesh.TriangleIndicesItr())
+		{
+			if (Mesh.GetTriangleGroup(tid) == FindGroupID)
+			{
+				tris.Add(tid);
+			}
+		}
+		return tris;
+	}
+
+	/**
+	* split input Mesh into submeshes based on group ID
+	* **does not** separate disconnected components w/ same group ID
+	*/
+	DYNAMICMESH_API void SeparateMeshByGroups(FDynamicMesh3& Mesh, TArray<FDynamicMesh3>& SplitMeshes)
+	{
+		FDynamicMeshEditor::SplitMesh(&Mesh, SplitMeshes, [&Mesh](int TID)
+		{
+			return Mesh.GetTriangleGroup(TID);
+		});
+	}
+
+	/**
+	* split input Mesh into submeshes based on group ID
+	* **does not** separate disconnected components w/ same group ID
+	*/
+	DYNAMICMESH_API void SeparateMeshByGroups(FDynamicMesh3& Mesh, TArray<FDynamicMesh3>& SplitMeshes, TArray<int>& GroupIDs)
+	{
+		// build split meshes
+		SeparateMeshByGroups(Mesh, SplitMeshes);
+
+		// build array of per-mesh group id
+		GroupIDs.Reset();
+		for (const FDynamicMesh3& M : SplitMeshes)
+		{
+			check(M.TriangleCount() > 0); // SplitMesh should never add an empty mesh
+			GroupIDs.Add(M.GetTriangleGroup(0));
+		}
+	}
+
+
+}
diff --git a/Engine/Plugins/Experimental/GeometryProcessing/Source/DynamicMesh/Public/GroupTopology.h b/Engine/Plugins/Experimental/GeometryProcessing/Source/DynamicMesh/Public/GroupTopology.h
--- a/Engine/Plugins/Experimental/GeometryProcessing/Source/DynamicMesh/Public/GroupTopology.h
+++ b/Engine/Plugins/Experimental/GeometryProcessing/Source/DynamicMesh/Public/GroupTopology.h
@@ -8,6 +8,99 @@
 #include "Containers/BitArray.h"
 
 
+
+
+
+/**
+ * FGroupTopologySelection represents a set of selected elements of a FGroupTopology
+ */
+struct DYNAMICMESH_API FGroupTopologySelection
+{
+	TArray<int32> SelectedGroupIDs;
+	TArray<int32> SelectedCornerIDs;
+	TArray<int32> SelectedEdgeIDs;
+
+	FGroupTopologySelection() { Clear(); }
+	
+	void Clear()
+	{
+		SelectedGroupIDs.Reset();
+		SelectedCornerIDs.Reset();
+		SelectedEdgeIDs.Reset();
+	}
+
+	bool IsEmpty() const
+	{
+		return SelectedCornerIDs.Num() == 0 && SelectedEdgeIDs.Num() == 0 && SelectedGroupIDs.Num() == 0;
+	}
+
+	void Append(const FGroupTopologySelection& Selection)
+	{
+		for (int32 i : Selection.SelectedCornerIDs)
+		{
+			SelectedCornerIDs.AddUnique(i);
+		}
+		for (int32 i : Selection.SelectedEdgeIDs)
+		{
+			SelectedEdgeIDs.AddUnique(i);
+		}
+		for (int32 i : Selection.SelectedGroupIDs)
+		{
+			SelectedGroupIDs.AddUnique(i);
+		}
+	}
+
+
+	void Remove(const FGroupTopologySelection& Selection)
+	{
+		for (int32 i : Selection.SelectedCornerIDs)
+		{
+			SelectedCornerIDs.RemoveSwap(i);
+		}
+		for (int32 i : Selection.SelectedEdgeIDs)
+		{
+			SelectedEdgeIDs.RemoveSwap(i);
+		}
+		for (int32 i : Selection.SelectedGroupIDs)
+		{
+			SelectedGroupIDs.RemoveSwap(i);
+		}
+	}
+
+
+	void Toggle(const FGroupTopologySelection& Selection)
+	{
+		for (int32 i : Selection.SelectedCornerIDs)
+		{
+			ToggleItem(SelectedCornerIDs, i);
+		}
+		for (int32 i : Selection.SelectedEdgeIDs)
+		{
+			ToggleItem(SelectedEdgeIDs, i);
+		}
+		for (int32 i : Selection.SelectedGroupIDs)
+		{
+			ToggleItem(SelectedGroupIDs, i);
+		}
+	}
+
+private:
+	void ToggleItem(TArray<int32>& List, int32 Item)
+	{
+		if (List.Contains(Item))
+		{
+			List.RemoveSwap(Item);
+		}
+		else
+		{
+			List.Add(Item);
+		}
+	}
+};
+
+
+
+
 /**
  * Given a per-triangle integer ("group"), FGroupTopology extracts a
  * group-level topological graph from an input Mesh. The graph consists
@@ -79,7 +172,7 @@ public:
 		/** GroupID for this group */
 		int GroupID;
 		/** List of triangles forming this group */
-		TArray<int> Faces;
+		TArray<int> Triangles;
 		
 		/** List of boundaries of this group (may be empty, eg on a closed component with only one group) */
 		TArray<FGroupBoundary> Boundaries;
@@ -169,6 +262,11 @@ public:
 	 */
 	bool GetGroupEdgeTangent(int GroupEdgeID, FVector3d& TangentOut) const;
 
+	/**
+	 * @return a 3D frame for the given selection
+	 */
+	FFrame3d GetSelectionFrame(const FGroupTopologySelection& Selection);
+
 protected:
 	const FDynamicMesh3* Mesh = nullptr;
 
diff --git a/Engine/Plugins/Experimental/GeometryProcessing/Source/DynamicMesh/Public/MeshNormals.h b/Engine/Plugins/Experimental/GeometryProcessing/Source/DynamicMesh/Public/MeshNormals.h
--- a/Engine/Plugins/Experimental/GeometryProcessing/Source/DynamicMesh/Public/MeshNormals.h
+++ b/Engine/Plugins/Experimental/GeometryProcessing/Source/DynamicMesh/Public/MeshNormals.h
@@ -88,7 +88,6 @@ public:
 	void CopyToOverlay(FDynamicMeshNormalOverlay* NormalOverlay, bool bInvert = false) const;
 
 
-
 	/**
 	 * Compute per-vertex normals for the given Mesh
 	 * @param bInvert if true, normals are flipped
@@ -107,6 +106,11 @@ public:
 	 */
 	static FVector3d ComputeOverlayNormal(const FDynamicMesh3& Mesh, FDynamicMeshNormalOverlay* NormalOverlay, int ElemIdx);
 
+	/**
+	 * Initialize the given NormalOverlay with per-vertex normals, ie single overlay element for each mesh vertex.
+	 * @param bUseMeshVertexNormalsIfAvailable if true and the parent mesh has per-vertex normals, use them instead of calculating new ones
+	 */
+	static void InitializeOverlayToPerVertexNormals(FDynamicMeshNormalOverlay* NormalOverlay, bool bUseMeshVertexNormalsIfAvailable = true);
 
 
 protected:
diff --git a/Engine/Plugins/Experimental/GeometryProcessing/Source/DynamicMesh/Public/MeshTransforms.h b/Engine/Plugins/Experimental/GeometryProcessing/Source/DynamicMesh/Public/MeshTransforms.h
--- a/Engine/Plugins/Experimental/GeometryProcessing/Source/DynamicMesh/Public/MeshTransforms.h
+++ b/Engine/Plugins/Experimental/GeometryProcessing/Source/DynamicMesh/Public/MeshTransforms.h
@@ -13,6 +13,22 @@
 namespace MeshTransforms
 {
 	
+	/**
+	 * Apply Translation to vertex positions of Mesh. Does not modify any other attributes.
+	 */
+	DYNAMICMESH_API void Translate(FDynamicMesh3& Mesh, const FVector3d& Translation);
+
+	/**
+	 * Transform Mesh into local coordinates of Frame
+	 */
+	DYNAMICMESH_API void WorldToFrameCoords(FDynamicMesh3& Mesh, const FFrame3d& Frame);
+
+	/**
+	 * Transform Mesh out of local coordinates of Frame
+	 */
+	DYNAMICMESH_API void FrameCoordsToWorld(FDynamicMesh3& Mesh, const FFrame3d& Frame);
+
+
 	/**
 	 * Apply given Transform to a Mesh.
 	 * Modifies Vertex Positions and Normals, and any Per-Triangle Normal Overlays
@@ -27,6 +43,14 @@ namespace MeshTransforms
 	DYNAMICMESH_API void ApplyTransformInverse(FDynamicMesh3& Mesh, const FTransform3d& Transform);
 
 
+	/**
+	 * Apply given Transform to a Mesh.
+	 * Modifies Vertex Positions and Normals, and any Per-Triangle Normal Overlays
+	 */
+	DYNAMICMESH_API void ApplyTransform(FDynamicMesh3& Mesh, 
+		TFunctionRef<FVector3d(const FVector3d&)> PositionTransform,
+		TFunctionRef<FVector3f(const FVector3f&)> NormalTransform);
+
 
 
 };
\ No newline at end of file
diff --git a/Engine/Plugins/Experimental/GeometryProcessing/Source/GeometricObjects/Private/Tests/VectorTypes.spec.cpp b/Engine/Plugins/Experimental/GeometryProcessing/Source/GeometricObjects/Private/Tests/VectorTypes.spec.cpp
--- a/Engine/Plugins/Experimental/GeometryProcessing/Source/GeometricObjects/Private/Tests/VectorTypes.spec.cpp
+++ b/Engine/Plugins/Experimental/GeometryProcessing/Source/GeometricObjects/Private/Tests/VectorTypes.spec.cpp
@@ -97,7 +97,7 @@ void FVectorTypesSpec::Define()
 			{
 				constexpr float initialValues[3]{0.5f, 0.5f, 0.5f};
 				Vec testVec{initialValues};
-				FLinearColor converted = {testVec};
+				FLinearColor converted = {(FVector)testVec};
 				TestEqual("Color.R", testVec.X, converted.R);
 				TestEqual("Color.G", testVec.Y, converted.G);
 				TestEqual("Color.B", testVec.Z, converted.B);
diff --git a/Engine/Plugins/Experimental/GeometryProcessing/Source/GeometricObjects/Private/Util/SmallListSet.cpp b/Engine/Plugins/Experimental/GeometryProcessing/Source/GeometricObjects/Private/Util/SmallListSet.cpp
--- a/Engine/Plugins/Experimental/GeometryProcessing/Source/GeometricObjects/Private/Util/SmallListSet.cpp
+++ b/Engine/Plugins/Experimental/GeometryProcessing/Source/GeometricObjects/Private/Util/SmallListSet.cpp
@@ -4,27 +4,6 @@
 
 const int32 FSmallListSet::NullValue = -1;
 
-
-FSmallListSet::FSmallListSet()
-{
-	ListHeads = TDynamicVector<int32>();
-	LinkedListElements = TDynamicVector<int32>();
-	FreeHeadIndex = NullValue;
-	ListBlocks = TDynamicVector<int32>();
-	FreeBlocks = TDynamicVector<int32>();
-}
-
-
-FSmallListSet::FSmallListSet(const FSmallListSet& copy)
-{
-	LinkedListElements = TDynamicVector<int32>(copy.LinkedListElements);
-	FreeHeadIndex = copy.FreeHeadIndex;
-	ListHeads = TDynamicVector<int32>(copy.ListHeads);
-	ListBlocks = TDynamicVector<int32>(copy.ListBlocks);
-	FreeBlocks = TDynamicVector<int32>(copy.FreeBlocks);
-}
-
-
 void FSmallListSet::Resize(int32 NewSize)
 {
 	int32 CurSize = (int32)ListHeads.GetLength();
@@ -397,4 +376,4 @@ bool FSmallListSet::RemoveFromLinkedList(int32 block_ptr, int32 val)
 		cur_ptr = LinkedListElements[cur_ptr + 1];
 	}
 	return false;
-}
\ No newline at end of file
+}
diff --git a/Engine/Plugins/Experimental/GeometryProcessing/Source/GeometricObjects/Public/BoxTypes.h b/Engine/Plugins/Experimental/GeometryProcessing/Source/GeometricObjects/Public/BoxTypes.h
--- a/Engine/Plugins/Experimental/GeometryProcessing/Source/GeometricObjects/Public/BoxTypes.h
+++ b/Engine/Plugins/Experimental/GeometryProcessing/Source/GeometricObjects/Public/BoxTypes.h
@@ -199,6 +199,23 @@ struct TAxisAlignedBox3
 		this->Min = Min;
 		this->Max = Max;
 	}
+
+	TAxisAlignedBox3(const TAxisAlignedBox3& OtherBox) = default;
+
+	template<typename OtherRealType>
+	explicit TAxisAlignedBox3(const TAxisAlignedBox3<OtherRealType>& OtherBox)
+	{
+		this->Min = FVector3<RealType>(OtherBox.Min);
+		this->Max = FVector3<RealType>(OtherBox.Max);
+	}
+
+	TAxisAlignedBox3(const FVector3<RealType>& Center, RealType HalfWidth)
+	{
+		this->Min = FVector3<RealType>(Center.X-HalfWidth, Center.Y-HalfWidth, Center.Z-HalfWidth);
+		this->Max = FVector3<RealType>(Center.X+HalfWidth, Center.Y+HalfWidth, Center.Z+HalfWidth);
+	}
+
+
 	TAxisAlignedBox3(const TAxisAlignedBox3& Box, const TFunction<FVector3<RealType>(const FVector3<RealType>&)> TransformF)
 	{
 		if (TransformF == nullptr)
@@ -217,7 +234,7 @@ struct TAxisAlignedBox3
 		}
 	}
 
-	operator FBox() const
+	explicit operator FBox() const
 	{
 		FVector MinV((float)Min.X, (float)Min.Y, (float)Min.Z);
 		FVector MaxV((float)Max.X, (float)Max.Y, (float)Max.Z);
@@ -236,9 +253,9 @@ struct TAxisAlignedBox3
 	FVector3<RealType> GetCorner(int Index) const
 	{
 		check(Index >= 0 && Index <= 7);
-		double X = (((Index & 1) != 0) ^ ((Index & 2) != 0)) ? (Max.X) : (Min.X);
-		double Y = ((Index / 2) % 2 == 0) ? (Min.Y) : (Max.Y);
-		double Z = (Index < 4) ? (Min.Z) : (Max.Z);
+		RealType X = (((Index & 1) != 0) ^ ((Index & 2) != 0)) ? (Max.X) : (Min.X);
+		RealType Y = ((Index / 2) % 2 == 0) ? (Min.Y) : (Max.Y);
+		RealType Z = (Index < 4) ? (Min.Z) : (Max.Z);
 		return FVector3<RealType>(X, Y, Z);
 	}
 
@@ -408,6 +425,16 @@ struct TAxisAlignedBox2
 		: Min(Min), Max(Max)
 	{
 	}
+
+	TAxisAlignedBox2(const TAxisAlignedBox2& OtherBox) = default;
+
+	template<typename OtherRealType>
+	explicit TAxisAlignedBox2(const TAxisAlignedBox2<OtherRealType>& OtherBox)
+	{
+		this->Min = FVector2<RealType>(OtherBox.Min);
+		this->Max = FVector2<RealType>(OtherBox.Max);
+	}
+
 	TAxisAlignedBox2(RealType SquareSize)
 		: Min((RealType)0, (RealType)0), Max(SquareSize, SquareSize)
 	{
@@ -417,7 +444,19 @@ struct TAxisAlignedBox2
 	{
 	}
 
-	operator FBox2D() const
+	TAxisAlignedBox2(const TArray<FVector2<RealType>>& Pts)
+	{
+		*this = Empty();
+		Contain(Pts);
+	}
+
+	TAxisAlignedBox2(const FVector2<RealType>& Center, RealType HalfWidth)
+	{
+		this->Min = FVector2<RealType>(Center.X - HalfWidth, Center.Y - HalfWidth);
+		this->Max = FVector2<RealType>(Center.X + HalfWidth, Center.Y + HalfWidth);
+	}
+
+	explicit operator FBox2D() const
 	{
 		FVector2D MinV((float)Min.X, (float)Min.Y);
 		FVector2D MaxV((float)Max.X, (float)Max.Y);
diff --git a/Engine/Plugins/Experimental/GeometryProcessing/Source/GeometricObjects/Public/FrameTypes.h b/Engine/Plugins/Experimental/GeometryProcessing/Source/GeometricObjects/Public/FrameTypes.h
--- a/Engine/Plugins/Experimental/GeometryProcessing/Source/GeometricObjects/Public/FrameTypes.h
+++ b/Engine/Plugins/Experimental/GeometryProcessing/Source/GeometricObjects/Public/FrameTypes.h
@@ -83,8 +83,23 @@ struct TFrame3
 	/** Construct a Frame from an FTransform */
 	explicit TFrame3(const FTransform& Transform)
 	{
-		Origin = Transform.GetTranslation();
-		Rotation = Transform.GetRotation();
+		Origin = FVector3<RealType>(Transform.GetTranslation());
+		Rotation = TQuaternion<RealType>(Transform.GetRotation());
+	}
+
+	/** Construct a Frame from an FVector and FQuat */
+	explicit TFrame3(const FVector& OriginIn, const FQuat& RotationIn)
+	{
+		Origin = FVector3<RealType>(OriginIn);
+		Rotation = TQuaternion<RealType>(RotationIn);
+	}
+
+	/** Convert between TFrame of different types */
+	template<typename RealType2>
+	explicit TFrame3(const TFrame3<RealType2>& OtherFrame)
+	{
+		Origin = static_cast<FVector3<RealType>>(OtherFrame.Origin);
+		Rotation = static_cast<TQuaternion<RealType>>(OtherFrame.Rotation);
 	}
 
 	
@@ -129,7 +144,7 @@ struct TFrame3
 	/** @return conversion of this Frame to FTransform */
 	FTransform ToFTransform() const
 	{
-		return FTransform(Rotation, Origin);
+		return FTransform((FQuat)Rotation, (FVector)Origin);
 	}
 
 	/** @return conversion of this Frame to TTransform */
@@ -290,11 +305,21 @@ struct TFrame3
 	 */
 	void Transform(const FTransform& XForm)
 	{
-		Origin = (FVector3<RealType>)XForm.TransformPosition((FVector3f)Origin);
+		Origin = FVector3<RealType>(XForm.TransformPosition((FVector)Origin));
 		Rotation = TQuaternion<RealType>(XForm.GetRotation()) * Rotation;
 	}
 
 
+	/**
+	 * transform this frame by the given transform
+	 */
+	void Transform(const TTransform3<RealType>& XForm)
+	{
+		Origin = XForm.TransformPosition(Origin);
+		Rotation = XForm.GetRotation() * Rotation;
+	}
+
+
 	/**
 	 * Align an axis of this frame with a target direction
 	 * @param AxisIndex which axis to align
diff --git a/Engine/Plugins/Experimental/GeometryProcessing/Source/GeometricObjects/Public/GeometryTypes.h b/Engine/Plugins/Experimental/GeometryProcessing/Source/GeometricObjects/Public/GeometryTypes.h
--- a/Engine/Plugins/Experimental/GeometryProcessing/Source/GeometricObjects/Public/GeometryTypes.h
+++ b/Engine/Plugins/Experimental/GeometryProcessing/Source/GeometricObjects/Public/GeometryTypes.h
@@ -89,14 +89,12 @@ protected:
 	TMap<IntType, IntType> ReverseMap;
 	bool bWantForward;
 	bool bWantReverse;
-	IntType InvalidID;
 
 public:
 
 	TIndexMap()
 	{
 		bWantForward = bWantReverse = true;
-		InvalidID = (IntType)-9999999;
 	}
 
 	void Reset()
@@ -105,8 +103,8 @@ public:
 		ReverseMap.Reset();
 	}
 
-	/** @return the value used to indicate "invalid" in the mapping */
-	inline IntType GetInvalidID() const { return InvalidID; }
+	/** @return the value used to indicate an ID is not present in the mapping */
+	constexpr IntType UnmappedID() const { return (IntType)-1; }
 
 	TMap<IntType, IntType>& GetForwardMap() { return ForwardMap; }
 	const TMap<IntType, IntType>& GetForwardMap() const { return ForwardMap; }
@@ -117,6 +115,7 @@ public:
 	/** add mapping from one index to another */
 	inline void Add(IntType FromID, IntType ToID)
 	{
+		checkSlow(FromID >= 0 && ToID >= 0);
 		ForwardMap.Add(FromID, ToID);
 		ReverseMap.Add(ToID, FromID);
 	}
@@ -124,6 +123,7 @@ public:
 	/** @return true if we can map forward from this value */
 	inline bool ContainsFrom(IntType FromID) const
 	{
+		checkSlow(FromID >= 0);
 		check(bWantForward);
 		return ForwardMap.Contains(FromID);
 	}
@@ -131,6 +131,7 @@ public:
 	/** @return true if we can reverse-map from this value */
 	inline bool ContainsTo(IntType ToID) const
 	{
+		checkSlow(ToID >= 0);
 		check(bWantReverse);
 		return ReverseMap.Contains(ToID);
 	}
@@ -139,22 +140,25 @@ public:
 	/** @return forward-map of input value */
 	inline IntType GetTo(IntType FromID) const
 	{
+		checkSlow(FromID >= 0);
 		check(bWantForward);
 		const IntType* FoundVal = ForwardMap.Find(FromID);
-		return (FoundVal == nullptr) ? InvalidID : *FoundVal;
+		return (FoundVal == nullptr) ? UnmappedID() : *FoundVal;
 	}
 
 	/** @return reverse-map of input value */
 	inline IntType GetFrom(IntType ToID) const
 	{
+		checkSlow(ToID >= 0);
 		check(bWantReverse);
 		const IntType* FoundVal = ReverseMap.Find(ToID);
-		return (FoundVal == nullptr) ? InvalidID : *FoundVal;
+		return (FoundVal == nullptr) ? UnmappedID() : *FoundVal;
 	}
 
 	/** @return forward-map of input value or null if not found */
 	inline const IntType* FindTo(IntType FromID) const
 	{
+		checkSlow(FromID >= 0);
 		check(bWantForward);
 		return ForwardMap.Find(FromID);
 	}
@@ -162,6 +166,7 @@ public:
 	/** @return reverse-map of input value or null if not found */
 	inline const IntType* FindFrom(IntType ToID) const
 	{
+		checkSlow(ToID >= 0);
 		check(bWantReverse);
 		return ReverseMap.Find(ToID);
 	}
@@ -169,6 +174,7 @@ public:
 
 	void Reserve(int NumElements)
 	{
+		checkSlow(NumElements >= 0);
 		if (bWantForward)
 		{
 			ForwardMap.Reserve(NumElements);
@@ -179,6 +185,7 @@ public:
 		}
 	}
 };
+
 typedef TIndexMap<int> FIndexMapi;
 
 
diff --git a/Engine/Plugins/Experimental/GeometryProcessing/Source/GeometricObjects/Public/Quaternion.h b/Engine/Plugins/Experimental/GeometryProcessing/Source/GeometricObjects/Public/Quaternion.h
--- a/Engine/Plugins/Experimental/GeometryProcessing/Source/GeometricObjects/Public/Quaternion.h
+++ b/Engine/Plugins/Experimental/GeometryProcessing/Source/GeometricObjects/Public/Quaternion.h
@@ -19,8 +19,10 @@ struct TQuaternion
 
 	TQuaternion();
 	TQuaternion(RealType X, RealType Y, RealType Z, RealType W);
-	TQuaternion(const RealType* Values);
+	explicit TQuaternion(const RealType* Values);
 	TQuaternion(const TQuaternion& Copy);
+	template<typename RealType2>
+	explicit TQuaternion(const TQuaternion<RealType2>& Copy);
 	TQuaternion(const FVector3<RealType>& Axis, RealType Angle, bool bAngleIsDegrees);
 	TQuaternion(const FVector3<RealType>& From, const FVector3<RealType>& To);
 	TQuaternion(const TQuaternion<RealType>& From, const TQuaternion<RealType>& To, RealType InterpT);
@@ -62,7 +64,7 @@ struct TQuaternion
 	//   SetFromRotationMatrix(FMatrix3f)
 
 
-	inline operator FQuat() const
+	explicit inline operator FQuat() const
 	{
 		FQuat Quat; 
 		Quat.X = (float)X;
@@ -71,7 +73,7 @@ struct TQuaternion
 		Quat.W = (float)W;
 		return Quat;
 	}
-	inline TQuaternion(const FQuat& Quat)
+	explicit inline TQuaternion(const FQuat& Quat)
 	{
 		X = (RealType)Quat.X;
 		Y = (RealType)Quat.Y;
@@ -113,6 +115,16 @@ TQuaternion<RealType>::TQuaternion(const RealType* Values)
 	W = Values[3];
 }
 
+template<typename RealType>
+template<typename RealType2>
+TQuaternion<RealType>::TQuaternion(const TQuaternion<RealType2>& Copy)
+{
+	X = (RealType)Copy.X;
+	Y = (RealType)Copy.Y;
+	Z = (RealType)Copy.Z;
+	W = (RealType)Copy.W;
+}
+
 template<typename RealType>
 TQuaternion<RealType>::TQuaternion(const TQuaternion& Copy)
 {
diff --git a/Engine/Plugins/Experimental/GeometryProcessing/Source/GeometricObjects/Public/RayTypes.h b/Engine/Plugins/Experimental/GeometryProcessing/Source/GeometricObjects/Public/RayTypes.h
--- a/Engine/Plugins/Experimental/GeometryProcessing/Source/GeometricObjects/Public/RayTypes.h
+++ b/Engine/Plugins/Experimental/GeometryProcessing/Source/GeometricObjects/Public/RayTypes.h
@@ -98,14 +98,14 @@ public:
 
 	// conversion operators
 
-	inline operator FRay() const
+	explicit inline operator FRay() const
 	{
 		return FRay((FVector)Origin, (FVector)Direction);
 	}
 	inline TRay3(const FRay & RayIn)
 	{
-		Origin = (FVector3<RealType>)RayIn.Origin;
-		Direction = (FVector3<RealType>)RayIn.Direction;
+		Origin = FVector3<RealType>(RayIn.Origin);
+		Direction = FVector3<RealType>(RayIn.Direction);
 	}
 
 
diff --git a/Engine/Plugins/Experimental/GeometryProcessing/Source/GeometricObjects/Public/TransformTypes.h b/Engine/Plugins/Experimental/GeometryProcessing/Source/GeometricObjects/Public/TransformTypes.h
--- a/Engine/Plugins/Experimental/GeometryProcessing/Source/GeometricObjects/Public/TransformTypes.h
+++ b/Engine/Plugins/Experimental/GeometryProcessing/Source/GeometricObjects/Public/TransformTypes.h
@@ -146,7 +146,6 @@ public:
 		return Rotation * V;
 	}
 
-
 	/**
 	 * Surface Normals are special, their transform is Rotate( Normalize( (1/Scale) * Normal) ) ).
 	 * However 1/Scale requires special handling in case any component is near-zero.
@@ -253,6 +252,53 @@ public:
 
 
 
+
+	// vector-type-conversion variants. This allows applying a float transform to double vector and vice-versa.
+	// Whether this should be allowed is debatable. However in practice it is extremely rare to convert an
+	// entire float transform to a double transform in order to apply to a double vector, which is the only
+	// case where this conversion is an issue
+
+	template<typename RealType2>
+	FVector3<RealType2> TransformPosition(const FVector3<RealType2>& P) const
+	{
+		return FVector3<RealType2>(TransformPosition(FVector3<RealType>(P)));
+	}
+	template<typename RealType2>
+	FVector3<RealType2> TransformVector(const FVector3<RealType2>& V) const
+	{
+		return FVector3<RealType2>(TransformVector(FVector3<RealType>(V)));
+	}
+	template<typename RealType2>
+	FVector3<RealType2> TransformVectorNoScale(const FVector3<RealType2>& V) const
+	{
+		return FVector3<RealType2>(TransformVectorNoScale(FVector3<RealType>(V)));
+	}
+	template<typename RealType2>
+	FVector3<RealType2> TransformNormal(const FVector3<RealType2>& V) const
+	{
+		return FVector3<RealType2>(TransformNormal(FVector3<RealType>(V)));
+	}
+	template<typename RealType2>
+	FVector3<RealType2> InverseTransformPosition(const FVector3<RealType2>& P) const
+	{
+		return FVector3<RealType2>(InverseTransformPosition(FVector3<RealType>(P)));
+	}
+	template<typename RealType2>
+	FVector3<RealType2> InverseTransformVector(const FVector3<RealType2>& V) const
+	{
+		return FVector3<RealType2>(InverseTransformVector(FVector3<RealType>(V)));
+	}
+	template<typename RealType2>
+	FVector3<RealType2> InverseTransformVectorNoScale(const FVector3<RealType2>& V) const
+	{
+		return FVector3<RealType2>(InverseTransformVectorNoScale(FVector3<RealType>(V)));
+	}
+	template<typename RealType2>
+	FVector3<RealType2> InverseTransformNormal(const FVector3<RealType2>& V) const
+	{
+		return FVector3<RealType2>(InverseTransformNormal(FVector3<RealType>(V)));
+	}
+
 };
 typedef TTransform3<float> FTransform3f;
 typedef TTransform3<double> FTransform3d;
diff --git a/Engine/Plugins/Experimental/GeometryProcessing/Source/GeometricObjects/Public/Util/CompactMaps.h b/Engine/Plugins/Experimental/GeometryProcessing/Source/GeometricObjects/Public/Util/CompactMaps.h
new file mode 100644
--- /dev/null
+++ b/Engine/Plugins/Experimental/GeometryProcessing/Source/GeometricObjects/Public/Util/CompactMaps.h
@@ -0,0 +1,148 @@
+// Copyright 1998-2019 Epic Games, Inc. All Rights Reserved.
+
+// Index remapping struct extracted from Dynamic Mesh for more general use
+
+#pragma once
+
+#include "Containers/Array.h"
+
+#include "IndexTypes.h"
+
+/**
+ * Stores index remapping arrays for different types of elements
+ * Should only be used for compacting, and should maintain invariant that Map*[Idx] <= Idx for all maps
+ */
+struct FCompactMaps
+{
+	constexpr static int InvalidID = IndexConstants::InvalidID;
+
+	TArray<int> MapV;
+	TArray<int> MapT;
+	bool bKeepTriangleMap = true; // if set to false, MapT should be empty
+
+	/**
+	 * Set up maps as identity maps.
+	 * Note triangle map will not be set if bKeepTriangleMap is false.
+	 *
+	 * @param MaxVID Vertex map will be created from 0 to MaxVID
+	 * @param MaxTID If bKeepTriangleMap is true, triangle map will be created from 0 to MaxTID
+	 */
+	void SetIdentity(int MaxVID, int MaxTID)
+	{
+		MapV.SetNumUninitialized(MaxVID);
+		for (int ID = 0; ID < MaxVID; ID++)
+		{
+			MapV[ID] = ID;
+		}
+		if (bKeepTriangleMap)
+		{
+			MapT.SetNumUninitialized(MaxTID);
+			for (int ID = 0; ID < MaxTID; ID++)
+			{
+				MapT[ID] = ID;
+			}
+		}
+		else
+		{
+			MapT.Empty();
+		}
+	}
+
+	/**
+	 * Reset all maps, and initialize with InvalidID
+	 * @param MaxVID Size of post-reset vertex map (filled w/ InvalidID)
+	 * @param MaxTID Size of post-reset triangle map (filled w/ InvalidID)
+	 */
+	void Reset(int MaxVID, int MaxTID = 0)
+	{
+		MapV.Reset();
+		MapT.Reset();
+		MapV.SetNumUninitialized(MaxVID);
+		MapT.SetNumUninitialized(MaxTID);
+		for (int VID = 0; VID < MaxVID; VID++)
+		{
+			MapV[VID] = InvalidID;
+		}
+		for (int TID = 0; TID < MaxTID; TID++)
+		{
+			MapT[TID] = InvalidID;
+		}
+	}
+
+	/** Reset all maps, leaving them empty */
+	void Reset()
+	{
+		MapV.Reset();
+		MapT.Reset();
+	}
+
+	/** Set mapping for a vertex */
+	inline void SetVertex(int FromID, int ToID)
+	{
+		checkSlow(FromID >= ToID);
+		MapV[FromID] = ToID;
+	}
+
+	/** Set mapping for a triangle */
+	inline void SetTriangle(int FromID, int ToID)
+	{
+		checkSlow(bKeepTriangleMap);
+		checkSlow(FromID >= ToID);
+		MapT[FromID] = ToID;
+	}
+
+	/** Get mapping for a vertex */
+	inline int GetVertex(int FromID) const
+	{
+		int ToID = MapV[FromID];
+		checkSlow(ToID != InvalidID);
+		return ToID;
+	}
+
+	/** Get mapping for a triangle */
+	inline int GetTriangle(int FromID) const
+	{
+		checkSlow(bKeepTriangleMap);
+		int ToID = MapT[FromID];
+		checkSlow(ToID != InvalidID);
+		return ToID;
+	}
+
+	/** 
+	 * Clear all triangles
+	 *
+	 * @param bPermanent If true, set flag to not create triangle map
+	 */
+	void ClearTriangleMap(bool bPermanent = true)
+	{
+		if (bPermanent)
+		{
+			bKeepTriangleMap = false;
+		}
+		MapT.Empty();
+	}
+
+	/** Check data for validity; for testing */
+	bool Validate() const
+	{
+		for (int Idx = 0; Idx < MapV.Num(); Idx++)
+		{
+			if (MapV[Idx] > Idx)
+			{
+				return false;
+			}
+		}
+		for (int Idx = 0; Idx < MapT.Num(); Idx++)
+		{
+			if (MapT[Idx] > Idx)
+			{
+				return false;
+			}
+		}
+		if (!bKeepTriangleMap && MapT.Num() > 0)
+		{
+			return false;
+		}
+		return true;
+	}
+};
diff --git a/Engine/Plugins/Experimental/GeometryProcessing/Source/GeometricObjects/Public/Util/IndexUtil.h b/Engine/Plugins/Experimental/GeometryProcessing/Source/GeometricObjects/Public/Util/IndexUtil.h
--- a/Engine/Plugins/Experimental/GeometryProcessing/Source/GeometricObjects/Public/Util/IndexUtil.h
+++ b/Engine/Plugins/Experimental/GeometryProcessing/Source/GeometricObjects/Public/Util/IndexUtil.h
@@ -256,7 +256,22 @@ namespace IndexUtil
 	{
 		return FVector3<T>(MapFunc[Val[0]], MapFunc[Val[1]], MapFunc[Val[2]]);
 	}
-
+	
+	/**
+	 * @return false if CheckFn returns false on any element of ToCheck, true otherwise
+	 */
+	template<typename T, typename Func>
+	bool ArrayCheck(const TArray<T>& ToCheck, Func CheckFn)
+	{
+		for (const T& Value : ToCheck)
+		{
+			if (!CheckFn(Value))
+			{
+				return false;
+			}
+		}
+		return true;
+	}
 
 
 	/**
diff --git a/Engine/Plugins/Experimental/GeometryProcessing/Source/GeometricObjects/Public/Util/SmallListSet.h b/Engine/Plugins/Experimental/GeometryProcessing/Source/GeometricObjects/Public/Util/SmallListSet.h
--- a/Engine/Plugins/Experimental/GeometryProcessing/Source/GeometricObjects/Public/Util/SmallListSet.h
+++ b/Engine/Plugins/Experimental/GeometryProcessing/Source/GeometricObjects/Public/Util/SmallListSet.h
@@ -32,36 +32,31 @@ protected:
 	static constexpr int32 BLOCK_LIST_OFFSET = BLOCKSIZE + 1;
 
 	/** mapping from list index to offset into block_store that contains list data */
-	TDynamicVector<int32> ListHeads;
+	TDynamicVector<int32> ListHeads{};
 
 	/** 
 	 * flat buffer used to store per-list linear-memory blocks. 
 	 * blocks are BLOCKSIZE+2 long, elements are [CurrentCount, item0...itemN, LinkedListPtr]
 	 */
-	TDynamicVector<int32> ListBlocks; 
+	TDynamicVector<int32> ListBlocks{};
 
 	/** list of free blocks as indices/offsets into block_store */
-	TDynamicVector<int32> FreeBlocks;
+	TDynamicVector<int32> FreeBlocks{};
 
 	/** number of allocated lists */
-	int32 AllocatedCount = 0;
+	int32 AllocatedCount{0};
 
 	/**
 	 * flat buffer used to store linked-list "spill" elements
 	 * each element is [value, next_ptr]
 	 */
-	TDynamicVector<int32> LinkedListElements;
+	TDynamicVector<int32> LinkedListElements{};
 
 	/** index of first free element in linked_store */
-	int32 FreeHeadIndex;
+	int32 FreeHeadIndex{NullValue};
 
 
 public:
-	FSmallListSet();
-
-	FSmallListSet(const FSmallListSet& copy);
-
-
 	/**
 	 * @return largest current list index
 	 */
@@ -381,4 +376,4 @@ public:
 	}
 
 
-};
\ No newline at end of file
+};
