make smoothing not set nans on isolated vertices

#rnx

diff --git a/Engine/Plugins/Experimental/GeometryProcessing/Source/DynamicMesh/Public/DynamicMesh3.h b/Engine/Plugins/Experimental/GeometryProcessing/Source/DynamicMesh/Public/DynamicMesh3.h
--- a/Engine/Plugins/Experimental/GeometryProcessing/Source/DynamicMesh/Public/DynamicMesh3.h
+++ b/Engine/Plugins/Experimental/GeometryProcessing/Source/DynamicMesh/Public/DynamicMesh3.h
@@ -275,6 +275,11 @@ public:
 	{
 		return VertexRefCounts.IsValid(VertexID);
 	}
+	/** @return true if VertexID is a valid vertex in this mesh AND is used by at least one triangle */
+	inline bool IsReferencedVertex(int VertexID) const
+	{
+		return VertexID >= 0 && VertexID < (int)VertexRefCounts.GetMaxIndex() && VertexRefCounts.GetRawRefCount(VertexID) > 1;
+	}
 	/** @return true if TriangleID is a valid triangle in this mesh */
 	inline bool IsTriangle(int TriangleID) const
 	{
diff --git a/Engine/Plugins/Experimental/MeshModelingToolset/Source/ModelingOperators/Private/SmoothingOps/IterativeSmoothingOp.cpp b/Engine/Plugins/Experimental/MeshModelingToolset/Source/ModelingOperators/Private/SmoothingOps/IterativeSmoothingOp.cpp
--- a/Engine/Plugins/Experimental/MeshModelingToolset/Source/ModelingOperators/Private/SmoothingOps/IterativeSmoothingOp.cpp
+++ b/Engine/Plugins/Experimental/MeshModelingToolset/Source/ModelingOperators/Private/SmoothingOps/IterativeSmoothingOp.cpp
@@ -44,7 +44,7 @@ void FIterativeSmoothingOp::Smooth()
 	for (int k = 0; k < SmoothIterations; ++k)
 	{
 		ParallelFor(NV, [=](int vid) {
-			if (ResultMesh->IsVertex(vid) == false)
+			if (ResultMesh->IsReferencedVertex(vid) == false)
 			{
 				return;
 			}
