GeomProcessing: Fix incorrect math in grid indexing that assumed only positive coordinates were in use. Need to use Floor and not int-cast.
#rnx

diff --git a/Engine/Plugins/Experimental/GeometryProcessing/Source/GeometricObjects/Public/Util/GridIndexing2.h b/Engine/Plugins/Experimental/GeometryProcessing/Source/GeometricObjects/Public/Util/GridIndexing2.h
--- a/Engine/Plugins/Experimental/GeometryProcessing/Source/GeometricObjects/Public/Util/GridIndexing2.h
+++ b/Engine/Plugins/Experimental/GeometryProcessing/Source/GeometricObjects/Public/Util/GridIndexing2.h
@@ -24,7 +24,9 @@ struct TScaleGridIndexer2
 	/** Convert real-valued point to integer grid coordinates */
 	inline FVector2i ToGrid(const FVector2<RealType>& P) const
 	{
-		return FVector2i(int(P.X / CellSize), int(P.Y / CellSize));
+		return FVector2i(
+			(int)TMathUtil<RealType>::Floor(P.X / CellSize),
+			(int)TMathUtil<RealType>::Floor(P.Y / CellSize));
 	}
 
 	/** Convert integer grid coordinates to real-valued point */
@@ -56,27 +58,27 @@ struct TShiftGridIndexer2
 	}
 
 	/** Convert real-valued point to integer grid coordinates */
-	inline FVector2i ToGrid(const FVector2<RealType>& point) const
+	inline FVector2i ToGrid(const FVector2<RealType>& Point) const
 	{
 		return FVector2i(
-			(int)((point.X - Origin.X) / CellSize),
-			(int)((point.Y - Origin.Y) / CellSize));
+			(int)TMathUtil<RealType>::Floor((Point.X - Origin.X) / CellSize),
+			(int)TMathUtil<RealType>::Floor((Point.Y - Origin.Y) / CellSize));
 	}
 
 	/** Convert integer grid coordinates to real-valued point */
-	inline FVector2<RealType> FromGrid(const FVector2i& gridpoint) const
+	inline FVector2<RealType> FromGrid(const FVector2i& GridPoint) const
 	{
 		return FVector2<RealType>(
-			((RealType)gridpoint.X * CellSize) + Origin.X,
-			((RealType)gridpoint.Y * CellSize) + Origin.Y);
+			((RealType)GridPoint.X * CellSize) + Origin.X,
+			((RealType)GridPoint.Y * CellSize) + Origin.Y);
 	}
 
 	/** Convert real-valued grid coordinates to real-valued point */
-	inline FVector2<RealType> FromGrid(const FVector2<RealType>& gridpointf) const
+	inline FVector2<RealType> FromGrid(const FVector2<RealType>& RealGridPoint) const
 	{
 		return FVector2<RealType>(
-			((RealType)gridpointf.X * CellSize) + Origin.X,
-			((RealType)gridpointf.Y * CellSize) + Origin.Y);
+			((RealType)RealGridPoint.X * CellSize) + Origin.X,
+			((RealType)RealGridPoint.Y * CellSize) + Origin.Y);
 	}
 };
 typedef TShiftGridIndexer2<float> FShiftGridIndexer2f;
diff --git a/Engine/Plugins/Experimental/GeometryProcessing/Source/GeometricObjects/Public/Util/GridIndexing3.h b/Engine/Plugins/Experimental/GeometryProcessing/Source/GeometricObjects/Public/Util/GridIndexing3.h
--- a/Engine/Plugins/Experimental/GeometryProcessing/Source/GeometricObjects/Public/Util/GridIndexing3.h
+++ b/Engine/Plugins/Experimental/GeometryProcessing/Source/GeometricObjects/Public/Util/GridIndexing3.h
@@ -18,6 +18,7 @@ struct TScaleGridIndexer3
 
 	TScaleGridIndexer3() : CellSize((RealType)1) 
 	{
+		ensure(CellSize >= TMathUtil<RealType>::ZeroTolerance);
 	}
 	
 	TScaleGridIndexer3(RealType CellSize) : CellSize(CellSize) 
@@ -27,7 +28,10 @@ struct TScaleGridIndexer3
 	/** Convert real-valued point to integer grid coordinates */
 	inline FVector3i ToGrid(const FVector3<RealType>& P) const
 	{
-		return FVector3i((int)(P.X / CellSize), (int)(P.Y / CellSize), (int)(P.Z / CellSize));
+		return FVector3i(
+			(int)TMathUtil<RealType>::Floor(P.X / CellSize),
+			(int)TMathUtil<RealType>::Floor(P.Y / CellSize),
+			(int)TMathUtil<RealType>::Floor(P.Z / CellSize));
 	}
 
 	/** Convert integer grid coordinates to real-valued point */
@@ -55,20 +59,22 @@ struct TShiftGridIndexer3
 	TShiftGridIndexer3() 
 		: CellSize((RealType)1), Origin(FVector3<RealType>::Zero())
 	{
+		ensure(CellSize >= TMathUtil<RealType>::ZeroTolerance);
 	}
 
 	TShiftGridIndexer3(const FVector3<RealType>& origin, RealType cellSize)
 		: CellSize(cellSize), Origin(origin)
 	{
+		ensure(CellSize >= TMathUtil<RealType>::ZeroTolerance);
 	}
 
 	/** Convert real-valued point to integer grid coordinates */
 	inline FVector3i ToGrid(const FVector3<RealType>& point) const
 	{
 		return FVector3i(
-			(int)((point.X - Origin.X) / CellSize),
-			(int)((point.Y - Origin.Y) / CellSize),
-			(int)((point.Z - Origin.Z) / CellSize));
+			(int)TMathUtil<RealType>::Floor((point.X - Origin.X) / CellSize),
+			(int)TMathUtil<RealType>::Floor((point.Y - Origin.Y) / CellSize),
+			(int)TMathUtil<RealType>::Floor((point.Z - Origin.Z) / CellSize));
 	}
 
 	/** Convert integer grid coordinates to real-valued point */
@@ -81,12 +87,12 @@ struct TShiftGridIndexer3
 	}
 
 	/** Convert real-valued grid coordinates to real-valued point */
-	inline FVector3<RealType> FromGrid(const FVector3<RealType>& gridpointf)  const
+	inline FVector3<RealType> FromGrid(const FVector3<RealType>& RealGridPoint)  const
 	{
-		return FVector3d(
-			((RealType)gridpointf.X * CellSize) + Origin.X,
-			((RealType)gridpointf.Y * CellSize) + Origin.Y,
-			((RealType)gridpointf.Z * CellSize) + Origin.Z);
+		return FVector3<RealType>(
+			((RealType)RealGridPoint.X * CellSize) + Origin.X,
+			((RealType)RealGridPoint.Y * CellSize) + Origin.Y,
+			((RealType)RealGridPoint.Z * CellSize) + Origin.Z);
 	}
 };
 typedef TShiftGridIndexer3<float> FShiftGridIndexer3f;
