GeometryProcessing: add corner index to MakeVertex() function, this allows us to override this function to generate rect-mesh on arbitrary 3d quads
#rnx

diff --git a/Engine/Plugins/Experimental/GeometryProcessing/Source/GeometricObjects/Private/Generators/RectangleMeshGenerator.cpp b/Engine/Plugins/Experimental/GeometryProcessing/Source/GeometricObjects/Private/Generators/RectangleMeshGenerator.cpp
--- a/Engine/Plugins/Experimental/GeometryProcessing/Source/GeometricObjects/Private/Generators/RectangleMeshGenerator.cpp
+++ b/Engine/Plugins/Experimental/GeometryProcessing/Source/GeometricObjects/Private/Generators/RectangleMeshGenerator.cpp
@@ -29,10 +29,10 @@ FMeshShapeGenerator& FRectangleMeshGenerator::Generate()
 	SetBufferSizes(TotalNumVertices, TotalNumTriangles, TotalNumVertices, TotalNumVertices);
 
 	// corner vertices
-	FVector3d v00 = MakeVertex(-Width / 2.0f, -Height / 2.0f);
-	FVector3d v01 = MakeVertex(Width / 2.0f, -Height / 2.0f);
-	FVector3d v11 = MakeVertex(Width / 2.0f, Height / 2.0f);
-	FVector3d v10 = MakeVertex(-Width / 2.0f, Height / 2.0f);
+	FVector3d v00 = MakeVertex(0, -Width / 2.0f, -Height / 2.0f);
+	FVector3d v01 = MakeVertex(1, Width / 2.0f, -Height / 2.0f);
+	FVector3d v11 = MakeVertex(2, Width / 2.0f, Height / 2.0f);
+	FVector3d v10 = MakeVertex(3, -Width / 2.0f, Height / 2.0f);
 
 	// corner UVs
 	float uvleft = 0.0f, uvright = 1.0f, uvbottom = 0.0f, uvtop = 1.0f;
@@ -144,10 +144,10 @@ FMeshShapeGenerator& FRoundedRectangleMeshGenerator::Generate()
 	
 
 	// corner vertices
-	FVector3d v00 = MakeVertex(-TotWidth / 2.0f, -TotHeight / 2.0f);
-	FVector3d v01 = MakeVertex(TotWidth / 2.0f, -TotHeight / 2.0f);
-	FVector3d v11 = MakeVertex(TotWidth / 2.0f, TotHeight / 2.0f);
-	FVector3d v10 = MakeVertex(-TotWidth / 2.0f, TotHeight / 2.0f);
+	FVector3d v00 = MakeVertex(0, -TotWidth / 2.0f, -TotHeight / 2.0f);
+	FVector3d v01 = MakeVertex(1, TotWidth / 2.0f, -TotHeight / 2.0f);
+	FVector3d v11 = MakeVertex(2, TotWidth / 2.0f, TotHeight / 2.0f);
+	FVector3d v10 = MakeVertex(3, -TotWidth / 2.0f, TotHeight / 2.0f);
 
 	// corner UVs
 	float uvleft = 0.0f, uvright = 1.0f, uvbottom = 0.0f, uvtop = 1.0f;
diff --git a/Engine/Plugins/Experimental/GeometryProcessing/Source/GeometricObjects/Public/Generators/RectangleMeshGenerator.h b/Engine/Plugins/Experimental/GeometryProcessing/Source/GeometricObjects/Public/Generators/RectangleMeshGenerator.h
--- a/Engine/Plugins/Experimental/GeometryProcessing/Source/GeometricObjects/Public/Generators/RectangleMeshGenerator.h
+++ b/Engine/Plugins/Experimental/GeometryProcessing/Source/GeometricObjects/Public/Generators/RectangleMeshGenerator.h
@@ -40,7 +40,7 @@ public:
 	virtual FMeshShapeGenerator& Generate() override;
 
 	/** Create vertex at position under IndicesMap, shifted to Origin*/
-	virtual FVector3d MakeVertex(double x, double y)
+	virtual FVector3d MakeVertex(int CornerIndex, double x, double y)
 	{
 		FVector3d v(0, 0, 0);
 		v[IndicesMap.A] = x;
