GeometryProcessing: need to define nonsense TMathUtilConstants<int> so that TMathUtil<int> works, because some vector types (eg Vector3i, Box3i) use TMathUtil::Abs/Sign/etc
#rnx

diff --git a/Engine/Plugins/Experimental/GeometryProcessing/Source/GeometricObjects/Public/MathUtil.h b/Engine/Plugins/Experimental/GeometryProcessing/Source/GeometricObjects/Public/MathUtil.h
--- a/Engine/Plugins/Experimental/GeometryProcessing/Source/GeometricObjects/Public/MathUtil.h
+++ b/Engine/Plugins/Experimental/GeometryProcessing/Source/GeometricObjects/Public/MathUtil.h
@@ -89,6 +89,30 @@ struct TMathUtilConstants<double>
 	GEOMETRICOBJECTS_API static constexpr double InvSqrt3 = 1.0 / Sqrt3;
 };
 
+
+// we use TMathUtil<int> so we need to define these nonsense constants
+template<>
+struct TMathUtilConstants<int>
+{
+	GEOMETRICOBJECTS_API static constexpr int Epsilon = 0;
+	GEOMETRICOBJECTS_API static constexpr int ZeroTolerance = 0;
+	GEOMETRICOBJECTS_API static constexpr int MaxReal = INT_MAX;
+	GEOMETRICOBJECTS_API static constexpr int Pi = 3;
+	GEOMETRICOBJECTS_API static constexpr int FourPi = 4 * Pi;
+	GEOMETRICOBJECTS_API static constexpr int TwoPi = 2 * Pi;
+	GEOMETRICOBJECTS_API static constexpr int HalfPi = 1;
+	GEOMETRICOBJECTS_API static constexpr double InvPi = 1;
+	GEOMETRICOBJECTS_API static constexpr double InvTwoPi = 1;
+	GEOMETRICOBJECTS_API static constexpr double DegToRad = 1;
+	GEOMETRICOBJECTS_API static constexpr double RadToDeg = 1;
+	GEOMETRICOBJECTS_API static constexpr double Sqrt2 = 1;
+	GEOMETRICOBJECTS_API static constexpr double InvSqrt2 = 1;
+	GEOMETRICOBJECTS_API static constexpr double Sqrt3 = 2;
+	GEOMETRICOBJECTS_API static constexpr double InvSqrt3 = 1;
+};
+
+
+
 template<typename RealType>
 class TMathUtil : public TMathUtilConstants<RealType>
 {
