fix linux build
#rnx

diff --git a/Engine/Plugins/Experimental/GeometryProcessing/Source/GeometricObjects/Public/MathUtil.h b/Engine/Plugins/Experimental/GeometryProcessing/Source/GeometricObjects/Public/MathUtil.h
--- a/Engine/Plugins/Experimental/GeometryProcessing/Source/GeometricObjects/Public/MathUtil.h
+++ b/Engine/Plugins/Experimental/GeometryProcessing/Source/GeometricObjects/Public/MathUtil.h
@@ -92,23 +92,23 @@ struct TMathUtilConstants<double>
 
 // we use TMathUtil<int> so we need to define these nonsense constants
 template<>
-struct TMathUtilConstants<int>
+struct TMathUtilConstants<int32>
 {
-	GEOMETRICOBJECTS_API static constexpr int Epsilon = 0;
-	GEOMETRICOBJECTS_API static constexpr int ZeroTolerance = 0;
-	GEOMETRICOBJECTS_API static constexpr int MaxReal = INT_MAX;
-	GEOMETRICOBJECTS_API static constexpr int Pi = 3;
-	GEOMETRICOBJECTS_API static constexpr int FourPi = 4 * Pi;
-	GEOMETRICOBJECTS_API static constexpr int TwoPi = 2 * Pi;
-	GEOMETRICOBJECTS_API static constexpr int HalfPi = 1;
-	GEOMETRICOBJECTS_API static constexpr double InvPi = 1;
-	GEOMETRICOBJECTS_API static constexpr double InvTwoPi = 1;
-	GEOMETRICOBJECTS_API static constexpr double DegToRad = 1;
-	GEOMETRICOBJECTS_API static constexpr double RadToDeg = 1;
-	GEOMETRICOBJECTS_API static constexpr double Sqrt2 = 1;
-	GEOMETRICOBJECTS_API static constexpr double InvSqrt2 = 1;
-	GEOMETRICOBJECTS_API static constexpr double Sqrt3 = 2;
-	GEOMETRICOBJECTS_API static constexpr double InvSqrt3 = 1;
+	GEOMETRICOBJECTS_API static constexpr int32 Epsilon = 0;
+	GEOMETRICOBJECTS_API static constexpr int32 ZeroTolerance = 0;
+	GEOMETRICOBJECTS_API static constexpr int32 MaxReal = 2147483647;
+	GEOMETRICOBJECTS_API static constexpr int32 Pi = 3;
+	GEOMETRICOBJECTS_API static constexpr int32 FourPi = 4 * Pi;
+	GEOMETRICOBJECTS_API static constexpr int32 TwoPi = 2 * Pi;
+	GEOMETRICOBJECTS_API static constexpr int32 HalfPi = 1;
+	GEOMETRICOBJECTS_API static constexpr int32 InvPi = 1;
+	GEOMETRICOBJECTS_API static constexpr int32 InvTwoPi = 1;
+	GEOMETRICOBJECTS_API static constexpr int32 DegToRad = 1;
+	GEOMETRICOBJECTS_API static constexpr int32 RadToDeg = 1;
+	GEOMETRICOBJECTS_API static constexpr int32 Sqrt2 = 1;
+	GEOMETRICOBJECTS_API static constexpr int32 InvSqrt2 = 1;
+	GEOMETRICOBJECTS_API static constexpr int32 Sqrt3 = 2;
+	GEOMETRICOBJECTS_API static constexpr int32 InvSqrt3 = 1;
 };
 
 
