minor fixes - export groups-color function, recalculate bounding box on fast vertex update in SimpleDynamicMeshComponent
#rnx

diff --git a/Engine/Plugins/Experimental/GeometryProcessing/Source/GeometricObjects/Private/Util/ColorConstants.cpp b/Engine/Plugins/Experimental/GeometryProcessing/Source/GeometricObjects/Private/Util/ColorConstants.cpp
new file mode 100644
--- /dev/null
+++ b/Engine/Plugins/Experimental/GeometryProcessing/Source/GeometricObjects/Private/Util/ColorConstants.cpp
@@ -0,0 +1,18 @@
+// Copyright 1998-2019 Epic Games, Inc. All Rights Reserved.
+
+#include "Util/ColorConstants.h"
+
+
+FColor LinearColors::SelectFColor(int32 Index)
+{
+	static const FColor ColorMap[] = {
+		SpringGreen3b(), Plum3b(), Khaki3b(),
+		PaleGreen3b(), LightSteelBlue3b(), Aquamarine3b(),
+		Salmon3b(), Goldenrod3b(), LightSeaGreen3b(),
+		IndianRed3b(), DarkSalmon3b(), Coral3b(),
+		Burlywood3b(), GreenYellow3b(), Lavender3b(),
+		MediumAquamarine3b(), Thistle3b(), Wheat3b(),
+		LightSkyBlue3b(), LightPink3b(), MediumSpringGreen3b()
+	};
+	return (Index < 0) ? FColor::White : ColorMap[Index % (7 * 3)];
+}
\ No newline at end of file
diff --git a/Engine/Plugins/Experimental/GeometryProcessing/Source/GeometricObjects/Public/Util/ColorConstants.h b/Engine/Plugins/Experimental/GeometryProcessing/Source/GeometricObjects/Public/Util/ColorConstants.h
--- a/Engine/Plugins/Experimental/GeometryProcessing/Source/GeometricObjects/Public/Util/ColorConstants.h
+++ b/Engine/Plugins/Experimental/GeometryProcessing/Source/GeometricObjects/Public/Util/ColorConstants.h
@@ -26,7 +26,7 @@ namespace LinearColors
 
 #define ColorConstant(ColorName, R, G, B) \
 	template <typename VectorType> VectorType ColorName##3f() { return MakeColor3f<VectorType>(R, G, B); } \
-	FColor ColorName##3b() { return MakeColor3f<FLinearColor>(R, G, B).ToFColor(false); }
+	inline FColor ColorName##3b() { return MakeColor3f<FLinearColor>(R, G, B).ToFColor(false); }
 
 	ColorConstant(Black, 0.0f, 0.0f, 0.0f);
 	ColorConstant(White, 1.0f, 1.0f, 1.0f);
@@ -221,18 +221,6 @@ namespace LinearColors
 	 * @param Index arbitrary integer
 	 * @return Color selected from a fixed set, or White if Index < 0
 	 */
-	FColor SelectFColor(int32 Index)
-	{
-		static const FColor ColorMap[] = {
-			SpringGreen3b(), Plum3b(), Khaki3b(),
-			PaleGreen3b(), LightSteelBlue3b(), Aquamarine3b(),
-			Salmon3b(), Goldenrod3b(), LightSeaGreen3b(),
-			IndianRed3b(), DarkSalmon3b(), Coral3b(),
-			Burlywood3b(), GreenYellow3b(), Lavender3b(),
-			MediumAquamarine3b(), Thistle3b(), Wheat3b(),
-			LightSkyBlue3b(), LightPink3b(), MediumSpringGreen3b()
-		};
-		return (Index < 0) ? FColor::White : ColorMap[Index % (7*3)];
-	}
+	GEOMETRICOBJECTS_API FColor SelectFColor(int32 Index);
 
 }
\ No newline at end of file
diff --git a/Engine/Plugins/Experimental/MeshModelingToolset/Source/ModelingComponents/Private/SimpleDynamicMeshComponent.cpp b/Engine/Plugins/Experimental/MeshModelingToolset/Source/ModelingComponents/Private/SimpleDynamicMeshComponent.cpp
--- a/Engine/Plugins/Experimental/MeshModelingToolset/Source/ModelingComponents/Private/SimpleDynamicMeshComponent.cpp
+++ b/Engine/Plugins/Experimental/MeshModelingToolset/Source/ModelingComponents/Private/SimpleDynamicMeshComponent.cpp
@@ -163,6 +163,8 @@ void USimpleDynamicMeshComponent::NotifyMeshUpdated()
 
 void USimpleDynamicMeshComponent::FastNotifyColorsUpdated()
 {
+	MarkRenderStateDirty();
+
 	if (CurrentProxy != nullptr)
 	{
 		CurrentProxy->FastUpdateVertices(false, false, true);
@@ -177,6 +179,9 @@ void USimpleDynamicMeshComponent::FastNotifyColorsUpdated()
 
 void USimpleDynamicMeshComponent::FastNotifyPositionsUpdated()
 {
+	MarkRenderStateDirty();
+	UpdateBounds();
+
 	if (CurrentProxy != nullptr)
 	{
 		CurrentProxy->FastUpdateVertices(true, false, false);
@@ -238,11 +243,8 @@ FColor USimpleDynamicMeshComponent::GetTriangleColor(const FDynamicMesh3* MeshIn
 FBoxSphereBounds USimpleDynamicMeshComponent::CalcBounds(const FTransform& LocalToWorld) const
 {
 	// Bounds are tighter if the box is generated from pre-transformed vertices.
-	FBox BoundingBox(ForceInit);
-	for ( FVector3d Vertex : Mesh->VerticesItr() ) 
-	{
-		BoundingBox += LocalToWorld.TransformPosition(Vertex);
-	}
+	FAxisAlignedBox3d MeshAABB = Mesh->GetCachedBounds();
+	FBox BoundingBox = (FBox)MeshAABB;
 
 	FBoxSphereBounds NewBounds;
 	NewBounds.BoxExtent = BoundingBox.GetExtent();
