GeometryProcessing: make AABBTree mesh a const pointer, Add MeshAdapterUtil::MakePointsAdapter
#rnx

diff --git a/Engine/Plugins/Experimental/GeometryProcessing/Source/DynamicMesh/Private/MeshAdapterUtil.cpp b/Engine/Plugins/Experimental/GeometryProcessing/Source/DynamicMesh/Private/MeshAdapterUtil.cpp
--- a/Engine/Plugins/Experimental/GeometryProcessing/Source/DynamicMesh/Private/MeshAdapterUtil.cpp
+++ b/Engine/Plugins/Experimental/GeometryProcessing/Source/DynamicMesh/Private/MeshAdapterUtil.cpp
@@ -2,6 +2,23 @@
 
 #include "MeshAdapterUtil.h"
 
+FPointSetAdapterd MeshAdapterUtil::MakePointsAdapter(const FDynamicPointSet3d* PointSet)
+{
+	FPointSetAdapterd Adapter;
+	Adapter.MaxPointID = [PointSet]() { return PointSet->MaxVertexID(); };
+	Adapter.PointCount = [PointSet]() { return PointSet->VertexCount(); };
+	Adapter.IsPoint = [PointSet](int Idx) { return PointSet->IsVertex(Idx); };
+	Adapter.GetPoint = [PointSet](int Idx) { return PointSet->GetVertex(Idx); };
+	Adapter.Timestamp = [PointSet] { return PointSet->GetTimestamp(); };
+
+	Adapter.HasNormals = [PointSet] { return false; };
+	Adapter.GetPointNormal = [PointSet](int Idx) {return FVector3d(0,0,1); };
+
+	return Adapter;
+}
+
+
+
 FPointSetAdapterd MeshAdapterUtil::MakeVerticesAdapter(const FDynamicMesh3* Mesh)
 {
 	FPointSetAdapterd Adapter;
diff --git a/Engine/Plugins/Experimental/GeometryProcessing/Source/DynamicMesh/Public/MeshAdapterUtil.h b/Engine/Plugins/Experimental/GeometryProcessing/Source/DynamicMesh/Public/MeshAdapterUtil.h
--- a/Engine/Plugins/Experimental/GeometryProcessing/Source/DynamicMesh/Public/MeshAdapterUtil.h
+++ b/Engine/Plugins/Experimental/GeometryProcessing/Source/DynamicMesh/Public/MeshAdapterUtil.h
@@ -4,6 +4,7 @@
 
 #include "CoreTypes.h"
 #include "DynamicMesh3.h"
+#include "DynamicPointSet3.h"
 #include "PointSetAdapter.h"
 #include "MeshAdapter.h"
 
@@ -33,6 +34,11 @@ namespace MeshAdapterUtil
 	 */
 	FPointSetAdapterd DYNAMICMESH_API MakeVerticesAdapter(const FDynamicMesh3* Mesh);
 
+	/**
+	 * @return PointSet points as a point set
+	 */
+	FPointSetAdapterd DYNAMICMESH_API MakePointsAdapter(const FDynamicPointSet3d* PointSet);
+
 	/**
 	 * @return Mesh triangle centroids as a point set
 	 */
diff --git a/Engine/Plugins/Experimental/GeometryProcessing/Source/GeometricObjects/Public/Generators/BoxSphereGenerator.h b/Engine/Plugins/Experimental/GeometryProcessing/Source/GeometricObjects/Public/Generators/BoxSphereGenerator.h
--- a/Engine/Plugins/Experimental/GeometryProcessing/Source/GeometricObjects/Public/Generators/BoxSphereGenerator.h
+++ b/Engine/Plugins/Experimental/GeometryProcessing/Source/GeometricObjects/Public/Generators/BoxSphereGenerator.h
@@ -44,9 +44,9 @@ public:
 				double y = V.Dot(AxY) / Box.Extents.Y;
 				double z = V.Dot(AxZ) / Box.Extents.Z;
 				double x2 = x * x, y2 = y * y, z2 = z * z;
-				double sx = x * FMath::Sqrt(1.0 - y2*0.5 - z2*0.5 + y2*z2/3.0);
-				double sy = y * FMath::Sqrt(1.0 - x2*0.5 - z2*0.5 + x2*z2/3.0);
-				double sz = z * FMath::Sqrt(1.0 - x2*0.5 - y2*0.5 + x2*y2/3.0);
+				double sx = x * FMathd::Sqrt(1.0 - y2*0.5 - z2*0.5 + y2*z2/3.0);
+				double sy = y * FMathd::Sqrt(1.0 - x2*0.5 - z2*0.5 + x2*z2/3.0);
+				double sz = z * FMathd::Sqrt(1.0 - x2*0.5 - y2*0.5 + x2*y2/3.0);
 				V = sx*AxX + sy*AxY + sz*AxZ;
 			}
 			V.Normalize();
diff --git a/Engine/Plugins/Experimental/GeometryProcessing/Source/GeometricObjects/Public/Spatial/FastWinding.h b/Engine/Plugins/Experimental/GeometryProcessing/Source/GeometricObjects/Public/Spatial/FastWinding.h
--- a/Engine/Plugins/Experimental/GeometryProcessing/Source/GeometricObjects/Public/Spatial/FastWinding.h
+++ b/Engine/Plugins/Experimental/GeometryProcessing/Source/GeometricObjects/Public/Spatial/FastWinding.h
@@ -166,11 +166,12 @@ namespace FastTriWinding
 	}
 } // namespace FastTriWinding
 
+
+
 /**
- *  Fast Mesh Winding Number computation
- *
- *  TODO: this is tightly coupled to the internal guts of TMeshAABBTree3 (because it is ported from code where it *was* the guts of the AABBTree)
- *			we should probably at least try to decouple it some
+ * Fast Mesh Winding Number extension to a TMeshAABBTree3.
+ * This class is an "add-on" to the AABBTree, that can compute the Fast Mesh Winding Number.
+ * This calculation requires a precomputation pass where information is cached at each tree node.
  */
 template <class TriangleMeshType>
 class TFastWindingTree
@@ -228,7 +229,10 @@ public:
 	}
 
 	/**
-	 *  Fast approximation of winding number using far-field approximations
+	 * Fast approximation of winding number using far-field approximations.
+	 * On a closed mesh the winding number will be 1 or more inside (depending on number of "winds").
+	 * Outside a closed mesh the winding number will be zero.
+	 * On an open mesh, the above holds near the mesh but in the "hole" areas the value will smoothly blend from 1 to 0 over a band of width dependent on the hole extent
 	 */
 	double FastWindingNumber(const FVector3d& P)
 	{
diff --git a/Engine/Plugins/Experimental/GeometryProcessing/Source/GeometricObjects/Public/Spatial/MeshAABBTree3.h b/Engine/Plugins/Experimental/GeometryProcessing/Source/GeometricObjects/Public/Spatial/MeshAABBTree3.h
--- a/Engine/Plugins/Experimental/GeometryProcessing/Source/GeometricObjects/Public/Spatial/MeshAABBTree3.h
+++ b/Engine/Plugins/Experimental/GeometryProcessing/Source/GeometricObjects/Public/Spatial/MeshAABBTree3.h
@@ -19,7 +19,7 @@ class TMeshAABBTree3 : public IMeshSpatial
 	friend class TFastWindingTree<TriangleMeshType>;
 
 protected:
-	TriangleMeshType* Mesh;
+	const TriangleMeshType* Mesh;
 	int MeshTimestamp = -1;
 	int TopDownLeafMaxTriCount = 4;
 
@@ -36,12 +36,12 @@ public:
 		Mesh = nullptr;
 	}
 
-	TMeshAABBTree3(TriangleMeshType* SourceMesh, bool bAutoBuild = true)
+	TMeshAABBTree3(const TriangleMeshType* SourceMesh, bool bAutoBuild = true)
 	{
 		SetMesh(SourceMesh, bAutoBuild);
 	}
 
-	void SetMesh(TriangleMeshType* SourceMesh, bool bAutoBuild = true)
+	void SetMesh(const TriangleMeshType* SourceMesh, bool bAutoBuild = true)
 	{
 		Mesh = SourceMesh;
 		MeshTimestamp = -1;
