fix incorrect type in template function
#rnx

diff --git a/Engine/Plugins/Experimental/GeometryProcessing/Source/DynamicMesh/Private/MeshAdapterUtil.cpp b/Engine/Plugins/Experimental/GeometryProcessing/Source/DynamicMesh/Private/MeshAdapterUtil.cpp
--- a/Engine/Plugins/Experimental/GeometryProcessing/Source/DynamicMesh/Private/MeshAdapterUtil.cpp
+++ b/Engine/Plugins/Experimental/GeometryProcessing/Source/DynamicMesh/Private/MeshAdapterUtil.cpp
@@ -12,7 +12,7 @@ FPointSetAdapterd MeshAdapterUtil::MakePointsAdapter(const FDynamicPointSet3d* P
 	Adapter.Timestamp = [PointSet] { return PointSet->GetTimestamp(); };
 
 	Adapter.HasNormals = [PointSet] { return false; };
-	Adapter.GetPointNormal = [PointSet](int Idx) {return FVector3d(0,0,1); };
+	Adapter.GetPointNormal = [PointSet](int Idx) {return FVector3f(0,0,1); };
 
 	return Adapter;
 }
