GeometryProcessing: add TDynamicMeshScalarTriangleAttribute subclass to simplify scalar-valued attribute code
#rnx

diff --git a/Engine/Plugins/Experimental/GeometryProcessing/Source/DynamicMesh/Public/DynamicMeshAttributeSet.h b/Engine/Plugins/Experimental/GeometryProcessing/Source/DynamicMesh/Public/DynamicMeshAttributeSet.h
--- a/Engine/Plugins/Experimental/GeometryProcessing/Source/DynamicMesh/Public/DynamicMeshAttributeSet.h
+++ b/Engine/Plugins/Experimental/GeometryProcessing/Source/DynamicMesh/Public/DynamicMeshAttributeSet.h
@@ -15,7 +15,7 @@ typedef TDynamicMeshVectorOverlay<float, 2, FVector2f> FDynamicMeshUVOverlay;
 typedef TDynamicMeshVectorOverlay<float, 3, FVector3f> FDynamicMeshNormalOverlay;
 
 /** Standard per-triangle integer material ID */
-typedef TDynamicMeshTriangleAttribute<int32, 1> FDynamicMeshMaterialAttribute;
+typedef TDynamicMeshScalarTriangleAttribute<int32> FDynamicMeshMaterialAttribute;
 
 /**
  * FDynamicMeshAttributeSet manages a set of extended attributes for a FDynamicMesh3.
diff --git a/Engine/Plugins/Experimental/GeometryProcessing/Source/DynamicMesh/Public/DynamicMeshTriangleAttribute.h b/Engine/Plugins/Experimental/GeometryProcessing/Source/DynamicMesh/Public/DynamicMeshTriangleAttribute.h
--- a/Engine/Plugins/Experimental/GeometryProcessing/Source/DynamicMesh/Public/DynamicMeshTriangleAttribute.h
+++ b/Engine/Plugins/Experimental/GeometryProcessing/Source/DynamicMesh/Public/DynamicMeshTriangleAttribute.h
@@ -287,3 +287,41 @@ bool FDynamicMeshTriangleAttributeChange<AttribValueType, AttribDimension>::Appl
 	return true;
 }
 
+
+
+/**
+ * TDynamicMeshScalarTriangleAttribute is an extension of TDynamicMeshTriangleAttribute for scalar-valued attributes.
+ * Adds some convenience functions to simplify get/set code.
+ */
+template<typename RealType>
+class TDynamicMeshScalarTriangleAttribute : public TDynamicMeshTriangleAttribute<RealType, 1>
+{
+public:
+	using BaseType = TDynamicMeshTriangleAttribute<RealType, 1>;
+	using BaseType::SetNewValue;
+	using BaseType::GetValue;
+	using BaseType::SetValue;
+
+	TDynamicMeshScalarTriangleAttribute() : BaseType()
+	{
+	}
+
+	TDynamicMeshScalarTriangleAttribute(FDynamicMesh3* ParentMeshIn) : BaseType(ParentMeshIn)
+	{
+	}
+
+	inline void SetNewValue(int NewTriangleID, RealType Value)
+	{
+		AttribValues.InsertAt(Value, NewTriangleID);
+	}
+
+	inline RealType GetValue(int TriangleID) const
+	{
+		return AttribValues[TriangleID];
+	}
+
+	inline void SetValue(int TriangleID, RealType Value)
+	{
+		AttribValues[TriangleID] = Value;
+	}
+};
