try to resolve CIS strict build failure
#rnx

diff --git a/Engine/Plugins/Experimental/GeometryProcessing/Source/DynamicMesh/Public/DynamicMeshTriangleAttribute.h b/Engine/Plugins/Experimental/GeometryProcessing/Source/DynamicMesh/Public/DynamicMeshTriangleAttribute.h
--- a/Engine/Plugins/Experimental/GeometryProcessing/Source/DynamicMesh/Public/DynamicMeshTriangleAttribute.h
+++ b/Engine/Plugins/Experimental/GeometryProcessing/Source/DynamicMesh/Public/DynamicMeshTriangleAttribute.h
@@ -312,16 +312,16 @@ public:
 
 	inline void SetNewValue(int NewTriangleID, RealType Value)
 	{
-		AttribValues.InsertAt(Value, NewTriangleID);
+		this->AttribValues.InsertAt(Value, NewTriangleID);
 	}
 
 	inline RealType GetValue(int TriangleID) const
 	{
-		return AttribValues[TriangleID];
+		return this->AttribValues[TriangleID];
 	}
 
 	inline void SetValue(int TriangleID, RealType Value)
 	{
-		AttribValues[TriangleID] = Value;
+		this->AttribValues[TriangleID] = Value;
 	}
 };
