cis fix for signed/unsigned comparison in dmesh attribs

#rnx

diff --git a/Engine/Plugins/Experimental/GeometryProcessing/Source/DynamicMesh/Public/DynamicMeshTriangleAttribute.h b/Engine/Plugins/Experimental/GeometryProcessing/Source/DynamicMesh/Public/DynamicMeshTriangleAttribute.h
--- a/Engine/Plugins/Experimental/GeometryProcessing/Source/DynamicMesh/Public/DynamicMeshTriangleAttribute.h
+++ b/Engine/Plugins/Experimental/GeometryProcessing/Source/DynamicMesh/Public/DynamicMeshTriangleAttribute.h
@@ -261,10 +261,14 @@ public:
 
 	inline void ResizeAttribStoreIfNeeded(int TriangleID)
 	{
-		int k = TriangleID * AttribDimension;
-		if (k+AttribDimension > AttribValues.Num())
+		if (!ensure(TriangleID >= 0))
+		{
+			return;
+		}
+		size_t NeededSize = size_t((TriangleID+1) * AttribDimension);
+		if (NeededSize > AttribValues.Num())
 		{
-			AttribValues.Resize(k + AttribDimension, GetDefaultAttributeValue());
+			AttribValues.Resize(NeededSize, GetDefaultAttributeValue());
 		}
 	}
 
diff --git a/Engine/Plugins/Experimental/GeometryProcessing/Source/DynamicMesh/Public/DynamicVertexAttribute.h b/Engine/Plugins/Experimental/GeometryProcessing/Source/DynamicMesh/Public/DynamicVertexAttribute.h
--- a/Engine/Plugins/Experimental/GeometryProcessing/Source/DynamicMesh/Public/DynamicVertexAttribute.h
+++ b/Engine/Plugins/Experimental/GeometryProcessing/Source/DynamicMesh/Public/DynamicVertexAttribute.h
@@ -204,10 +204,14 @@ public:
 
 	inline void ResizeAttribStoreIfNeeded(int VertexID)
 	{
-		int k = VertexID * AttribDimension;
-		if (k+AttribDimension > AttribValues.Num())
+		if (!ensure(VertexID >= 0))
+		{
+			return;
+		}
+		size_t NeededSize = size_t((VertexID+1) * AttribDimension);
+		if (NeededSize > AttribValues.Num())
 		{
-			AttribValues.Resize(k + AttribDimension, GetDefaultAttributeValue());
+			AttribValues.Resize(NeededSize, GetDefaultAttributeValue());
 		}
 	}
 
