Fix for warning after upgrading to PVS-Studio 7.0.4:

V1028: Possible overflow. Consider casting operands, not the result.

#rnx

diff --git a/Engine/Plugins/Experimental/GeometryProcessing/Source/DynamicMesh/Public/DynamicMeshTriangleAttribute.h b/Engine/Plugins/Experimental/GeometryProcessing/Source/DynamicMesh/Public/DynamicMeshTriangleAttribute.h
--- a/Engine/Plugins/Experimental/GeometryProcessing/Source/DynamicMesh/Public/DynamicMeshTriangleAttribute.h
+++ b/Engine/Plugins/Experimental/GeometryProcessing/Source/DynamicMesh/Public/DynamicMeshTriangleAttribute.h
@@ -305,7 +305,7 @@ public:
 		{
 			return;
 		}
-		size_t NeededSize = size_t((TriangleID+1) * AttribDimension);
+		size_t NeededSize = (((size_t)TriangleID+1) * AttribDimension);
 		if (NeededSize > AttribValues.Num())
 		{
 			AttribValues.Resize(NeededSize, GetDefaultAttributeValue());
