fix remesher crashing on mesh with duplicate triangles by refusing to append duplicate triangles + add test coverage and checkvalidity flags to support testing
#rnx

diff --git a/Engine/Plugins/Experimental/GeometryProcessing/Source/DynamicMesh/Private/DynamicMesh3.cpp b/Engine/Plugins/Experimental/GeometryProcessing/Source/DynamicMesh/Private/DynamicMesh3.cpp
--- a/Engine/Plugins/Experimental/GeometryProcessing/Source/DynamicMesh/Private/DynamicMesh3.cpp
+++ b/Engine/Plugins/Experimental/GeometryProcessing/Source/DynamicMesh/Private/DynamicMesh3.cpp
@@ -8,6 +8,7 @@
 // NB: These have to be here until C++17 allows inline variables
 constexpr int       FDynamicMesh3::InvalidID;
 constexpr int       FDynamicMesh3::NonManifoldID;
+constexpr int       FDynamicMesh3::DuplicateTriangleID;
 constexpr int       FDynamicMesh3::InvalidGroupID;
 constexpr FVector3d FDynamicMesh3::InvalidVertex;
 constexpr FIndex3i  FDynamicMesh3::InvalidTriangle;
@@ -686,12 +687,11 @@ bool FDynamicMesh3::IsSameMesh(const FDynamicMesh3& m2, bool bCheckConnectivity,
 	return true;
 }
 
-bool FDynamicMesh3::CheckValidity(bool bAllowNonManifoldVertices, EValidityCheckFailMode FailMode) const
+bool FDynamicMesh3::CheckValidity(FValidityOptions ValidityOptions, EValidityCheckFailMode FailMode) const
 {
 
 	TArray<int> triToVtxRefs;
 	triToVtxRefs.SetNum(MaxVertexID());
-	//int[] triToVtxRefs = new int[this.MaxVertexID];
 
 	bool is_ok = true;
 	TFunction<void(bool)> CheckOrFailF = [&](bool b)
@@ -762,9 +762,12 @@ bool FDynamicMesh3::CheckValidity(bool bAllowNonManifoldVertices, EValidityCheck
 			CheckOrFailF(IndexUtil::SamePairUnordered(a, b, ev[0], ev[1]));
 
 			// also check that nbr edge has opposite orientation
-			FIndex3i othertv = GetTriangle(tOther);
-			int found = IndexUtil::FindTriOrderedEdge(b, a, othertv);
-			CheckOrFailF(found != InvalidID);
+			if (ValidityOptions.bAllowAdjacentFacesReverseOrientation == false)
+			{
+				FIndex3i othertv = GetTriangle(tOther);
+				int found = IndexUtil::FindTriOrderedEdge(b, a, othertv);
+				CheckOrFailF(found != InvalidID);
+			}
 		}
 	}
 
@@ -832,7 +835,7 @@ bool FDynamicMesh3::CheckValidity(bool bAllowNonManifoldVertices, EValidityCheck
 		GetVtxTriangles(vID, vTris2, true);
 		CheckOrFailF(vTris.Num() == vTris2.Num());
 		//System.Console.WriteLine(string.Format("{0} {1} {2}", vID, vTris.Count, GetVtxEdges(vID).Count));
-		if (bAllowNonManifoldVertices)
+		if (ValidityOptions.bAllowNonManifoldVertices)
 		{
 			CheckOrFailF(vTris.Num() <= GetVtxEdgeCount(vID));
 		}
diff --git a/Engine/Plugins/Experimental/GeometryProcessing/Source/DynamicMesh/Private/DynamicMesh3_Edits.cpp b/Engine/Plugins/Experimental/GeometryProcessing/Source/DynamicMesh/Private/DynamicMesh3_Edits.cpp
--- a/Engine/Plugins/Experimental/GeometryProcessing/Source/DynamicMesh/Private/DynamicMesh3_Edits.cpp
+++ b/Engine/Plugins/Experimental/GeometryProcessing/Source/DynamicMesh/Private/DynamicMesh3_Edits.cpp
@@ -205,6 +205,20 @@ int FDynamicMesh3::AppendTriangle(const FIndex3i& tv, int gid)
 		return NonManifoldID;
 	}
 
+	// if all edges were already present, check for duplicate triangle
+	if (e0 != InvalidID && e1 != InvalidID && e2 != InvalidID)
+	{
+		// check if the triangle attached to edge e0 (tv[0] - tv[1]) also contains tv[2]
+		int e0i = 4 * e0;
+		int ti = 3 * Edges[e0i + 2];
+		if (Triangles[ti] == tv[2] || Triangles[ti + 1] == tv[2] || Triangles[ti + 2] == tv[2])
+		{
+			return DuplicateTriangleID;
+		}
+		// there's no other triangle on the edge to check -- if there were, we would have already returned NonManifoldID
+		checkSlow(Edges[e0i + 3] == InvalidID);
+	}
+
 	bool bHasGroups = HasTriangleGroups();  // have to check before changing .triangles
 
 	// now safe to insert triangle
diff --git a/Engine/Plugins/Experimental/GeometryProcessing/Source/DynamicMesh/Private/DynamicMeshEditor.cpp b/Engine/Plugins/Experimental/GeometryProcessing/Source/DynamicMesh/Private/DynamicMeshEditor.cpp
--- a/Engine/Plugins/Experimental/GeometryProcessing/Source/DynamicMesh/Private/DynamicMeshEditor.cpp
+++ b/Engine/Plugins/Experimental/GeometryProcessing/Source/DynamicMesh/Private/DynamicMeshEditor.cpp
@@ -590,7 +590,7 @@ bool FDynamicMeshEditor::ReinsertSubmesh(const FDynamicSubmesh3& Region, FOption
 
 
 		int new_tid = Mesh->AppendTriangle(new_t, gid);
-		ensure(new_tid != FDynamicMesh3::InvalidID && new_tid != FDynamicMesh3::NonManifoldID);
+		ensure(new_tid >= 0);
 		if (!Mesh->IsTriangle(new_tid))
 		{
 			bAllOK = false;
diff --git a/Engine/Plugins/Experimental/GeometryProcessing/Source/DynamicMesh/Private/MeshRefinerBase.cpp b/Engine/Plugins/Experimental/GeometryProcessing/Source/DynamicMesh/Private/MeshRefinerBase.cpp
--- a/Engine/Plugins/Experimental/GeometryProcessing/Source/DynamicMesh/Private/MeshRefinerBase.cpp
+++ b/Engine/Plugins/Experimental/GeometryProcessing/Source/DynamicMesh/Private/MeshRefinerBase.cpp
@@ -247,7 +247,7 @@ void FMeshRefinerBase::DoDebugChecks(bool bEndOfPass)
 
 	if ((DEBUG_CHECK_LEVEL > 2) || (bEndOfPass && DEBUG_CHECK_LEVEL > 1))
 	{
-		Mesh->CheckValidity(true);
+		Mesh->CheckValidity(FDynamicMesh3::FValidityOptions::Permissive());
 		DebugCheckUVSeamConstraints();
 	}
 }
diff --git a/Engine/Plugins/Experimental/GeometryProcessing/Source/DynamicMesh/Public/DynamicMesh3.h b/Engine/Plugins/Experimental/GeometryProcessing/Source/DynamicMesh/Public/DynamicMesh3.h
--- a/Engine/Plugins/Experimental/GeometryProcessing/Source/DynamicMesh/Public/DynamicMesh3.h
+++ b/Engine/Plugins/Experimental/GeometryProcessing/Source/DynamicMesh/Public/DynamicMesh3.h
@@ -96,6 +96,8 @@ using FVertexSplitInfo = DynamicMeshInfo::FVertexSplitInfo;
 	constexpr static int InvalidID = IndexConstants::InvalidID;
 	/** NonManifoldID is returned by AppendTriangle() to indicate that the added triangle would result in nonmanifold geometry and hence was ignored */
 	constexpr static int NonManifoldID = -2;
+	/** DuplicateTriangleID is returned by AppendTriangle() to indicate that the added triangle already exists in the mesh, and was ignored because we do not support duplicate triangles */
+	constexpr static int DuplicateTriangleID = -3;
 	/** InvalidGroupID indicates that a group ID is invalid */
 	constexpr static int InvalidGroupID = IndexConstants::InvalidID;
 
@@ -1150,11 +1152,37 @@ public:
 							bool bCheckGroups = false,
 							float Epsilon = TMathUtil<float>::Epsilon);
 
+	/**
+	 * Options for what the validity check will permit
+	 */
+	struct FValidityOptions
+	{
+		bool bAllowNonManifoldVertices = false;
+		bool bAllowAdjacentFacesReverseOrientation = false;
+
+		/**
+		 * Construct validity checking options
+		 */
+		FValidityOptions(bool bAllowNonManifoldVertices = false, bool bAllowAdjacentFacesReverseOrientation = false)
+			: bAllowNonManifoldVertices(bAllowNonManifoldVertices), bAllowAdjacentFacesReverseOrientation(bAllowAdjacentFacesReverseOrientation)
+		{}
+
+		/**
+		 * Construct with most-permissive options that we still consider valid for processing
+		 */
+		static FValidityOptions Permissive()
+		{
+			FValidityOptions ToRet;
+			ToRet.bAllowAdjacentFacesReverseOrientation = true;
+			ToRet.bAllowNonManifoldVertices = true;
+			return ToRet;
+		}
+	};
 
 	/**
 	 * Checks that the mesh is well-formed, ie all internal data structures are consistent
 	 */
-	virtual bool CheckValidity(bool bAllowNonManifoldVertices = false, EValidityCheckFailMode FailMode = EValidityCheckFailMode::Check) const;
+	virtual bool CheckValidity(FValidityOptions ValidityOptions = FValidityOptions(), EValidityCheckFailMode FailMode = EValidityCheckFailMode::Check) const;
 
 	//
 	// Internal functions
diff --git a/Engine/Plugins/Experimental/GeometryProcessing/Source/MeshConversion/Private/MeshDescriptionToDynamicMesh.cpp b/Engine/Plugins/Experimental/GeometryProcessing/Source/MeshConversion/Private/MeshDescriptionToDynamicMesh.cpp
--- a/Engine/Plugins/Experimental/GeometryProcessing/Source/MeshConversion/Private/MeshDescriptionToDynamicMesh.cpp
+++ b/Engine/Plugins/Experimental/GeometryProcessing/Source/MeshConversion/Private/MeshDescriptionToDynamicMesh.cpp
@@ -215,6 +215,11 @@ void FMeshDescriptionToDynamicMesh::Convert(const FMeshDescription* MeshIn, FDyn
 			int VertexID1 = MeshIn->GetVertexInstanceVertex(InstanceTri[1]).GetValue();
 			int VertexID2 = MeshIn->GetVertexInstanceVertex(InstanceTri[2]).GetValue();
 			int NewTriangleID = MeshOut.AppendTriangle(VertexID0, VertexID1, VertexID2, GroupID);
+
+			if (NewTriangleID == FDynamicMesh3::DuplicateTriangleID)
+			{
+				continue;
+			}
 			
 			// if append failed due to non-manifold, duplicate verts
 			if (NewTriangleID == FDynamicMesh3::NonManifoldID)
@@ -266,6 +271,8 @@ void FMeshDescriptionToDynamicMesh::Convert(const FMeshDescription* MeshIn, FDyn
 				checkSlow(NewTriangleID != FDynamicMesh3::NonManifoldID);
 			}
 
+			checkSlow(NewTriangleID >= 0);
+
 			FIndex3i Tri(VertexID0, VertexID1, VertexID2);
 
 			if (bCalculateMaps)
