propagate MeshRegionBoundaryLoops failure cases up to caller, via return value and flag, rather than just check()ing on them
Handle this and related failures more consistently upstream.

Also don't rely on triangle orientation for edge walk direction after the first edge.

#rnx

diff --git a/Engine/Plugins/Experimental/GeometryProcessing/Source/DynamicMesh/Private/DynamicMeshEditor.cpp b/Engine/Plugins/Experimental/GeometryProcessing/Source/DynamicMesh/Private/DynamicMeshEditor.cpp
--- a/Engine/Plugins/Experimental/GeometryProcessing/Source/DynamicMesh/Private/DynamicMeshEditor.cpp
+++ b/Engine/Plugins/Experimental/GeometryProcessing/Source/DynamicMesh/Private/DynamicMeshEditor.cpp
@@ -381,7 +381,10 @@ bool FDynamicMeshEditor::DisconnectTriangles(const TArray<int>& Triangles, TArra
 		}
 
 		FEdgeLoop InnerLoop;
-		InnerLoop.InitializeFromVertices(Mesh, NewVertexLoop, false);
+		if (!ensure(InnerLoop.InitializeFromVertices(Mesh, NewVertexLoop, false)))
+		{
+			return false;
+		}
 		LoopPair.InnerVertices = MoveTemp(InnerLoop.Vertices);
 		LoopPair.InnerEdges = MoveTemp(InnerLoop.Edges);
 		LoopPair.bOuterIncludesIsolatedVertices = bSawBoundaryInLoop;
diff --git a/Engine/Plugins/Experimental/GeometryProcessing/Source/DynamicMesh/Private/GroupTopology.cpp b/Engine/Plugins/Experimental/GeometryProcessing/Source/DynamicMesh/Private/GroupTopology.cpp
--- a/Engine/Plugins/Experimental/GeometryProcessing/Source/DynamicMesh/Private/GroupTopology.cpp
+++ b/Engine/Plugins/Experimental/GeometryProcessing/Source/DynamicMesh/Private/GroupTopology.cpp
@@ -501,6 +501,7 @@ void FGroupTopology::ForGroupSetEdges(const TArray<int>& GroupIDs,
 void FGroupTopology::ExtractGroupEdges(FGroup& Group)
 {
 	FMeshRegionBoundaryLoops BdryLoops(Mesh, Group.Triangles, true);
+	ensure(!BdryLoops.bFailed); // TODO: handle failure to find boundary loops?
 	int NumLoops = BdryLoops.Loops.Num();
 
 	Group.Boundaries.SetNum(NumLoops);
diff --git a/Engine/Plugins/Experimental/GeometryProcessing/Source/DynamicMesh/Private/MeshBoundaryLoops.cpp b/Engine/Plugins/Experimental/GeometryProcessing/Source/DynamicMesh/Private/MeshBoundaryLoops.cpp
--- a/Engine/Plugins/Experimental/GeometryProcessing/Source/DynamicMesh/Private/MeshBoundaryLoops.cpp
+++ b/Engine/Plugins/Experimental/GeometryProcessing/Source/DynamicMesh/Private/MeshBoundaryLoops.cpp
@@ -385,7 +385,7 @@ int FMeshBoundaryLoops::FindLeftTurnEdge(int incoming_e, int bowtie_v, TArray<in
 
 		// compute projected angle
 		FVector3d bc = Mesh->GetVertex(bdry_ev.B) - Mesh->GetVertex(bowtie_v);
-		double fAngleS = VectorUtil::PlaneAngleSignedD(ab, bc, n);
+		double fAngleS = -VectorUtil::PlaneAngleSignedD(ab, bc, n);
 
 		// turn left!
 		if (best_angle == TNumericLimits<double>::Max() || fAngleS < best_angle)
diff --git a/Engine/Plugins/Experimental/GeometryProcessing/Source/DynamicMesh/Private/MeshRegionBoundaryLoops.cpp b/Engine/Plugins/Experimental/GeometryProcessing/Source/DynamicMesh/Private/MeshRegionBoundaryLoops.cpp
--- a/Engine/Plugins/Experimental/GeometryProcessing/Source/DynamicMesh/Private/MeshRegionBoundaryLoops.cpp
+++ b/Engine/Plugins/Experimental/GeometryProcessing/Source/DynamicMesh/Private/MeshRegionBoundaryLoops.cpp
@@ -73,6 +73,8 @@ int FMeshRegionBoundaryLoops::GetMaxVerticesLoopIndex() const
 
 bool FMeshRegionBoundaryLoops::Compute()
 {
+	bFailed = false; // reset
+
 	// This algorithm assumes that triangles are oriented consistently, 
 	// so closed boundary-loop can be followed by walking edges in-order
 	Loops.SetNum(0);
@@ -109,6 +111,7 @@ bool FMeshRegionBoundaryLoops::Compute()
 		loop_edges.Add(eStart);
 
 		int eCur = eid;
+		int eFirstVert = -1; // the first vertex on eCur, in terms of our walking order
 
 		// follow the chain : order of oriented edges
 		bool bClosed = false;
@@ -119,14 +122,32 @@ bool FMeshRegionBoundaryLoops::Compute()
 			int tid_in = IndexConstants::InvalidID, tid_out = IndexConstants::InvalidID;
 			IsEdgeOnBoundary(eCur, tid_in, tid_out);
 
-			FIndex2i ev = GetOrientedEdgeVerts(eCur, tid_in, tid_out);
-			int cure_a = ev.A, cure_b = ev.B;
+			int cure_a, cure_b;
+			if (eFirstVert == -1)
+			{
+				FIndex2i ev = GetOrientedEdgeVerts(eCur, tid_in);
+				cure_a = ev.A;
+				cure_b = ev.B;
+			}
+			else
+			{
+				// once we've walked on at least one edge, no longer need to rely on triangle orientation to know which way we were walking
+				FIndex2i edgev = Mesh->GetEdgeV(eCur);
+				checkSlow(edgev.Contains(eFirstVert));
+				cure_a = eFirstVert;
+				cure_b = edgev.A == cure_a ? edgev.B : edgev.A;
+			}
 			loop_verts.Add(cure_a);
 
 			int e0 = -1, e1 = 1;
 			int bdry_nbrs = GetVertexBoundaryEdges(cure_b, e0, e1);
 
-			check(bdry_nbrs >= 2); //  if (bdry_nbrs < 2) throw new MeshBoundaryLoopsException("MeshRegionBoundaryLoops.Compute: found broken neighbourhood at vertex " + cure_b){ UnclosedLoop = true };
+			if (bdry_nbrs < 2)
+			{
+				// found broken neighbourhood at vertex cure_b -- unrecoverable failure (unclosed loop)
+				bFailed = true;
+				return false;
+			}
 
 			int eNext = -1;
 			if (bdry_nbrs > 2)
@@ -153,7 +174,12 @@ bool FMeshRegionBoundaryLoops::Compute()
 					// Try to pick the best "turn left" vertex.
 					eNext = FindLeftTurnEdge(eCur, cure_b, all_e, num_be, used_edge);
 
-					check(eNext != -1); // throw new MeshBoundaryLoopsException("MeshRegionBoundaryLoops.Compute: cannot find valid outgoing edge at bowtie vertex " + cure_b){ BowtieFailure = true };
+					if (eNext == -1)
+					{
+						// Cannot find valid outgoing edge at bowtie vertex cure_b -- unrecoverable failure
+						bFailed = true;
+						return false;
+					}
 				}
 
 				if (bowties.Contains(cure_b) == false)
@@ -187,16 +213,27 @@ bool FMeshRegionBoundaryLoops::Compute()
 				eCur = eNext;
 				used_edge.Add(eCur);
 			}
+
+			eFirstVert = cure_b;
 		}
 
 		// if we saw a bowtie vertex, we might need to break up this loop,
 		// so call ExtractSubloops
 		if (bowties.Num() > 0)
 		{
-			TArray<FEdgeLoop> subloops = ExtractSubloops(loop_verts, loop_edges, bowties);
-			for (int i = 0; i < subloops.Num(); ++i)
+			TArray<FEdgeLoop> subloops;
+			bool bExtractedLoops = TryExtractSubloops(loop_verts, loop_edges, bowties, subloops);
+			if (!bExtractedLoops)
+			{
+				// skip adding subloops and mark as failure (but go on computing the rest of the boundary loops)
+				bFailed = true;
+			}
+			else
 			{
-				Loops.Add(subloops[i]);
+				for (int i = 0; i < subloops.Num(); ++i)
+				{
+					Loops.Add(subloops[i]);
+				}
 			}
 		}
 		else
@@ -214,7 +251,7 @@ bool FMeshRegionBoundaryLoops::Compute()
 		bowties.SetNum(0);
 	}
 
-	return true;
+	return !bFailed;
 }
 
 
@@ -254,7 +291,7 @@ bool FMeshRegionBoundaryLoops::IsEdgeOnBoundary(int eid, int& tid_in, int& tid_o
 
 
 // return same indices as GetEdgeV, but oriented based on attached triangle
-FIndex2i FMeshRegionBoundaryLoops::GetOrientedEdgeVerts(int eID, int tid_in, int tid_out)
+FIndex2i FMeshRegionBoundaryLoops::GetOrientedEdgeVerts(int eID, int tid_in)
 {
 	FIndex2i edgev = Mesh->GetEdgeV(eID);
 	int a = edgev.A, b = edgev.B;
@@ -344,7 +381,7 @@ int FMeshRegionBoundaryLoops::FindLeftTurnEdge(int incoming_e, int bowtie_v, TAr
 		// [TODO] can do this more efficiently?
 		int tid_in = IndexConstants::InvalidID, tid_out = IndexConstants::InvalidID;
 		IsEdgeOnBoundary(bdry_eid, tid_in, tid_out);
-		FIndex2i bdry_ev = GetOrientedEdgeVerts(bdry_eid, tid_in, tid_out);
+		FIndex2i bdry_ev = GetOrientedEdgeVerts(bdry_eid, tid_in);
 		//FIndex2i bdry_ev = Mesh.GetOrientedBoundaryEdgeV(bdry_eid);
 
 		if (bdry_ev.A != bowtie_v) {
@@ -353,7 +390,7 @@ int FMeshRegionBoundaryLoops::FindLeftTurnEdge(int incoming_e, int bowtie_v, TAr
 
 		// compute projected angle
 		FVector3d bc = Mesh->GetVertex(bdry_ev.B) - Mesh->GetVertex(bowtie_v);
-		double fAngleS = VectorUtil::PlaneAngleSignedD(ab, bc, n);
+		double fAngleS = -VectorUtil::PlaneAngleSignedD(ab, bc, n);
 
 		// turn left!
 		if (best_angle == TNumericLimits<double>::Max() || fAngleS < best_angle)
@@ -362,7 +399,6 @@ int FMeshRegionBoundaryLoops::FindLeftTurnEdge(int incoming_e, int bowtie_v, TAr
 			best_e = bdry_eid;
 		}
 	}
-	check(best_e != -1);
 
 	return best_e;
 }
@@ -379,9 +415,9 @@ int FMeshRegionBoundaryLoops::FindLeftTurnEdge(int incoming_e, int bowtie_v, TAr
 //
 // Currently loopE is not used, and the returned FEdgeLoop objects do not have their Edges
 // arrays initialized. Perhaps to improve : future.
-TArray<FEdgeLoop> FMeshRegionBoundaryLoops::ExtractSubloops(TArray<int>& loopV, const TArray<int>& loopE, const TArray<int>& bowties)
+bool FMeshRegionBoundaryLoops::TryExtractSubloops(TArray<int>& loopV, const TArray<int>& loopE, const TArray<int>& bowties, TArray<FEdgeLoop>& SubLoopsOut)
 {
-	TArray<FEdgeLoop> subs;
+	SubLoopsOut.Reset();
 
 	// figure out which bowties we saw are actually duplicated : loopV
 	TArray<int> dupes;
@@ -400,8 +436,8 @@ TArray<FEdgeLoop> FMeshRegionBoundaryLoops::ExtractSubloops(TArray<int>& loopV,
 		NewLoop.Vertices = loopV;
 		NewLoop.Edges = loopE;
 		NewLoop.BowtieVertices = bowties;
-		subs.Add(NewLoop);
-		return subs;
+		SubLoopsOut.Add(NewLoop);
+		return true;
 	}
 
 	// This loop extracts subloops until we have dealt with all the
@@ -428,7 +464,12 @@ TArray<FEdgeLoop> FMeshRegionBoundaryLoops::ExtractSubloops(TArray<int>& loopV,
 				}
 			}
 		}
-		check(bv_shortest != -1); //  throw new MeshBoundaryLoopsException("MeshRegionBoundaryLoops.Compute: Cannot find a valid simple loop");
+		if (bv_shortest == -1)
+		{
+			// Cannot find a valid simple loop -- unrecoverable failure
+			return false;
+		}
+
 		if (bv != bv_shortest)
 		{
 			bv = bv_shortest;
@@ -442,7 +483,7 @@ TArray<FEdgeLoop> FMeshRegionBoundaryLoops::ExtractSubloops(TArray<int>& loopV,
 		FMeshBoundaryLoops::ExtractSpan(loopV, start_i, end_i, true, loop.Vertices);
 		FEdgeLoop::VertexLoopToEdgeLoop(Mesh, loop.Vertices, loop.Edges);
 		loop.BowtieVertices = bowties;
-		subs.Add(loop);
+		SubLoopsOut.Add(loop);
 
 		// If there are no more duplicates of this bowtie, we can treat
 		// it like a regular vertex now
@@ -476,8 +517,8 @@ TArray<FEdgeLoop> FMeshRegionBoundaryLoops::ExtractSubloops(TArray<int>& loopV,
 		}
 		FEdgeLoop::VertexLoopToEdgeLoop(Mesh, loop.Vertices, loop.Edges);
 		loop.BowtieVertices = bowties;
-		subs.Add(loop);
+		SubLoopsOut.Add(loop);
 	}
 
-	return subs;
+	return true;
 }
\ No newline at end of file
diff --git a/Engine/Plugins/Experimental/GeometryProcessing/Source/DynamicMesh/Private/Operations/ExtrudeMesh.cpp b/Engine/Plugins/Experimental/GeometryProcessing/Source/DynamicMesh/Private/Operations/ExtrudeMesh.cpp
--- a/Engine/Plugins/Experimental/GeometryProcessing/Source/DynamicMesh/Private/Operations/ExtrudeMesh.cpp
+++ b/Engine/Plugins/Experimental/GeometryProcessing/Source/DynamicMesh/Private/Operations/ExtrudeMesh.cpp
@@ -118,60 +118,65 @@ bool FExtrudeMesh::ApplyExtrude(FExtrusionInfo& Region, FMeshNormals* UseNormals
 		}
 
 		FDynamicMeshEditResult StitchResult;
+		bool bStitchSuccess;
 		if (IsPositiveOffset)
 		{
-			Editor.StitchVertexLoopsMinimal(OffsetLoop, BaseLoop.Vertices, StitchResult);
+			bStitchSuccess = Editor.StitchVertexLoopsMinimal(OffsetLoop, BaseLoop.Vertices, StitchResult);
 		}
 		else
 		{
-			Editor.StitchVertexLoopsMinimal(BaseLoop.Vertices, OffsetLoop, StitchResult);
+			bStitchSuccess = Editor.StitchVertexLoopsMinimal(BaseLoop.Vertices, OffsetLoop, StitchResult);
 		}
-		StitchResult.GetAllTriangles(Region.StitchTriangles[LoopIndex]);
-		Region.StitchPolygonIDs[LoopIndex] = StitchResult.NewGroups;
 
-		// for each polygon we created in stitch, set UVs and normals
-		if (Mesh->HasAttributes())
+		if (bStitchSuccess)
 		{
-			float AccumUVTranslation = 0;
+			StitchResult.GetAllTriangles(Region.StitchTriangles[LoopIndex]);
+			Region.StitchPolygonIDs[LoopIndex] = StitchResult.NewGroups;
 
-			FFrame3d FirstProjectFrame;
-			FVector3d FrameUp;
-
-			int NumNewQuads = StitchResult.NewQuads.Num();
-			for (int k = 0; k < NumNewQuads; k++)
+			// for each polygon we created in stitch, set UVs and normals
+			if (Mesh->HasAttributes())
 			{
-				FVector3f Normal = Editor.ComputeAndSetQuadNormal(StitchResult.NewQuads[k], true);
+				float AccumUVTranslation = 0;
 
-				// align axis 0 of projection frame to first edge, then for further edges,
-				// rotate around 'up' axis to keep normal aligned and frame horizontal
-				FFrame3d ProjectFrame;
-				if (k == 0)
-				{
-					FVector3d FirstEdge = Mesh->GetVertex(BaseLoop.Vertices[1]) - Mesh->GetVertex(BaseLoop.Vertices[0]);
-					FirstEdge.Normalize();
-					FirstProjectFrame = FFrame3d(FVector3d::Zero(), (FVector3d)Normal);
-					FirstProjectFrame.ConstrainedAlignAxis(0, FirstEdge, (FVector3d)Normal);
-					FrameUp = FirstProjectFrame.GetAxis(1);
-					ProjectFrame = FirstProjectFrame;
-				}
-				else
-				{
-					ProjectFrame = FirstProjectFrame;
-					ProjectFrame.ConstrainedAlignAxis(2, (FVector3d)Normal, FrameUp);
-				}
+				FFrame3d FirstProjectFrame;
+				FVector3d FrameUp;
 
-				if (k > 0)
+				int NumNewQuads = StitchResult.NewQuads.Num();
+				for (int k = 0; k < NumNewQuads; k++)
 				{
-					AccumUVTranslation += Mesh->GetVertex(BaseLoop.Vertices[k]).Distance(Mesh->GetVertex(BaseLoop.Vertices[k - 1]));
+					FVector3f Normal = Editor.ComputeAndSetQuadNormal(StitchResult.NewQuads[k], true);
+
+					// align axis 0 of projection frame to first edge, then for further edges,
+					// rotate around 'up' axis to keep normal aligned and frame horizontal
+					FFrame3d ProjectFrame;
+					if (k == 0)
+					{
+						FVector3d FirstEdge = Mesh->GetVertex(BaseLoop.Vertices[1]) - Mesh->GetVertex(BaseLoop.Vertices[0]);
+						FirstEdge.Normalize();
+						FirstProjectFrame = FFrame3d(FVector3d::Zero(), (FVector3d)Normal);
+						FirstProjectFrame.ConstrainedAlignAxis(0, FirstEdge, (FVector3d)Normal);
+						FrameUp = FirstProjectFrame.GetAxis(1);
+						ProjectFrame = FirstProjectFrame;
+					}
+					else
+					{
+						ProjectFrame = FirstProjectFrame;
+						ProjectFrame.ConstrainedAlignAxis(2, (FVector3d)Normal, FrameUp);
+					}
+
+					if (k > 0)
+					{
+						AccumUVTranslation += Mesh->GetVertex(BaseLoop.Vertices[k]).Distance(Mesh->GetVertex(BaseLoop.Vertices[k - 1]));
+					}
+
+					// translate horizontally such that vertical spans are adjacent in UV space (so textures tile/wrap properly)
+					float TranslateU = UVScaleFactor * AccumUVTranslation;
+					Editor.SetQuadUVsFromProjection(StitchResult.NewQuads[k], ProjectFrame, UVScaleFactor, FVector2f(TranslateU, 0));
 				}
-
-				// translate horizontally such that vertical spans are adjacent in UV space (so textures tile/wrap properly)
-				float TranslateU = UVScaleFactor * AccumUVTranslation;
-				Editor.SetQuadUVsFromProjection(StitchResult.NewQuads[k], ProjectFrame, UVScaleFactor, FVector2f(TranslateU, 0));
 			}
-		}
 
-		Region.NewLoops[LoopIndex].InitializeFromVertices(Mesh, OffsetLoop);
+			Region.NewLoops[LoopIndex].InitializeFromVertices(Mesh, OffsetLoop);
+		}
 		LoopIndex++;
 	}
 
diff --git a/Engine/Plugins/Experimental/GeometryProcessing/Source/DynamicMesh/Public/MeshRegionBoundaryLoops.h b/Engine/Plugins/Experimental/GeometryProcessing/Source/DynamicMesh/Public/MeshRegionBoundaryLoops.h
--- a/Engine/Plugins/Experimental/GeometryProcessing/Source/DynamicMesh/Public/MeshRegionBoundaryLoops.h
+++ b/Engine/Plugins/Experimental/GeometryProcessing/Source/DynamicMesh/Public/MeshRegionBoundaryLoops.h
@@ -15,11 +15,19 @@
 class DYNAMICMESH_API FMeshRegionBoundaryLoops
 {
 public:
+
+	// INPUTS
+
 	/** Mesh we are finding loops on */
 	const FDynamicMesh3* Mesh = nullptr;
 	/** Resulting set of loops filled by Compute() */
 	TArray<FEdgeLoop> Loops;
 
+	// OUTPUTS
+
+	/** If true, we did not completely succeed in extracting all loops */
+	bool bFailed = false;
+
 public:
 	FMeshRegionBoundaryLoops() {}
 
@@ -76,7 +84,7 @@ protected:
 	bool IsEdgeOnBoundary(int eid, int& tid_in, int& tid_out) const;
 
 	// return same indices as GetEdgeV, but oriented based on attached triangle
-	FIndex2i GetOrientedEdgeVerts(int eID, int tid_in, int tid_out);
+	FIndex2i GetOrientedEdgeVerts(int eID, int tid_in);
 
 	// returns first two boundary edges, and count of total boundary edges
 	int GetVertexBoundaryEdges(int vID, int& e0, int& e1);
@@ -99,7 +107,7 @@ protected:
 	// This is called when loopV contains one or more "bowtie" vertices.
 	// These vertices *might* be duplicated : loopV (but not necessarily)
 	// If they are, we have to break loopV into subloops that don't contain duplicates.
-	TArray<FEdgeLoop> ExtractSubloops(TArray<int>& loopV, const TArray<int>& loopE, const TArray<int>& bowties);
+	bool TryExtractSubloops(TArray<int>& loopV, const TArray<int>& loopE, const TArray<int>& bowties, TArray<FEdgeLoop>& SubLoopsOut);
 
 
 };
diff --git a/Engine/Plugins/Experimental/GeometryProcessing/Source/GeometricObjects/Public/VectorTypes.h b/Engine/Plugins/Experimental/GeometryProcessing/Source/GeometricObjects/Public/VectorTypes.h
--- a/Engine/Plugins/Experimental/GeometryProcessing/Source/GeometricObjects/Public/VectorTypes.h
+++ b/Engine/Plugins/Experimental/GeometryProcessing/Source/GeometricObjects/Public/VectorTypes.h
@@ -551,7 +551,10 @@ struct FVector3
 			X * V2.Y - Y * V2.X);
 	}
 
-	// Angle in Degrees
+	/**
+	 * Assumes this vector and V2 are both already normalized
+	 * @return the angle to vector V2 in degrees
+	 */
 	T AngleD(const FVector3<T>& V2) const
 	{
 		T DotVal = Dot(V2);
@@ -559,7 +562,10 @@ struct FVector3
 		return (T)(acos(ClampedDot) * (T)(180.0 / 3.14159265358979323846));
 	}
 
-	// Angle in Radians
+	/**
+	 * Assumes this vector and V2 are both already normalized
+	 * @return the angle to vector V2 in radians
+	 */
 	T AngleR(const FVector3<T>& V2) const
 	{
 		T DotVal = Dot(V2);
diff --git a/Engine/Plugins/Experimental/MeshModelingToolset/Source/MeshModelingTools/Private/EditMeshPolygonsTool.cpp b/Engine/Plugins/Experimental/MeshModelingToolset/Source/MeshModelingTools/Private/EditMeshPolygonsTool.cpp
--- a/Engine/Plugins/Experimental/MeshModelingToolset/Source/MeshModelingTools/Private/EditMeshPolygonsTool.cpp
+++ b/Engine/Plugins/Experimental/MeshModelingToolset/Source/MeshModelingTools/Private/EditMeshPolygonsTool.cpp
@@ -1414,7 +1414,7 @@ void UEditMeshPolygonsTool::ApplyRetriangulate()
 		const TArray<int32>& Triangles = Topology->GetGroupTriangles(GroupID);
 		ChangeTracker.SaveTriangles(Triangles, true);
 		FMeshRegionBoundaryLoops RegionLoops(Mesh, Triangles, true);
-		if (RegionLoops.Loops.Num() == 1 && Triangles.Num() > 1)
+		if (!RegionLoops.bFailed && RegionLoops.Loops.Num() == 1 && Triangles.Num() > 1)
 		{
 			Editor.RemoveTriangles(Topology->GetGroupTriangles(GroupID), true);
 			RegionLoops.Loops[0].Reverse();
