Merging CL 11341541.

MeshConstraints: before a new constraint is added to a vertex, make sure that we don't overwrite an existing constraint with something less restrictive.

#rnx

diff --git a/Engine/Plugins/Experimental/GeometryProcessing/Source/DynamicMesh/Private/MeshConstraintsUtil.cpp b/Engine/Plugins/Experimental/GeometryProcessing/Source/DynamicMesh/Private/MeshConstraintsUtil.cpp
--- a/Engine/Plugins/Experimental/GeometryProcessing/Source/DynamicMesh/Private/MeshConstraintsUtil.cpp
+++ b/Engine/Plugins/Experimental/GeometryProcessing/Source/DynamicMesh/Private/MeshConstraintsUtil.cpp
@@ -59,7 +59,7 @@ FMeshConstraintsUtil::ConstrainAllBoundariesAndSeams(FMeshConstraints& Constrain
 			const bool bIsGroupBoundary = Mesh.IsGroupBoundaryEdge(EdgeID);
 			const bool bIsMaterialBoundary = Attributes && Attributes->IsMaterialBoundaryEdge(EdgeID);
 			const bool bIsSeam = Attributes && Attributes->IsSeamEdge(EdgeID);
-			FVertexConstraint VtxConstraint{false, true};
+			FVertexConstraint VtxConstraint = FVertexConstraint::Unconstrained();
 			EEdgeRefineFlags EdgeFlags{};
 			auto ApplyBoundaryConstraint =
 				[&VtxConstraint, &EdgeFlags](uint8 BoundaryConstraint)
@@ -98,9 +98,19 @@ FMeshConstraintsUtil::ConstrainAllBoundariesAndSeams(FMeshConstraints& Constrain
 				FIndex2i EdgeVerts = Mesh.GetEdgeV(EdgeID);
 
 				ConstraintSetLock.Lock();
-				Constraints.SetOrUpdateEdgeConstraint(EdgeID, FEdgeConstraint{EdgeFlags});
-				Constraints.SetOrUpdateVertexConstraint(EdgeVerts.A, VtxConstraint);
-				Constraints.SetOrUpdateVertexConstraint(EdgeVerts.B, VtxConstraint);
+
+				Constraints.SetOrUpdateEdgeConstraint(EdgeID, FEdgeConstraint{ EdgeFlags });
+
+				// If any vertex constraints exist, we can only make them more restrictive!
+							
+				FVertexConstraint ConstraintA = VtxConstraint;
+				ConstraintA.CombineConstraint(Constraints.GetVertexConstraint(EdgeVerts.A));
+				Constraints.SetOrUpdateVertexConstraint(EdgeVerts.A, ConstraintA);
+
+				FVertexConstraint ConstraintB = VtxConstraint;
+				ConstraintB.CombineConstraint(Constraints.GetVertexConstraint(EdgeVerts.B));
+				Constraints.SetOrUpdateVertexConstraint(EdgeVerts.B, ConstraintB);
+
 				ConstraintSetLock.Unlock();
 			}
 		}
diff --git a/Engine/Plugins/Experimental/GeometryProcessing/Source/DynamicMesh/Public/MeshConstraints.h b/Engine/Plugins/Experimental/GeometryProcessing/Source/DynamicMesh/Public/MeshConstraints.h
--- a/Engine/Plugins/Experimental/GeometryProcessing/Source/DynamicMesh/Public/MeshConstraints.h
+++ b/Engine/Plugins/Experimental/GeometryProcessing/Source/DynamicMesh/Public/MeshConstraints.h
@@ -157,6 +157,15 @@ public:
 
 	/** @return a Pinned-but-Movable vertex constraint */
 	static FVertexConstraint PinnedMovable() { return FVertexConstraint(true, true); }
+
+	/** Combine the incoming constraint with this constraint. The combination cannot be less restrictive than the
+	    existing constraint. */
+	void CombineConstraint(const FVertexConstraint& OtherConstraint)
+	{
+		Fixed = Fixed || OtherConstraint.Fixed;
+		Movable = Movable && OtherConstraint.Movable;
+	}
+
 };
 
 
@@ -322,4 +331,4 @@ public:
 		Vertices.Remove(VertexID);
 	}
 
-};
\ No newline at end of file
+};
