disable check that is sometimes hit due to numerical issues
#rnx

diff --git a/Engine/Plugins/Experimental/GeometryProcessing/Source/GeometricObjects/Public/Spatial/SparseDynamicOctree3.h b/Engine/Plugins/Experimental/GeometryProcessing/Source/GeometricObjects/Public/Spatial/SparseDynamicOctree3.h
--- a/Engine/Plugins/Experimental/GeometryProcessing/Source/GeometricObjects/Public/Spatial/SparseDynamicOctree3.h
+++ b/Engine/Plugins/Experimental/GeometryProcessing/Source/GeometricObjects/Public/Spatial/SparseDynamicOctree3.h
@@ -345,6 +345,8 @@ protected:
 	}
 
 
+	// warning: result here appears to be unstable (due to optimization?) if any of the position values are on the border of a cell
+	// (in testing, pragma-optimization-disabled code produced off-by-one result from calling this function)
 	FVector3i PointToIndex(uint32 Level, const FVector3d& Position) const
 	{
 		double CellWidth = GetCellWidth(Level);
@@ -460,7 +462,8 @@ void FSparseDynamicOctree3::Insert_NewChildCell(int32 ObjectID, const FAxisAlign
 	OrigParentCell.SetChild(ChildIdx, NewChildCell);
 
 	check(CanFit(NewChildCell, Bounds));
-	check(PointToIndex(NewChildCell.Level, Bounds.Center()) == NewChildCell.Index);
+	// this check is unstable if the center point is within machine-epsilon of the cell border
+	//check(PointToIndex(NewChildCell.Level, Bounds.Center()) == NewChildCell.Index);
 }
 
 
@@ -474,7 +477,8 @@ void FSparseDynamicOctree3::Insert_ToCell(int32 ObjectID, const FAxisAlignedBox3
 	CellObjectLists.Insert(ExistingCell.CellID, ObjectID);
 
 	check(CanFit(ExistingCell, Bounds));
-	check(PointToIndex(ExistingCell.Level, Bounds.Center()) == ExistingCell.Index);
+	// this check is unstable if the center point is within machine-epsilon of the cell border
+	//check(PointToIndex(ExistingCell.Level, Bounds.Center()) == ExistingCell.Index);
 }
 
 
