ModelingTools: fix issues w/ bad mesh changes being created in dynamic mesh sculpting.

Problem is that during remesh passes the affected triangle set can expand (due to splits/collapses on ROI boundary) and our tracking of it is not perfect. Remesher now tells ChangeTracker before it modifies an edge, rather than Sculpt tool trying to predict set of modified triangles up front.

In theory this allows for smaller changes, because often only a subset of triangles are modified. However ChangeTracker currently only emits modified vertices as part of modified triangles. So we still need to save the full triangle ROI. Made several MeshChangeTracker functions protected for now, because they cannot be used externally due to this limitation. Updated call sites to those functions (which were likely also buggy).

#rnx

diff --git a/Engine/Plugins/Experimental/GeometryProcessing/Source/DynamicMesh/Private/MeshRefinerBase.cpp b/Engine/Plugins/Experimental/GeometryProcessing/Source/DynamicMesh/Private/MeshRefinerBase.cpp
--- a/Engine/Plugins/Experimental/GeometryProcessing/Source/DynamicMesh/Private/MeshRefinerBase.cpp
+++ b/Engine/Plugins/Experimental/GeometryProcessing/Source/DynamicMesh/Private/MeshRefinerBase.cpp
@@ -2,7 +2,7 @@
 
 #include "MeshRefinerBase.h"
 #include "DynamicMeshAttributeSet.h"
-
+#include "DynamicMeshChangeTracker.h"
 
 
 
@@ -232,6 +232,46 @@ bool FMeshRefinerBase::CanCollapseVertex(int eid, int a, int b, int& collapse_to
 
 
 
+void FMeshRefinerBase::SetMeshChangeTracker(FDynamicMeshChangeTracker* Tracker)
+{
+	ActiveChangeTracker = Tracker;
+}
+
+
+void FMeshRefinerBase::SaveTriangleBeforeModify(int32 TriangleID)
+{
+	if (ActiveChangeTracker)
+	{
+		ActiveChangeTracker->SaveTriangle(TriangleID, true);
+		checkSlow(ActiveChangeTracker->VerifyTriangleIsSaved(TriangleID));
+	}
+}
+
+void FMeshRefinerBase::SaveVertexTrianglesBeforeModify(int32 VertexID)
+{
+	if (ActiveChangeTracker)
+	{
+		for (int32 TriangleID : Mesh->VtxTrianglesItr(VertexID))
+		{
+			ActiveChangeTracker->SaveTriangle(TriangleID, true );
+			checkSlow(ActiveChangeTracker->VerifyTriangleIsSaved(TriangleID));
+		}
+	}
+}
+
+
+void FMeshRefinerBase::SaveEdgeBeforeModify(int32 EdgeID)
+{
+	if (ActiveChangeTracker)
+	{
+		FIndex2i EdgeVerts = Mesh->GetEdgeV(EdgeID);
+		SaveVertexTrianglesBeforeModify(EdgeVerts.A);
+		SaveVertexTrianglesBeforeModify(EdgeVerts.B);
+	}
+}
+
+
+
 void FMeshRefinerBase::RuntimeDebugCheck(int eid)
 {
 	if (DebugEdges.Contains(eid))
diff --git a/Engine/Plugins/Experimental/GeometryProcessing/Source/DynamicMesh/Private/Remesher.cpp b/Engine/Plugins/Experimental/GeometryProcessing/Source/DynamicMesh/Private/Remesher.cpp
--- a/Engine/Plugins/Experimental/GeometryProcessing/Source/DynamicMesh/Private/Remesher.cpp
+++ b/Engine/Plugins/Experimental/GeometryProcessing/Source/DynamicMesh/Private/Remesher.cpp
@@ -205,6 +205,7 @@ FRemesher::EProcessResult FRemesher::ProcessEdge(int edgeID)
 
 		// lots of cases where we cannot collapse, but we should just let
 		// mesh sort that out, right?
+		SaveEdgeBeforeModify(edgeID);
 		COUNT_COLLAPSES++;
 		FDynamicMesh3::FEdgeCollapseInfo collapseInfo;
 		EMeshResult result = Mesh->CollapseEdge(iKeep, iCollapse, collapse_t, collapseInfo);
@@ -266,6 +267,7 @@ abort_collapse:
 
 		if (bTryFlip) 
 		{
+			SaveEdgeBeforeModify(edgeID);
 			FDynamicMesh3::FEdgeFlipInfo flipInfo;
 			COUNT_FLIPS++;
 			EMeshResult result = Mesh->FlipEdge(edgeID, flipInfo);
@@ -290,6 +292,7 @@ abort_collapse:
 	bool bTriedSplit = false;
 	if (bEnableSplits && constraint.CanSplit() && edge_len_sqr > MaxEdgeLength*MaxEdgeLength) 
 	{
+		SaveEdgeBeforeModify(edgeID);
 		FDynamicMesh3::FEdgeSplitInfo SplitInfo;
 		COUNT_SPLITS++;
 		EMeshResult result = Mesh->SplitEdge(edgeID, SplitInfo);
diff --git a/Engine/Plugins/Experimental/GeometryProcessing/Source/DynamicMesh/Public/DynamicMeshChangeTracker.h b/Engine/Plugins/Experimental/GeometryProcessing/Source/DynamicMesh/Public/DynamicMeshChangeTracker.h
--- a/Engine/Plugins/Experimental/GeometryProcessing/Source/DynamicMesh/Public/DynamicMeshChangeTracker.h
+++ b/Engine/Plugins/Experimental/GeometryProcessing/Source/DynamicMesh/Public/DynamicMeshChangeTracker.h
@@ -256,6 +256,8 @@ protected:
  * Call BeginChange() before making any changes to the mesh, then call SaveVertex()
  * and SaveTriangle() before modifying the respective elements. Then call EndChange()
  * to construct a FDynamicMeshChange that represents the mesh delta.
+ *
+ * @warning Currently only vertices that are part of saved triangles will be stored in the emitted FMeshChange!
  * 
  */
 class DYNAMICMESH_API FDynamicMeshChangeTracker
@@ -269,18 +271,11 @@ public:
 	/** Construct a change object that represents the delta between the Begin and End states */
 	TUniquePtr<FDynamicMeshChange> EndChange();
 
-	/** Save necessary information about a vertex before it is modified */
-	void SaveVertex(int32 VertexID);
-
 	/** Save necessary information about a triangle before it is modified */
 	void SaveTriangle(int32 TriangleID, bool bSaveVertices);
 
 	/** Save necessary information about an edge before it is modified */
-	inline void SaveEdge(int32 EdgeID, bool bTriangles, bool bVertices);
-
-	/** Save necessary information about a set of vertices before they are modified */
-	template<typename EnumerableType>
-	void SaveVertices(EnumerableType VertexIDs);
+	inline void SaveEdge(int32 EdgeID, bool bVertices);
 
 	/** Save necessary information about a set of triangles before they are modified */
 	template<typename EnumerableType>
@@ -300,6 +295,22 @@ public:
 	/** Do (limited) sanity checking to make sure that the change is well-formed*/
 	void VerifySaveState();
 
+protected:
+
+	//
+	// Currently EndChange() only stores vertices that are part of modified triangles.
+	// So calling SaveVertex() independently, on triangles that are not saved, will be lost.
+	// (This needs to be fixed)
+	//
+
+	/** Save necessary information about a vertex before it is modified */
+	void SaveVertex(int32 VertexID);
+
+	/** Save necessary information about a set of vertices before they are modified */
+	template<typename EnumerableType>
+	void SaveVertices(EnumerableType VertexIDs);
+
+
 protected:
 	const FDynamicMesh3* Mesh = nullptr;
 
@@ -320,21 +331,13 @@ protected:
 
 
 
-void FDynamicMeshChangeTracker::SaveEdge(int32 EdgeID, bool bTriangles, bool bVertices)
+void FDynamicMeshChangeTracker::SaveEdge(int32 EdgeID, bool bVertices)
 {
 	FIndex4i EdgeInfo = Mesh->GetEdge(EdgeID);
-	if (bTriangles)
-	{
-		SaveTriangle(EdgeInfo.C, bVertices);
-		if (EdgeInfo.D != FDynamicMesh3::InvalidID)
-		{
-			SaveTriangle(EdgeInfo.D, bVertices);
-		}
-	}
-	else
+	SaveTriangle(EdgeInfo.C, bVertices);
+	if (EdgeInfo.D != FDynamicMesh3::InvalidID)
 	{
-		SaveVertex(EdgeInfo.A);
-		SaveVertex(EdgeInfo.B);
+		SaveTriangle(EdgeInfo.D, bVertices);
 	}
 }
 
diff --git a/Engine/Plugins/Experimental/GeometryProcessing/Source/DynamicMesh/Public/MeshRefinerBase.h b/Engine/Plugins/Experimental/GeometryProcessing/Source/DynamicMesh/Public/MeshRefinerBase.h
--- a/Engine/Plugins/Experimental/GeometryProcessing/Source/DynamicMesh/Public/MeshRefinerBase.h
+++ b/Engine/Plugins/Experimental/GeometryProcessing/Source/DynamicMesh/Public/MeshRefinerBase.h
@@ -8,6 +8,7 @@
 #include "MeshConstraints.h"
 #include "Util/ProgressCancel.h"
 
+class FDynamicMeshChangeTracker;
 
 /**
  * This is a base class that implements common functionality for various triangle mesh resampling strategies
@@ -25,7 +26,6 @@ protected:
 	/** Vertices can be projected onto this surface when they are modified */
 	IProjectionTarget* ProjTarget = nullptr;
 
-
 	FMeshRefinerBase(FDynamicMesh3* MeshIn)
 	{
 		this->Mesh = MeshIn;
@@ -255,8 +255,21 @@ protected:
 
 
 
+public:
+	//
+	// Mesh Change Tracking support
+	//
+
+	void SetMeshChangeTracker(FDynamicMeshChangeTracker* Tracker);
+
+protected:
+	FDynamicMeshChangeTracker* ActiveChangeTracker = nullptr;
+	virtual void SaveTriangleBeforeModify(int32 TriangleID);
+	virtual void SaveEdgeBeforeModify(int32 EdgeID);
+	virtual void SaveVertexTrianglesBeforeModify(int32 VertexID);
 
 
+protected:
 	/*
 	 * testing/debug/profiling stuff
 	 */
diff --git a/Engine/Plugins/Experimental/MeshModelingToolset/Source/MeshModelingTools/Private/DynamicMeshSculptTool.cpp b/Engine/Plugins/Experimental/MeshModelingToolset/Source/MeshModelingTools/Private/DynamicMeshSculptTool.cpp
--- a/Engine/Plugins/Experimental/MeshModelingToolset/Source/MeshModelingTools/Private/DynamicMeshSculptTool.cpp
+++ b/Engine/Plugins/Experimental/MeshModelingToolset/Source/MeshModelingTools/Private/DynamicMeshSculptTool.cpp
@@ -697,7 +697,7 @@ void UDynamicMeshSculptTool::ApplyOffsetBrush(const FRay& WorldRay, bool bUseVie
 	FVector3d LocalNormal = CurTargetTransform.InverseTransformNormal(LastBrushPosNormalWorld);
 
 	double Direction = (bInvert) ? -1.0 : 1.0;
-	double UseSpeed = Direction * FMathd::Sqrt(CurrentBrushRadius) * (SculptProperties->PrimaryBrushSpeed) * ActivePressure;
+	double UseSpeed = 0.5 * Direction * FMathd::Sqrt(CurrentBrushRadius) * (SculptProperties->PrimaryBrushSpeed) * ActivePressure;
 	double MaxOffset = CurrentBrushRadius;
 
 	FDynamicMesh3* Mesh = DynamicMeshComponent->GetMesh();
@@ -1702,6 +1702,11 @@ void UDynamicMeshSculptTool::RemeshROIPass()
 	FSubRegionRemesher Remesher(Mesh);
 	ConfigureRemesher(Remesher);
 
+	if (ActiveMeshChange != nullptr)
+	{
+		Remesher.SetMeshChangeTracker(ActiveMeshChange);
+	}
+
 	bool bIsUniformSmooth = (Remesher.SmoothType == FRemesher::ESmoothTypes::Uniform);
 	for (int k = 0; k < 5; ++k)
 	{
@@ -1717,6 +1722,9 @@ void UDynamicMeshSculptTool::RemeshROIPass()
 
 			if (ActiveMeshChange != nullptr)
 			{
+				// [TODO] would like to only save vertices here, as triangles will be saved by Remesher as necessary.
+				// However currently FDynamicMeshChangeTracker cannot independently save vertices, only vertices 
+				// that are part of saved triangles will be included in the output FDynamicMeshChange
 				Remesher.SaveActiveROI(ActiveMeshChange);
 				//ActiveMeshChange->VerifySaveState();    // useful for debugging
 			}
@@ -2206,10 +2214,8 @@ void UDynamicMeshSculptTool::SaveActiveROI()
 {
 	if (ActiveMeshChange != nullptr)
 	{
-		for (int vid : VertexROI)
-		{
-			ActiveMeshChange->SaveVertex(vid);
-		}
+		// must save triangles containing vertex ROI or they will not be included in emitted mesh change (due to limitations of change tracker)
+		ActiveMeshChange->SaveVertexOneRingTriangles(VertexROI, true);
 	}
 }
 
diff --git a/Engine/Plugins/Experimental/MeshModelingToolset/Source/MeshModelingTools/Private/EditMeshPolygonsTool.cpp b/Engine/Plugins/Experimental/MeshModelingToolset/Source/MeshModelingTools/Private/EditMeshPolygonsTool.cpp
--- a/Engine/Plugins/Experimental/MeshModelingToolset/Source/MeshModelingTools/Private/EditMeshPolygonsTool.cpp
+++ b/Engine/Plugins/Experimental/MeshModelingToolset/Source/MeshModelingTools/Private/EditMeshPolygonsTool.cpp
@@ -1617,7 +1617,7 @@ void UEditMeshPolygonsTool::ApplyStraightenEdges()
 		int32 NumV = EdgeVerts.Num();
 		if ( NumV > 2 )
 		{
-			ChangeTracker.SaveVertices(EdgeVerts);
+			ChangeTracker.SaveVertexOneRingTriangles(EdgeVerts, true);
 			FVector3d A(Mesh->GetVertex(EdgeVerts[0])), B(Mesh->GetVertex(EdgeVerts[NumV-1]));
 			TArray<double> VtxArcLengths;
 			double EdgeArcLen = Topology->GetEdgeArcLength(Edge.EdgeTopoID, &VtxArcLengths);
diff --git a/Engine/Plugins/Experimental/MeshModelingToolset/Source/MeshModelingTools/Private/MeshSelectionTool.cpp b/Engine/Plugins/Experimental/MeshModelingToolset/Source/MeshModelingTools/Private/MeshSelectionTool.cpp
--- a/Engine/Plugins/Experimental/MeshModelingToolset/Source/MeshModelingTools/Private/MeshSelectionTool.cpp
+++ b/Engine/Plugins/Experimental/MeshModelingToolset/Source/MeshModelingTools/Private/MeshSelectionTool.cpp
@@ -1259,15 +1259,8 @@ void UMeshSelectionTool::DisconnectSelectedTriangles()
 		FMeshRegionBoundaryLoops BoundaryLoops(&Mesh, SelectedFaces);
 		for (const FEdgeLoop& Loop : BoundaryLoops.Loops)
 		{
-			for (int VID : Loop.Vertices)
-			{
-				ChangeTracker.SaveVertex(VID);
-				// include the whole one-ring in case the disconnect creates bowties that need to be split
-				for (int TID : Mesh.VtxTrianglesItr(VID))
-				{
-					ChangeTracker.SaveTriangle(TID, true);
-				}
-			}
+			// include the whole one-ring in case the disconnect creates bowties that need to be split
+			ChangeTracker.SaveVertexOneRingTriangles(Loop.Vertices, true);
 		}
 
 		FDynamicMeshEditor Editor(&Mesh);
