Remove checkSlow calls to function that does not exist
#rnx

diff --git a/Engine/Plugins/Experimental/GeometryProcessing/Source/DynamicMesh/Private/MeshRefinerBase.cpp b/Engine/Plugins/Experimental/GeometryProcessing/Source/DynamicMesh/Private/MeshRefinerBase.cpp
--- a/Engine/Plugins/Experimental/GeometryProcessing/Source/DynamicMesh/Private/MeshRefinerBase.cpp
+++ b/Engine/Plugins/Experimental/GeometryProcessing/Source/DynamicMesh/Private/MeshRefinerBase.cpp
@@ -243,7 +243,6 @@ void FMeshRefinerBase::SaveTriangleBeforeModify(int32 TriangleID)
 	if (ActiveChangeTracker)
 	{
 		ActiveChangeTracker->SaveTriangle(TriangleID, true);
-		checkSlow(ActiveChangeTracker->VerifyTriangleIsSaved(TriangleID));
 	}
 }
 
@@ -254,7 +253,6 @@ void FMeshRefinerBase::SaveVertexTrianglesBeforeModify(int32 VertexID)
 		for (int32 TriangleID : Mesh->VtxTrianglesItr(VertexID))
 		{
 			ActiveChangeTracker->SaveTriangle(TriangleID, true );
-			checkSlow(ActiveChangeTracker->VerifyTriangleIsSaved(TriangleID));
 		}
 	}
 }
