fix crash on edit normals if mesh has duplicate triangles

#rnx

diff --git a/Engine/Plugins/Experimental/GeometryProcessing/Source/MeshConversion/Private/DynamicMeshToMeshDescription.cpp b/Engine/Plugins/Experimental/GeometryProcessing/Source/MeshConversion/Private/DynamicMeshToMeshDescription.cpp
--- a/Engine/Plugins/Experimental/GeometryProcessing/Source/MeshConversion/Private/DynamicMeshToMeshDescription.cpp
+++ b/Engine/Plugins/Experimental/GeometryProcessing/Source/MeshConversion/Private/DynamicMeshToMeshDescription.cpp
@@ -103,6 +103,7 @@ void FDynamicMeshToMeshDescription::UpdateAttributes(const FDynamicMesh3* MeshIn
 		{
 			if (MeshIn->HasAttributes())
 			{
+				check(MeshIn->TriangleCount() == MeshOut.Triangles().Num())
 				for (int UVLayerIndex = 0, NumLayers = MeshIn->Attributes()->NumUVLayers(); UVLayerIndex < NumLayers; UVLayerIndex++)
 				{
 					DynamicMeshToMeshDescriptionConversionHelper::SetAttributesFromOverlay(MeshIn, MeshOut, InstanceAttrib, MeshIn->Attributes()->GetUVLayer(UVLayerIndex), UVLayerIndex);
@@ -110,6 +111,7 @@ void FDynamicMeshToMeshDescription::UpdateAttributes(const FDynamicMesh3* MeshIn
 			}
 			else
 			{
+				check(MeshIn->VertexCount() == MeshOut.Vertices().Num());
 				for (int VertID : MeshIn->VertexIndicesItr())
 				{
 					FVector2D UV = (FVector2D)MeshIn->GetVertexUV(VertID);
diff --git a/Engine/Plugins/Experimental/GeometryProcessing/Source/MeshConversion/Public/DynamicMeshToMeshDescription.h b/Engine/Plugins/Experimental/GeometryProcessing/Source/MeshConversion/Public/DynamicMeshToMeshDescription.h
--- a/Engine/Plugins/Experimental/GeometryProcessing/Source/MeshConversion/Public/DynamicMeshToMeshDescription.h
+++ b/Engine/Plugins/Experimental/GeometryProcessing/Source/MeshConversion/Public/DynamicMeshToMeshDescription.h
@@ -28,6 +28,43 @@ public:
 	{
 	}
 
+	/**
+	 * Checks if element counts match. If false then Update can't be called -- you must call Convert
+	 *
+	 * @param DynamicMesh The dynamic mesh with updated vertices or attributes
+	 * @param MeshDescription The corresponding mesh description
+	 * @param bVerticesOnly If true, only check vertex counts match
+	 * @param bAttributesOnly If true, only check what needs to be checked for UpdateAttributes
+							 (will check vertices or triangles depending on whether attributes are per vertex or in overlays)
+	 */
+	static bool HaveMatchingElementCounts(const FDynamicMesh3* DynamicMesh, const FMeshDescription* MeshDescription, bool bVerticesOnly, bool bAttributesOnly)
+	{
+		bool bVerticesMatch = DynamicMesh->IsCompactV() && DynamicMesh->VertexCount() == MeshDescription->Vertices().Num();
+		bool bTrianglesMatch = DynamicMesh->IsCompactT() && DynamicMesh->TriangleCount() == MeshDescription->Triangles().Num();
+		if (bVerticesOnly || (bAttributesOnly && !DynamicMesh->HasAttributes()))
+		{
+			return bVerticesMatch;
+		}
+		else if (bAttributesOnly && DynamicMesh->HasAttributes())
+		{
+			return bTrianglesMatch;
+		}
+		return bVerticesMatch && bTrianglesMatch;
+	}
+
+	/**
+	 * Checks if element counts match. If false then Update can't be called -- you must call Convert
+	 * Result is based on the current ConversionOptions (e.g. if you are only updating vertices, mismatched triangles are ok)
+	 *
+	 * @param DynamicMesh The dynamic mesh with updated vertices or attributes
+	 * @param MeshDescription The corresponding mesh description
+	 */
+	bool HaveMatchingElementCounts(const FDynamicMesh3* DynamicMesh, const FMeshDescription* MeshDescription)
+	{
+		bool bUpdateAttributes = ConversionOptions.bUpdateNormals || ConversionOptions.bUpdateUVs;
+		return HaveMatchingElementCounts(DynamicMesh, MeshDescription, !bUpdateAttributes, !ConversionOptions.bUpdatePositions);
+	}
+
 	/**
 	 * Default conversion of DynamicMesh to MeshDescription. Calls functions below depending on mesh state
 	 */
diff --git a/Engine/Plugins/Experimental/MeshModelingToolset/Source/MeshModelingTools/Private/EditNormalsTool.cpp b/Engine/Plugins/Experimental/MeshModelingToolset/Source/MeshModelingTools/Private/EditNormalsTool.cpp
--- a/Engine/Plugins/Experimental/MeshModelingToolset/Source/MeshModelingTools/Private/EditNormalsTool.cpp
+++ b/Engine/Plugins/Experimental/MeshModelingToolset/Source/MeshModelingTools/Private/EditNormalsTool.cpp
@@ -269,7 +269,7 @@ void UEditNormalsTool::GenerateAsset(const TArray<FDynamicMeshOpResult>& Results
 		{
 			FDynamicMeshToMeshDescription Converter;
 			
-			if (BasicProperties->WillTopologyChange())
+			if (BasicProperties->WillTopologyChange() || !FDynamicMeshToMeshDescription::HaveMatchingElementCounts(Results[ComponentIdx].Mesh.Get(), CommitParams.MeshDescription, false, true))
 			{
 				// full conversion if normal topology changed or faces were inverted
 				Converter.Convert(Results[ComponentIdx].Mesh.Get(), *CommitParams.MeshDescription);
diff --git a/Engine/Plugins/Experimental/MeshModelingToolset/Source/ModelingComponents/Private/OctreeDynamicMeshComponent.cpp b/Engine/Plugins/Experimental/MeshModelingToolset/Source/ModelingComponents/Private/OctreeDynamicMeshComponent.cpp
--- a/Engine/Plugins/Experimental/MeshModelingToolset/Source/ModelingComponents/Private/OctreeDynamicMeshComponent.cpp
+++ b/Engine/Plugins/Experimental/MeshModelingToolset/Source/ModelingComponents/Private/OctreeDynamicMeshComponent.cpp
@@ -125,14 +125,13 @@ void UOctreeDynamicMeshComponent::ApplyTransform(const FTransform3d& Transform,
 
 void UOctreeDynamicMeshComponent::Bake(FMeshDescription* MeshDescription, bool bHaveModifiedTopology, const FConversionToMeshDescriptionOptions& ConversionOptions)
 {
-	if (bHaveModifiedTopology == false && Mesh.Get()->VertexCount() == MeshDescription->Vertices().Num())
+	FDynamicMeshToMeshDescription Converter(ConversionOptions);
+	if (bHaveModifiedTopology == false && Converter.HaveMatchingElementCounts(Mesh.Get(), MeshDescription))
 	{
-		FDynamicMeshToMeshDescription Converter(ConversionOptions);
 		Converter.Update(Mesh.Get(), *MeshDescription);
 	}
 	else
 	{
-		FDynamicMeshToMeshDescription Converter(ConversionOptions);
 		Converter.Convert(Mesh.Get(), *MeshDescription);
 
 		//UE_LOG(LogTemp, Warning, TEXT("MeshDescription has %d instances"), MeshDescription->VertexInstances().Num());
diff --git a/Engine/Plugins/Experimental/MeshModelingToolset/Source/ModelingComponents/Private/SimpleDynamicMeshComponent.cpp b/Engine/Plugins/Experimental/MeshModelingToolset/Source/ModelingComponents/Private/SimpleDynamicMeshComponent.cpp
--- a/Engine/Plugins/Experimental/MeshModelingToolset/Source/ModelingComponents/Private/SimpleDynamicMeshComponent.cpp
+++ b/Engine/Plugins/Experimental/MeshModelingToolset/Source/ModelingComponents/Private/SimpleDynamicMeshComponent.cpp
@@ -85,9 +85,9 @@ void USimpleDynamicMeshComponent::InitializeNewMesh()
 
 void USimpleDynamicMeshComponent::Bake(FMeshDescription* MeshDescription, bool bHaveModifiedTopology, const FConversionToMeshDescriptionOptions& ConversionOptions)
 {
-	if (bHaveModifiedTopology == false && Mesh.Get()->VertexCount() == MeshDescription->Vertices().Num())
+	FDynamicMeshToMeshDescription Converter(ConversionOptions);
+	if (bHaveModifiedTopology == false && Converter.HaveMatchingElementCounts(Mesh.Get(), MeshDescription))
 	{
-		FDynamicMeshToMeshDescription Converter(ConversionOptions);
 		if (ConversionOptions.bUpdatePositions)
 		{
 			Converter.Update(Mesh.Get(), *MeshDescription, ConversionOptions.bUpdateNormals, ConversionOptions.bUpdateUVs);
@@ -99,7 +99,6 @@ void USimpleDynamicMeshComponent::Bake(FMeshDescription* MeshDescription, bool b
 	}
 	else
 	{
-		FDynamicMeshToMeshDescription Converter(ConversionOptions);
 		Converter.Convert(Mesh.Get(), *MeshDescription);
 
 		//UE_LOG(LogTemp, Warning, TEXT("MeshDescription has %d instances"), MeshDescription->VertexInstances().Num());
