Change Modeling Mode plugins to be Editor-only temporarily
#rnx

diff --git a/Engine/Plugins/Experimental/GeometryProcessing/GeometryProcessing.uplugin b/Engine/Plugins/Experimental/GeometryProcessing/GeometryProcessing.uplugin
--- a/Engine/Plugins/Experimental/GeometryProcessing/GeometryProcessing.uplugin
+++ b/Engine/Plugins/Experimental/GeometryProcessing/GeometryProcessing.uplugin
@@ -16,27 +16,27 @@
 	"Modules": [
 		{
 			"Name": "GeometricObjects",
-			"Type": "Runtime",
+			"Type": "Editor",
 			"LoadingPhase": "Default"
 		},
 		{
 			"Name": "GeometryAlgorithms",
-			"Type": "Runtime",
+			"Type": "Editor",
 			"LoadingPhase": "Default"
 		},
 		{
 			"Name": "DynamicMesh",
-			"Type": "Runtime",
+			"Type": "Editor",
 			"LoadingPhase": "Default"
 		},
 		{
 			"Name": "MeshConversion",
-			"Type": "Runtime",
+			"Type": "Editor",
 			"LoadingPhase": "Default"
 		},
 		{
 			"Name": "MeshSolverUtilities",
-			"Type": "Runtime",
+			"Type": "Editor",
 			"LoadingPhase": "Default"
 		}
 	]
diff --git a/Engine/Plugins/Experimental/MeshModelingToolset/MeshModelingToolset.uplugin b/Engine/Plugins/Experimental/MeshModelingToolset/MeshModelingToolset.uplugin
--- a/Engine/Plugins/Experimental/MeshModelingToolset/MeshModelingToolset.uplugin
+++ b/Engine/Plugins/Experimental/MeshModelingToolset/MeshModelingToolset.uplugin
@@ -17,13 +17,13 @@
 	"Modules": [
 		{
 			"Name": "ModelingComponents",
-			"Type": "Runtime",
+			"Type": "Editor",
 			"LoadingPhase": "Default",
 			"WhitelistPlatforms": [ "Win64" ]
 		},
 		{
 			"Name": "ModelingOperators",
-			"Type": "Runtime",
+			"Type": "Editor",
 			"LoadingPhase": "Default",
 			"WhitelistPlatforms": [ "Win64" ]
 		},
@@ -35,7 +35,7 @@
 		},
 		{
 			"Name": "MeshModelingTools",
-			"Type": "Runtime",
+			"Type": "Editor",
 			"LoadingPhase": "Default",
 			"WhitelistPlatforms": [ "Win64" ]
 		},
