Edigrate Dev-Editor updates to InteractiveToolsFramework, GeometryProcessing, MeshModelingToolset, and ModelingTools Editor Mode Plugins. Update a few files in new MeshPaint Plugin for UInteractiveTool rename of ::Tick() to ::OnTick().
#rnx

diff --git a/Engine/Plugins/Experimental/GeometryProcessing/Source/DynamicMesh/Private/DynamicMesh3.cpp b/Engine/Plugins/Experimental/GeometryProcessing/Source/DynamicMesh/Private/DynamicMesh3.cpp
--- a/Engine/Plugins/Experimental/GeometryProcessing/Source/DynamicMesh/Private/DynamicMesh3.cpp
+++ b/Engine/Plugins/Experimental/GeometryProcessing/Source/DynamicMesh/Private/DynamicMesh3.cpp
@@ -156,34 +156,63 @@ void FDynamicMesh3::Copy(const FMeshShapeGenerator* Generator)
 	Edges = TDynamicVector<int>();
 	EdgeRefCounts = FRefCountVector();
 
-	EnableAttributes();
-	FDynamicMeshUVOverlay* UVOverlay = Attributes()->PrimaryUV();
-	FDynamicMeshNormalOverlay* NormalOverlay = Attributes()->PrimaryNormals();
+	if (Generator->HasAttributes())
+	{
+		EnableAttributes();
+	}
+	
 
 	int NumVerts = Generator->Vertices.Num();
 	for (int i = 0; i < NumVerts; ++i)
 	{
 		AppendVertex(Generator->Vertices[i]);
 	}
-	int NumUVs = Generator->UVs.Num();
-	for (int i = 0; i < NumUVs; ++i)
+
+	if (Generator->HasAttributes())
 	{
-		UVOverlay->AppendElement(Generator->UVs[i], Generator->UVParentVertex[i]);
+		FDynamicMeshUVOverlay* UVOverlay = Attributes()->PrimaryUV();
+		FDynamicMeshNormalOverlay* NormalOverlay = Attributes()->PrimaryNormals();
+		int NumUVs = Generator->UVs.Num();
+		for (int i = 0; i < NumUVs; ++i)
+		{
+			UVOverlay->AppendElement(Generator->UVs[i]);
+		}
+		int NumNormals = Generator->Normals.Num();
+		for (int i = 0; i < NumNormals; ++i)
+		{
+			NormalOverlay->AppendElement(Generator->Normals[i]);
+		}
+
+		int NumTris = Generator->Triangles.Num();
+		for (int i = 0; i < NumTris; ++i)
+		{
+			int PolyID = Generator->TrianglePolygonIDs.Num() > 0 ? Generator->TrianglePolygonIDs[i] : i;
+			int tid = AppendTriangle(Generator->Triangles[i], PolyID);
+			check(tid == i);
+			UVOverlay->SetTriangle(tid, Generator->TriangleUVs[i]);
+			NormalOverlay->SetTriangle(tid, Generator->TriangleNormals[i]);
+		}
 	}
-	int NumNormals = Generator->Normals.Num();
-	for (int i = 0; i < NumNormals; ++i)
+	else if (Generator->TrianglePolygonIDs.Num()) // no attributes, yes polygon ids
 	{
-		NormalOverlay->AppendElement(Generator->Normals[i], Generator->NormalParentVertex[i]);
+		int NumTris = Generator->Triangles.Num();
+		for (int i = 0; i < NumTris; ++i)
+		{
+			int tid = AppendTriangle(Generator->Triangles[i], Generator->TrianglePolygonIDs[i]);
+			check(tid == i);
+		}
 	}
-
-	int NumTris = Generator->Triangles.Num();
-	for (int i = 0; i < NumTris; ++i)
+	else // no attribute and no polygon ids
 	{
-		int tid = AppendTriangle(Generator->Triangles[i], Generator->TrianglePolygonIDs[i]);
-		check(tid == i);
-		UVOverlay->SetTriangle(tid, Generator->TriangleUVs[i]);
-		NormalOverlay->SetTriangle(tid, Generator->TriangleNormals[i]);
+		int NumTris = Generator->Triangles.Num();
+		for (int i = 0; i < NumTris; ++i)
+		{
+			int tid = AppendTriangle(Generator->Triangles[i], i);
+			check(tid == i);
+		}
 	}
+
+
 }
 
 void FDynamicMesh3::Copy(const FDynamicMesh3& copy, bool bNormals, bool bColors, bool bUVs, bool bAttributes)
@@ -529,7 +558,7 @@ FIndex3i FDynamicMesh3::GetTriNeighbourTris(int tID) const
 void FDynamicMesh3::GetVertexOneRingTriangles(int VertexID, TArray<int>& TrianglesOut) const
 {
 	check(VertexRefCounts.IsValid(VertexID));
-	for (int EdgeID : VertexEdgeLists.Values(VertexID))
+	VertexEdgeLists.Enumerate(VertexID, [&](int32 EdgeID)
 	{
 		FIndex2i EdgePair = GetOrderedOneRingEdgeTris(VertexID, EdgeID);
 		if (EdgePair.A != InvalidID)
@@ -540,9 +569,28 @@ void FDynamicMesh3::GetVertexOneRingTriangles(int VertexID, TArray<int>& Triangl
 		{
 			TrianglesOut.Add(EdgePair.B);
 		}
-	}
+	});
 }
 
+
+void FDynamicMesh3::EnumerateVertexTriangles(int32 VertexID, TFunctionRef<void(int32)> ApplyFunc) const
+{
+	check(VertexRefCounts.IsValid(VertexID));
+	VertexEdgeLists.Enumerate(VertexID, [&](int32 EdgeID)
+	{
+		FIndex2i EdgePair = GetOrderedOneRingEdgeTris(VertexID, EdgeID);
+		if (EdgePair.A != InvalidID)
+		{
+			ApplyFunc(EdgePair.A);
+		}
+		if (EdgePair.B != InvalidID)
+		{
+			ApplyFunc(EdgePair.B);
+		}
+	});
+}
+
+
 FString FDynamicMesh3::MeshInfoString()
 {
 	FString VtxString = FString::Printf(TEXT("Vertices count %d max %d  %s  VtxEdges %s"),
@@ -1006,6 +1054,25 @@ int FDynamicMesh3::FindTriangleEdge(int tID, int vA, int vB) const
 }
 
 
+int32 FDynamicMesh3::FindEdgeInternal(int32 vA, int32 vB, bool& bIsBoundary) const
+{
+	// edge vertices must be sorted (min,max), that means we only need one index-check in inner loop.
+	int32 vMax = vA, vMin = vB;
+	if (vB > vA)
+	{
+		vMax = vB; vMin = vA;
+	}
+	return VertexEdgeLists.Find(vMin, [&](int32 eid)
+	{
+		int32 idx = 4 * eid;
+		if (Edges[idx + 1] == vMax)
+		{
+			bIsBoundary = (Edges[idx + 3] == InvalidID);
+			return true;
+		}
+		return false;
+	}, InvalidID);
+}
 
 
 
@@ -1020,20 +1087,17 @@ int FDynamicMesh3::FindEdge(int vA, int vB) const
 		return InvalidID;
 	}
 
-	// [RMS] edge vertices must be sorted (min,max),
+	// edge vertices must be sorted (min,max),
 	//   that means we only need one index-check in inner loop.
-	//   commented out code is robust to incorrect ordering, but slower.
-	int vO = FMath::Max(vA, vB);
-	int vI = FMath::Min(vA, vB);
-	for (int eid : VertexEdgeLists.Values(vI)) 
+	int32 vMax = vA, vMin = vB;
+	if (vB > vA)
 	{
-		if (Edges[4 * eid + 1] == vO)
-		{
-			//if (DoesEdgeHaveVertex(eid, vO))
-			return eid;
-		}
+		vMax = vB; vMin = vA;
 	}
-	return InvalidID;
+	return VertexEdgeLists.Find(vMin, [&](int32 eid)
+	{
+		return (Edges[4*eid + 1] == vMax);
+	}, InvalidID);
 
 	// this is slower, likely because it creates func<> every time. can we do w/o that?
 	//return VertexEdgeLists.Find(vI, (eid) => { return Edges[4 * eid + 1] == vO; }, InvalidID);
diff --git a/Engine/Plugins/Experimental/GeometryProcessing/Source/DynamicMesh/Private/DynamicMesh3_Edits.cpp b/Engine/Plugins/Experimental/GeometryProcessing/Source/DynamicMesh/Private/DynamicMesh3_Edits.cpp
--- a/Engine/Plugins/Experimental/GeometryProcessing/Source/DynamicMesh/Private/DynamicMesh3_Edits.cpp
+++ b/Engine/Plugins/Experimental/GeometryProcessing/Source/DynamicMesh/Private/DynamicMesh3_Edits.cpp
@@ -195,12 +195,13 @@ int FDynamicMesh3::AppendTriangle(const FIndex3i& tv, int gid)
 
 	// look up edges. if any already have two triangles, this would 
 	// create non-manifold geometry and so we do not allow it
-	int e0 = FindEdge(tv[0], tv[1]);
-	int e1 = FindEdge(tv[1], tv[2]);
-	int e2 = FindEdge(tv[2], tv[0]);
-	if ((e0 != InvalidID && IsBoundaryEdge(e0) == false)
-		|| (e1 != InvalidID && IsBoundaryEdge(e1) == false)
-		|| (e2 != InvalidID && IsBoundaryEdge(e2) == false)) 
+	bool boundary0, boundary1, boundary2;
+	int e0 = FindEdgeInternal(tv[0], tv[1], boundary0);
+	int e1 = FindEdgeInternal(tv[1], tv[2], boundary1);
+	int e2 = FindEdgeInternal(tv[2], tv[0], boundary2);
+	if ((e0 != InvalidID && boundary0 == false)
+		|| (e1 != InvalidID && boundary1 == false)
+		|| (e2 != InvalidID && boundary2 == false))
 	{
 		return NonManifoldID;
 	}
@@ -224,10 +225,10 @@ int FDynamicMesh3::AppendTriangle(const FIndex3i& tv, int gid)
 	// now safe to insert triangle
 	int tid = TriangleRefCounts.Allocate();
 	int i = 3 * tid;
-	Triangles.InsertAt(tv[2], i + 2);
-	Triangles.InsertAt(tv[1], i + 1);
 	Triangles.InsertAt(tv[0], i);
-	if (bHasGroups) 
+	Triangles.InsertAt(tv[1], i + 1);
+	Triangles.InsertAt(tv[2], i + 2);
+	if (bHasGroups)
 	{
 		TriangleGroups->InsertAt(gid, tid);
 		GroupIDCounter = FMath::Max(GroupIDCounter, gid + 1);
@@ -1401,7 +1402,7 @@ EMeshResult FDynamicMesh3::CollapseEdge(int vKeep, int vRemove, double collapse_
 		}
 	}
 
-	// [RMS] I am not sure this tetrahedron case will detect bowtie vertices.
+	// I am not sure this tetrahedron case will detect bowtie vertices.
 	// But the single-triangle case does
 
 	// We cannot collapse if we have a tetrahedron. In this case a has 3 nbr edges,
@@ -1434,9 +1435,7 @@ EMeshResult FDynamicMesh3::CollapseEdge(int vKeep, int vRemove, double collapse_
 		}
 	}
 
-	// [RMS] this was added from C++ version...seems like maybe I never hit
-	//   this case? Conceivably could be porting bug but looking at the
-	//   C++ code I cannot see how we could possibly have caught this case...
+	// TODO: it's unclear how this case would ever be encountered; check if it is needed
 	//
 	// cannot collapse an edge where both vertices are boundary vertices
 	// because that would create a bowtie
diff --git a/Engine/Plugins/Experimental/GeometryProcessing/Source/DynamicMesh/Private/DynamicMesh3_Queries.cpp b/Engine/Plugins/Experimental/GeometryProcessing/Source/DynamicMesh/Private/DynamicMesh3_Queries.cpp
--- a/Engine/Plugins/Experimental/GeometryProcessing/Source/DynamicMesh/Private/DynamicMesh3_Queries.cpp
+++ b/Engine/Plugins/Experimental/GeometryProcessing/Source/DynamicMesh/Private/DynamicMesh3_Queries.cpp
@@ -127,6 +127,23 @@ int FDynamicMesh3::GetVtxTriangleCount(int vID, bool bBruteForce) const
 }
 
 
+int FDynamicMesh3::GetVtxSingleTriangle(int VertexID) const
+{
+	if (!IsVertex(VertexID))
+	{
+		return FDynamicMesh3::InvalidID;
+	}
+
+	for (int EID : VertexEdgeLists.Values(VertexID))
+	{
+		int i = 4 * EID;
+		int T0 = Edges[i + 2];
+		return T0;
+	}
+
+	return FDynamicMesh3::InvalidID;
+}
+
 
 EMeshResult FDynamicMesh3::GetVtxTriangles(int vID, TArray<int>& TrianglesOut, bool bUseOrientation) const
 {
diff --git a/Engine/Plugins/Experimental/GeometryProcessing/Source/DynamicMesh/Private/DynamicMeshChangeTracker.cpp b/Engine/Plugins/Experimental/GeometryProcessing/Source/DynamicMesh/Private/DynamicMeshChangeTracker.cpp
--- a/Engine/Plugins/Experimental/GeometryProcessing/Source/DynamicMesh/Private/DynamicMeshChangeTracker.cpp
+++ b/Engine/Plugins/Experimental/GeometryProcessing/Source/DynamicMesh/Private/DynamicMeshChangeTracker.cpp
@@ -112,7 +112,7 @@ void TDynamicMeshAttributeChange<RealType,ElementSize>::ApplyReplaceChange(TDyna
 		if (Overlay->IsElement(ElemInfo.ElementID) == false)
 		{
 			const RealType* Data = &InsertElementData[ElemInfo.DataIndex];
-			EMeshResult Result = Overlay->InsertElement(ElemInfo.ElementID, Data, ElemInfo.ParentVertexID, bUseUnsafe);
+			EMeshResult Result = Overlay->InsertElement(ElemInfo.ElementID, Data, bUseUnsafe);
 			check(Result == EMeshResult::Ok);
 		}
 		else
diff --git a/Engine/Plugins/Experimental/GeometryProcessing/Source/DynamicMesh/Private/DynamicMeshEditor.cpp b/Engine/Plugins/Experimental/GeometryProcessing/Source/DynamicMesh/Private/DynamicMeshEditor.cpp
--- a/Engine/Plugins/Experimental/GeometryProcessing/Source/DynamicMesh/Private/DynamicMeshEditor.cpp
+++ b/Engine/Plugins/Experimental/GeometryProcessing/Source/DynamicMesh/Private/DynamicMeshEditor.cpp
@@ -277,7 +277,7 @@ void FDynamicMeshEditor::DuplicateTriangles(const TArray<int>& Triangles, FMeshI
 	{
 		FIndex3i Tri = Mesh->GetTriangle(TriangleID);
 
-		int NewGroupID = FindOrCreateDuplicateGroup(TriangleID, IndexMaps, ResultOut);
+		int NewGroupID = Mesh->HasTriangleGroups() ? FindOrCreateDuplicateGroup(TriangleID, IndexMaps, ResultOut) : -1;
 
 		FIndex3i NewTri;
 		NewTri[0] = FindOrCreateDuplicateVertex(Tri[0], IndexMaps, ResultOut);
@@ -535,7 +535,7 @@ bool FDynamicMeshEditor::ReinsertSubmesh(const FDynamicSubmesh3& Region, FOption
 					int base_v = Region.MapVertexToBaseMesh(sub_v);
 					if (base_v >= 0 && Mesh->IsVertex(base_v) && Region.InBaseBorderVertices(base_v) == true)
 					{
-						// [RMS] this should always be true, but assert in tests to find out
+						// this should always be true
 						if (ensure(Mesh->IsBoundaryVertex(base_v)))
 						{
 							new_v = base_v;
@@ -638,9 +638,9 @@ void FDynamicMeshEditor::SetQuadNormals(const FIndex2i& QuadTris, const FVector3
 	FIndex3i Triangle1 = Mesh->GetTriangle(QuadTris.A);
 
 	FIndex3i NormalTriangle1;
-	NormalTriangle1[0] = Normals->AppendElement(Normal, Triangle1[0]);
-	NormalTriangle1[1] = Normals->AppendElement(Normal, Triangle1[1]);
-	NormalTriangle1[2] = Normals->AppendElement(Normal, Triangle1[2]);
+	NormalTriangle1[0] = Normals->AppendElement(Normal);
+	NormalTriangle1[1] = Normals->AppendElement(Normal);
+	NormalTriangle1[2] = Normals->AppendElement(Normal);
 	Normals->SetTriangle(QuadTris.A, NormalTriangle1);
 
 	if (Mesh->IsTriangle(QuadTris.B))
@@ -652,7 +652,7 @@ void FDynamicMeshEditor::SetQuadNormals(const FIndex2i& QuadTris, const FVector3
 			int i = Triangle1.IndexOf(Triangle2[j]);
 			if (i == -1)
 			{
-				NormalTriangle2[j] = Normals->AppendElement(Normal, Triangle2[j]);
+				NormalTriangle2[j] = Normals->AppendElement(Normal);
 			}
 			else
 			{
@@ -681,7 +681,7 @@ void FDynamicMeshEditor::SetTriangleNormals(const TArray<int>& Triangles, const
 			const int* FoundElementID = Vertices.Find(BaseTri[j]);
 			if (FoundElementID == nullptr)
 			{
-				ElemTri[j] = Normals->AppendElement(Normal, BaseTri[j]);
+				ElemTri[j] = Normals->AppendElement(Normal);
 				Vertices.Add(BaseTri[j], ElemTri[j]);
 			}
 			else
@@ -715,7 +715,7 @@ void FDynamicMeshEditor::SetTriangleNormals(const TArray<int>& Triangles)
 			if (FoundElementID == nullptr)
 			{
 				FVector3d VtxROINormal = FMeshNormals::ComputeVertexNormal(*Mesh, BaseTri[j], TrianglePredicate);
-				ElemTri[j] = Normals->AppendElement( (FVector3f)VtxROINormal, BaseTri[j]);
+				ElemTri[j] = Normals->AppendElement( (FVector3f)VtxROINormal);
 				Vertices.Add(BaseTri[j], ElemTri[j]);
 			}
 			else
@@ -808,8 +808,7 @@ void FDynamicMeshEditor::SetTubeNormals(const TArray<int>& Triangles, const TArr
 		for (int Idx = 0; Idx < NumVertices; Idx++)
 		{
 			int VID = VertexIDs[Idx];
-			VertToElID.Add(VID, Normals->AppendElement(VertNormals[Side][Idx], VID));
-
+			VertToElID.Add(VID, Normals->AppendElement(VertNormals[Side][Idx]));
 		}
 	}
 	for (int TID : Triangles)
@@ -891,9 +890,9 @@ void FDynamicMeshEditor::SetGeneralTubeUVs(const TArray<int>& Triangles,
 		for (int Idx = 0; Idx < NumVertices; Idx++)
 		{
 			int VID = VertexIDs[Idx];
-			VertToElID.Add(VID, UVs->AppendElement(VertUVs[Side][Idx], VID));
+			VertToElID.Add(VID, UVs->AppendElement(VertUVs[Side][Idx]));
 		}
-		DuplicateMappingForLastVert[Side] = UVs->AppendElement(VertUVs[Side].Last(), VertexIDs[0]);
+		DuplicateMappingForLastVert[Side] = UVs->AppendElement(VertUVs[Side].Last());
 	}
 	
 	bool bPastInitialVertices[2]{ false, false };
@@ -945,7 +944,7 @@ void FDynamicMeshEditor::SetTriangleUVsFromProjection(const TArray<int>& Triangl
 			{
 				FVector2f UV = (FVector2f)ProjectionFrame.ToPlaneUV(Mesh->GetVertex(BaseTri[j]), 2);
 				UVBounds.Contain(UV);
-				ElemTri[j] = UVs->AppendElement(UV, BaseTri[j]);
+				ElemTri[j] = UVs->AppendElement(UV);
 				AllUVIndices.Add(ElemTri[j]);
 				BaseToOverlayVIDMap.Add(BaseTri[j], ElemTri[j]);
 			}
@@ -982,7 +981,7 @@ void FDynamicMeshEditor::SetQuadUVsFromProjection(const FIndex2i& QuadTris, cons
 	for (int j = 0; j < 3; ++j)
 	{
 		FVector2f UV = (FVector2f)ProjectionFrame.ToPlaneUV(Mesh->GetVertex(Triangle1[j]), 2);
-		UVTriangle1[j] = UVs->AppendElement(UV, Triangle1[j]);
+		UVTriangle1[j] = UVs->AppendElement(UV);
 		AllUVs[j] = UV;
 		AllUVIndices[j] = UVTriangle1[j];
 	}
@@ -999,7 +998,7 @@ void FDynamicMeshEditor::SetQuadUVsFromProjection(const FIndex2i& QuadTris, cons
 			if (i == -1)
 			{
 				FVector2f UV = (FVector2f)ProjectionFrame.ToPlaneUV(Mesh->GetVertex(Triangle2[j]), 2);
-				UVTriangle2[j] = UVs->AppendElement(UV, Triangle2[j]);
+				UVTriangle2[j] = UVs->AppendElement(UV);
 				AllUVs[3] = UV;
 				AllUVIndices[3] = UVTriangle2[j];
 			}
@@ -1193,15 +1192,7 @@ int FDynamicMeshEditor::FindOrCreateDuplicateUV(int ElementID, int UVLayerIndex,
 	if (NewElementID == IndexMaps.InvalidID())
 	{
 		FDynamicMeshUVOverlay* UVOverlay = Mesh->Attributes()->GetUVLayer(UVLayerIndex);
-
-		// need to determine new parent vertex. It should be in the map already!
-		int ParentVertexID = UVOverlay->GetParentVertex(ElementID);
-		int NewParentVertexID = IndexMaps.GetNewVertex(ParentVertexID);
-		check(NewParentVertexID != IndexMaps.InvalidID());
-
-		NewElementID = UVOverlay->AppendElement(
-			UVOverlay->GetElement(ElementID), NewParentVertexID);
-
+		NewElementID = UVOverlay->AppendElement(UVOverlay->GetElement(ElementID));
 		IndexMaps.SetUV(UVLayerIndex, ElementID, NewElementID);
 	}
 	return NewElementID;
@@ -1215,15 +1206,7 @@ int FDynamicMeshEditor::FindOrCreateDuplicateNormal(int ElementID, int NormalLay
 	if (NewElementID == IndexMaps.InvalidID())
 	{
 		FDynamicMeshNormalOverlay* NormalOverlay = Mesh->Attributes()->GetNormalLayer(NormalLayerIndex);
-
-		// need to determine new parent vertex. It should be in the map already!
-		int ParentVertexID = NormalOverlay->GetParentVertex(ElementID);
-		int NewParentVertexID = IndexMaps.GetNewVertex(ParentVertexID);
-		check(NewParentVertexID != IndexMaps.InvalidID());
-
-		NewElementID = NormalOverlay->AppendElement(
-			NormalOverlay->GetElement(ElementID), NewParentVertexID);
-
+		NewElementID = NormalOverlay->AppendElement(NormalOverlay->GetElement(ElementID));
 		IndexMaps.SetNormal(NormalLayerIndex, ElementID, NewElementID);
 	}
 	return NewElementID;
@@ -1392,7 +1375,7 @@ void FDynamicMeshEditor::AppendNormals(const FDynamicMesh3* AppendMesh,
 		{
 			Normal = (FVector3f)NormalTransform(ParentVertID, (FVector3d)Normal);
 		}
-		int NewElemID = ToNormals->AppendElement(Normal, VertexMap.GetTo(ParentVertID));
+		int NewElemID = ToNormals->AppendElement(Normal);
 		NormalMapOut.Add(ElemID, NewElemID);
 	}
 
@@ -1421,9 +1404,8 @@ void FDynamicMeshEditor::AppendUVs(const FDynamicMesh3* AppendMesh,
 	// copy over uv elements
 	for (int ElemID : FromUVs->ElementIndicesItr())
 	{
-		int ParentVertID = FromUVs->GetParentVertex(ElemID);
 		FVector2f UV = FromUVs->GetElement(ElemID);
-		int NewElemID = ToUVs->AppendElement(UV, VertexMap.GetTo(ParentVertID));
+		int NewElemID = ToUVs->AppendElement(UV);
 		UVMapOut.Add(ElemID, NewElemID);
 	}
 
@@ -1464,15 +1446,7 @@ static int AppendTriangleUVAttribute(const FDynamicMesh3* FromMesh, int FromElem
 	{
 		const FDynamicMeshUVOverlay* FromUVOverlay = FromMesh->Attributes()->GetUVLayer(UVLayerIndex);
 		FDynamicMeshUVOverlay* ToUVOverlay = ToMesh->Attributes()->GetUVLayer(UVLayerIndex);
-
-		// need to determine new parent vertex. It should be in the map already!
-		int ParentVertexID = FromUVOverlay->GetParentVertex(FromElementID);
-		int NewParentVertexID = IndexMaps.GetNewVertex(ParentVertexID);
-		check(NewParentVertexID != IndexMaps.InvalidID());
-
-		NewElementID = ToUVOverlay->AppendElement(
-			FromUVOverlay->GetElement(FromElementID), NewParentVertexID);
-
+		NewElementID = ToUVOverlay->AppendElement(FromUVOverlay->GetElement(FromElementID));
 		IndexMaps.SetUV(UVLayerIndex, FromElementID, NewElementID);
 	}
 	return NewElementID;
@@ -1487,15 +1461,7 @@ static int AppendTriangleNormalAttribute(const FDynamicMesh3* FromMesh, int From
 	{
 		const FDynamicMeshNormalOverlay* FromNormalOverlay = FromMesh->Attributes()->GetNormalLayer(NormalLayerIndex);
 		FDynamicMeshNormalOverlay* ToNormalOverlay = ToMesh->Attributes()->GetNormalLayer(NormalLayerIndex);
-
-		// need to determine new parent vertex. It should be in the map already!
-		int ParentVertexID = FromNormalOverlay->GetParentVertex(FromElementID);
-		int NewParentVertexID = IndexMaps.GetNewVertex(ParentVertexID);
-		check(NewParentVertexID != IndexMaps.InvalidID());
-
-		NewElementID = ToNormalOverlay->AppendElement(
-			FromNormalOverlay->GetElement(FromElementID), NewParentVertexID);
-
+		NewElementID = ToNormalOverlay->AppendElement(FromNormalOverlay->GetElement(FromElementID));
 		IndexMaps.SetNormal(NormalLayerIndex, FromElementID, NewElementID);
 	}
 	return NewElementID;
diff --git a/Engine/Plugins/Experimental/GeometryProcessing/Source/DynamicMesh/Private/DynamicMeshOverlay.cpp b/Engine/Plugins/Experimental/GeometryProcessing/Source/DynamicMesh/Private/DynamicMeshOverlay.cpp
--- a/Engine/Plugins/Experimental/GeometryProcessing/Source/DynamicMesh/Private/DynamicMeshOverlay.cpp
+++ b/Engine/Plugins/Experimental/GeometryProcessing/Source/DynamicMesh/Private/DynamicMeshOverlay.cpp
@@ -18,7 +18,7 @@ void TDynamicMeshOverlay<RealType, ElementSize>::ClearElements()
 
 
 template<typename RealType, int ElementSize>
-int TDynamicMeshOverlay<RealType, ElementSize>::AppendElement(RealType ConstantValue, int ParentVertex)
+int TDynamicMeshOverlay<RealType, ElementSize>::AppendElement(RealType ConstantValue)
 {
 	int vid = ElementsRefCounts.Allocate();
 	int i = ElementSize * vid;
@@ -26,7 +26,7 @@ int TDynamicMeshOverlay<RealType, ElementSize>::AppendElement(RealType ConstantV
 	{
 		Elements.InsertAt(ConstantValue, i + k);
 	}
-	ParentVertices.InsertAt(ParentVertex, vid);
+	ParentVertices.InsertAt(FDynamicMesh3::InvalidID, vid);
 
 	//updateTimeStamp(true);
 	return vid;
@@ -34,7 +34,7 @@ int TDynamicMeshOverlay<RealType, ElementSize>::AppendElement(RealType ConstantV
 
 
 template<typename RealType, int ElementSize>
-int TDynamicMeshOverlay<RealType, ElementSize>::AppendElement(const RealType* Value, int ParentVertex)
+int TDynamicMeshOverlay<RealType, ElementSize>::AppendElement(const RealType* Value)
 {
 	int vid = ElementsRefCounts.Allocate();
 	int i = ElementSize * vid;
@@ -44,8 +44,7 @@ int TDynamicMeshOverlay<RealType, ElementSize>::AppendElement(const RealType* Va
 	{
 		Elements.InsertAt(Value[k], i + k);
 	}
-
-	ParentVertices.InsertAt(ParentVertex, vid);
+	ParentVertices.InsertAt(FDynamicMesh3::InvalidID, vid);
 
 	//updateTimeStamp(true);
 	return vid;
@@ -53,7 +52,7 @@ int TDynamicMeshOverlay<RealType, ElementSize>::AppendElement(const RealType* Va
 
 
 template<typename RealType, int ElementSize>
-EMeshResult TDynamicMeshOverlay<RealType, ElementSize>::InsertElement(int ElementID, const RealType* Value, int ParentVertex, bool bUnsafe)
+EMeshResult TDynamicMeshOverlay<RealType, ElementSize>::InsertElement(int ElementID, const RealType* Value, bool bUnsafe)
 {
 	if (ElementsRefCounts.IsValid(ElementID))
 	{
@@ -74,7 +73,7 @@ EMeshResult TDynamicMeshOverlay<RealType, ElementSize>::InsertElement(int Elemen
 		Elements.InsertAt(Value[k], i + k);
 	}
 
-	ParentVertices.InsertAt(ParentVertex, ElementID);
+	ParentVertices.InsertAt(FDynamicMesh3::InvalidID, ElementID);
 
 	//UpdateTimeStamp(true, true);
 	return EMeshResult::Ok;
@@ -143,7 +142,7 @@ void TDynamicMeshOverlay<RealType, ElementSize>::CreateFromPredicate(TFunctionRe
 
 			for (int Idx = 0; Idx < NumGroupID; Idx++)
 			{
-				AppendedElements.Add(AppendElement(InitElementValue, VertexID));
+				AppendedElements.Add(AppendElement(InitElementValue));
 			}
 			for (int TriSubIdx = 0; TriSubIdx < GroupNum; TriSubIdx++)
 			{
@@ -151,8 +150,10 @@ void TDynamicMeshOverlay<RealType, ElementSize>::CreateFromPredicate(TFunctionRe
 				FIndex3i TriVertIDs = ParentMesh->GetTriangle(TriID);
 				int VertSubIdx = IndexUtil::FindTriIndex(VertexID, TriVertIDs);
 				int i = 3 * TriID;
-				ElementTriangles.InsertAt(AppendedElements[TrisActiveSubGroup[TriSubIdx]], i + VertSubIdx);
-				ElementsRefCounts.Increment(AppendedElements[TrisActiveSubGroup[TriSubIdx]]);
+				int ElementIndex = AppendedElements[TrisActiveSubGroup[TriSubIdx]];
+				ElementTriangles.InsertAt(ElementIndex, i + VertSubIdx);
+				ElementsRefCounts.Increment(ElementIndex);
+				ParentVertices.InsertAt(VertexID, ElementIndex);
 			}
 			GroupStart += GroupNum;
 		}
@@ -161,7 +162,7 @@ void TDynamicMeshOverlay<RealType, ElementSize>::CreateFromPredicate(TFunctionRe
 
 
 template<typename RealType, int ElementSize>
-void TDynamicMeshOverlay<RealType, ElementSize>::SplitVerticesWithPredicate(TFunctionRef<bool(int ElementID, int TriID)> ShouldSplitOutVertex, TFunctionRef<void(int ElementID, int TriID, RealType* FillVect)> GetNewElementValue)
+void TDynamicMeshOverlay<RealType, ElementSize>::SplitVerticesWithPredicate(TFunctionRef<bool(int ElementID, int TriID)> ShouldSplitOutVertex, TFunctionRef<void(int ElementID, int TriID, RealType * FillVect)> GetNewElementValue)
 {
 	for (int TriID : ParentMesh->TriangleIndicesItr())
 	{
@@ -191,14 +192,13 @@ void TDynamicMeshOverlay<RealType, ElementSize>::SplitVerticesWithPredicate(TFun
 				TriChanged = true;
 				RealType NewElementData[ElementSize];
 				GetNewElementValue(ElementID, TriID, NewElementData);
-				ElTri[SubIdx] = AppendElement(NewElementData, ParentVertices[ElementID]);
+				ElTri[SubIdx] = AppendElement(NewElementData);
 			}
 		}
 		if (TriChanged)
 		{
 			InternalSetTriangle(TriID, ElTri, true);
 		}
-		
 	}
 }
 
@@ -216,7 +216,7 @@ int TDynamicMeshOverlay<RealType, ElementSize>::SplitElementWithNewParent(int El
 {
 	RealType SourceData[ElementSize];
 	GetElement(ElementID, SourceData);
-	int NewElID = AppendElement(SourceData, NewParentID);
+	int NewElID = AppendElement(SourceData);
 	for (int TriID : TrianglesToUpdate)
 	{
 		int ElementTriStart = TriID * 3;
@@ -231,7 +231,17 @@ int TDynamicMeshOverlay<RealType, ElementSize>::SplitElementWithNewParent(int El
 			}
 		}
 	}
-	checkSlow(ElementsRefCounts.IsValid(ElementID)); // simple sanity check; shouldn't completely remove an element with SplitElement
+	ParentVertices.InsertAt(NewParentID, NewElID);
+
+	checkSlow(ElementsRefCounts.IsValid(ElementID));
+
+	// An element may have become isolated after changing all of its incident triangles. Delete such an element.
+	if (ElementsRefCounts.GetRefCount(ElementID) == 1) 
+	{
+		ElementsRefCounts.Decrement(ElementID);
+		ParentVertices[ElementID] = -1;
+	}
+
 	return NewElID;
 }
 
@@ -389,6 +399,8 @@ void TDynamicMeshOverlay<RealType, ElementSize>::UnsetTriangle(int TriangleID)
 template<typename RealType, int ElementSize>
 void TDynamicMeshOverlay<RealType, ElementSize>::InternalSetTriangle(int tid, const FIndex3i& tv, bool bIncrementRefCounts)
 {
+	check(ParentMesh);
+
 	int i = 3 * tid;
 	ElementTriangles.InsertAt(tv[2], i + 2);
 	ElementTriangles.InsertAt(tv[1], i + 1);
@@ -400,6 +412,18 @@ void TDynamicMeshOverlay<RealType, ElementSize>::InternalSetTriangle(int tid, co
 		ElementsRefCounts.Increment(tv[1]);
 		ElementsRefCounts.Increment(tv[2]);
 	}
+
+	if (tv != FDynamicMesh3::InvalidTriangle)
+	{
+		// Set parent vertex IDs
+		const FIndex3i ParentTriangle = ParentMesh->GetTriangle(tid);
+
+		for (int VInd = 0; VInd < 3; ++VInd)
+		{
+			checkSlow(ParentVertices[tv[VInd]] == ParentTriangle[VInd] || ParentVertices[tv[VInd]] == FDynamicMesh3::InvalidID);
+			ParentVertices.InsertAt(ParentTriangle[VInd], tv[VInd]);
+		}
+	}
 }
 
 
@@ -449,7 +473,7 @@ bool TDynamicMeshOverlay<RealType,ElementSize>::IsSeamEdge(int eid) const
 	return !IndexUtil::SamePairUnordered(Triangle0[idx_base_a1], Triangle0[idx_base_b1], Triangle1[idx_base_a2], Triangle1[idx_base_b2]);
 
 
-	// [RMS] this doesn't seem to work but it should, and would be more efficient. 
+	// TODO: this doesn't seem to work but it should, and would be more efficient:
 	//   - add ParentMesh->FindTriEdgeIndex(tid,eid)
 	//   - SamePairUnordered query could directly index into ElementTriangles[]
 	//FIndex3i TriangleA = GetTriangle(et.A);
@@ -666,12 +690,12 @@ void TDynamicMeshOverlay<RealType, ElementSize>::OnSplitEdge(const FDynamicMesh3
 	// insert invalid triangle as needed
 	if (!bT0Set)
 	{
-		InternalSetTriangle(splitInfo.NewTriangles.A, FDynamicMesh3::InvalidTriangle, false);
+		InitializeNewTriangle(splitInfo.NewTriangles.A);
 	}
 	if (!bT1Set && splitInfo.NewTriangles.B >= 0)
 	{
 		// new triangle is invalid
-		InternalSetTriangle(splitInfo.NewTriangles.B, FDynamicMesh3::InvalidTriangle, false);
+		InitializeNewTriangle(splitInfo.NewTriangles.B);
 	}
 	// if neither tri was set, nothing else to do
 	if (!bT0Set && !bT1Set)
@@ -692,7 +716,7 @@ void TDynamicMeshOverlay<RealType, ElementSize>::OnSplitEdge(const FDynamicMesh3
 		int idx_base_c = IndexUtil::GetOtherTriIndex(idx_base_a1, idx_base_b1);
 
 		// create new element at lerp position
-		NewElemID = AppendElement((RealType)0, splitInfo.NewVertex);
+		NewElemID = AppendElement((RealType)0);
 		SetElementFromLerp(NewElemID, Triangle0[idx_base_a1], Triangle0[idx_base_b1], (RealType)splitInfo.SplitT);
 
 		// rewrite triangle 0
@@ -728,7 +752,7 @@ void TDynamicMeshOverlay<RealType, ElementSize>::OnSplitEdge(const FDynamicMesh3
 		if (bHasSharedUVEdge == false)
 		{
 			// create new element at lerp position
-			OtherNewElemID = AppendElement((RealType)0, splitInfo.NewVertex);
+			OtherNewElemID = AppendElement((RealType)0);
 			SetElementFromLerp(OtherNewElemID, Triangle1[idx_base_a2], Triangle1[idx_base_b2], (RealType)splitInfo.SplitT);
 		}
 
@@ -996,7 +1020,7 @@ void TDynamicMeshOverlay<RealType, ElementSize>::OnPokeTriangle(const FDynamicMe
 	FIndex3i Triangle = GetTriangle(PokeInfo.OriginalTriangle);
 
 	// create new element at barycentric position
-	int CenterElemID = AppendElement((RealType)0, PokeInfo.NewVertex);
+	int CenterElemID = AppendElement((RealType)0);
 	FVector3<RealType> BaryCoords((RealType)PokeInfo.BaryCoords.X, (RealType)PokeInfo.BaryCoords.Y, (RealType)PokeInfo.BaryCoords.Z);
 	SetElementFromBary(CenterElemID, Triangle[0], Triangle[1], Triangle[2], BaryCoords);
 
@@ -1019,28 +1043,24 @@ void TDynamicMeshOverlay<RealType, ElementSize>::OnMergeEdges(const FDynamicMesh
 	// MergeEdges just merges vertices. For now we will not also merge UVs. So all we need to
 	// do is rewrite the UV parent vertices
 
-	FIndex3i ModifiedEdges(MergeInfo.KeptEdge, MergeInfo.ExtraKeptEdges.A, MergeInfo.ExtraKeptEdges.B);
-	for (int EdgeIdx = 0; EdgeIdx < 3; ++EdgeIdx)
+	for (int i = 0; i < 2; i++)
 	{
-		if (ParentMesh->IsEdge(ModifiedEdges[EdgeIdx]) == false)
+		int KeptVID = MergeInfo.KeptVerts[i];
+		int RemovedVID = MergeInfo.RemovedVerts[i];
+		// this for loop is very similar to GetVertexElements() but accounts for the base mesh already being updated
+		for (int TID : ParentMesh->VtxTrianglesItr(KeptVID)) // only care about triangles connected to the *new* vertex; these are updated
 		{
-			continue;
-		}
-
-		FIndex2i EdgeTris = ParentMesh->GetEdgeT(ModifiedEdges[EdgeIdx]);
-		for (int j = 0; j < 2; ++j)
-		{
-			FIndex3i ElemTriangle = GetTriangle(EdgeTris[j]);
-			for (int k = 0; k < 3; ++k)
+			if (!IsSetTriangle(TID))
 			{
-				int ParentVID = ParentVertices[ElemTriangle[k]];
-				if (ParentVID == MergeInfo.RemovedVerts.A)
-				{
-					ParentVertices[ElemTriangle[k]] = MergeInfo.KeptVerts.A;
-				}
-				else if (ParentVID == MergeInfo.RemovedVerts.B)
+				continue;
+			}
+			FIndex3i Triangle = GetTriangle(TID);
+			for (int j = 0; j < 3; ++j)
+			{
+				// though the ParentMesh vertex is NewVertex in the source mesh, it is still OriginalVertex in the ParentVertices array (since that hasn't been updated yet)
+				if (Triangle[j] != FDynamicMesh3::InvalidID && ParentVertices[Triangle[j]] == RemovedVID)
 				{
-					ParentVertices[ElemTriangle[k]] = MergeInfo.KeptVerts.B;
+					ParentVertices[Triangle[j]] = KeptVID;
 				}
 			}
 		}
@@ -1139,11 +1159,22 @@ bool TDynamicMeshOverlay<RealType, ElementSize>::CheckValidity(bool bAllowNonMan
 
 	// @todo: check that all connected element-pairs are also edges in parentmesh
 
-	// check that parent vtx of each element is actually a vertex
+	// check that parent vtx of a non-isolated element is actually a vertex
 	for (int elemid : ElementIndicesItr())
 	{
 		int ParentVID = GetParentVertex(elemid);
-		CheckOrFailF(ParentMesh->IsVertex(ParentVID));
+		bool bParentIsVertex = ParentMesh->IsVertex(ParentVID);
+		bool bElementIsIsolated = (ElementsRefCounts.GetRefCount(elemid) == 1);
+
+		// bParentIsVertex XOR bElementIsIsolated
+		if ( bParentIsVertex )
+		{
+			CheckOrFailF(!bElementIsIsolated);
+		}
+		else
+		{
+			CheckOrFailF(bElementIsIsolated);
+		}
 	}
 
 	// check that parent vertices of each element triangle are the same as base triangle
diff --git a/Engine/Plugins/Experimental/GeometryProcessing/Source/DynamicMesh/Private/MeshBoundaryLoops.cpp b/Engine/Plugins/Experimental/GeometryProcessing/Source/DynamicMesh/Private/MeshBoundaryLoops.cpp
--- a/Engine/Plugins/Experimental/GeometryProcessing/Source/DynamicMesh/Private/MeshBoundaryLoops.cpp
+++ b/Engine/Plugins/Experimental/GeometryProcessing/Source/DynamicMesh/Private/MeshBoundaryLoops.cpp
@@ -395,9 +395,7 @@ int FMeshBoundaryLoops::FindLeftTurnEdge(int incoming_e, int bowtie_v, TArray<in
 		}
 	}
 
-	// [RMS] w/ bowtie vertices and open spans, this does happen
-	//Debug.Assert(best_e != -1);
-
+	// Note w/ bowtie vertices and open spans, best_e CAN be invalid (== -1)
 	return best_e;
 }
 
diff --git a/Engine/Plugins/Experimental/GeometryProcessing/Source/DynamicMesh/Private/MeshIndexUtil.cpp b/Engine/Plugins/Experimental/GeometryProcessing/Source/DynamicMesh/Private/MeshIndexUtil.cpp
--- a/Engine/Plugins/Experimental/GeometryProcessing/Source/DynamicMesh/Private/MeshIndexUtil.cpp
+++ b/Engine/Plugins/Experimental/GeometryProcessing/Source/DynamicMesh/Private/MeshIndexUtil.cpp
@@ -55,14 +55,33 @@ void MeshIndexUtil::VertexToTriangleOneRing(const FDynamicMesh3* Mesh, const TAr
 {
 	int NumVerts = VertexIDs.Num();
 	TriangleIDsOut.Reserve( (NumVerts < 5) ? NumVerts*6 : NumVerts*4);
+	TArray<int32> OneRingT;  // temp storage
+	OneRingT.Reserve(16);
+
+	//auto AddFunc = [&](int32 TriID)
+	//{
+	//	TriangleIDsOut.Add(TriID);
+	//};
+
 	for (int k = 0; k < NumVerts; ++k)
 	{
 		if (Mesh->IsVertex(VertexIDs[k]))
 		{
-			for (int TriID : Mesh->VtxTrianglesItr(VertexIDs[k]))
+			// slowest
+			//for (int TriID : Mesh->VtxTrianglesItr(VertexIDs[k]))
+			//{
+			//	TriangleIDsOut.Add(TriID);
+			//}
+
+			OneRingT.SetNum(0, false);
+			Mesh->GetVertexOneRingTriangles(VertexIDs[k], OneRingT);
+			for (int TriID : OneRingT)
 			{
 				TriangleIDsOut.Add(TriID);
 			}
+
+			// slightly slower?
+			//Mesh->EnumerateVertexTriangles(VertexIDs[k], AddFunc);
 		}
 	}
 }
\ No newline at end of file
diff --git a/Engine/Plugins/Experimental/GeometryProcessing/Source/DynamicMesh/Private/MeshNormals.cpp b/Engine/Plugins/Experimental/GeometryProcessing/Source/DynamicMesh/Private/MeshNormals.cpp
--- a/Engine/Plugins/Experimental/GeometryProcessing/Source/DynamicMesh/Private/MeshNormals.cpp
+++ b/Engine/Plugins/Experimental/GeometryProcessing/Source/DynamicMesh/Private/MeshNormals.cpp
@@ -225,7 +225,7 @@ bool FMeshNormals::QuickRecomputeOverlayNormals(FDynamicMesh3& Mesh, bool bInver
 FVector3d FMeshNormals::ComputeVertexNormal(const FDynamicMesh3& Mesh, int VertIdx, bool bWeightByArea, bool bWeightByAngle)
 {
 	FVector3d SumNormal = FVector3d::Zero();
-	for (int TriIdx : Mesh.VtxTrianglesItr(VertIdx))
+	Mesh.EnumerateVertexTriangles(VertIdx, [&](int32 TriIdx)
 	{
 		FVector3d TriNormal, TriCentroid; double TriArea;
 		Mesh.GetTriInfo(TriIdx, TriNormal, TriArea, TriCentroid);
@@ -234,7 +234,7 @@ FVector3d FMeshNormals::ComputeVertexNormal(const FDynamicMesh3& Mesh, int VertI
 		FIndex3i Triangle = Mesh.GetTriangle(TriIdx);
 		int32 j = IndexUtil::FindTriIndex(VertIdx, Triangle);
 		SumNormal += TriNormal * TriNormalWeights[j];
-	}
+	});
 	return SumNormal.Normalized();
 }
 
@@ -242,7 +242,7 @@ FVector3d FMeshNormals::ComputeVertexNormal(const FDynamicMesh3& Mesh, int VertI
 FVector3d FMeshNormals::ComputeVertexNormal(const FDynamicMesh3& Mesh, int32 VertIdx, TFunctionRef<bool(int32)> TriangleFilterFunc, bool bWeightByArea, bool bWeightByAngle)
 {
 	FVector3d NormalSum(0, 0, 0);
-	for (int TriIdx : Mesh.VtxTrianglesItr(VertIdx))
+	Mesh.EnumerateVertexTriangles(VertIdx, [&](int32 TriIdx)
 	{
 		if (TriangleFilterFunc(TriIdx))
 		{
@@ -254,7 +254,7 @@ FVector3d FMeshNormals::ComputeVertexNormal(const FDynamicMesh3& Mesh, int32 Ver
 			int32 j = IndexUtil::FindTriIndex(VertIdx, Triangle);
 			NormalSum += TriNormal * TriNormalWeights[j];
 		}
-	}
+	});
 	return NormalSum.Normalized();
 }
 
@@ -265,7 +265,7 @@ FVector3d FMeshNormals::ComputeOverlayNormal(const FDynamicMesh3& Mesh, FDynamic
 	int ParentVertexID = NormalOverlay->GetParentVertex(ElemIdx);
 	FVector3d SumNormal = FVector3d::Zero();
 	int Count = 0;
-	for (int TriIdx : Mesh.VtxTrianglesItr(ParentVertexID))
+	Mesh.EnumerateVertexTriangles(ParentVertexID, [&](int32 TriIdx)
 	{
 		if (NormalOverlay->TriangleHasElement(TriIdx, ElemIdx))
 		{
@@ -274,7 +274,7 @@ FVector3d FMeshNormals::ComputeOverlayNormal(const FDynamicMesh3& Mesh, FDynamic
 			SumNormal += Area * Normal;
 			Count++;
 		}
-	}
+	});
 
 	return (Count > 0) ? SumNormal.Normalized() : FVector3d::Zero();
 }
@@ -297,7 +297,7 @@ void FMeshNormals::InitializeOverlayToPerVertexNormals(FDynamicMeshNormalOverlay
 	for (int vid : Mesh->VertexIndicesItr())
 	{
 		FVector3f Normal = (bUseMeshNormals) ? Mesh->GetVertexNormal(vid) : (FVector3f)Normals[vid];
-		int nid = NormalOverlay->AppendElement(Normal, vid);
+		int nid = NormalOverlay->AppendElement(Normal);
 		VertToNormalMap[vid] = nid;
 	}
 
@@ -309,4 +309,4 @@ void FMeshNormals::InitializeOverlayToPerVertexNormals(FDynamicMeshNormalOverlay
 		Tri.C = VertToNormalMap[Tri.C];
 		NormalOverlay->SetTriangle(tid, Tri);
 	}
-}
\ No newline at end of file
+}
diff --git a/Engine/Plugins/Experimental/GeometryProcessing/Source/DynamicMesh/Private/MeshRefinerBase.cpp b/Engine/Plugins/Experimental/GeometryProcessing/Source/DynamicMesh/Private/MeshRefinerBase.cpp
--- a/Engine/Plugins/Experimental/GeometryProcessing/Source/DynamicMesh/Private/MeshRefinerBase.cpp
+++ b/Engine/Plugins/Experimental/GeometryProcessing/Source/DynamicMesh/Private/MeshRefinerBase.cpp
@@ -192,7 +192,7 @@ bool FMeshRefinerBase::CanCollapseVertex(int eid, int a, int b, int& collapse_to
 		return true;
 	}
 	// if both fixed, and options allow, treat this edge as unconstrained (eg collapse to midpoint)
-	// [RMS] tried picking a or b here, but something weird happens, where
+	// TODO: tried picking a or b here, but something weird happens, where
 	//   eg cylinder cap will entirely erode away. Somehow edge lengths stay below threshold??
 	if (AllowCollapseFixedVertsWithSameSetID
 		&& ca.FixedSetID >= 0
diff --git a/Engine/Plugins/Experimental/GeometryProcessing/Source/DynamicMesh/Private/MeshSimplification.cpp b/Engine/Plugins/Experimental/GeometryProcessing/Source/DynamicMesh/Private/MeshSimplification.cpp
--- a/Engine/Plugins/Experimental/GeometryProcessing/Source/DynamicMesh/Private/MeshSimplification.cpp
+++ b/Engine/Plugins/Experimental/GeometryProcessing/Source/DynamicMesh/Private/MeshSimplification.cpp
@@ -766,13 +766,7 @@ void TMeshSimplification<QuadricErrorType>::FullProjectionPass()
 
 	ApplyToProjectVertices(project);
 
-	// [RMS] not sure how to do this...
-	//if (EnableParallelProjection) {
-	//    gParallel.ForEach<int>(project_vertices(), project);
-	//} else {
-	//    foreach (int vid in project_vertices())
-	//        project(vid);
-	//}
+	// TODO: optionally do projection in parallel?
 }
 
 template <typename QuadricErrorType>
diff --git a/Engine/Plugins/Experimental/GeometryProcessing/Source/DynamicMesh/Private/MeshWeights.cpp b/Engine/Plugins/Experimental/GeometryProcessing/Source/DynamicMesh/Private/MeshWeights.cpp
--- a/Engine/Plugins/Experimental/GeometryProcessing/Source/DynamicMesh/Private/MeshWeights.cpp
+++ b/Engine/Plugins/Experimental/GeometryProcessing/Source/DynamicMesh/Private/MeshWeights.cpp
@@ -31,14 +31,17 @@ FVector3d FMeshWeights::MeanValueCentroid(const FDynamicMesh3 & mesh, int v_i)
 		FVector3d Vj = mesh.GetVertex(v_j);
 		FVector3d vVj = (Vj - Vi);
 		double len_vVj = vVj.Normalize();
-		// [RMS] is this the right thing to do? if vertices are coincident,
+		// TODO: is this the right thing to do? if vertices are coincident,
 		//   weight of this vertex should be very high!
 		if (len_vVj < FMathd::ZeroTolerance)
+		{
 			continue;
+		}
 		FVector3d vVdelta = (mesh.GetVertex(opp_v1) - Vi).Normalized();
 		double w_ij = VectorUtil::VectorTanHalfAngle(vVj, vVdelta);
 
-		if (opp_v2 != FDynamicMesh3::InvalidID) {
+		if (opp_v2 != FDynamicMesh3::InvalidID)
+		{
 			FVector3d vVgamma = (mesh.GetVertex(opp_v2) - Vi).Normalized();
 			w_ij += VectorUtil::VectorTanHalfAngle(vVj, vVgamma);
 		}
diff --git a/Engine/Plugins/Experimental/GeometryProcessing/Source/DynamicMesh/Private/Operations/MeshPlaneCut.cpp b/Engine/Plugins/Experimental/GeometryProcessing/Source/DynamicMesh/Private/Operations/MeshPlaneCut.cpp
--- a/Engine/Plugins/Experimental/GeometryProcessing/Source/DynamicMesh/Private/Operations/MeshPlaneCut.cpp
+++ b/Engine/Plugins/Experimental/GeometryProcessing/Source/DynamicMesh/Private/Operations/MeshPlaneCut.cpp
@@ -7,6 +7,7 @@
 
 #include "Operations/SimpleHoleFiller.h"
 #include "Operations/PlanarHoleFiller.h"
+#include "Operations/MinimalHoleFiller.h"
 #include "MeshNormals.h"
 #include "DynamicMeshEditor.h"
 #include "MathUtil.h"
@@ -537,7 +538,7 @@ bool FMeshPlaneCut::SimpleHoleFill(int ConstantGroupID)
 	HoleFillTriangles.Empty();
 	for (FOpenBoundary& Boundary : OpenBoundaries)
 	{
-		TArray<int> BoundaryFillTriangles = HoleFillTriangles.Emplace_GetRef();
+		TArray<int>& BoundaryFillTriangles = HoleFillTriangles.Emplace_GetRef();
 		FFrame3d ProjectionFrame(PlaneOrigin, PlaneNormal);
 
 		for (const FEdgeLoop& Loop : Boundary.CutLoops)
@@ -561,6 +562,39 @@ bool FMeshPlaneCut::SimpleHoleFill(int ConstantGroupID)
 }
 
 
+
+bool FMeshPlaneCut::MinimalHoleFill(int ConstantGroupID)
+{
+	bool bAllOk = true;
+
+	HoleFillTriangles.Empty();
+	for (FOpenBoundary& Boundary : OpenBoundaries)
+	{
+		TArray<int>& BoundaryFillTriangles = HoleFillTriangles.Emplace_GetRef();
+		FFrame3d ProjectionFrame(PlaneOrigin, PlaneNormal);
+
+		for (const FEdgeLoop& Loop : Boundary.CutLoops)
+		{
+			FMinimalHoleFiller Filler(Mesh, Loop);
+			int GID = ConstantGroupID >= 0 ? ConstantGroupID : Mesh->AllocateTriangleGroup();
+			bAllOk = Filler.Fill(GID) && bAllOk;
+
+			BoundaryFillTriangles.Append(Filler.NewTriangles);
+
+			if (Mesh->HasAttributes())
+			{
+				FDynamicMeshEditor Editor(Mesh);
+				Editor.SetTriangleNormals(Filler.NewTriangles, (FVector3f)PlaneNormal * Boundary.NormalSign);
+				Editor.SetTriangleUVsFromProjection(Filler.NewTriangles, ProjectionFrame, UVScaleFactor);
+			}
+		}
+	}
+
+	return bAllOk;
+}
+
+
+
 bool FMeshPlaneCut::HoleFill(TFunction<TArray<FIndex3i>(const FGeneralPolygon2d&)> PlanarTriangulationFunc, bool bFillSpans, int ConstantGroupID)
 {
 	bool bAllOk = true;
diff --git a/Engine/Plugins/Experimental/GeometryProcessing/Source/DynamicMesh/Private/Operations/RepairOrientation.cpp b/Engine/Plugins/Experimental/GeometryProcessing/Source/DynamicMesh/Private/Operations/RepairOrientation.cpp
--- a/Engine/Plugins/Experimental/GeometryProcessing/Source/DynamicMesh/Private/Operations/RepairOrientation.cpp
+++ b/Engine/Plugins/Experimental/GeometryProcessing/Source/DynamicMesh/Private/Operations/RepairOrientation.cpp
@@ -90,10 +90,10 @@ void FMeshRepairOrientation::ComputeComponentStatistics(FDynamicMeshAABBTree3* T
 	
 	// only want to raycast triangles in this component
 	TSet<int> TrisInComponent; TrisInComponent.Append(C.Triangles);
-	Tree->TriangleFilterF = [&TrisInComponent](int TID)
-	{
-		return TrisInComponent.Contains(TID);
-	};
+	IMeshSpatial::FQueryOptions RaycastOptions([&TrisInComponent](int TID)
+		{
+			return TrisInComponent.Contains(TID);
+		});
 	
 	// We want to try to figure out what is 'outside' relative to the world.
 	// Assumption is that faces we can hit from far away should be oriented outwards.
@@ -105,9 +105,9 @@ void FMeshRepairOrientation::ComputeComponentStatistics(FDynamicMeshAABBTree3* T
 	
 	FCriticalSection Mutex;
 
-	// note the C# version of this explicitly set a block size for parallel groups, and used a spin lock; TODO: revisit whether that should be ported / if the naive version is causing badness
+	// TODO: profile the parallel for and consider whether block size / locking method should be different
 	bool bNoParallel = false;
-	ParallelFor(C.Triangles.Num(), [this, &Tree, &Mutex, &C, &Dist](int32 Idx)
+	ParallelFor(C.Triangles.Num(), [this, &Tree, &RaycastOptions, &Mutex, &C, &Dist](int32 Idx)
 	{
 		int TID = C.Triangles[Idx];
 		FVector3d Normal, Centroid;
@@ -120,8 +120,8 @@ void FMeshRepairOrientation::ComputeComponentStatistics(FDynamicMeshAABBTree3* T
 
 		FVector3d PosPt = Centroid + Dist * Normal;
 		FVector3d NegPt = Centroid - Dist * Normal;
-		int HitPos = Tree->FindNearestHitTriangle(FRay3d(PosPt, -Normal));
-		int HitNeg = Tree->FindNearestHitTriangle(FRay3d(NegPt, Normal));
+		int HitPos = Tree->FindNearestHitTriangle(FRay3d(PosPt, -Normal), RaycastOptions);
+		int HitNeg = Tree->FindNearestHitTriangle(FRay3d(NegPt, Normal), RaycastOptions);
 		if (HitPos != TID && HitNeg != TID)
 		{
 			return; // no evidence
@@ -144,8 +144,6 @@ void FMeshRepairOrientation::ComputeComponentStatistics(FDynamicMeshAABBTree3* T
 		}
 	},
 		bNoParallel);
-	
-	Tree->TriangleFilterF = nullptr;
 }
 
 
diff --git a/Engine/Plugins/Experimental/GeometryProcessing/Source/DynamicMesh/Private/ProjectionTargets.cpp b/Engine/Plugins/Experimental/GeometryProcessing/Source/DynamicMesh/Private/ProjectionTargets.cpp
--- a/Engine/Plugins/Experimental/GeometryProcessing/Source/DynamicMesh/Private/ProjectionTargets.cpp
+++ b/Engine/Plugins/Experimental/GeometryProcessing/Source/DynamicMesh/Private/ProjectionTargets.cpp
@@ -7,7 +7,7 @@ FVector3d FMeshProjectionTarget::Project(const FVector3d& Point, int Identifier)
 {
 	double fDistSqr;
 	int tNearestID = Spatial->FindNearestTriangle(Point, fDistSqr);
-	if (tNearestID <= 0)
+	if (tNearestID < 0)
 	{
 		return Point;
 	}
@@ -33,7 +33,7 @@ FVector3d FMeshProjectionTarget::Project(const FVector3d& Point, FVector3d& Proj
 {
 	double fDistSqr;
 	int tNearestID = Spatial->FindNearestTriangle(Point, fDistSqr);
-	if (tNearestID <= 0)
+	if (tNearestID < 0)
 	{
 		return Point;
 	}
diff --git a/Engine/Plugins/Experimental/GeometryProcessing/Source/DynamicMesh/Private/Remesher.cpp b/Engine/Plugins/Experimental/GeometryProcessing/Source/DynamicMesh/Private/Remesher.cpp
--- a/Engine/Plugins/Experimental/GeometryProcessing/Source/DynamicMesh/Private/Remesher.cpp
+++ b/Engine/Plugins/Experimental/GeometryProcessing/Source/DynamicMesh/Private/Remesher.cpp
@@ -4,7 +4,7 @@
 #include "DynamicMeshAttributeSet.h"
 #include "MeshWeights.h"
 
-
+#include "Async/ParallelFor.h"
 
 
 void FRemesher::SetTargetEdgeLength(double fLength)
@@ -104,7 +104,7 @@ void FRemesher::BasicRemeshPass()
 	ProfileBeginProject();
 	if (ProjTarget != nullptr && ProjectionMode == ETargetProjectionMode::AfterRefinement)
 	{
-		FullProjectionPass();
+		FullProjectionPass(bEnableParallelProjection);
 		DoDebugChecks();
 	}
 	ProfileEndProject();
@@ -335,7 +335,7 @@ void FRemesher::UpdateAfterSplit(int edgeID, int va, int vb, const FDynamicMesh3
 		// inherit edge constraint
 		Constraints->SetOrUpdateEdgeConstraint(SplitInfo.NewEdges.A, Constraints->GetEdgeConstraint(edgeID));
 
-		// [RMS] update vertex constraints. Note that there is some ambiguity here.
+		// Update vertex constraints. Note that there is some ambiguity here.
 		//   Both verts being constrained doesn't inherently mean that the edge is on
 		//   a constraint, that's why these checks are only applied if edge is constrained.
 		//   But constrained edge doesn't necessarily mean we want to inherit vert constraints!!
@@ -343,7 +343,7 @@ void FRemesher::UpdateAfterSplit(int edgeID, int va, int vb, const FDynamicMesh3
 		//   although, pretty safe to assume that we would at least disable flips
 		//   if both vertices are constrained to same line/curve. So, maybe this makes sense...
 		//
-		//   (perhaps edge constraint should be explicitly tagged to resolve this ambiguity??)
+		//   (TODO: perhaps edge constraint should be explicitly tagged to resolve this ambiguity??)
 
 		// vert inherits Fixed if both orig edge verts Fixed, and both tagged with same SetID
 		FVertexConstraint ca = Constraints->GetVertexConstraint(va);
@@ -475,30 +475,38 @@ TFunction<FVector3d(const FDynamicMesh3&, int, double)> FRemesher::GetSmoothFunc
 
 void FRemesher::FullSmoothPass_Buffer(bool bParallel)
 {
-	InitializeVertexBufferForPass();
-
 	TFunction<FVector3d(const FDynamicMesh3&, int, double)> UseSmoothFunc = GetSmoothFunction();
 
-	auto SmoothAndUpdateFunc = [&](int vID) 
+	if (bParallel)
 	{
-		bool bModified = false;
-		FVector3d vSmoothed = ComputeSmoothedVertexPos(vID, UseSmoothFunc, bModified);
-		if (bModified) 
+		auto VertexMoveFunction = [&UseSmoothFunc, this](int VertexID, bool& bModified)
 		{
-			TempFlagBuffer[vID] = true;
-			TempPosBuffer[vID] = vSmoothed;
-		}
-	};
+			return ComputeSmoothedVertexPos(VertexID, UseSmoothFunc, bModified);
+		};
 
-	//if (bParallel) {
-	//    gParallel.ForEach<int>(smooth_vertices(), smooth);
-	//} else {
-	//    foreach (int vID in smooth_vertices())
-	//        smooth(vID);
-	//}
-	ApplyToSmoothVertices(SmoothAndUpdateFunc);
+		MoveVerticesParallel(VertexMoveFunction);
+	}
+	else
+	{
+		// Serial
 
-	ApplyVertexBuffer(bParallel);
+		InitializeVertexBufferForPass();
+
+		auto SmoothAndUpdateFunc = [this, UseSmoothFunc](int vID)
+		{
+			bool bModified = false;
+			FVector3d vSmoothed = ComputeSmoothedVertexPos(vID, UseSmoothFunc, bModified);
+			if (bModified)
+			{
+				TempFlagBuffer[vID] = true;
+				TempPosBuffer[vID] = vSmoothed;
+			}
+		};
+
+		ApplyToSmoothVertices(SmoothAndUpdateFunc);
+
+		ApplyVertexBuffer(false);
+	}
 }
 
 
@@ -512,42 +520,59 @@ void FRemesher::InitializeVertexBufferForPass()
 	{
 		TempPosBuffer.Resize(Mesh->MaxVertexID() + Mesh->MaxVertexID() / 5);
 	}
-	if (TempFlagBuffer.Num() < Mesh->MaxVertexID()) 
+	if (TempFlagBuffer.Num() < Mesh->MaxVertexID())
 	{
 		TempFlagBuffer.SetNum(2 * Mesh->MaxVertexID());
 	}
-
 	TempFlagBuffer.Init(false, TempFlagBuffer.Num());
-	//TempFlagBuffer.assign(TempFlagBuffer.size(), false);
 }
 
+
 void FRemesher::ApplyVertexBuffer(bool bParallel)
 {
-	for (int vid : Mesh->VertexIndicesItr()) 
+	check(TempFlagBuffer.Num() >= Mesh->MaxVertexID());
+	check(static_cast<int>(TempPosBuffer.Num()) >= Mesh->MaxVertexID());
+
+	if (!bParallel)
 	{
-		if (TempFlagBuffer[vid])
+		// Serial
+		for (int vid : Mesh->VertexIndicesItr())
 		{
-			Mesh->SetVertex(vid, TempPosBuffer[vid]);
+			if (TempFlagBuffer[vid])
+			{
+				Mesh->SetVertex(vid, TempPosBuffer[vid]);
+			}
 		}
 	}
+	else
+	{
+		// The serial version calls Mesh->SetVertex which sets the vertex position and also updates the timestamps. In
+		// the parallel version, to avoid contention for the timestamps we work in batches, calling 
+		// SetVertex_NoTimeStampUpdate several times for each async task, then finally locking and incrementing the 
+		// timestamp once per task.
 
-	// [TODO] can probably use block-parallel here...
-	//if (bParallel) {
-	//    gParallel.BlockStartEnd(0, Mesh->MaxVertexID-1, (a,b) => {
-	//        for (int vid = a; vid <= b; vid++) {
-	//            if (TempFlagBuffer[vid])
-	//                Mesh->SetVertex(vid, TempPosBuffer[vid]);
-	//        }
-	//    });
-	//} else {
-	//    foreach (int vid in Mesh->VertexIndices()) {
-	//        if (TempFlagBuffer[vid])
-	//            Mesh->SetVertex(vid, TempPosBuffer[vid]);
-	//    }
-	//}
-}
+		// TODO: Try varying the batch size to observe a trade-off between work done per thread and overhead due to 
+		// locking.
 
+		const int BatchSize = 1000;
+		const int NumBatches = FMath::CeilToInt(Mesh->MaxVertexID() / static_cast<float>(BatchSize));
 
+		ParallelFor(NumBatches, [this, BatchSize](int32 BatchID)
+		{
+			int NumVerticesSetThisBatch = 0;
+			for (int VertexID = BatchID * BatchSize; VertexID < (BatchID + 1) * BatchSize; ++VertexID)
+			{
+				if (VertexID < TempFlagBuffer.Num() && TempFlagBuffer[VertexID] && Mesh->IsVertex(VertexID))
+				{
+					Mesh->SetVertex_NoTimeStampUpdate(VertexID, TempPosBuffer[VertexID]);
+					++NumVerticesSetThisBatch;
+				}
+			}
+			Mesh->IncrementTimeStamps(NumVerticesSetThisBatch, true, false);
+
+		}, false);
+	}
+}
 
 FVector3d FRemesher::ComputeSmoothedVertexPos(int vID,
 	TFunction<FVector3d(const FDynamicMesh3&, int, double)> smoothFunc, bool& bModified)
@@ -609,33 +634,43 @@ void FRemesher::ApplyToSmoothVertices(const TFunction<void(int)>& VertexSmoothFu
 
 
 // Project vertices onto projection target. 
-// We can do projection in parallel if we have .net 
-void FRemesher::FullProjectionPass()
+void FRemesher::FullProjectionPass(bool bParallel)
 {
-	auto UseProjectionFunc = [&](int vID) 
+	if (bParallel)
 	{
-		if (IsVertexConstrained(vID))
+		auto VertexMoveFunction = [this](int VertexID, bool& bModified)
 		{
-			return;
-		}
-		if (VertexControlF != nullptr && ((int)VertexControlF(vID) & (int)EVertexControl::NoProject) != 0)
+			FVector3d CurrentPosition = Mesh->GetVertex(VertexID);
+			FVector3d ProjectedPosition = ProjTarget->Project(CurrentPosition, VertexID);
+			bModified = !VectorUtil::EpsilonEqual(CurrentPosition, ProjectedPosition, FMathd::ZeroTolerance);
+
+			return ProjectedPosition;
+		};
+
+		MoveVerticesParallel(VertexMoveFunction);
+	}
+	else
+	{
+		// Serial
+
+		auto UseProjectionFunc = [this](int vID)
 		{
-			return;
-		}
-		FVector3d curpos = Mesh->GetVertex(vID);
-		FVector3d projected = ProjTarget->Project(curpos, vID);
-		Mesh->SetVertex(vID, projected);
-	};
-
-	ApplyToProjectVertices(UseProjectionFunc);
-
-	// [RMS] not sure how to do this...
-	//if (EnableParallelProjection) {
-	//    gParallel.ForEach<int>(project_vertices(), project);
-	//} else {
-	//    foreach (int vid in project_vertices())
-	//        project(vid);
-	//}
+			if (IsVertexConstrained(vID))
+			{
+				return;
+			}
+			if (VertexControlF != nullptr && ((int)VertexControlF(vID) & (int)EVertexControl::NoProject) != 0)
+			{
+				return;
+			}
+			FVector3d curpos = Mesh->GetVertex(vID);
+			FVector3d projected = ProjTarget->Project(curpos, vID);
+			Mesh->SetVertex(vID, projected);
+		};
+
+		ApplyToProjectVertices(UseProjectionFunc);
+	}
+
 }
 
 
@@ -651,5 +686,35 @@ void FRemesher::ApplyToProjectVertices(const TFunction<void(int)>& VertexProject
 
 
 
+void FRemesher::MoveVerticesParallel(TFunction<FVector3d(int, bool&)> NewVertexPosition)
+{
+	// This is done in two passes:
+	// 1. Looping over all vertices, compute new vertex positions and put them in a buffer. Simultaneously fill a 
+	//    parallel buffer of bools indicating which vertices were given new positions.
+	// 2. Copy all vertex positions saved in step 1 from the buffer into the Mesh data structure. Update the timestamps.
+
+	InitializeVertexBufferForPass();
+
+	check(TempFlagBuffer.Num() >= Mesh->MaxVertexID());
+	check(static_cast<int>(TempPosBuffer.Num()) >= Mesh->MaxVertexID());
 
+	// First compute all vertex displacements and put them into a buffer
+	ParallelFor(Mesh->MaxVertexID(), [this, NewVertexPosition](int32 VertexID)
+	{
+		if (!Mesh->IsVertex(VertexID)) { return; }
 
+		const FVector3d CurrentPosition = Mesh->GetVertex(VertexID);
+		bool bModified = false;
+		const FVector3d NewPosition = NewVertexPosition(VertexID, bModified);
+
+		if (bModified)
+		{
+			TempFlagBuffer[VertexID] = true;
+			TempPosBuffer[VertexID] = NewPosition;
+		}
+	}, false);
+
+	// Finally move the vertex positions according to the buffer
+	ApplyVertexBuffer(true);
+
+}
diff --git a/Engine/Plugins/Experimental/GeometryProcessing/Source/DynamicMesh/Public/DynamicMesh3.h b/Engine/Plugins/Experimental/GeometryProcessing/Source/DynamicMesh/Public/DynamicMesh3.h
--- a/Engine/Plugins/Experimental/GeometryProcessing/Source/DynamicMesh/Public/DynamicMesh3.h
+++ b/Engine/Plugins/Experimental/GeometryProcessing/Source/DynamicMesh/Public/DynamicMesh3.h
@@ -381,6 +381,9 @@ public:
 		});
 	}
 
+	/** Call ApplyFunc for each one-ring triangle of a vertex. Currently this is significantly more efficient than VtxTrianglesItr() in many use cases. */
+	void EnumerateVertexTriangles(int32 VertexID, TFunctionRef<void(int32)> ApplyFunc) const;
+
 	//
 	// Mesh Construction
 	//
@@ -478,6 +481,43 @@ public:
 		}
 	}
 
+	/** Set vertex position. DOES NOT UPDATE TIMESTAMPS.
+	This is meant to be used in conjunction with IncrementTimeStamps, below. Set a bunch of vertex positions then 
+	increment the timestamps with the number of vertices that were moved, rather than incrementing the timestamps for
+	each vertex.
+	*/
+	void SetVertex_NoTimeStampUpdate(int VertexID, const FVector3d& vNewPos)
+	{
+		checkSlow(VectorUtil::IsFinite(vNewPos));
+		check(IsVertex(VertexID));
+
+		if (VectorUtil::IsFinite(vNewPos))
+		{
+			int i = 3 * VertexID;
+			Vertices[i] = vNewPos.X;
+			Vertices[i + 1] = vNewPos.Y;
+			Vertices[i + 2] = vNewPos.Z;
+		}
+	}
+
+	/** Increment the specified timestamps by the given amount. Thread-safe. 
+	This is meant to be used in conjunction with SetVertex_NoTimeStampUpdate, above. Set a bunch of vertex positions then
+	increment the timestamps with the number of vertices that were moved, rather than incrementing the timestamps for
+	each vertex.
+	*/
+	void IncrementTimeStamps(int Amount, bool bShapeChange, bool bTopologyChange)
+	{
+		FPlatformAtomics::InterlockedAdd(&Timestamp, Amount);
+		if (bShapeChange)
+		{
+			FPlatformAtomics::InterlockedAdd(&ShapeTimestamp, Amount);
+		}
+		if (bTopologyChange)
+		{
+			FPlatformAtomics::InterlockedAdd(&TopologyTimestamp, Amount);
+		}
+	}
+
 	/** Get extended vertex information */
 	bool GetVertex(int VertexID, FVertexInfo& VertInfo, bool bWantNormals, bool bWantColors, bool bWantUVs) const;
 
@@ -777,6 +817,11 @@ public:
 	 */
 	EMeshResult GetVtxTriangles(int VertexID, TArray<int>& TrianglesOut, bool bUseOrientation) const;
 
+	/**
+	 * @return Triangle ID for a single triangle connected to VertexID, or InvalidID if VertexID does not exist or has no attached triangles
+	 */
+	int GetVtxSingleTriangle(int VertexID) const;
+
 	/**
 	* Get triangles connected to vertex in contiguous order, with multiple groups if vertex is a bowtie.
 	* @param VertexID Vertex ID to search around
@@ -1292,6 +1337,8 @@ protected:
 
 	int FindTriangleEdge(int TriangleID, int vA, int vB) const;
 
+	int32 FindEdgeInternal(int32 vA, int32 vB, bool& bIsBoundary) const;
+
 	inline bool EdgeHasVertex(int EdgeID, int VertexID) const
 	{
 		int i = 4 * EdgeID;
diff --git a/Engine/Plugins/Experimental/GeometryProcessing/Source/DynamicMesh/Public/DynamicMeshOverlay.h b/Engine/Plugins/Experimental/GeometryProcessing/Source/DynamicMesh/Public/DynamicMeshOverlay.h
--- a/Engine/Plugins/Experimental/GeometryProcessing/Source/DynamicMesh/Public/DynamicMeshOverlay.h
+++ b/Engine/Plugins/Experimental/GeometryProcessing/Source/DynamicMesh/Public/DynamicMeshOverlay.h
@@ -9,6 +9,7 @@
 #include "VectorTypes.h"
 #include "GeometryTypes.h"
 #include "InfoTypes.h"
+#include "IndexTypes.h"
 
 class FDynamicMesh3;
 
@@ -98,7 +99,7 @@ public:
 			if (Copy.IsElement(EID))
 			{
 				Copy.GetElement(EID, Data);
-				MapE[EID] = AppendElement(Data, CompactMaps.MapV[Copy.ParentVertices[EID]]);
+				MapE[EID] = AppendElement(Data);
 			}
 			else
 			{
@@ -150,7 +151,14 @@ public:
 			GetElement(iLastE, Data);
 			SetElement(iCurE, Data);
 			int OrigParent = ParentVertices[iLastE];
-			ParentVertices[iCurE] = CompactMaps.GetVertex(OrigParent);
+			if (OrigParent == IndexConstants::InvalidID)
+			{
+				ParentVertices[iCurE] = IndexConstants::InvalidID;
+			}
+			else 
+			{
+				ParentVertices[iCurE] = CompactMaps.GetVertex(OrigParent);
+			}
 			ERef[iCurE] = ERef[iLastE];
 			ERef[iLastE] = FRefCountVector::INVALID_REF_COUNT;
 			MapE[iLastE] = iCurE;
@@ -207,10 +215,10 @@ public:
 	element_iterator ElementIndicesItr() const { return ElementsRefCounts.Indices(); }
 
 
-	/** Allocate a new element with the given constant value and parent-mesh vertex */
-	int AppendElement(RealType ConstantValue, int ParentVertex);
-	/** Allocate a new element with the given value and parent-mesh vertex */
-	int AppendElement(const RealType* Value, int ParentVertex);
+	/** Allocate a new element with the given constant value */
+	int AppendElement(RealType ConstantValue);
+	/** Allocate a new element with the given value */
+	int AppendElement(const RealType* Value);
 
 	/** Initialize the triangle list to the given size, and set all triangles to InvalidID */
 	void InitializeTriangles(int MaxTriangleID);
@@ -265,6 +273,7 @@ public:
 
 	/**
 	* Create a new copy of ElementID, and update connected triangles in the TrianglesToUpdate array to reference the copy of ElementID where they used to reference ElementID.  The new element will have the given parent vertex ID.
+	* Deletes any elements that are no longer used after the triangles are changed.
 	*
 	* @param ElementID the element to copy
 	* @param SplitParentVertexID the new parent vertex for copied elements
@@ -305,7 +314,7 @@ public:
 	 * If bUnsafe, we use fast id allocation that does not update free list.
 	 * You should only be using this between BeginUnsafeElementsInsert() / EndUnsafeElementsInsert() calls
 	 */
-	EMeshResult InsertElement(int ElementID, const RealType* Value, int ParentVertex, bool bUnsafe = false);
+	EMeshResult InsertElement(int ElementID, const RealType* Value, bool bUnsafe = false);
 
 
 	//
@@ -395,9 +404,29 @@ public:
 	/** find the triangles connected to an element */
 	void GetElementTriangles(int ElementID, TArray<int>& OutTriangles) const;
 
-	/** Currently an iteration every time */
+	/** @return true if overlay has any interior seam edges. This requires an O(N) search unless it early-outs. */
 	bool HasInteriorSeamEdges() const;
 
+	/**
+	 * Compute interpolated parameter value inside triangle using barycentric coordinates
+	 * @param TriangleID index of triangle
+	 * @param BaryCoords 3 barycentric coordinates inside triangle
+	 * @param DataOut resulting interpolated overlay parameter value (of size ElementSize)
+	 */
+	template<typename AsType>
+	void GetTriBaryInterpolate(int32 TriangleID, const AsType* BaryCoords, AsType* DataOut) const
+	{
+		int32 TriIndex = 3 * TriangleID;
+		int32 ElemIndex0 = ElementTriangles[TriIndex] * ElementSize;
+		int32 ElemIndex1 = ElementTriangles[TriIndex+1] * ElementSize;
+		int32 ElemIndex2 = ElementTriangles[TriIndex+2] * ElementSize;
+		const AsType Bary0 = (AsType)BaryCoords[0], Bary1 = (AsType)BaryCoords[1], Bary2 = (AsType)BaryCoords[2];
+		for (int32 i = 0; i < ElementSize; ++i)
+		{
+			DataOut[i] = Bary0*(AsType)Elements[ElemIndex0+i] + Bary1*(AsType)Elements[ElemIndex1+i] + Bary2*(AsType)Elements[ElemIndex2+i];
+		}
+	}
+
 	/**
 	 * Checks that the overlay mesh is well-formed, ie all internal data structures are consistent
 	 */
@@ -432,6 +461,7 @@ protected:
 
 	/** updates the triangles array and optionally the element reference counts */
 	void InternalSetTriangle(int TriangleID, const FIndex3i& TriElements, bool bIncrementRefCounts);
+
 };
 
 
@@ -463,9 +493,9 @@ public:
 	/**
 	 * Append a new Element to the overlay
 	 */
-	inline int AppendElement(const VectorType& Value, int SourceVertex)
+	inline int AppendElement(const VectorType& Value)
 	{
-		return BaseType::AppendElement((const RealType*)Value, SourceVertex);
+		return BaseType::AppendElement((const RealType*)Value);
 	}
 
 	/**
diff --git a/Engine/Plugins/Experimental/GeometryProcessing/Source/DynamicMesh/Public/MeshRefinerBase.h b/Engine/Plugins/Experimental/GeometryProcessing/Source/DynamicMesh/Public/MeshRefinerBase.h
--- a/Engine/Plugins/Experimental/GeometryProcessing/Source/DynamicMesh/Public/MeshRefinerBase.h
+++ b/Engine/Plugins/Experimental/GeometryProcessing/Source/DynamicMesh/Public/MeshRefinerBase.h
@@ -85,7 +85,7 @@ public:
 	/** Set to true to profile various passes @todo re-enable this! */
 	bool ENABLE_PROFILING = false;
 
-	/** 0 = no checking, 1 = check constraints and validity each pass, 2 = check validity after every mesh change (v slow but best for debugging) */
+	/** 0 = no checking, 1 = check constraints each pass, 2 = and check validity each pass, 3 = and check validity after every mesh change (v slow but best for debugging) */
 	int DEBUG_CHECK_LEVEL = 0;
 
 
diff --git a/Engine/Plugins/Experimental/GeometryProcessing/Source/DynamicMesh/Public/MeshTangents.h b/Engine/Plugins/Experimental/GeometryProcessing/Source/DynamicMesh/Public/MeshTangents.h
--- a/Engine/Plugins/Experimental/GeometryProcessing/Source/DynamicMesh/Public/MeshTangents.h
+++ b/Engine/Plugins/Experimental/GeometryProcessing/Source/DynamicMesh/Public/MeshTangents.h
@@ -53,6 +53,20 @@ public:
 		BitangentOut = Bitangents[k];
 	}
 
+	/**
+	 * @return Interpolated tangent and bitangent at barycentric coordinates inside a triangle, for per-triangle computed tangents
+	 * @param TriangleID triangle index in mesh
+	 * @param BaryCoords barycentric coordinates in triangle
+	 * @param TangentOut interpolated tangent
+	 * @param BitangentOut interpolated bitangent
+	 */
+	void GetInterpolatedTriangleTangent(int32 TriangleID, const FVector3<RealType>& BaryCoords, FVector3<RealType>& TangentOut, FVector3<RealType>& BitangentOut) const
+	{
+		int32 k = TriangleID * 3;
+		TangentOut = BaryCoords.X*Tangents[k] + BaryCoords.Y*Tangents[k+1] + BaryCoords.Z*Tangents[k+2];
+		BitangentOut = BaryCoords.X*Bitangents[k] + BaryCoords.Y*Bitangents[k+1] + BaryCoords.Z*Bitangents[k+2];
+	}
+
 
 	/**
 	 * Set tangent and bitangent at a vertex of triangle for per-triangle computed tangents.
diff --git a/Engine/Plugins/Experimental/GeometryProcessing/Source/DynamicMesh/Public/Operations/EmbedSurfacePath.h b/Engine/Plugins/Experimental/GeometryProcessing/Source/DynamicMesh/Public/Operations/EmbedSurfacePath.h
--- a/Engine/Plugins/Experimental/GeometryProcessing/Source/DynamicMesh/Public/Operations/EmbedSurfacePath.h
+++ b/Engine/Plugins/Experimental/GeometryProcessing/Source/DynamicMesh/Public/Operations/EmbedSurfacePath.h
@@ -125,7 +125,7 @@ public:
 	/**
 	 * Embed a surface path in mesh provided that the path only crosses vertices and edges except at the start and end, so we can add the path easily with local edge splits and possibly two triangle pokes (rather than needing general remeshing machinery)
 	 *
-	 * @param bUpdatePath Updating the Path array with the new vertices (if false, the path will no longer be valid after running this function)\
+	 * @param bUpdatePath Updating the Path array with the new vertices (if false, the path will no longer be valid after running this function)
 	 * @param PathVertices Indices of the vertices on the path after embedding succeeds; NOTE these will not be 1:1 with the input Path
 	 * @return true if embedding succeeded.
 	 */
diff --git a/Engine/Plugins/Experimental/GeometryProcessing/Source/DynamicMesh/Public/Operations/MeshPlaneCut.h b/Engine/Plugins/Experimental/GeometryProcessing/Source/DynamicMesh/Public/Operations/MeshPlaneCut.h
--- a/Engine/Plugins/Experimental/GeometryProcessing/Source/DynamicMesh/Public/Operations/MeshPlaneCut.h
+++ b/Engine/Plugins/Experimental/GeometryProcessing/Source/DynamicMesh/Public/Operations/MeshPlaneCut.h
@@ -138,6 +138,8 @@ public:
 	 */
 	virtual bool SimpleHoleFill(int ConstantGroupID = -1);
 
+	virtual bool MinimalHoleFill(int ConstantGroupID = -1);
+
 	/**
 	 *  Fill cut loops with FPlanarHoleFiller, using a caller-provided triangulation function
 	 */
diff --git a/Engine/Plugins/Experimental/GeometryProcessing/Source/DynamicMesh/Public/Operations/MeshRegionOperator.h b/Engine/Plugins/Experimental/GeometryProcessing/Source/DynamicMesh/Public/Operations/MeshRegionOperator.h
--- a/Engine/Plugins/Experimental/GeometryProcessing/Source/DynamicMesh/Public/Operations/MeshRegionOperator.h
+++ b/Engine/Plugins/Experimental/GeometryProcessing/Source/DynamicMesh/Public/Operations/MeshRegionOperator.h
@@ -10,6 +10,8 @@
 #include "DynamicMesh3.h"
 #include "DynamicSubmesh3.h"
 #include "FaceGroupUtil.h"
+#include "MeshNormals.h"
+#include "Algo/RemoveIf.h"
 
 /**
  * This class automatically extracts a submesh from a mesh, 
@@ -30,7 +32,7 @@ public:
 	// this is only valid after BackPropagate() call!! maps submeshverts to base mesh
 	FOptionallySparseIndexMap ReinsertSubToBaseMapV;
 
-	// [RMS] computation of this is kind of a hack right now...
+	// Note computation of this is kind of a hack right now; see TODO in class comment above
 	FOptionallySparseIndexMap ReinsertSubToBaseMapT;
 
 	// handle a tricky problem...see comments for EDuplicateTriBehavior enum
@@ -180,20 +182,34 @@ public:
 		Editor.RemoveTriangles(CurrentBaseTris, true);
 
 		// insert new submesh
-		TArray<int> new_tris; new_tris.SetNum(Region.GetSubmesh().TriangleCount());
+		TArray<int> new_tris; 
+		new_tris.Reserve(Region.GetSubmesh().TriangleCount());
 		bool bOK = Editor.ReinsertSubmesh(Region, ReinsertSubToBaseMapV, &new_tris, ReinsertDuplicateTriBehavior);
 
 		// reconstruct this...hacky? TODO: have ReinsertSubmesh directly build this
 		int NT = Region.GetSubmesh().MaxTriangleID();
 		ReinsertSubToBaseMapT.Initialize(NT, Region.GetSubmesh().TriangleCount());
 		int nti = 0;
+		bool bHasInvalid = false;
 		for (int ti = 0; ti < NT; ++ti)
 		{
 			if (Region.GetSubmesh().IsTriangle(ti) == false)
 			{
 				continue;
 			}
-			ReinsertSubToBaseMapT.Set(ti, new_tris[nti++]);
+			int new_tri = new_tris[nti++];
+			// should only happen if not allowing submesh repairs and ReinsertDuplicateTriBehavior is EnsureContinue
+			if (!bAllowSubmeshRepairs && !ensure(new_tri != -2))
+			{
+				bHasInvalid = true;
+				continue;
+			}
+			ReinsertSubToBaseMapT.Set(ti, new_tri);
+		}
+		// remove invalid triangles due to insertion failures (must be kept before this point for correspondence)
+		if (bHasInvalid)
+		{
+			new_tris.SetNum(Algo::RemoveIf(new_tris, [](int id) { return id == -2; }));
 		}
 
 		// assert that new triangles are all valid (goes wrong sometimes??)
diff --git a/Engine/Plugins/Experimental/GeometryProcessing/Source/DynamicMesh/Public/Operations/MinimalHoleFiller.h b/Engine/Plugins/Experimental/GeometryProcessing/Source/DynamicMesh/Public/Operations/MinimalHoleFiller.h
new file mode 100644
--- /dev/null
+++ b/Engine/Plugins/Experimental/GeometryProcessing/Source/DynamicMesh/Public/Operations/MinimalHoleFiller.h
@@ -0,0 +1,73 @@
+// Copyright Epic Games, Inc. All Rights Reserved.
+// Port of geometry3Sharp MinimalHoleFiller
+
+#pragma once
+
+#include "HoleFiller.h"
+#include "MathUtil.h"
+#include "VectorTypes.h"
+#include "GeometryTypes.h"
+#include "MeshBoundaryLoops.h"
+#include "SimpleHoleFiller.h"
+#include "Containers/Set.h"
+#include "MeshRegionOperator.h"
+#include "DynamicMesh3.h"
+#include "QueueRemesher.h"
+#include "MeshConstraintsUtil.h"
+
+
+/// Construct a "minimal" fill surface for the hole. This surface is often quasi-developable, reconstructs sharp edges, 
+/// etc. There are various options.
+
+class DYNAMICMESH_API FMinimalHoleFiller : public IHoleFiller
+{
+public:
+
+	FMinimalHoleFiller(FDynamicMesh3* InMesh, FEdgeLoop InFillLoop) :
+		Mesh(InMesh),
+		FillLoop(InFillLoop),
+		FillMesh(nullptr)
+	{}
+
+	bool Fill(int32 GroupID = -1) override;
+
+	// Settings
+	bool bIgnoreBoundaryTriangles = false;
+	bool bOptimizeDevelopability = true;
+	bool bOptimizeTriangles = true;
+	double DevelopabilityTolerance = 0.0001;
+
+	//  Outputs
+	/// Final fill vertices (should be empty?)
+	TArray<int> NewVertices;
+
+private:
+
+	// Inputs
+	FDynamicMesh3* Mesh;
+	FEdgeLoop FillLoop;
+
+	TUniquePtr<FMeshRegionOperator> RegionOp;
+	FDynamicMesh3* FillMesh;
+	TSet<int> BoundaryVertices;
+	TMap<int, double> ExteriorAngleSums;
+	TArray<double> Curvatures;
+
+	void AddAllEdges(int EdgeID, TSet<int>& EdgeSet);
+
+	double AspectMetric(int eid);
+	static double GetTriAspect(const FDynamicMesh3& mesh, FIndex3i& tri);
+
+	void UpdateCurvature(int vid);
+	double CurvatureMetricCached(int a, int b, int c, int d);
+	double CurvatureMetricEval(int a, int b, int c, int d);
+	double ComputeGaussCurvature(int vid);
+
+	// Steps in Fill:
+	void CollapseToMinimal();
+	void RemoveRemainingInteriorVerts();
+	void FlipToFlatter();
+	void FlipToMinimizeCurvature();
+	void FlipToImproveAspectRatios();
+};
+
diff --git a/Engine/Plugins/Experimental/GeometryProcessing/Source/DynamicMesh/Public/Operations/PlanarHoleFiller.h b/Engine/Plugins/Experimental/GeometryProcessing/Source/DynamicMesh/Public/Operations/PlanarHoleFiller.h
--- a/Engine/Plugins/Experimental/GeometryProcessing/Source/DynamicMesh/Public/Operations/PlanarHoleFiller.h
+++ b/Engine/Plugins/Experimental/GeometryProcessing/Source/DynamicMesh/Public/Operations/PlanarHoleFiller.h
@@ -4,6 +4,7 @@
 
 #pragma once
 
+#include "HoleFiller.h"
 #include "CoreMinimal.h"
 #include "MathUtil.h"
 #include "VectorTypes.h"
@@ -19,7 +20,7 @@ class FDynamicMesh3;
 /**
  * Fill a set of boundary loops with planar surfaces.  User must provide the triangulation function.
  */
-class DYNAMICMESH_API FPlanarHoleFiller
+class DYNAMICMESH_API FPlanarHoleFiller : public IHoleFiller
 {
 public:
 	//
@@ -33,11 +34,6 @@ public:
 	FVector3d PlaneOrigin;
 	FVector3d PlaneNormal;
 
-	//
-	// Outputs
-	//
-	TArray<int> NewTriangles;
-
 public:
 	/**
 	 *  Construct simple hole filler (just adds a central vertex and a triangle fan)
@@ -59,5 +55,5 @@ public:
 		return EOperationValidationResult::Ok;
 	}
 
-	virtual bool Fill(int GroupID = -1);
+	bool Fill(int GroupID = -1) override;
 };
diff --git a/Engine/Plugins/Experimental/GeometryProcessing/Source/DynamicMesh/Public/Operations/SimpleHoleFiller.h b/Engine/Plugins/Experimental/GeometryProcessing/Source/DynamicMesh/Public/Operations/SimpleHoleFiller.h
--- a/Engine/Plugins/Experimental/GeometryProcessing/Source/DynamicMesh/Public/Operations/SimpleHoleFiller.h
+++ b/Engine/Plugins/Experimental/GeometryProcessing/Source/DynamicMesh/Public/Operations/SimpleHoleFiller.h
@@ -4,6 +4,7 @@
 
 #pragma once
 
+#include "HoleFiller.h"
 #include "MathUtil.h"
 #include "VectorTypes.h"
 #include "GeometryTypes.h"
@@ -17,7 +18,7 @@ class FDynamicMesh3;
  * Fill an EdgeLoop hole with triangles.
  * Supports two fill modes, either a fan connected to a new central vertex, or a triangulation of the boundary polygon
  */
-class DYNAMICMESH_API FSimpleHoleFiller
+class DYNAMICMESH_API FSimpleHoleFiller : public IHoleFiller
 {
 public:
 	enum class EFillType
@@ -37,15 +38,18 @@ public:
 	// Outputs
 	//
 	int32 NewVertex = IndexConstants::InvalidID;
-	TArray<int32> NewTriangles;
 
 public:
 	/**
 	 *  Construct simple hole filler (just adds a central vertex and a triangle fan)
 	 */
-	FSimpleHoleFiller(FDynamicMesh3* Mesh, FEdgeLoop Loop) : Mesh(Mesh), Loop(Loop)
+	FSimpleHoleFiller(FDynamicMesh3* Mesh, FEdgeLoop Loop, EFillType InFillType = EFillType::TriangleFan) : 
+		Mesh(Mesh), 
+		Loop(Loop), 
+		FillType(InFillType)
 	{
 	}
+
 	virtual ~FSimpleHoleFiller() {}
 	
 	
@@ -62,7 +66,7 @@ public:
 		return EOperationValidationResult::Ok;
 	}
 
-	virtual bool Fill(int32 GroupID = -1);	
+	bool Fill(int32 GroupID = -1) override;	
 
 
 protected:
diff --git a/Engine/Plugins/Experimental/GeometryProcessing/Source/DynamicMesh/Public/Remesher.h b/Engine/Plugins/Experimental/GeometryProcessing/Source/DynamicMesh/Public/Remesher.h
--- a/Engine/Plugins/Experimental/GeometryProcessing/Source/DynamicMesh/Public/Remesher.h
+++ b/Engine/Plugins/Experimental/GeometryProcessing/Source/DynamicMesh/Public/Remesher.h
@@ -226,8 +226,6 @@ protected:
 	// want to also constrain this new vertex, possibly project to constraint target. 
 	virtual void UpdateAfterSplit(int edgeID, int va, int vb, const FDynamicMesh3::FEdgeSplitInfo& splitInfo);
 
-
-
 	TDynamicVector<FVector3d> TempPosBuffer;	// this is a temporary buffer used by smoothing and projection
 	TArray<bool> TempFlagBuffer;				// list of which indices in TempPosBuffer were modified
 
@@ -251,17 +249,23 @@ protected:
 	virtual TFunction<FVector3d(const FDynamicMesh3&, int, double)> GetSmoothFunction();
 
 
-
 	// Project vertices onto projection target.
-	virtual void FullProjectionPass();
+	virtual void FullProjectionPass(bool bParallel);
 
+	// Project a single vertex using the given target
 	virtual void ProjectVertex(int VertexID, IProjectionTarget* UseTarget);
 
 	// used by collapse-edge to get projected position for new vertex
 	virtual FVector3d GetProjectedCollapsePosition(int VertexID, const FVector3d& vNewPos);
 
+	// Apply the given projection function to all vertices in the mesh
 	virtual void ApplyToProjectVertices(const TFunction<void(int)>& VertexProjectFunc);
 
+	// Move all vertices in parallel. Update the timestamps.
+	// The NewVertexPosition function accepts a vertex index and returns a new vertex position, as well as setting the 
+	// bool out-param to indicate that the vertex has moved.
+	void MoveVerticesParallel(TFunction<FVector3d(int, bool&)> NewVertexPosition);
+
 
 	/*
 	 * testing/debug/profiling stuff
diff --git a/Engine/Plugins/Experimental/GeometryProcessing/Source/GeometricObjects/Private/Generators/DiscMeshGenerator.cpp b/Engine/Plugins/Experimental/GeometryProcessing/Source/GeometricObjects/Private/Generators/DiscMeshGenerator.cpp
--- a/Engine/Plugins/Experimental/GeometryProcessing/Source/GeometricObjects/Private/Generators/DiscMeshGenerator.cpp
+++ b/Engine/Plugins/Experimental/GeometryProcessing/Source/GeometricObjects/Private/Generators/DiscMeshGenerator.cpp
@@ -25,7 +25,7 @@ FPuncturedDiscMeshGenerator::FPuncturedDiscMeshGenerator()
 
 FMeshShapeGenerator& FDiscMeshGenerator::Generate()
 {
-	int AngleNV = AngleSamples > 2 ? AngleSamples : 2;
+	int AngleNV = AngleSamples > 3 ? AngleSamples : 3;
 	int RadialNV = RadialSamples > 1 ? RadialSamples : 1;
 	int NumVertices = AngleNV * RadialNV + 1;
 	
@@ -126,7 +126,7 @@ FMeshShapeGenerator& FDiscMeshGenerator::Generate()
 
 FMeshShapeGenerator& FPuncturedDiscMeshGenerator::Generate()
 {
-	int AngleNV = AngleSamples > 2 ? AngleSamples : 2;
+	int AngleNV = AngleSamples > 3 ? AngleSamples : 3;
 	int RadialNV = RadialSamples > 2 ? RadialSamples : 2;
 	int NumVertices = AngleNV * RadialNV;
 
diff --git a/Engine/Plugins/Experimental/GeometryProcessing/Source/GeometricObjects/Private/IndexUtil.cpp b/Engine/Plugins/Experimental/GeometryProcessing/Source/GeometricObjects/Private/IndexUtil.cpp
--- a/Engine/Plugins/Experimental/GeometryProcessing/Source/GeometricObjects/Private/IndexUtil.cpp
+++ b/Engine/Plugins/Experimental/GeometryProcessing/Source/GeometricObjects/Private/IndexUtil.cpp
@@ -2,6 +2,29 @@
 
 #include "Util/IndexUtil.h"
 
+// integer indices offsets in x/y directions
+const FVector2i IndexUtil::GridOffsets4[] =
+{
+	FVector2i(-1, 0), FVector2i(1, 0),
+	FVector2i(0, -1), FVector2i(0, 1)
+};
+
+// integer indices offsets in x/y directions and diagonals
+const FVector2i IndexUtil::GridOffsets8[] =
+{
+	FVector2i(-1, 0),  FVector2i(1, 0),
+	FVector2i(0, -1),  FVector2i(0, 1),
+	FVector2i(-1, 1),  FVector2i(1, 1),
+	FVector2i(-1, -1), FVector2i(1, -1)
+};
+
+// integer indices offsets in x/y/z directions, corresponds w/ BoxFaces directions
+const FVector3i IndexUtil::GridOffsets6[] =
+{
+	FVector3i( 0, 0,-1), FVector3i( 0, 0, 1),
+	FVector3i(-1, 0, 0), FVector3i( 1, 0, 0),
+	FVector3i( 0,-1, 0), FVector3i( 0, 1, 0)
+};
 
 // integer indices offsets in x/y/z directions and diagonals
 const FVector3i IndexUtil::GridOffsets26[] = 
diff --git a/Engine/Plugins/Experimental/GeometryProcessing/Source/GeometricObjects/Private/Util/SmallListSet.cpp b/Engine/Plugins/Experimental/GeometryProcessing/Source/GeometricObjects/Private/Util/SmallListSet.cpp
--- a/Engine/Plugins/Experimental/GeometryProcessing/Source/GeometricObjects/Private/Util/SmallListSet.cpp
+++ b/Engine/Plugins/Experimental/GeometryProcessing/Source/GeometricObjects/Private/Util/SmallListSet.cpp
@@ -334,6 +334,41 @@ bool FSmallListSet::Replace(int32 ListIndex, const TFunction<bool(int32)>& Predi
 
 
 
+void FSmallListSet::Enumerate(int32 ListIndex, TFunctionRef<void(int32)> ApplyFunc) const
+{
+	int32 block_ptr = ListHeads[ListIndex];
+	if (block_ptr != NullValue)
+	{
+		int32 N = ListBlocks[block_ptr];
+		if (N < BLOCKSIZE)
+		{
+			int32 iEnd = block_ptr + N;
+			for (int32 i = block_ptr + 1; i <= iEnd; ++i)
+			{
+				ApplyFunc(ListBlocks[i]);
+			}
+		}
+		else
+		{
+			// we spilled to linked list, have to iterate through it as well
+			int32 iEnd = block_ptr + BLOCKSIZE;
+			for (int32 i = block_ptr + 1; i <= iEnd; ++i)
+			{
+				ApplyFunc(ListBlocks[i]);
+			}
+			int32 cur_ptr = ListBlocks[block_ptr + BLOCK_LIST_OFFSET];
+			while (cur_ptr != NullValue)
+			{
+				ApplyFunc(LinkedListElements[cur_ptr]);
+				cur_ptr = LinkedListElements[cur_ptr + 1];
+			}
+		}
+	}
+}
+
+
+
+
 int32 FSmallListSet::AllocateBlock()
 {
 	int32 nfree = (int32)FreeBlocks.GetLength();
diff --git a/Engine/Plugins/Experimental/GeometryProcessing/Source/GeometricObjects/Public/BoxTypes.h b/Engine/Plugins/Experimental/GeometryProcessing/Source/GeometricObjects/Public/BoxTypes.h
--- a/Engine/Plugins/Experimental/GeometryProcessing/Source/GeometricObjects/Public/BoxTypes.h
+++ b/Engine/Plugins/Experimental/GeometryProcessing/Source/GeometricObjects/Public/BoxTypes.h
@@ -5,6 +5,7 @@
 #include "Math/Box.h"
 #include "Math/Box2D.h"
 #include "VectorTypes.h"
+#include "TransformTypes.h"
 
 template <typename RealType>
 struct TInterval1
@@ -177,6 +178,11 @@ struct TInterval1
 	{
 		return TInterval1(Min * f, Max * f);
 	}
+
+	inline bool IsEmpty() const
+	{
+		return Max > Min;
+	}
 };
 
 typedef TInterval1<float> FInterval1f;
@@ -200,6 +206,13 @@ struct TAxisAlignedBox3
 		this->Max = Max;
 	}
 
+	TAxisAlignedBox3(const FVector3<RealType>& V0, const FVector3<RealType>& V1, const FVector3<RealType>& V2)
+	{
+		TMathUtil<RealType>::MinMax(V0.X, V1.X, V2.X, Min.X, Max.X);
+		TMathUtil<RealType>::MinMax(V0.Y, V1.Y, V2.Y, Min.Y, Max.Y);
+		TMathUtil<RealType>::MinMax(V0.Z, V1.Z, V2.Z, Min.Z, Max.Z);
+	}
+
 	TAxisAlignedBox3(const TAxisAlignedBox3& OtherBox) = default;
 
 	template<typename OtherRealType>
@@ -234,6 +247,27 @@ struct TAxisAlignedBox3
 		}
 	}
 
+	TAxisAlignedBox3(const TAxisAlignedBox3& Box, const FTransform3d& Transform)
+	{
+		FVector3<RealType> C0 = Transform.TransformPosition(Box.GetCorner(0));
+		Min = C0;
+		Max = C0;
+		for (int i = 1; i < 8; ++i)
+		{
+			Contain(Transform.TransformPosition(Box.GetCorner(i)));
+		}
+	}
+
+	bool operator==(const TAxisAlignedBox3<RealType>& Other) const
+	{
+		return Max == Other.Max && Min == Other.Min;
+	}
+	bool operator!=(const TAxisAlignedBox3<RealType>& Other) const
+	{
+		return Max != Other.Max || Min != Other.Min;
+	}
+
+
 	explicit operator FBox() const
 	{
 		FVector MinV((float)Min.X, (float)Min.Y, (float)Min.Z);
@@ -410,6 +444,11 @@ struct TAxisAlignedBox3
 	{
 		return FVector3<RealType>(Max.X - Min.X, Max.Y - Min.Y, Max.Z - Min.Z);
 	}
+
+	inline bool IsEmpty() const
+	{
+		return Max.X > Min.X || Max.Y > Min.Y || Max.Z > Min.Z;
+	}
 };
 
 template <typename RealType>
@@ -588,6 +627,11 @@ struct TAxisAlignedBox2
 	{
 		return TMathUtil<RealType>::Min(Width(), Height());
 	}
+
+	inline bool IsEmpty() const
+	{
+		return Max.X > Min.X || Max.Y > Min.Y;
+	}
 };
 
 typedef TAxisAlignedBox2<float> FAxisAlignedBox2f;
diff --git a/Engine/Plugins/Experimental/GeometryProcessing/Source/GeometricObjects/Public/Curve/GeneralPolygon2.h b/Engine/Plugins/Experimental/GeometryProcessing/Source/GeometricObjects/Public/Curve/GeneralPolygon2.h
--- a/Engine/Plugins/Experimental/GeometryProcessing/Source/GeometricObjects/Public/Curve/GeneralPolygon2.h
+++ b/Engine/Plugins/Experimental/GeometryProcessing/Source/GeometricObjects/Public/Curve/GeneralPolygon2.h
@@ -83,7 +83,6 @@ public:
 				return false;
 			}
 
-			// [RMS] segment/segment intersection broken?
 			for (const TPolygon2<T>& ExistingHole : Holes)
 			{
 				if (Hole.Overlaps(ExistingHole))
