Edigrate Dev-Editor updates to InteractiveToolsFramework, GeometryProcessing, MeshModelingToolset, and ModelingTools Editor Mode Plugins. Update a few files in new MeshPaint Plugin for UInteractiveToo

diff --git a/Engine/Plugins/Experimental/GeometryProcessing/Source/GeometricObjects/Public/Generators/MarchingCubes.h b/Engine/Plugins/Experimental/GeometryProcessing/Source/GeometricObjects/Public/Generators/MarchingCubes.h
new file mode 100644
--- /dev/null
+++ b/Engine/Plugins/Experimental/GeometryProcessing/Source/GeometricObjects/Public/Generators/MarchingCubes.h
@@ -0,0 +1,1163 @@
+// Copyright Epic Games, Inc. All Rights Reserved.
+
+// Port of geometry3Sharp MarchingCubesPro
+
+#pragma once
+
+#include "Async/ParallelFor.h"
+
+#include "Math/UnrealMathUtility.h"
+#include "MeshShapeGenerator.h"
+#include "Misc/AssertionMacros.h"
+
+#include "Spatial/DenseGrid3.h"
+
+#include "CompGeom/PolygonTriangulation.h"
+
+
+enum class /*GEOMETRICOBJECTS_API*/ ERootfindingModes
+{
+	SingleLerp,
+	LerpSteps,
+	Bisection
+};
+
+class /*GEOMETRICOBJECTS_API*/ FMarchingCubes : public FMeshShapeGenerator
+{
+public:
+	/**
+	*  this is the function we will evaluate
+	*/
+	TFunction<double(FVector3<double>)> Implicit;
+
+	/**
+	*  mesh surface will be at this isovalue. Normally 0 unless you want
+	*  offset surface or field is not a distance-field.
+	*/
+	double IsoValue = 0;
+
+	/** bounding-box we will mesh inside of. We use the min-corner and
+	 *  the width/height/depth, but do not clamp vertices to stay within max-corner,
+	 *  we may spill one cell over
+	 */
+	TAxisAlignedBox3<double> Bounds;
+
+	/**
+	 *  Length of edges of cubes that are marching.
+	 *  currently, # of cells along axis = (int)(bounds_dimension / CellSize) + 1
+	 */
+	double CubeSize = 0.1;
+
+	/**
+	 *  Use multi-threading? Generally a good idea unless problem is very small or
+	 *  you are multi-threading at a higher level (which may be more efficient)
+	 */
+	bool bParallelCompute = true;
+
+
+
+	/**
+	 *  Which rootfinding method will be used to converge on surface along edges
+	 */
+	ERootfindingModes RootMode = ERootfindingModes::SingleLerp;
+
+	/**
+	 *  number of iterations of rootfinding method (ignored for SingleLerp)
+	 */
+	int RootModeSteps = 5;
+
+
+	/** if this function returns true, we should abort calculation */
+	TFunction<bool(void)> CancelF = []() { return false; };
+
+	/*
+	 * Outputs
+	 */
+
+	// cube indices range from [Origin,CellDimensions)   
+	FVector3i CellDimensions;
+
+
+	FMarchingCubes()
+	{
+		Bounds = TAxisAlignedBox3<double>(FVector3<double>::Zero(), 8);
+		CubeSize = 0.25;
+	}
+
+	virtual ~FMarchingCubes()
+	{
+	}
+
+
+
+	/**
+	*  Run MC algorithm and generate Output mesh
+	*/
+	FMeshShapeGenerator& Generate() override
+	{
+		int32 nx = (int32)(Bounds.Width() / CubeSize) + 1;
+		int32 ny = (int32)(Bounds.Height() / CubeSize) + 1;
+		int32 nz = (int32)(Bounds.Depth() / CubeSize) + 1;
+		CellDimensions = FVector3i(nx, ny, nz);
+		GridBounds = TAxisAlignedBox3<int>(FVector3i::Zero(), CellDimensions - FVector3i(1,1,1)); // grid bounds are inclusive
+
+		corner_values_grid = FDenseGrid3f(nx+1, ny+1, nz+1, FMathf::MaxReal);
+		edge_vertices.Reset();
+		corner_values.Reset();
+
+		if (bParallelCompute) {
+			generate_parallel();
+		} else {
+			generate_basic();
+		}
+
+		return *this;
+	}
+
+
+	FMeshShapeGenerator& GenerateContinuation(TArrayView<const FVector3<double>> Seeds)
+	{
+		int nx = (int)(Bounds.Width() / CubeSize) + 1;
+		int ny = (int)(Bounds.Height() / CubeSize) + 1;
+		int nz = (int)(Bounds.Depth() / CubeSize) + 1;
+		CellDimensions = FVector3i(nx, ny, nz);
+		GridBounds = TAxisAlignedBox3<int>(FVector3i::Zero(), CellDimensions - FVector3i(1,1,1)); // grid bounds are inclusive
+
+		UE_LOG(LogTemp, Warning, TEXT("marching cubes dims = %d %d %d"), nx, ny, nz);
+
+		if (LastGridBounds != GridBounds)
+		{
+			corner_values_grid = FDenseGrid3f(nx + 1, ny + 1, nz + 1, FMathf::MaxReal);
+			edge_vertices.Reset();
+			corner_values.Reset();
+			if (bParallelCompute)
+			{
+				done_cells = FDenseGrid3i(CellDimensions.X, CellDimensions.Y, CellDimensions.Z, 0);
+			}
+		}
+		else
+		{
+			edge_vertices.Reset();
+			corner_values.Reset();
+			corner_values_grid.Assign(FMathf::MaxReal);
+			if (bParallelCompute)
+			{
+				done_cells.Assign(0);
+			}
+		}
+
+		if (bParallelCompute) {
+			generate_continuation_parallel(Seeds);
+		} else {
+			generate_continuation(Seeds);
+		}
+
+		LastGridBounds = GridBounds;
+
+		return *this;
+	}
+
+
+protected:
+
+
+	FAxisAlignedBox3i GridBounds;
+	FAxisAlignedBox3i LastGridBounds;
+
+
+	// we pass Cells around, this makes code cleaner
+	struct FGridCell
+	{
+		// TODO we do not actually need to store i, we just need the min-corner!
+		FVector3i i[8];    // indices of corners of cell
+		double f[8];      // field values at corners
+	};
+
+
+
+	void corner_pos(const FVector3i& IJK, FVector3<double>& Pos)
+	{
+		Pos.X = Bounds.Min.X + CubeSize * IJK.X;
+		Pos.Y = Bounds.Min.Y + CubeSize * IJK.Y;
+		Pos.Z = Bounds.Min.Z + CubeSize * IJK.Z;
+	}
+	FVector3<double> corner_pos(const FVector3i& IJK)
+	{
+		return FVector3<double>(Bounds.Min.X + CubeSize * IJK.X,
+			Bounds.Min.Y + CubeSize * IJK.Y,
+			Bounds.Min.Z + CubeSize * IJK.Z);
+	}
+	FVector3i cell_index(const FVector3<double>& Pos)
+	{
+		return FVector3i(
+			(int)((Pos.X - Bounds.Min.X) / CubeSize),
+			(int)((Pos.Y - Bounds.Min.Y) / CubeSize),
+			(int)((Pos.Z - Bounds.Min.Z) / CubeSize));
+	}
+
+
+
+	//
+	// corner and edge hash functions, these pack the coordinate
+	// integers into 16-bits, so max of 65536 in any dimension.
+	//
+
+
+	int64 corner_hash(const FVector3i& Idx)
+	{
+		return ((int64)Idx.X&0xFFFF) | (((int64)Idx.Y&0xFFFF) << 16) | (((int64)Idx.Z&0xFFFF) << 32);
+	}
+	int64 corner_hash(int X, int Y, int Z)
+	{
+		return ((int64)X & 0xFFFF) | (((int64)Y & 0xFFFF) << 16) | (((int64)Z & 0xFFFF) << 32);
+	}
+
+	const int64 EDGE_X = int64(1) << 60;
+	const int64 EDGE_Y = int64(1) << 61;
+	const int64 EDGE_Z = int64(1) << 62;
+
+	int64 edge_hash(const FVector3i& Idx1, const FVector3i& Idx2)
+	{
+		if ( Idx1.X != Idx2.X )
+		{
+			int xlo = FMath::Min(Idx1.X, Idx2.X);
+			return corner_hash(xlo, Idx1.Y, Idx1.Z) | EDGE_X;
+		}
+		else if ( Idx1.Y != Idx2.Y )
+		{
+			int ylo = FMath::Min(Idx1.Y, Idx2.Y);
+			return corner_hash(Idx1.X, ylo, Idx1.Z) | EDGE_Y;
+		}
+		else
+		{
+			int zlo = FMath::Min(Idx1.Z, Idx2.Z);
+			return corner_hash(Idx1.X, Idx1.Y, zlo) | EDGE_Z;
+		}
+	}
+
+
+
+	//
+	// Hash table for edge vertices
+	//
+
+	TMap<int64, int> edge_vertices;
+	FCriticalSection edge_vertices_lock;
+
+	int edge_vertex_id(const FVector3i& Idx1, const FVector3i& Idx2, double F1, double F2)
+	{
+		int64 hash = edge_hash(Idx1, Idx2);
+
+		int vid = IndexConstants::InvalidID;
+		bool found = false;
+		{
+			FScopeLock Lock(&edge_vertices_lock);
+			int* foundvid = edge_vertices.Find(hash);
+			if (foundvid)
+			{
+				vid = *foundvid;
+				found = true;
+			}
+		}
+		
+		if (found)
+		{
+			return vid;
+		}
+
+		// ok this is a bit messy. We do not want to lock the entire hash table 
+		// while we do find_iso. However it is possible that during this time we
+		// are unlocked we have re-entered with the same edge. So when we
+		// re-acquire the lock we need to check again that we have not already
+		// computed this edge, otherwise we will end up with duplicate vertices!
+
+		FVector3<double> pa = FVector3<double>::Zero(), pb = FVector3<double>::Zero();
+		corner_pos(Idx1, pa);
+		corner_pos(Idx2, pb);
+		FVector3<double> Pos = FVector3<double>::Zero();
+		find_iso(pa, pb, F1, F2, Pos);
+
+		{
+			FScopeLock Lock(&edge_vertices_lock);
+			int* foundvid = edge_vertices.Find(hash);
+			if (!foundvid)
+			{
+				vid = append_vertex(Pos);
+				edge_vertices.Add(hash, vid);
+			}
+			else
+			{
+				vid = *foundvid;
+			}
+		}
+
+		return vid;
+	}
+
+
+
+
+
+
+	//
+	// Store corner values in hash table. This doesn't make
+	// sense if we are evaluating entire grid, way too slow.
+	//
+
+	TMap<int64, double> corner_values;
+	FCriticalSection corner_values_lock;
+
+	double corner_value(const FVector3i& Idx)
+	{
+		int64 hash = corner_hash(Idx);
+		double value = 0;
+		double* findval = corner_values.Find(hash);
+		if (findval == nullptr)
+		{
+			FVector3<double> V = corner_pos(Idx);
+			value = Implicit(V);
+			corner_values.Add(hash, value);
+		} 
+		else
+		{
+			value = *findval;
+		}
+		return value;
+	}
+	void initialize_cell_values(FGridCell& Cell, bool Shift)
+	{
+		FScopeLock Lock(&corner_values_lock);
+
+		if ( Shift )
+		{
+			Cell.f[1] = corner_value(Cell.i[1]);
+			Cell.f[2] = corner_value(Cell.i[2]);
+			Cell.f[5] = corner_value(Cell.i[5]);
+			Cell.f[6] = corner_value(Cell.i[6]);
+		}
+		else
+		{
+			for (int i = 0; i < 8; ++i)
+			{
+				Cell.f[i] = corner_value(Cell.i[i]);
+			}
+		}
+	}
+
+
+
+	//
+	// store corner values in pre-allocated grid that has
+	// FMathf::MaxReal as sentinel. 
+	// (note this is float grid, not double...)
+	//
+
+	FDenseGrid3f corner_values_grid;
+
+	double corner_value_grid(const FVector3i& Idx)
+	{
+		double val = corner_values_grid[Idx];
+		if (val != FMathf::MaxReal)
+		{
+			return val;
+		}
+
+		FVector3<double> V = corner_pos(Idx);
+		val = Implicit(V);
+		corner_values_grid[Idx] = (float)val;
+		return val;
+	}
+	void initialize_cell_values_grid(FGridCell& Cell, bool Shift)
+	{
+		if (Shift)
+		{
+			Cell.f[1] = corner_value_grid(Cell.i[1]);
+			Cell.f[2] = corner_value_grid(Cell.i[2]);
+			Cell.f[5] = corner_value_grid(Cell.i[5]);
+			Cell.f[6] = corner_value_grid(Cell.i[6]);
+		}
+		else
+		{
+			for (int i = 0; i < 8; ++i)
+			{
+				Cell.f[i] = corner_value_grid(Cell.i[i]);
+			}
+		}
+	}
+
+
+
+	//
+	// explicitly compute corner values as necessary
+	//
+	//
+
+	double corner_value_nohash(const FVector3i& Idx) {
+		FVector3<double> V = corner_pos(Idx);
+		return Implicit(V);
+	}
+	void initialize_cell_values_nohash(FGridCell& Cell, bool Shift)
+	{
+		if (Shift)
+		{
+			Cell.f[1] = corner_value_nohash(Cell.i[1]);
+			Cell.f[2] = corner_value_nohash(Cell.i[2]);
+			Cell.f[5] = corner_value_nohash(Cell.i[5]);
+			Cell.f[6] = corner_value_nohash(Cell.i[6]);
+		}
+		else
+		{
+			for (int i = 0; i < 8; ++i)
+			{
+				Cell.f[i] = corner_value_nohash(Cell.i[i]);
+			}
+		}
+	}
+
+
+
+
+
+	/**
+	*  compute 3D corner-positions and field values for cell at index
+	*/
+	void initialize_cell(FGridCell& Cell, const FVector3i& Idx)
+	{
+		Cell.i[0] = FVector3i(Idx.X + 0, Idx.Y + 0, Idx.Z + 0);
+		Cell.i[1] = FVector3i(Idx.X + 1, Idx.Y + 0, Idx.Z + 0);
+		Cell.i[2] = FVector3i(Idx.X + 1, Idx.Y + 0, Idx.Z + 1);
+		Cell.i[3] = FVector3i(Idx.X + 0, Idx.Y + 0, Idx.Z + 1);
+		Cell.i[4] = FVector3i(Idx.X + 0, Idx.Y + 1, Idx.Z + 0);
+		Cell.i[5] = FVector3i(Idx.X + 1, Idx.Y + 1, Idx.Z + 0);
+		Cell.i[6] = FVector3i(Idx.X + 1, Idx.Y + 1, Idx.Z + 1);
+		Cell.i[7] = FVector3i(Idx.X + 0, Idx.Y + 1, Idx.Z + 1);
+
+		//initialize_cell_values(Cell, false);
+		initialize_cell_values_grid(Cell, false);
+		//initialize_cell_values_nohash(Cell, false);
+	}
+
+
+	// assume we just want to slide cell at XIdx-1 to cell at XIdx, while keeping
+	// yi and ZIdx constant. Then only x-coords change, and we have already 
+	// computed half the values
+	void shift_cell_x(FGridCell& Cell, int XIdx)
+	{
+		Cell.f[0] = Cell.f[1];
+		Cell.f[3] = Cell.f[2];
+		Cell.f[4] = Cell.f[5];
+		Cell.f[7] = Cell.f[6];
+
+		Cell.i[0].X = XIdx; Cell.i[1].X = XIdx+1; Cell.i[2].X = XIdx+1; Cell.i[3].X = XIdx;
+		Cell.i[4].X = XIdx; Cell.i[5].X = XIdx+1; Cell.i[6].X = XIdx+1; Cell.i[7].X = XIdx;
+
+		//initialize_cell_values(Cell, true);
+		initialize_cell_values_grid(Cell, true);
+		//initialize_cell_values_nohash(Cell, true);
+	}
+
+
+	bool parallel_mesh_access = false;
+	FCriticalSection mesh_lock;
+
+	/**
+	*  processing z-slabs of cells in parallel
+	*/
+	void generate_parallel()
+	{
+		parallel_mesh_access = true;
+
+		// [TODO] maybe shouldn't alway use Z axis here?
+		ParallelFor(CellDimensions.Z, [this](int32 ZIdx)
+		{
+			FGridCell Cell;
+			int vertTArray[12];
+			for (int yi = 0; yi < CellDimensions.Y; ++yi)
+			{
+				if (CancelF())
+				{
+					return;
+				}
+				// compute full cell at x=0, then slide along x row, which saves half of value computes
+				FVector3i Idx(0, yi, ZIdx);
+				initialize_cell(Cell, Idx);
+				polygonize_cell(Cell, vertTArray);
+				for (int XIdx = 1; XIdx < CellDimensions.X; ++XIdx)
+				{
+					shift_cell_x(Cell, XIdx);
+					polygonize_cell(Cell, vertTArray);
+				}
+			}
+		});
+
+
+		parallel_mesh_access = false;
+	}
+
+
+
+
+	/**
+	*  fully sequential version, no threading
+	*/
+	void generate_basic()
+	{
+		FGridCell Cell;
+		int vertTArray[12];
+
+		for (int ZIdx = 0; ZIdx < CellDimensions.Z; ++ZIdx)
+		{
+			for (int yi = 0; yi < CellDimensions.Y; ++yi)
+			{
+				if (CancelF())
+				{
+					return;
+				}
+				// compute full Cell at x=0, then slide along x row, which saves half of value computes
+				FVector3i Idx(0, yi, ZIdx);
+				initialize_cell(Cell, Idx);
+				polygonize_cell(Cell, vertTArray);
+				for (int XIdx = 1; XIdx < CellDimensions.X; ++XIdx)
+				{
+					shift_cell_x(Cell, XIdx);
+					polygonize_cell(Cell, vertTArray);
+				}
+
+			}
+		}
+	}
+
+
+
+
+	/**
+	*  fully sequential version, no threading
+	*/
+	void generate_continuation(TArrayView<const FVector3<double>> Seeds)
+	{
+		FGridCell Cell;
+		int vertTArray[12];
+
+		done_cells = FDenseGrid3i(CellDimensions.X, CellDimensions.Y, CellDimensions.Z, 0);
+
+		TArray<FVector3i> stack;
+
+		for (FVector3<double> seed : Seeds)
+		{
+			FVector3i seed_idx = cell_index(seed);
+			if (done_cells[seed_idx] == 1)
+			{
+				continue;
+			}
+			stack.Add(seed_idx);
+			done_cells[seed_idx] = 1;
+
+			while ( stack.Num() > 0 )
+			{
+				FVector3i Idx = stack[stack.Num()-1]; 
+				stack.RemoveAt(stack.Num()-1);
+				if (CancelF())
+				{
+					return;
+				}
+
+				initialize_cell(Cell, Idx);
+				if ( polygonize_cell(Cell, vertTArray) )
+				{     // found crossing
+					for ( FVector3i o : IndexUtil::GridOffsets6 )
+					{
+						FVector3i nbr_idx = Idx + o;
+						if (GridBounds.Contains(nbr_idx) && done_cells[nbr_idx] == 0)
+						{
+							stack.Add(nbr_idx);
+							done_cells[nbr_idx] = 1;
+						}
+					}
+				}
+			}
+		}
+	}
+
+
+
+
+	/**
+	*  parallel seed evaluation
+	*/
+	void generate_continuation_parallel(TArrayView<const FVector3<double>> Seeds)
+	{
+		parallel_mesh_access = true;
+
+		ParallelFor(Seeds.Num(), [&](int32 Index)
+		{
+			FVector3<double> Seed = Seeds[Index];
+			FVector3i seed_idx = cell_index(Seed);
+			if (set_cell_if_not_done(seed_idx) == false)
+			{
+				return;
+			}
+
+			FGridCell Cell;
+			int vertTArray[12];
+
+			TArray<FVector3i> stack;
+			stack.Add(seed_idx);
+
+			while (stack.Num() > 0)
+			{
+				FVector3i Idx = stack[stack.Num() - 1];
+				stack.RemoveAt(stack.Num() - 1);
+				if (CancelF())
+				{
+					return;
+				}
+
+				initialize_cell(Cell, Idx);
+				if (polygonize_cell(Cell, vertTArray))
+				{     // found crossing
+					for (FVector3i o : IndexUtil::GridOffsets6)
+					{
+						FVector3i nbr_idx = Idx + o;
+						if (GridBounds.Contains(nbr_idx))
+						{
+							if (set_cell_if_not_done(nbr_idx) == true)
+							{ 
+								stack.Add(nbr_idx);
+							}
+						}
+					}
+				}
+			}
+		});
+
+		parallel_mesh_access = false;
+	}
+
+
+
+	FDenseGrid3i done_cells;
+	FCriticalSection done_cells_lock;
+
+	bool set_cell_if_not_done(const FVector3i& Idx)
+	{
+		bool was_set = false;
+		FScopeLock Lock(&done_cells_lock);
+		if (done_cells[Idx] == 0)
+		{
+			done_cells[Idx] = 1;
+			was_set = true;
+		}
+		return was_set;
+	}
+
+
+
+
+
+
+
+
+
+
+	/**
+	*  find edge crossings and generate triangles for this cell
+	*/
+	bool polygonize_cell(FGridCell& Cell, int VertIndexArray[8])
+	{
+		// construct bits of index into edge table, where bit for each
+		// corner is 1 if that value is < isovalue.
+		// This tell us which edges have sign-crossings, and the int value
+		// of the bitmap is an index into the edge and triangle tables
+		int cubeindex = 0, Shift = 1;
+		for (int i = 0; i < 8; ++i)
+		{
+			if (Cell.f[i] < IsoValue)
+				cubeindex |= Shift;
+			Shift <<= 1;
+		}
+
+		// no crossings!
+		if (EdgeTable[cubeindex] == 0)
+		{
+			return false;
+		}
+
+		// check each bit of value in edge table. If it is 1, we
+		// have a crossing on that edge. Look up the indices of this
+		// edge and find the intersection point along it
+		Shift = 1;
+		FVector3<double> pa = FVector3<double>::Zero(), pb = FVector3<double>::Zero();
+		for (int i = 0; i <= 11; i++)
+		{
+			if ((EdgeTable[cubeindex] & Shift) != 0)
+			{
+				int a = EdgeIndices[i][0], b = EdgeIndices[i][1];
+				VertIndexArray[i] = edge_vertex_id(Cell.i[a], Cell.i[b], Cell.f[a], Cell.f[b]);
+			}
+			Shift <<= 1;
+		}
+
+		// now iterate through the set of triangles in TriTable for this cube,
+		// and emit triangles using the vertices we found.
+		int tri_count = 0;
+		for (int i = 0; TriTable[cubeindex][i] != -1; i += 3)
+		{
+			int ta = TriTable[cubeindex][i];
+			int tb = TriTable[cubeindex][i + 1];
+			int tc = TriTable[cubeindex][i + 2];
+			int a = VertIndexArray[ta], b = VertIndexArray[tb], c = VertIndexArray[tc];
+
+			// if a corner is within tolerance of isovalue, then some triangles
+			// will be degenerate, and we can skip them w/o resulting in cracks (right?)
+			// !! this should never happen anymore...artifact of old hashtable impl
+			if (a == b || a == c || b == c)
+			{
+				continue;
+			}
+
+			/*int tid = */
+			append_triangle(a, b, c);
+			tri_count++;
+		}
+
+		return (tri_count > 0);
+	}
+
+
+
+
+	/**
+	*  add vertex to mesh, with locking if we are computing in parallel
+	*/
+	int append_vertex(FVector3<double> V)
+	{
+		if (parallel_mesh_access)
+		{
+			FScopeLock Lock(&mesh_lock);
+			return AppendVertex((FVector3d)V);
+		}
+		else
+		{
+			return AppendVertex((FVector3d)V);
+		}
+	}
+
+
+
+	/**
+	*  add triangle to mesh, with locking if we are computing in parallel
+	*/
+	int append_triangle(int A, int B, int C)
+	{
+		if (parallel_mesh_access)
+		{
+			FScopeLock Lock(&mesh_lock);
+			return AppendTriangle(A, B, C);
+		}
+		else
+		{
+			return AppendTriangle(A, B, C);
+		}
+	}
+
+
+
+	/**
+	*  root-find the intersection along edge from f(P1)=ValP1 to f(P2)=ValP2
+	*/
+	void find_iso(const FVector3<double>& P1, const FVector3<double>& P2, double ValP1, double ValP2, FVector3<double>& PIso)
+	{
+		// Ok, this is a bit hacky but seems to work? If both isovalues
+		// are the same, we just return the midpoint. If one is nearly zero, we can
+		// but assume that's where the surface is. *However* if we return that point exactly,
+		// we can get nonmanifold vertices, because multiple fans may connect there. 
+		// Since FDynamicMesh3 disallows that, it results in holes. So we pull 
+		// slightly towards the other point along this edge. This means we will get
+		// repeated nearly-coincident vertices, but the mesh will be manifold.
+		const double dt = 0.999999;
+		if (FMath::Abs(ValP1 - ValP2) < 0.00001)
+		{
+			PIso = (P1 + P2) * 0.5;
+			return;
+		}
+		if (FMath::Abs(IsoValue - ValP1) < 0.00001)
+		{
+			PIso = dt * P1 + (1.0 - dt) * P2;
+			return;
+		}
+		if (FMath::Abs(IsoValue - ValP2) < 0.00001)
+		{
+			PIso = (dt) * P2 + (1.0 - dt) * P1;
+			return;
+		}
+
+		// Note: if we don't maintain min/max order here, then numerical error means
+		//   that hashing on point x/y/z doesn't work
+		FVector3<double> a = P1, b = P2;
+		double fa = ValP1, fb = ValP2;
+		if (ValP2 < ValP1)
+		{
+			a = P2; b = P1;
+			fb = ValP1; fa = ValP2;
+		}
+
+		// converge on root
+		if (RootMode == ERootfindingModes::Bisection)
+		{
+			for (int k = 0; k < RootModeSteps; ++k)
+			{
+				PIso.X = (a.X + b.X) * 0.5; PIso.Y = (a.Y + b.Y) * 0.5; PIso.Z = (a.Z + b.Z) * 0.5;
+				double mid_f = Implicit(PIso);
+				if (mid_f < IsoValue)
+				{
+					a = PIso; fa = mid_f;
+				}
+				else
+				{
+					b = PIso; fb = mid_f;
+				}
+			}
+			PIso = FVector3<double>::Lerp(a, b, 0.5);
+
+		}
+		else
+		{
+			double mu = 0;
+			if (RootMode == ERootfindingModes::LerpSteps)
+			{
+				for (int k = 0; k < RootModeSteps; ++k)
+				{
+					mu = (IsoValue - fa) / (fb - fa);
+					PIso.X = a.X + mu * (b.X - a.X);
+					PIso.Y = a.Y + mu * (b.Y - a.Y);
+					PIso.Z = a.Z + mu * (b.Z - a.Z);
+					double mid_f = Implicit(PIso);
+					if (mid_f < IsoValue)
+					{
+						a = PIso; fa = mid_f;
+					}
+					else
+					{
+						b = PIso; fb = mid_f;
+					}
+				}
+			}
+
+			// final lerp
+			mu = (IsoValue - fa) / (fb - fa);
+			PIso.X = a.X + mu * (b.X - a.X);
+			PIso.Y = a.Y + mu * (b.Y - a.Y);
+			PIso.Z = a.Z + mu * (b.Z - a.Z);
+		}
+	}
+
+
+
+
+	/*
+	* Below here are standard marching-cubes tables. 
+	*/
+
+
+	constexpr static int EdgeIndices[12][2] = {
+		{0,1}, {1,2}, {2,3}, {3,0}, {4,5}, {5,6}, {6,7}, {7,4}, {0,4}, {1,5}, {2,6}, {3,7}
+	};
+
+	constexpr static int EdgeTable[256] = {
+		0x0  , 0x109, 0x203, 0x30a, 0x406, 0x50f, 0x605, 0x70c,
+		0x80c, 0x905, 0xa0f, 0xb06, 0xc0a, 0xd03, 0xe09, 0xf00,
+		0x190, 0x99 , 0x393, 0x29a, 0x596, 0x49f, 0x795, 0x69c,
+		0x99c, 0x895, 0xb9f, 0xa96, 0xd9a, 0xc93, 0xf99, 0xe90,
+		0x230, 0x339, 0x33 , 0x13a, 0x636, 0x73f, 0x435, 0x53c,
+		0xa3c, 0xb35, 0x83f, 0x936, 0xe3a, 0xf33, 0xc39, 0xd30,
+		0x3a0, 0x2a9, 0x1a3, 0xaa , 0x7a6, 0x6af, 0x5a5, 0x4ac,
+		0xbac, 0xaa5, 0x9af, 0x8a6, 0xfaa, 0xea3, 0xda9, 0xca0,
+		0x460, 0x569, 0x663, 0x76a, 0x66 , 0x16f, 0x265, 0x36c,
+		0xc6c, 0xd65, 0xe6f, 0xf66, 0x86a, 0x963, 0xa69, 0xb60,
+		0x5f0, 0x4f9, 0x7f3, 0x6fa, 0x1f6, 0xff , 0x3f5, 0x2fc,
+		0xdfc, 0xcf5, 0xfff, 0xef6, 0x9fa, 0x8f3, 0xbf9, 0xaf0,
+		0x650, 0x759, 0x453, 0x55a, 0x256, 0x35f, 0x55 , 0x15c,
+		0xe5c, 0xf55, 0xc5f, 0xd56, 0xa5a, 0xb53, 0x859, 0x950,
+		0x7c0, 0x6c9, 0x5c3, 0x4ca, 0x3c6, 0x2cf, 0x1c5, 0xcc ,
+		0xfcc, 0xec5, 0xdcf, 0xcc6, 0xbca, 0xac3, 0x9c9, 0x8c0,
+		0x8c0, 0x9c9, 0xac3, 0xbca, 0xcc6, 0xdcf, 0xec5, 0xfcc,
+		0xcc , 0x1c5, 0x2cf, 0x3c6, 0x4ca, 0x5c3, 0x6c9, 0x7c0,
+		0x950, 0x859, 0xb53, 0xa5a, 0xd56, 0xc5f, 0xf55, 0xe5c,
+		0x15c, 0x55 , 0x35f, 0x256, 0x55a, 0x453, 0x759, 0x650,
+		0xaf0, 0xbf9, 0x8f3, 0x9fa, 0xef6, 0xfff, 0xcf5, 0xdfc,
+		0x2fc, 0x3f5, 0xff , 0x1f6, 0x6fa, 0x7f3, 0x4f9, 0x5f0,
+		0xb60, 0xa69, 0x963, 0x86a, 0xf66, 0xe6f, 0xd65, 0xc6c,
+		0x36c, 0x265, 0x16f, 0x66 , 0x76a, 0x663, 0x569, 0x460,
+		0xca0, 0xda9, 0xea3, 0xfaa, 0x8a6, 0x9af, 0xaa5, 0xbac,
+		0x4ac, 0x5a5, 0x6af, 0x7a6, 0xaa , 0x1a3, 0x2a9, 0x3a0,
+		0xd30, 0xc39, 0xf33, 0xe3a, 0x936, 0x83f, 0xb35, 0xa3c,
+		0x53c, 0x435, 0x73f, 0x636, 0x13a, 0x33 , 0x339, 0x230,
+		0xe90, 0xf99, 0xc93, 0xd9a, 0xa96, 0xb9f, 0x895, 0x99c,
+		0x69c, 0x795, 0x49f, 0x596, 0x29a, 0x393, 0x99 , 0x190,
+		0xf00, 0xe09, 0xd03, 0xc0a, 0xb06, 0xa0f, 0x905, 0x80c,
+		0x70c, 0x605, 0x50f, 0x406, 0x30a, 0x203, 0x109, 0x0   };
+
+
+	constexpr static int TriTable[256][16] =
+	{
+		{-1, -1, -1, -1, -1, -1, -1, -1, -1, -1, -1, -1, -1, -1, -1, -1},
+		{0, 8, 3, -1, -1, -1, -1, -1, -1, -1, -1, -1, -1, -1, -1, -1},
+		{0, 1, 9, -1, -1, -1, -1, -1, -1, -1, -1, -1, -1, -1, -1, -1},
+		{1, 8, 3, 9, 8, 1, -1, -1, -1, -1, -1, -1, -1, -1, -1, -1},
+		{1, 2, 10, -1, -1, -1, -1, -1, -1, -1, -1, -1, -1, -1, -1, -1},
+		{0, 8, 3, 1, 2, 10, -1, -1, -1, -1, -1, -1, -1, -1, -1, -1},
+		{9, 2, 10, 0, 2, 9, -1, -1, -1, -1, -1, -1, -1, -1, -1, -1},
+		{2, 8, 3, 2, 10, 8, 10, 9, 8, -1, -1, -1, -1, -1, -1, -1},
+		{3, 11, 2, -1, -1, -1, -1, -1, -1, -1, -1, -1, -1, -1, -1, -1},
+		{0, 11, 2, 8, 11, 0, -1, -1, -1, -1, -1, -1, -1, -1, -1, -1},
+		{1, 9, 0, 2, 3, 11, -1, -1, -1, -1, -1, -1, -1, -1, -1, -1},
+		{1, 11, 2, 1, 9, 11, 9, 8, 11, -1, -1, -1, -1, -1, -1, -1},
+		{3, 10, 1, 11, 10, 3, -1, -1, -1, -1, -1, -1, -1, -1, -1, -1},
+		{0, 10, 1, 0, 8, 10, 8, 11, 10, -1, -1, -1, -1, -1, -1, -1},
+		{3, 9, 0, 3, 11, 9, 11, 10, 9, -1, -1, -1, -1, -1, -1, -1},
+		{9, 8, 10, 10, 8, 11, -1, -1, -1, -1, -1, -1, -1, -1, -1, -1},
+		{4, 7, 8, -1, -1, -1, -1, -1, -1, -1, -1, -1, -1, -1, -1, -1},
+		{4, 3, 0, 7, 3, 4, -1, -1, -1, -1, -1, -1, -1, -1, -1, -1},
+		{0, 1, 9, 8, 4, 7, -1, -1, -1, -1, -1, -1, -1, -1, -1, -1},
+		{4, 1, 9, 4, 7, 1, 7, 3, 1, -1, -1, -1, -1, -1, -1, -1},
+		{1, 2, 10, 8, 4, 7, -1, -1, -1, -1, -1, -1, -1, -1, -1, -1},
+		{3, 4, 7, 3, 0, 4, 1, 2, 10, -1, -1, -1, -1, -1, -1, -1},
+		{9, 2, 10, 9, 0, 2, 8, 4, 7, -1, -1, -1, -1, -1, -1, -1},
+		{2, 10, 9, 2, 9, 7, 2, 7, 3, 7, 9, 4, -1, -1, -1, -1},
+		{8, 4, 7, 3, 11, 2, -1, -1, -1, -1, -1, -1, -1, -1, -1, -1},
+		{11, 4, 7, 11, 2, 4, 2, 0, 4, -1, -1, -1, -1, -1, -1, -1},
+		{9, 0, 1, 8, 4, 7, 2, 3, 11, -1, -1, -1, -1, -1, -1, -1},
+		{4, 7, 11, 9, 4, 11, 9, 11, 2, 9, 2, 1, -1, -1, -1, -1},
+		{3, 10, 1, 3, 11, 10, 7, 8, 4, -1, -1, -1, -1, -1, -1, -1},
+		{1, 11, 10, 1, 4, 11, 1, 0, 4, 7, 11, 4, -1, -1, -1, -1},
+		{4, 7, 8, 9, 0, 11, 9, 11, 10, 11, 0, 3, -1, -1, -1, -1},
+		{4, 7, 11, 4, 11, 9, 9, 11, 10, -1, -1, -1, -1, -1, -1, -1},
+		{9, 5, 4, -1, -1, -1, -1, -1, -1, -1, -1, -1, -1, -1, -1, -1},
+		{9, 5, 4, 0, 8, 3, -1, -1, -1, -1, -1, -1, -1, -1, -1, -1},
+		{0, 5, 4, 1, 5, 0, -1, -1, -1, -1, -1, -1, -1, -1, -1, -1},
+		{8, 5, 4, 8, 3, 5, 3, 1, 5, -1, -1, -1, -1, -1, -1, -1},
+		{1, 2, 10, 9, 5, 4, -1, -1, -1, -1, -1, -1, -1, -1, -1, -1},
+		{3, 0, 8, 1, 2, 10, 4, 9, 5, -1, -1, -1, -1, -1, -1, -1},
+		{5, 2, 10, 5, 4, 2, 4, 0, 2, -1, -1, -1, -1, -1, -1, -1},
+		{2, 10, 5, 3, 2, 5, 3, 5, 4, 3, 4, 8, -1, -1, -1, -1},
+		{9, 5, 4, 2, 3, 11, -1, -1, -1, -1, -1, -1, -1, -1, -1, -1},
+		{0, 11, 2, 0, 8, 11, 4, 9, 5, -1, -1, -1, -1, -1, -1, -1},
+		{0, 5, 4, 0, 1, 5, 2, 3, 11, -1, -1, -1, -1, -1, -1, -1},
+		{2, 1, 5, 2, 5, 8, 2, 8, 11, 4, 8, 5, -1, -1, -1, -1},
+		{10, 3, 11, 10, 1, 3, 9, 5, 4, -1, -1, -1, -1, -1, -1, -1},
+		{4, 9, 5, 0, 8, 1, 8, 10, 1, 8, 11, 10, -1, -1, -1, -1},
+		{5, 4, 0, 5, 0, 11, 5, 11, 10, 11, 0, 3, -1, -1, -1, -1},
+		{5, 4, 8, 5, 8, 10, 10, 8, 11, -1, -1, -1, -1, -1, -1, -1},
+		{9, 7, 8, 5, 7, 9, -1, -1, -1, -1, -1, -1, -1, -1, -1, -1},
+		{9, 3, 0, 9, 5, 3, 5, 7, 3, -1, -1, -1, -1, -1, -1, -1},
+		{0, 7, 8, 0, 1, 7, 1, 5, 7, -1, -1, -1, -1, -1, -1, -1},
+		{1, 5, 3, 3, 5, 7, -1, -1, -1, -1, -1, -1, -1, -1, -1, -1},
+		{9, 7, 8, 9, 5, 7, 10, 1, 2, -1, -1, -1, -1, -1, -1, -1},
+		{10, 1, 2, 9, 5, 0, 5, 3, 0, 5, 7, 3, -1, -1, -1, -1},
+		{8, 0, 2, 8, 2, 5, 8, 5, 7, 10, 5, 2, -1, -1, -1, -1},
+		{2, 10, 5, 2, 5, 3, 3, 5, 7, -1, -1, -1, -1, -1, -1, -1},
+		{7, 9, 5, 7, 8, 9, 3, 11, 2, -1, -1, -1, -1, -1, -1, -1},
+		{9, 5, 7, 9, 7, 2, 9, 2, 0, 2, 7, 11, -1, -1, -1, -1},
+		{2, 3, 11, 0, 1, 8, 1, 7, 8, 1, 5, 7, -1, -1, -1, -1},
+		{11, 2, 1, 11, 1, 7, 7, 1, 5, -1, -1, -1, -1, -1, -1, -1},
+		{9, 5, 8, 8, 5, 7, 10, 1, 3, 10, 3, 11, -1, -1, -1, -1},
+		{5, 7, 0, 5, 0, 9, 7, 11, 0, 1, 0, 10, 11, 10, 0, -1},
+		{11, 10, 0, 11, 0, 3, 10, 5, 0, 8, 0, 7, 5, 7, 0, -1},
+		{11, 10, 5, 7, 11, 5, -1, -1, -1, -1, -1, -1, -1, -1, -1, -1},
+		{10, 6, 5, -1, -1, -1, -1, -1, -1, -1, -1, -1, -1, -1, -1, -1},
+		{0, 8, 3, 5, 10, 6, -1, -1, -1, -1, -1, -1, -1, -1, -1, -1},
+		{9, 0, 1, 5, 10, 6, -1, -1, -1, -1, -1, -1, -1, -1, -1, -1},
+		{1, 8, 3, 1, 9, 8, 5, 10, 6, -1, -1, -1, -1, -1, -1, -1},
+		{1, 6, 5, 2, 6, 1, -1, -1, -1, -1, -1, -1, -1, -1, -1, -1},
+		{1, 6, 5, 1, 2, 6, 3, 0, 8, -1, -1, -1, -1, -1, -1, -1},
+		{9, 6, 5, 9, 0, 6, 0, 2, 6, -1, -1, -1, -1, -1, -1, -1},
+		{5, 9, 8, 5, 8, 2, 5, 2, 6, 3, 2, 8, -1, -1, -1, -1},
+		{2, 3, 11, 10, 6, 5, -1, -1, -1, -1, -1, -1, -1, -1, -1, -1},
+		{11, 0, 8, 11, 2, 0, 10, 6, 5, -1, -1, -1, -1, -1, -1, -1},
+		{0, 1, 9, 2, 3, 11, 5, 10, 6, -1, -1, -1, -1, -1, -1, -1},
+		{5, 10, 6, 1, 9, 2, 9, 11, 2, 9, 8, 11, -1, -1, -1, -1},
+		{6, 3, 11, 6, 5, 3, 5, 1, 3, -1, -1, -1, -1, -1, -1, -1},
+		{0, 8, 11, 0, 11, 5, 0, 5, 1, 5, 11, 6, -1, -1, -1, -1},
+		{3, 11, 6, 0, 3, 6, 0, 6, 5, 0, 5, 9, -1, -1, -1, -1},
+		{6, 5, 9, 6, 9, 11, 11, 9, 8, -1, -1, -1, -1, -1, -1, -1},
+		{5, 10, 6, 4, 7, 8, -1, -1, -1, -1, -1, -1, -1, -1, -1, -1},
+		{4, 3, 0, 4, 7, 3, 6, 5, 10, -1, -1, -1, -1, -1, -1, -1},
+		{1, 9, 0, 5, 10, 6, 8, 4, 7, -1, -1, -1, -1, -1, -1, -1},
+		{10, 6, 5, 1, 9, 7, 1, 7, 3, 7, 9, 4, -1, -1, -1, -1},
+		{6, 1, 2, 6, 5, 1, 4, 7, 8, -1, -1, -1, -1, -1, -1, -1},
+		{1, 2, 5, 5, 2, 6, 3, 0, 4, 3, 4, 7, -1, -1, -1, -1},
+		{8, 4, 7, 9, 0, 5, 0, 6, 5, 0, 2, 6, -1, -1, -1, -1},
+		{7, 3, 9, 7, 9, 4, 3, 2, 9, 5, 9, 6, 2, 6, 9, -1},
+		{3, 11, 2, 7, 8, 4, 10, 6, 5, -1, -1, -1, -1, -1, -1, -1},
+		{5, 10, 6, 4, 7, 2, 4, 2, 0, 2, 7, 11, -1, -1, -1, -1},
+		{0, 1, 9, 4, 7, 8, 2, 3, 11, 5, 10, 6, -1, -1, -1, -1},
+		{9, 2, 1, 9, 11, 2, 9, 4, 11, 7, 11, 4, 5, 10, 6, -1},
+		{8, 4, 7, 3, 11, 5, 3, 5, 1, 5, 11, 6, -1, -1, -1, -1},
+		{5, 1, 11, 5, 11, 6, 1, 0, 11, 7, 11, 4, 0, 4, 11, -1},
+		{0, 5, 9, 0, 6, 5, 0, 3, 6, 11, 6, 3, 8, 4, 7, -1},
+		{6, 5, 9, 6, 9, 11, 4, 7, 9, 7, 11, 9, -1, -1, -1, -1},
+		{10, 4, 9, 6, 4, 10, -1, -1, -1, -1, -1, -1, -1, -1, -1, -1},
+		{4, 10, 6, 4, 9, 10, 0, 8, 3, -1, -1, -1, -1, -1, -1, -1},
+		{10, 0, 1, 10, 6, 0, 6, 4, 0, -1, -1, -1, -1, -1, -1, -1},
+		{8, 3, 1, 8, 1, 6, 8, 6, 4, 6, 1, 10, -1, -1, -1, -1},
+		{1, 4, 9, 1, 2, 4, 2, 6, 4, -1, -1, -1, -1, -1, -1, -1},
+		{3, 0, 8, 1, 2, 9, 2, 4, 9, 2, 6, 4, -1, -1, -1, -1},
+		{0, 2, 4, 4, 2, 6, -1, -1, -1, -1, -1, -1, -1, -1, -1, -1},
+		{8, 3, 2, 8, 2, 4, 4, 2, 6, -1, -1, -1, -1, -1, -1, -1},
+		{10, 4, 9, 10, 6, 4, 11, 2, 3, -1, -1, -1, -1, -1, -1, -1},
+		{0, 8, 2, 2, 8, 11, 4, 9, 10, 4, 10, 6, -1, -1, -1, -1},
+		{3, 11, 2, 0, 1, 6, 0, 6, 4, 6, 1, 10, -1, -1, -1, -1},
+		{6, 4, 1, 6, 1, 10, 4, 8, 1, 2, 1, 11, 8, 11, 1, -1},
+		{9, 6, 4, 9, 3, 6, 9, 1, 3, 11, 6, 3, -1, -1, -1, -1},
+		{8, 11, 1, 8, 1, 0, 11, 6, 1, 9, 1, 4, 6, 4, 1, -1},
+		{3, 11, 6, 3, 6, 0, 0, 6, 4, -1, -1, -1, -1, -1, -1, -1},
+		{6, 4, 8, 11, 6, 8, -1, -1, -1, -1, -1, -1, -1, -1, -1, -1},
+		{7, 10, 6, 7, 8, 10, 8, 9, 10, -1, -1, -1, -1, -1, -1, -1},
+		{0, 7, 3, 0, 10, 7, 0, 9, 10, 6, 7, 10, -1, -1, -1, -1},
+		{10, 6, 7, 1, 10, 7, 1, 7, 8, 1, 8, 0, -1, -1, -1, -1},
+		{10, 6, 7, 10, 7, 1, 1, 7, 3, -1, -1, -1, -1, -1, -1, -1},
+		{1, 2, 6, 1, 6, 8, 1, 8, 9, 8, 6, 7, -1, -1, -1, -1},
+		{2, 6, 9, 2, 9, 1, 6, 7, 9, 0, 9, 3, 7, 3, 9, -1},
+		{7, 8, 0, 7, 0, 6, 6, 0, 2, -1, -1, -1, -1, -1, -1, -1},
+		{7, 3, 2, 6, 7, 2, -1, -1, -1, -1, -1, -1, -1, -1, -1, -1},
+		{2, 3, 11, 10, 6, 8, 10, 8, 9, 8, 6, 7, -1, -1, -1, -1},
+		{2, 0, 7, 2, 7, 11, 0, 9, 7, 6, 7, 10, 9, 10, 7, -1},
+		{1, 8, 0, 1, 7, 8, 1, 10, 7, 6, 7, 10, 2, 3, 11, -1},
+		{11, 2, 1, 11, 1, 7, 10, 6, 1, 6, 7, 1, -1, -1, -1, -1},
+		{8, 9, 6, 8, 6, 7, 9, 1, 6, 11, 6, 3, 1, 3, 6, -1},
+		{0, 9, 1, 11, 6, 7, -1, -1, -1, -1, -1, -1, -1, -1, -1, -1},
+		{7, 8, 0, 7, 0, 6, 3, 11, 0, 11, 6, 0, -1, -1, -1, -1},
+		{7, 11, 6, -1, -1, -1, -1, -1, -1, -1, -1, -1, -1, -1, -1, -1},
+		{7, 6, 11, -1, -1, -1, -1, -1, -1, -1, -1, -1, -1, -1, -1, -1},
+		{3, 0, 8, 11, 7, 6, -1, -1, -1, -1, -1, -1, -1, -1, -1, -1},
+		{0, 1, 9, 11, 7, 6, -1, -1, -1, -1, -1, -1, -1, -1, -1, -1},
+		{8, 1, 9, 8, 3, 1, 11, 7, 6, -1, -1, -1, -1, -1, -1, -1},
+		{10, 1, 2, 6, 11, 7, -1, -1, -1, -1, -1, -1, -1, -1, -1, -1},
+		{1, 2, 10, 3, 0, 8, 6, 11, 7, -1, -1, -1, -1, -1, -1, -1},
+		{2, 9, 0, 2, 10, 9, 6, 11, 7, -1, -1, -1, -1, -1, -1, -1},
+		{6, 11, 7, 2, 10, 3, 10, 8, 3, 10, 9, 8, -1, -1, -1, -1},
+		{7, 2, 3, 6, 2, 7, -1, -1, -1, -1, -1, -1, -1, -1, -1, -1},
+		{7, 0, 8, 7, 6, 0, 6, 2, 0, -1, -1, -1, -1, -1, -1, -1},
+		{2, 7, 6, 2, 3, 7, 0, 1, 9, -1, -1, -1, -1, -1, -1, -1},
+		{1, 6, 2, 1, 8, 6, 1, 9, 8, 8, 7, 6, -1, -1, -1, -1},
+		{10, 7, 6, 10, 1, 7, 1, 3, 7, -1, -1, -1, -1, -1, -1, -1},
+		{10, 7, 6, 1, 7, 10, 1, 8, 7, 1, 0, 8, -1, -1, -1, -1},
+		{0, 3, 7, 0, 7, 10, 0, 10, 9, 6, 10, 7, -1, -1, -1, -1},
+		{7, 6, 10, 7, 10, 8, 8, 10, 9, -1, -1, -1, -1, -1, -1, -1},
+		{6, 8, 4, 11, 8, 6, -1, -1, -1, -1, -1, -1, -1, -1, -1, -1},
+		{3, 6, 11, 3, 0, 6, 0, 4, 6, -1, -1, -1, -1, -1, -1, -1},
+		{8, 6, 11, 8, 4, 6, 9, 0, 1, -1, -1, -1, -1, -1, -1, -1},
+		{9, 4, 6, 9, 6, 3, 9, 3, 1, 11, 3, 6, -1, -1, -1, -1},
+		{6, 8, 4, 6, 11, 8, 2, 10, 1, -1, -1, -1, -1, -1, -1, -1},
+		{1, 2, 10, 3, 0, 11, 0, 6, 11, 0, 4, 6, -1, -1, -1, -1},
+		{4, 11, 8, 4, 6, 11, 0, 2, 9, 2, 10, 9, -1, -1, -1, -1},
+		{10, 9, 3, 10, 3, 2, 9, 4, 3, 11, 3, 6, 4, 6, 3, -1},
+		{8, 2, 3, 8, 4, 2, 4, 6, 2, -1, -1, -1, -1, -1, -1, -1},
+		{0, 4, 2, 4, 6, 2, -1, -1, -1, -1, -1, -1, -1, -1, -1, -1},
+		{1, 9, 0, 2, 3, 4, 2, 4, 6, 4, 3, 8, -1, -1, -1, -1},
+		{1, 9, 4, 1, 4, 2, 2, 4, 6, -1, -1, -1, -1, -1, -1, -1},
+		{8, 1, 3, 8, 6, 1, 8, 4, 6, 6, 10, 1, -1, -1, -1, -1},
+		{10, 1, 0, 10, 0, 6, 6, 0, 4, -1, -1, -1, -1, -1, -1, -1},
+		{4, 6, 3, 4, 3, 8, 6, 10, 3, 0, 3, 9, 10, 9, 3, -1},
+		{10, 9, 4, 6, 10, 4, -1, -1, -1, -1, -1, -1, -1, -1, -1, -1},
+		{4, 9, 5, 7, 6, 11, -1, -1, -1, -1, -1, -1, -1, -1, -1, -1},
+		{0, 8, 3, 4, 9, 5, 11, 7, 6, -1, -1, -1, -1, -1, -1, -1},
+		{5, 0, 1, 5, 4, 0, 7, 6, 11, -1, -1, -1, -1, -1, -1, -1},
+		{11, 7, 6, 8, 3, 4, 3, 5, 4, 3, 1, 5, -1, -1, -1, -1},
+		{9, 5, 4, 10, 1, 2, 7, 6, 11, -1, -1, -1, -1, -1, -1, -1},
+		{6, 11, 7, 1, 2, 10, 0, 8, 3, 4, 9, 5, -1, -1, -1, -1},
+		{7, 6, 11, 5, 4, 10, 4, 2, 10, 4, 0, 2, -1, -1, -1, -1},
+		{3, 4, 8, 3, 5, 4, 3, 2, 5, 10, 5, 2, 11, 7, 6, -1},
+		{7, 2, 3, 7, 6, 2, 5, 4, 9, -1, -1, -1, -1, -1, -1, -1},
+		{9, 5, 4, 0, 8, 6, 0, 6, 2, 6, 8, 7, -1, -1, -1, -1},
+		{3, 6, 2, 3, 7, 6, 1, 5, 0, 5, 4, 0, -1, -1, -1, -1},
+		{6, 2, 8, 6, 8, 7, 2, 1, 8, 4, 8, 5, 1, 5, 8, -1},
+		{9, 5, 4, 10, 1, 6, 1, 7, 6, 1, 3, 7, -1, -1, -1, -1},
+		{1, 6, 10, 1, 7, 6, 1, 0, 7, 8, 7, 0, 9, 5, 4, -1},
+		{4, 0, 10, 4, 10, 5, 0, 3, 10, 6, 10, 7, 3, 7, 10, -1},
+		{7, 6, 10, 7, 10, 8, 5, 4, 10, 4, 8, 10, -1, -1, -1, -1},
+		{6, 9, 5, 6, 11, 9, 11, 8, 9, -1, -1, -1, -1, -1, -1, -1},
+		{3, 6, 11, 0, 6, 3, 0, 5, 6, 0, 9, 5, -1, -1, -1, -1},
+		{0, 11, 8, 0, 5, 11, 0, 1, 5, 5, 6, 11, -1, -1, -1, -1},
+		{6, 11, 3, 6, 3, 5, 5, 3, 1, -1, -1, -1, -1, -1, -1, -1},
+		{1, 2, 10, 9, 5, 11, 9, 11, 8, 11, 5, 6, -1, -1, -1, -1},
+		{0, 11, 3, 0, 6, 11, 0, 9, 6, 5, 6, 9, 1, 2, 10, -1},
+		{11, 8, 5, 11, 5, 6, 8, 0, 5, 10, 5, 2, 0, 2, 5, -1},
+		{6, 11, 3, 6, 3, 5, 2, 10, 3, 10, 5, 3, -1, -1, -1, -1},
+		{5, 8, 9, 5, 2, 8, 5, 6, 2, 3, 8, 2, -1, -1, -1, -1},
+		{9, 5, 6, 9, 6, 0, 0, 6, 2, -1, -1, -1, -1, -1, -1, -1},
+		{1, 5, 8, 1, 8, 0, 5, 6, 8, 3, 8, 2, 6, 2, 8, -1},
+		{1, 5, 6, 2, 1, 6, -1, -1, -1, -1, -1, -1, -1, -1, -1, -1},
+		{1, 3, 6, 1, 6, 10, 3, 8, 6, 5, 6, 9, 8, 9, 6, -1},
+		{10, 1, 0, 10, 0, 6, 9, 5, 0, 5, 6, 0, -1, -1, -1, -1},
+		{0, 3, 8, 5, 6, 10, -1, -1, -1, -1, -1, -1, -1, -1, -1, -1},
+		{10, 5, 6, -1, -1, -1, -1, -1, -1, -1, -1, -1, -1, -1, -1, -1},
+		{11, 5, 10, 7, 5, 11, -1, -1, -1, -1, -1, -1, -1, -1, -1, -1},
+		{11, 5, 10, 11, 7, 5, 8, 3, 0, -1, -1, -1, -1, -1, -1, -1},
+		{5, 11, 7, 5, 10, 11, 1, 9, 0, -1, -1, -1, -1, -1, -1, -1},
+		{10, 7, 5, 10, 11, 7, 9, 8, 1, 8, 3, 1, -1, -1, -1, -1},
+		{11, 1, 2, 11, 7, 1, 7, 5, 1, -1, -1, -1, -1, -1, -1, -1},
+		{0, 8, 3, 1, 2, 7, 1, 7, 5, 7, 2, 11, -1, -1, -1, -1},
+		{9, 7, 5, 9, 2, 7, 9, 0, 2, 2, 11, 7, -1, -1, -1, -1},
+		{7, 5, 2, 7, 2, 11, 5, 9, 2, 3, 2, 8, 9, 8, 2, -1},
+		{2, 5, 10, 2, 3, 5, 3, 7, 5, -1, -1, -1, -1, -1, -1, -1},
+		{8, 2, 0, 8, 5, 2, 8, 7, 5, 10, 2, 5, -1, -1, -1, -1},
+		{9, 0, 1, 5, 10, 3, 5, 3, 7, 3, 10, 2, -1, -1, -1, -1},
+		{9, 8, 2, 9, 2, 1, 8, 7, 2, 10, 2, 5, 7, 5, 2, -1},
+		{1, 3, 5, 3, 7, 5, -1, -1, -1, -1, -1, -1, -1, -1, -1, -1},
+		{0, 8, 7, 0, 7, 1, 1, 7, 5, -1, -1, -1, -1, -1, -1, -1},
+		{9, 0, 3, 9, 3, 5, 5, 3, 7, -1, -1, -1, -1, -1, -1, -1},
+		{9, 8, 7, 5, 9, 7, -1, -1, -1, -1, -1, -1, -1, -1, -1, -1},
+		{5, 8, 4, 5, 10, 8, 10, 11, 8, -1, -1, -1, -1, -1, -1, -1},
+		{5, 0, 4, 5, 11, 0, 5, 10, 11, 11, 3, 0, -1, -1, -1, -1},
+		{0, 1, 9, 8, 4, 10, 8, 10, 11, 10, 4, 5, -1, -1, -1, -1},
+		{10, 11, 4, 10, 4, 5, 11, 3, 4, 9, 4, 1, 3, 1, 4, -1},
+		{2, 5, 1, 2, 8, 5, 2, 11, 8, 4, 5, 8, -1, -1, -1, -1},
+		{0, 4, 11, 0, 11, 3, 4, 5, 11, 2, 11, 1, 5, 1, 11, -1},
+		{0, 2, 5, 0, 5, 9, 2, 11, 5, 4, 5, 8, 11, 8, 5, -1},
+		{9, 4, 5, 2, 11, 3, -1, -1, -1, -1, -1, -1, -1, -1, -1, -1},
+		{2, 5, 10, 3, 5, 2, 3, 4, 5, 3, 8, 4, -1, -1, -1, -1},
+		{5, 10, 2, 5, 2, 4, 4, 2, 0, -1, -1, -1, -1, -1, -1, -1},
+		{3, 10, 2, 3, 5, 10, 3, 8, 5, 4, 5, 8, 0, 1, 9, -1},
+		{5, 10, 2, 5, 2, 4, 1, 9, 2, 9, 4, 2, -1, -1, -1, -1},
+		{8, 4, 5, 8, 5, 3, 3, 5, 1, -1, -1, -1, -1, -1, -1, -1},
+		{0, 4, 5, 1, 0, 5, -1, -1, -1, -1, -1, -1, -1, -1, -1, -1},
+		{8, 4, 5, 8, 5, 3, 9, 0, 5, 0, 3, 5, -1, -1, -1, -1},
+		{9, 4, 5, -1, -1, -1, -1, -1, -1, -1, -1, -1, -1, -1, -1, -1},
+		{4, 11, 7, 4, 9, 11, 9, 10, 11, -1, -1, -1, -1, -1, -1, -1},
+		{0, 8, 3, 4, 9, 7, 9, 11, 7, 9, 10, 11, -1, -1, -1, -1},
+		{1, 10, 11, 1, 11, 4, 1, 4, 0, 7, 4, 11, -1, -1, -1, -1},
+		{3, 1, 4, 3, 4, 8, 1, 10, 4, 7, 4, 11, 10, 11, 4, -1},
+		{4, 11, 7, 9, 11, 4, 9, 2, 11, 9, 1, 2, -1, -1, -1, -1},
+		{9, 7, 4, 9, 11, 7, 9, 1, 11, 2, 11, 1, 0, 8, 3, -1},
+		{11, 7, 4, 11, 4, 2, 2, 4, 0, -1, -1, -1, -1, -1, -1, -1},
+		{11, 7, 4, 11, 4, 2, 8, 3, 4, 3, 2, 4, -1, -1, -1, -1},
+		{2, 9, 10, 2, 7, 9, 2, 3, 7, 7, 4, 9, -1, -1, -1, -1},
+		{9, 10, 7, 9, 7, 4, 10, 2, 7, 8, 7, 0, 2, 0, 7, -1},
+		{3, 7, 10, 3, 10, 2, 7, 4, 10, 1, 10, 0, 4, 0, 10, -1},
+		{1, 10, 2, 8, 7, 4, -1, -1, -1, -1, -1, -1, -1, -1, -1, -1},
+		{4, 9, 1, 4, 1, 7, 7, 1, 3, -1, -1, -1, -1, -1, -1, -1},
+		{4, 9, 1, 4, 1, 7, 0, 8, 1, 8, 7, 1, -1, -1, -1, -1},
+		{4, 0, 3, 7, 4, 3, -1, -1, -1, -1, -1, -1, -1, -1, -1, -1},
+		{4, 8, 7, -1, -1, -1, -1, -1, -1, -1, -1, -1, -1, -1, -1, -1},
+		{9, 10, 8, 10, 11, 8, -1, -1, -1, -1, -1, -1, -1, -1, -1, -1},
+		{3, 0, 9, 3, 9, 11, 11, 9, 10, -1, -1, -1, -1, -1, -1, -1},
+		{0, 1, 10, 0, 10, 8, 8, 10, 11, -1, -1, -1, -1, -1, -1, -1},
+		{3, 1, 10, 11, 3, 10, -1, -1, -1, -1, -1, -1, -1, -1, -1, -1},
+		{1, 2, 11, 1, 11, 9, 9, 11, 8, -1, -1, -1, -1, -1, -1, -1},
+		{3, 0, 9, 3, 9, 11, 1, 2, 9, 2, 11, 9, -1, -1, -1, -1},
+		{0, 2, 11, 8, 0, 11, -1, -1, -1, -1, -1, -1, -1, -1, -1, -1},
+		{3, 2, 11, -1, -1, -1, -1, -1, -1, -1, -1, -1, -1, -1, -1, -1},
+		{2, 3, 8, 2, 8, 10, 10, 8, 9, -1, -1, -1, -1, -1, -1, -1},
+		{9, 10, 2, 0, 9, 2, -1, -1, -1, -1, -1, -1, -1, -1, -1, -1},
+		{2, 3, 8, 2, 8, 10, 0, 1, 8, 1, 10, 8, -1, -1, -1, -1},
+		{1, 10, 2, -1, -1, -1, -1, -1, -1, -1, -1, -1, -1, -1, -1, -1},
+		{1, 3, 8, 9, 1, 8, -1, -1, -1, -1, -1, -1, -1, -1, -1, -1},
+		{0, 9, 1, -1, -1, -1, -1, -1, -1, -1, -1, -1, -1, -1, -1, -1},
+		{0, 3, 8, -1, -1, -1, -1, -1, -1, -1, -1, -1, -1, -1, -1, -1},
+		{-1, -1, -1, -1, -1, -1, -1, -1, -1, -1, -1, -1, -1, -1, -1, -1}
+	};
+
+};
+
+
diff --git a/Engine/Plugins/Experimental/GeometryProcessing/Source/GeometricObjects/Public/Generators/MeshShapeGenerator.h b/Engine/Plugins/Experimental/GeometryProcessing/Source/GeometricObjects/Public/Generators/MeshShapeGenerator.h
--- a/Engine/Plugins/Experimental/GeometryProcessing/Source/GeometricObjects/Public/Generators/MeshShapeGenerator.h
+++ b/Engine/Plugins/Experimental/GeometryProcessing/Source/GeometricObjects/Public/Generators/MeshShapeGenerator.h
@@ -9,7 +9,7 @@
  * Base class for triangle mesh generators (eg like to generate sphere, cylinder, etc)
  * Subclasses must implement ::Generate() 
  */
-class GEOMETRICOBJECTS_API FMeshShapeGenerator
+class /*GEOMETRICOBJECTS_API*/ FMeshShapeGenerator
 {
 public:
 	/** Array of vertex positions */
@@ -31,7 +31,7 @@ public:
 	TArray<FIndex3i> TriangleUVs;
 	/** Array of triangle corner Normals, stored as tuples of indices into Normals array. Same length as Triangles array. */
 	TArray<FIndex3i> TriangleNormals;
-	/** ARray of per-triangle integer polygon IDs. Same length as Triangles array. */
+	/** Array of per-triangle integer polygon IDs. Same length as Triangles array. */
 	TArray<int> TrianglePolygonIDs;
 
 
@@ -49,6 +49,10 @@ public:
 	 */
 	virtual FMeshShapeGenerator& Generate() = 0;
 
+	bool HasAttributes() const
+	{
+		return UVs.Num() > 0 || Normals.Num() > 0;
+	}
 
 	/** clear arrays so that Generate() can be run again */
 	void Reset()
@@ -178,7 +182,7 @@ public:
 	}
 
 	
-	inline int AppendTriangle(int Index, int A, int B, int C)
+	inline int AppendTriangle(int A, int B, int C)
 	{
 		Triangles.Add( bReverseOrientation ? FIndex3i(C, B, A) : FIndex3i(A, B, C) );
 		return Triangles.Num() - 1;
diff --git a/Engine/Plugins/Experimental/GeometryProcessing/Source/GeometricObjects/Public/Generators/SphereGenerator.h b/Engine/Plugins/Experimental/GeometryProcessing/Source/GeometricObjects/Public/Generators/SphereGenerator.h
--- a/Engine/Plugins/Experimental/GeometryProcessing/Source/GeometricObjects/Public/Generators/SphereGenerator.h
+++ b/Engine/Plugins/Experimental/GeometryProcessing/Source/GeometricObjects/Public/Generators/SphereGenerator.h
@@ -15,15 +15,14 @@ public:
 	/** Radius */
 	double Radius = 1;
 
-	int NumPhi = 16; // number of vertices along vertical extent from north pole to south pole
-	int NumTheta = 16; // number of vertices around circles
+	int NumPhi = 3; // number of vertices along vertical extent from north pole to south pole
+	int NumTheta = 3; // number of vertices around circles
 
-	/** If true, each quad of sphere gets a separate polygroup */
+	/** If true, each quad of box gets a separate polygroup */
 	bool bPolygroupPerQuad = false;
 
-public:
-
-	inline static FVector3d SphericalToCartesian(double r, double theta, double phi)
+private:
+	static FVector3d SphericalToCartesian(double r, double theta, double phi)
 	{
 		double Sphi = sin(phi);
 		double Cphi = cos(phi);
@@ -33,186 +32,175 @@ public:
 		return FVector3d(r * Ctheta * Sphi, r * Stheta * Sphi, r * Cphi);
 	}
 
-	/** Generate the mesh */
-	virtual FMeshShapeGenerator& Generate() override
+	void GenerateVertices()
 	{
-		// enforce sane values for vertex counts
-		if (NumPhi < 3)
+		auto SetVertex = [this](int32 VtxIdx,
+								FVector3d Pos, FVector3f Normal )
 		{
-			NumPhi = 3;
-		}
-		if (NumTheta < 3)
-		{
-			NumTheta = 3;
-		}
-		int32 NumVertices = (NumPhi - 2) * NumTheta + 2;
-		int32 NumUVs = NumPhi * (NumTheta + 1);
-		int32 NumTris = (NumPhi - 2) * NumTheta * 2;
-		SetBufferSizes(NumVertices, NumTris, NumUVs, NumVertices);
-
-		double Dphi = FMathd::Pi / double(NumPhi - 1);
-		double Dtheta = FMathd::TwoPi / double(NumTheta);
-		float DUVphi = 1.0 / float(NumPhi - 1);
-		float DUVtheta = 1.0 / float(NumTheta);
-
-		TArray<int32> Grid2d; Grid2d.SetNumUninitialized(NumTheta * NumPhi);
-		for (int32 i = 0; i < NumTheta * NumPhi; ++i)
-		{
-			Grid2d[i] = -1;
-		}
-
-
-		auto GridInterface = [&Grid2d, this](int32 itheta, int32 jphi)->int32&
-		{
-			// make periodic in j
-			while (itheta >= NumTheta) { itheta = itheta - NumTheta; }
-
-			int32 offset = itheta + jphi * NumTheta;
-
-			return Grid2d[offset];
+			Vertices[VtxIdx] = Pos;
+			Normals[VtxIdx] = Normal;
+			NormalParentVertex[VtxIdx] = VtxIdx;
 		};
-
-		int32 VtxIdx = 0;
-		for (int32 p = 1; p < NumPhi - 1; ++p) // NB: this skips the poles.
 		{
-			double Phi = p * Dphi;
+			const double Dphi = FMathd::Pi / double(NumPhi - 1);
+			const double Dtheta = FMathd::TwoPi / double(NumTheta);
+
+			int32 p,t;
+			double Phi, Theta;
+			int32 VtxIdx = 0;
 
-			for (int32 t = 0; t < NumTheta; ++t)
+			// Add points between the poles
+			for (p = 1, Phi = Dphi; p < NumPhi - 1; ++p, Phi += Dphi) // NB: this skips the poles.
 			{
-				double Theta = t * Dtheta;
-				FVector3d Pos = SphericalToCartesian(1., Theta, Phi);
-
-				Vertices[VtxIdx] = Pos * Radius;
-				Normals[VtxIdx] = FVector3f(Pos);
-				NormalParentVertex[VtxIdx] = VtxIdx;
-				GridInterface(t, p) = VtxIdx;
-				VtxIdx++;
+				for (t = 0, Theta = 0; t < NumTheta; ++t, ++VtxIdx, Theta += Dtheta)
+				{
+					FVector3d Normal = SphericalToCartesian(1., Theta, Phi);
+					SetVertex(VtxIdx, Normal * Radius, FVector3f(Normal));
+				}
 			}
+			// add a single point at the North Pole
+			SetVertex(VtxIdx++, FVector3d::UnitZ() * Radius, FVector3f::UnitZ());
+			// add a single point at the South Pole
+			SetVertex(VtxIdx++, -FVector3d::UnitZ() * Radius, -FVector3f::UnitZ());
 		}
+	}
 
+	void GenerateUVVertices()
+	{
+		// generate the UV's
+		const float DUVphi = 1.0 / float(NumPhi - 1);
+		const float DUVtheta = -1.0 / float(NumTheta);
 
-		// add a single point at the North Pole
+		int32 UVIdx = 0;
+		int32 p,t;
+		float UVPhi, UVTheta;
+		for ( p = 1, UVPhi = DUVphi; p < NumPhi - 1; ++p, UVPhi += DUVphi)
 		{
-			FVector3d Pos = SphericalToCartesian(1., 0., 0.);
-			Vertices[VtxIdx] = Pos * Radius;
-			Normals[VtxIdx] = FVector3f(Pos);
-			NormalParentVertex[VtxIdx] = VtxIdx;
-
-			for (int32 t = 0; t < NumTheta; ++t)
+			for (t = 0, UVTheta = 1; t < NumTheta; ++t, ++UVIdx, UVTheta += DUVtheta)
 			{
-				GridInterface(t, 0) = VtxIdx;
+				UVs[UVIdx] = FVector2f(UVTheta, UVPhi);
+				UVParentVertex[UVIdx] = (p - 1) * NumTheta + t;
 			}
-
-			VtxIdx++;
+			UVs[UVIdx] = FVector2f(UVTheta, UVPhi);
+			UVParentVertex[UVIdx] = (p - 1) * NumTheta; // Wrap around
+			++UVIdx;
 		}
-		// add a single point at the south pole
+		int32 NorthPoleVtxIdx = (NumPhi - 2) * NumTheta;
+		for (t = 0, UVTheta = 1 + DUVtheta; t < NumTheta; ++t, ++UVIdx, UVTheta += DUVtheta)
 		{
-			FVector3d Pos = SphericalToCartesian(-1., 0., 0.);
-			Vertices[VtxIdx] = Pos * Radius;
-			Normals[VtxIdx] = FVector3f(Pos);
-			NormalParentVertex[VtxIdx] = VtxIdx;
-
-			for (int32 t = 0; t < NumTheta; ++t)
-			{
-				GridInterface(t, NumPhi - 1) = VtxIdx;
-			}
-
-			VtxIdx++;
+			UVs[UVIdx] = FVector2f(UVTheta, 0.0);
+			UVParentVertex[UVIdx] = NorthPoleVtxIdx;
 		}
-
-		int32 UVIdx = 0;
-		auto UVIdxLookup = [this](int32 t, int32 p)
-		{
-			check(t >= 0 && t <= NumTheta);
-			return p * (NumTheta+1) + t;
-		};
-		for (int32 p = 0; p < NumPhi; ++p) // NB: this skips the poles.
+		int32 SouthPoleVtxIdx = NorthPoleVtxIdx + 1;
+		for (t = 0, UVTheta = 1 + DUVtheta; t < NumTheta; ++t, ++UVIdx, UVTheta += DUVtheta)
 		{
-			float Phi = p * DUVphi;
-
-			for (int32 t = 0; t <= NumTheta; ++t)
-			{
-				float Theta = 1 - t * DUVtheta;
-				UVs[UVIdx] = FVector2f(Theta, Phi);
-				UVParentVertex[UVIdx] = GridInterface(t, p);
-				ensure(UVIdxLookup(t, p) == UVIdx);
-
-				UVIdx++;
-			}
+			UVs[UVIdx] = FVector2f(UVTheta, 1.0);
+			UVParentVertex[UVIdx] = SouthPoleVtxIdx;
 		}
+	}
 
+	using CornerIndices =  FVector3i;
+	void OutputTriangle(int TriIdx, int PolyIdx,  CornerIndices Corners, CornerIndices UVCorners)
+	{
+		SetTriangle(TriIdx, Corners.X, Corners.Y, Corners.Z);
+		SetTrianglePolygon(TriIdx, PolyIdx);
+		SetTriangleUVs(TriIdx, UVCorners.X, UVCorners.Y, UVCorners.Z);
+		SetTriangleNormals(TriIdx, Corners.X, Corners.Y, Corners.Z);
+	}
+
+	void OutputEquatorialTriangles()
+	{
 		int32 TriIdx = 0, PolyIdx = 0;
+
+		// Generate equatorial triangles
+		int32 Corners[4] =   { 0, 1,     NumTheta + 1, NumTheta};
+		int32 UVCorners[4] = { 0, 1, NumTheta + 2, NumTheta + 1};
 		for (int32 p = 1; p < NumPhi - 2; ++p)
 		{
-			for (int32 t = 0; t < NumTheta; ++t)
+			for (int32 t = 0; t < NumTheta - 1; ++t)
 			{
-				// (p,t), (p, t+1), (p+1, t), (p+1, t+1) are the 4 corners of a quad.
-				// counter clockwise.
-				int32 Corners[4] = { GridInterface(t, p + 1), GridInterface(t + 1, p + 1), GridInterface(t + 1, p), GridInterface(t, p) };
-				int32 UVCorners[4] = { UVIdxLookup(t, p + 1),   UVIdxLookup(t + 1, p + 1),   UVIdxLookup(t + 1, p),   UVIdxLookup(t, p) };
-
 				// convert each quad into 2 triangles.
-
-				SetTriangle(TriIdx, Corners[0], Corners[2], Corners[1]);
-				SetTrianglePolygon(TriIdx, PolyIdx);
-				SetTriangleUVs(TriIdx, UVCorners[0], UVCorners[2], UVCorners[1]);
-				SetTriangleNormals(TriIdx, Corners[0], Corners[2], Corners[1]);
-				TriIdx++;
-
-				SetTriangle(TriIdx, Corners[2], Corners[0], Corners[3]);
-				SetTrianglePolygon(TriIdx, PolyIdx);
-				SetTriangleUVs(TriIdx, UVCorners[2], UVCorners[0], UVCorners[3]);
-				SetTriangleNormals(TriIdx, Corners[2], Corners[0], Corners[3]);
-				TriIdx++;
+				OutputTriangle(TriIdx++, PolyIdx,
+							   {Corners[0],   Corners[1],   Corners[2]},
+							   {UVCorners[0], UVCorners[1], UVCorners[2]});
+				OutputTriangle(TriIdx++, PolyIdx,
+							   {Corners[2],   Corners[3],   Corners[0]},
+							   {UVCorners[2], UVCorners[3], UVCorners[0]});
+				for (auto& i : Corners) ++i; 
+				for (auto& i : UVCorners) ++i;
 				if (bPolygroupPerQuad)
 				{
 					PolyIdx++;
 				}
 			}
+			OutputTriangle(TriIdx++, PolyIdx,
+						   {Corners[0], Corners[1] - NumTheta, Corners[2] - NumTheta},
+						   {UVCorners[0]         , UVCorners[1],            UVCorners[2] });
+			OutputTriangle(TriIdx++, PolyIdx,
+						   {Corners[2] - NumTheta, Corners[3],   Corners[0]},
+						   {UVCorners[2],          UVCorners[3],            UVCorners[0]});
+			for (auto& i : Corners) ++i; 
+			for (auto& i : UVCorners) i += 2;
+			if (bPolygroupPerQuad)
+			{
+				PolyIdx++;
+			}
 		}
+	}
 
-		// Do Triangles that connect to north pole
+	void OutputPolarTriangles()
+	{
+		const int32 NumEquatorialVtx = (NumPhi - 2) * NumTheta;
+		const int32 NumEquatorialUVVtx = (NumPhi - 2) * (NumTheta + 1);
+		const int32 NorthPoleVtxIdx = NumEquatorialVtx;
+		const int32 SouthPoleVtxIdx = NumEquatorialVtx + 1;
+		int32 PolyIdx = NumTheta  * (NumPhi - 3);
+		int32 TriIdx = PolyIdx * 2;
+
+		// Triangles that connect to north pole
+		for (int32 t = 0; t < NumTheta; ++t)
 		{
-			int32 p = 0; // North pole 
-			for (int32 t = 0; t < NumTheta; ++t)
+			OutputTriangle(TriIdx++, PolyIdx,
+						   {t, NorthPoleVtxIdx,       (t + 1) % NumTheta},
+						   {t, NumEquatorialUVVtx + t, t + 1});
+			if (bPolygroupPerQuad)
 			{
-				int32 Corners[3] = { GridInterface(t, p + 1), GridInterface(t + 1, p + 1), GridInterface(t + 1, p) };
-				int32 UVCorners[3] = { UVIdxLookup(t, p + 1),   UVIdxLookup(t + 1, p + 1),   UVIdxLookup(t + 1, p) };
-
-				SetTriangle(TriIdx, Corners[0], Corners[2], Corners[1]);
-				SetTrianglePolygon(TriIdx, PolyIdx);
-				SetTriangleUVs(TriIdx, UVCorners[0], UVCorners[2], UVCorners[1]);
-				SetTriangleNormals(TriIdx, Corners[0], Corners[2], Corners[1]);
-				TriIdx++;
-				if (bPolygroupPerQuad)
-				{
-					PolyIdx++;
-				}
+				PolyIdx++;
 			}
 		}
 
-		// Do Triangles that connect to south pole
+		// Triangles that connect to South pole
+		const int32 Offset   = NumEquatorialVtx - NumTheta;
+		const int32 OffsetUV = NumEquatorialUVVtx - (NumTheta + 1);
+		for (int32 t = 0; t < NumTheta; ++t)
 		{
-			int32 p = NumPhi - 2; // South pole 
-			for (int32 t = 0; t < NumTheta; ++t)
+			OutputTriangle(TriIdx++, PolyIdx,
+						   {t + Offset,   ((t + 1) % NumTheta) + Offset, SouthPoleVtxIdx},
+						   {t + OffsetUV, t + 1 + OffsetUV             , NumEquatorialUVVtx + NumTheta + t});
+			if (bPolygroupPerQuad)
 			{
-				int32 Corners[4] = { GridInterface(t, p + 1), GridInterface(t + 1, p), GridInterface(t, p) };
-				int32 UVCorners[3] = { UVIdxLookup(t, p + 1),   UVIdxLookup(t + 1, p),   UVIdxLookup(t, p) };
-
-				SetTriangle(TriIdx, Corners[0], Corners[2], Corners[1]);
-				SetTrianglePolygon(TriIdx, PolyIdx);
-				SetTriangleUVs(TriIdx, UVCorners[0], UVCorners[2], UVCorners[1]);
-				SetTriangleNormals(TriIdx, Corners[0], Corners[2], Corners[1]);
-				TriIdx++;
-				if (bPolygroupPerQuad)
-				{
-					PolyIdx++;
-				}
+				PolyIdx++;
 			}
 		}
+	}
+public:
+	/** Generate the mesh */
+	FMeshShapeGenerator& Generate() override
+	{
+		// enforce sane values for vertex counts
+		NumPhi = FMath::Max(NumPhi, 3);
+		NumTheta = FMath::Max(NumTheta, 3);
+		const int32 NumVertices = (NumPhi - 2) * NumTheta + 2;
+		const int32 NumUVs = (NumPhi - 2) * (NumTheta + 1) + (2 * NumTheta);
+		const int32 NumTris = (NumPhi - 2) * NumTheta * 2;
+		SetBufferSizes(NumVertices, NumTris, NumUVs, NumVertices);
 
+		GenerateVertices();
+		GenerateUVVertices();
+		OutputEquatorialTriangles();
+		OutputPolarTriangles();
 		return *this;
 	}
 
 
-};
\ No newline at end of file
+};
diff --git a/Engine/Plugins/Experimental/GeometryProcessing/Source/GeometricObjects/Public/Generators/SweepGenerator.h b/Engine/Plugins/Experimental/GeometryProcessing/Source/GeometricObjects/Public/Generators/SweepGenerator.h
--- a/Engine/Plugins/Experimental/GeometryProcessing/Source/GeometricObjects/Public/Generators/SweepGenerator.h
+++ b/Engine/Plugins/Experimental/GeometryProcessing/Source/GeometricObjects/Public/Generators/SweepGenerator.h
@@ -440,6 +440,11 @@ public:
 	/** Generate the mesh */
 	virtual FMeshShapeGenerator& Generate() override
 	{
+		if (bCapped)
+		{
+			AngleSamples = FMath::Max(AngleSamples, 3);
+		}
+
 		TArray<float> Radii, Heights;
 
 		Radii.Add(Radius[0]);
@@ -468,7 +473,7 @@ public:
 	float StickRadius = 0.5f;
 	float StickLength = 1.0f;
 	float HeadBaseRadius = 1.0f;
-	float TipRadius = 0.01f;
+	float HeadTipRadius = 0.01f;
 	float HeadLength = 0.5f;
 
 	int AdditionalLengthSamples[3]{ 0,0,0 }; // additional length-wise samples on the three segments (along stick, along arrow base, along arrow cone)
@@ -476,7 +481,7 @@ public:
 	void DistributeAdditionalLengthSamples(int TargetSamples)
 	{
 		TArray<float> AlongPercents;
-		TArray<float> Radii{ StickRadius, StickRadius, HeadBaseRadius, TipRadius };
+		TArray<float> Radii{ StickRadius, StickRadius, HeadBaseRadius, HeadTipRadius };
 		TArray<float> Heights{ 0, StickLength, StickLength, StickLength + HeadLength };
 		float LenAlong = ComputeSegLengths(Radii, Heights, AlongPercents);
 		for (int Idx = 0; Idx < 3; Idx++)
@@ -491,7 +496,7 @@ public:
 	{
 		TArray<float> Radii, Heights;
 		TArray<int> SharpNormalsAlongLength;
-		float SrcRadii[] { StickRadius, StickRadius, HeadBaseRadius, TipRadius };
+		float SrcRadii[] { StickRadius, StickRadius, HeadBaseRadius, HeadTipRadius };
 		float SrcHeights[] { 0, StickLength, StickLength, StickLength + HeadLength };
 
 		int SegIdx = 0;
diff --git a/Engine/Plugins/Experimental/GeometryProcessing/Source/GeometricObjects/Public/Implicit/GridInterpolant.h b/Engine/Plugins/Experimental/GeometryProcessing/Source/GeometricObjects/Public/Implicit/GridInterpolant.h
new file mode 100644
--- /dev/null
+++ b/Engine/Plugins/Experimental/GeometryProcessing/Source/GeometricObjects/Public/Implicit/GridInterpolant.h
@@ -0,0 +1,186 @@
+// Copyright Epic Games, Inc. All Rights Reserved.
+
+// Port of geometry3Sharp's CachingMeshSDFImplicit
+
+#pragma once
+
+#include "MathUtil.h"
+
+
+/**
+ * Tri-linear interpolant for a 3D dense Grid. Supports Grid translation
+ * via GridOrigin, but does not support scaling or rotation. If you need those,
+ * you can wrap this in something that does the xform.
+ *
+ * GridType must have a GetValue() that returns a value to interpolate at a given FVector3i location -- (w/ locations ranging from [0,0,0] to Dimensions (exclusive))
+ */
+template <class GridType>
+class TTriLinearGridInterpolant /*: public BoundedImplicitFunction3d (TODO: consider add ImplicitFunction3d interface concept once more implicit functions are available*/
+{
+public:
+	GridType* Grid;
+	FVector3d GridOrigin;
+	double CellSize;
+	FVector3i Dimensions;
+
+	// value to return if query point is outside Grid (in an SDF
+	// outside is usually positive). Need to do math with this value,
+	// so don't use double.MaxValue or square will overflow
+	double Outside = FMathd::Sqrt(FMathd::Sqrt(FMathd::MaxReal));
+
+	TTriLinearGridInterpolant(GridType* Grid, FVector3d GridOrigin, double CellSize, FVector3i Dimensions) : Grid(Grid), GridOrigin(GridOrigin), CellSize(CellSize), Dimensions(Dimensions)
+	{
+	}
+
+	FAxisAlignedBox3d Bounds() const
+	{
+		return FAxisAlignedBox3d(
+			{ GridOrigin.X, GridOrigin.Y, GridOrigin.Z },
+			{ GridOrigin.X + CellSize * Dimensions.X,
+			  GridOrigin.Y + CellSize * Dimensions.Y,
+			  GridOrigin.Z + CellSize * Dimensions.Z });
+	}
+
+
+	double Value(const FVector3d& Pt) const
+	{
+		FVector3d gridPt(
+			((Pt.X - GridOrigin.X) / CellSize),
+			((Pt.Y - GridOrigin.Y) / CellSize),
+			((Pt.Z - GridOrigin.Z) / CellSize));
+
+		// compute integer coordinates
+		int X0 = (int)gridPt.X;
+		int Y0 = (int)gridPt.Y, Y1 = Y0 + 1;
+		int Z0 = (int)gridPt.Z, Z1 = Z0 + 1;
+
+		// clamp to Grid
+		if (X0 < 0 || (X0 + 1) >= Dimensions.X ||
+			Y0 < 0 || Y1 >= Dimensions.Y ||
+			Z0 < 0 || Z1 >= Dimensions.Z)
+		{
+			return Outside;
+		}
+
+		// convert double coords to [0,1] range
+		double fAx = gridPt.X - (double)X0;
+		double fAy = gridPt.Y - (double)Y0;
+		double fAz = gridPt.Z - (double)Z0;
+		double OneMinusfAx = 1.0 - fAx;
+
+		// compute trilinear interpolant. The code below tries to do this with the fewest 
+		// number of variables, in hopes that optimizer will be clever about re-using registers, etc.
+		// Commented code at bottom is fully-expanded version.
+		// [TODO] it is possible to implement lerps here as A+(B-A)*t, saving a multiply and a variable.
+		//   This is numerically worse, but since the Grid values are floats and
+		//   we are computing in doubles, does it matter?
+		double xa, xb;
+
+		get_value_pair(X0, Y0, Z0, xa, xb);
+		double yz = (1 - fAy) * (1 - fAz);
+		double sum = (OneMinusfAx * xa + fAx * xb) * yz;
+
+		get_value_pair(X0, Y0, Z1, xa, xb);
+		yz = (1 - fAy) * (fAz);
+		sum += (OneMinusfAx * xa + fAx * xb) * yz;
+
+		get_value_pair(X0, Y1, Z0, xa, xb);
+		yz = (fAy) * (1 - fAz);
+		sum += (OneMinusfAx * xa + fAx * xb) * yz;
+
+		get_value_pair(X0, Y1, Z1, xa, xb);
+		yz = (fAy) * (fAz);
+		sum += (OneMinusfAx * xa + fAx * xb) * yz;
+
+		return sum;
+
+		// fV### is Grid cell corner index
+		//return
+		//    fV000 * (1 - fAx) * (1 - fAy) * (1 - fAz) +
+		//    fV001 * (1 - fAx) * (1 - fAy) * (fAz) +
+		//    fV010 * (1 - fAx) * (fAy) * (1 - fAz) +
+		//    fV011 * (1 - fAx) * (fAy) * (fAz) +
+		//    fV100 * (fAx) * (1 - fAy) * (1 - fAz) +
+		//    fV101 * (fAx) * (1 - fAy) * (fAz) +
+		//    fV110 * (fAx) * (fAy) * (1 - fAz) +
+		//    fV111 * (fAx) * (fAy) * (fAz);
+	}
+
+
+protected:
+	void get_value_pair(int I, int J, int K, double& A, double& B) const
+	{
+		A = (double)Grid->GetValue(FVector3i(I, J, K));
+		B = (double)Grid->GetValue(FVector3i(I + 1, J, K));
+	}
+
+
+public:
+	FVector3d Gradient(const FVector3d& Pt) const
+	{
+		FVector3d gridPt = FVector3d(
+			((Pt.X - GridOrigin.X) / CellSize),
+			((Pt.Y - GridOrigin.Y) / CellSize),
+			((Pt.Z - GridOrigin.Z) / CellSize));
+
+		// clamp to Grid
+		if (gridPt.X < 0 || gridPt.X >= Dimensions.X - 1 ||
+			gridPt.Y < 0 || gridPt.Y >= Dimensions.Y - 1 ||
+			gridPt.Z < 0 || gridPt.Z >= Dimensions.Z - 1)
+		{
+			return FVector3d::Zero();
+		}
+
+		// compute integer coordinates
+		int X0 = (int)gridPt.X;
+		int Y0 = (int)gridPt.Y, Y1 = Y0 + 1;
+		int Z0 = (int)gridPt.Z, Z1 = Z0 + 1;
+
+		// convert double coords to [0,1] range
+		double fAx = gridPt.X - (double)X0;
+		double fAy = gridPt.Y - (double)Y0;
+		double fAz = gridPt.Z - (double)Z0;
+
+		double fV000, fV100;
+		get_value_pair(X0, Y0, Z0, fV000, fV100);
+		double fV010, fV110;
+		get_value_pair(X0, Y1, Z0, fV010, fV110);
+		double fV001, fV101;
+		get_value_pair(X0, Y0, Z1, fV001, fV101);
+		double fV011, fV111;
+		get_value_pair(X0, Y1, Z1, fV011, fV111);
+
+		// [TODO] can re-order this to vastly reduce number of ops!
+		double gradX =
+			-fV000 * (1 - fAy) * (1 - fAz) +
+			-fV001 * (1 - fAy) * (fAz)+
+			-fV010 * (fAy) * (1 - fAz) +
+			-fV011 * (fAy) * (fAz)+
+			fV100 * (1 - fAy) * (1 - fAz) +
+			fV101 * (1 - fAy) * (fAz)+
+			fV110 * (fAy) * (1 - fAz) +
+			fV111 * (fAy) * (fAz);
+
+		double gradY =
+			-fV000 * (1 - fAx) * (1 - fAz) +
+			-fV001 * (1 - fAx) * (fAz)+
+			fV010 * (1 - fAx) * (1 - fAz) +
+			fV011 * (1 - fAx) * (fAz)+
+			-fV100 * (fAx) * (1 - fAz) +
+			-fV101 * (fAx) * (fAz)+
+			fV110 * (fAx) * (1 - fAz) +
+			fV111 * (fAx) * (fAz);
+
+		double gradZ =
+			-fV000 * (1 - fAx) * (1 - fAy) +
+			fV001 * (1 - fAx) * (1 - fAy) +
+			-fV010 * (1 - fAx) * (fAy)+
+			fV011 * (1 - fAx) * (fAy)+
+			-fV100 * (fAx) * (1 - fAy) +
+			fV101 * (fAx) * (1 - fAy) +
+			-fV110 * (fAx) * (fAy)+
+			fV111 * (fAx) * (fAy);
+
+		return FVector3d(gradX, gradY, gradZ);
+	}
+};
diff --git a/Engine/Plugins/Experimental/GeometryProcessing/Source/GeometricObjects/Public/Intersection/IntersectionUtil.h b/Engine/Plugins/Experimental/GeometryProcessing/Source/GeometricObjects/Public/Intersection/IntersectionUtil.h
--- a/Engine/Plugins/Experimental/GeometryProcessing/Source/GeometricObjects/Public/Intersection/IntersectionUtil.h
+++ b/Engine/Plugins/Experimental/GeometryProcessing/Source/GeometricObjects/Public/Intersection/IntersectionUtil.h
@@ -97,7 +97,7 @@ namespace IntersectionUtil
 		const FVector3<RealType>& SphereCenter, 
 		RealType SphereRadius)
 	{
-		// [RMS] adapted from GeometricTools GTEngine
+		// adapted from GeometricTools GTEngine
 		// https://www.geometrictools.com/GTEngine/Include/Mathematics/GteIntrLine3Sphere3.h
 
 		// The sphere is (X-C)^T*(X-C)-1 = 0 and the line is X = P+t*D.
@@ -127,7 +127,7 @@ namespace IntersectionUtil
 		RealType SphereRadius,
 		FLinearIntersection& ResultOut)
 	{
-		// [RMS] adapted from GeometricTools GTEngine
+		// adapted from GeometricTools GTEngine
 		// https://www.geometrictools.com/GTEngine/Include/Mathematics/GteIntrLine3Sphere3.h
 
 		// The sphere is (X-C)^T*(X-C)-1 = 0 and the line is X = P+t*D.
@@ -187,7 +187,7 @@ namespace IntersectionUtil
 		const FVector3<RealType>& SphereCenter, 
 		RealType SphereRadius)
 	{
-		// [RMS] adapted from GeometricTools GTEngine
+		// adapted from GeometricTools GTEngine
 		// https://www.geometrictools.com/GTEngine/Include/Mathematics/GteIntrRay3Sphere3.h
 
 		// The sphere is (X-C)^T*(X-C)-1 = 0 and the line is X = P+t*D.
@@ -226,7 +226,7 @@ namespace IntersectionUtil
 		RealType SphereRadius, 
 		FLinearIntersection& Result)
 	{
-		// [RMS] adapted from GeometricTools GTEngine
+		// adapted from GeometricTools GTEngine
 		// https://www.geometrictools.com/GTEngine/Include/Mathematics/GteIntrRay3Sphere3.h
 
 		RaySphereIntersection(RayOrigin, RayDirection, SphereCenter, SphereRadius, Result);
diff --git a/Engine/Plugins/Experimental/GeometryProcessing/Source/GeometricObjects/Public/Intersection/IntrLine2Line2.h b/Engine/Plugins/Experimental/GeometryProcessing/Source/GeometricObjects/Public/Intersection/IntrLine2Line2.h
--- a/Engine/Plugins/Experimental/GeometryProcessing/Source/GeometricObjects/Public/Intersection/IntrLine2Line2.h
+++ b/Engine/Plugins/Experimental/GeometryProcessing/Source/GeometricObjects/Public/Intersection/IntrLine2Line2.h
@@ -104,7 +104,7 @@ public:
 			return (Result == EIntersectionResult::Intersects);
 		}
 
-		// [RMS] if either line direction is not a normalized vector, 
+		// if either line direction is not a normalized vector, 
 		//   results are garbage, so fail query
 		if (Line1.Direction.IsNormalized() == false || Line2.Direction.IsNormalized() == false) 
 		{
diff --git a/Engine/Plugins/Experimental/GeometryProcessing/Source/GeometricObjects/Public/Intersection/IntrSegment2Segment2.h b/Engine/Plugins/Experimental/GeometryProcessing/Source/GeometricObjects/Public/Intersection/IntrSegment2Segment2.h
--- a/Engine/Plugins/Experimental/GeometryProcessing/Source/GeometricObjects/Public/Intersection/IntrSegment2Segment2.h
+++ b/Engine/Plugins/Experimental/GeometryProcessing/Source/GeometricObjects/Public/Intersection/IntrSegment2Segment2.h
@@ -139,7 +139,7 @@ public:
 			return (Result == EIntersectionResult::Intersects);
 		}
 
-		// [RMS] if either segment direction is not a normalized vector, 
+		// if either segment direction is not a normalized vector, 
 		//   results are garbage, so fail query
 		if (Segment1.Direction.IsNormalized() == false || Segment2.Direction.IsNormalized() == false)
 		{
@@ -210,7 +210,7 @@ public:
 		Result = (Type != EIntersectionType::Empty) ?
 			EIntersectionResult::Intersects : EIntersectionResult::NoIntersection;
 
-		// [RMS] for debugging...
+		// for debugging...
 		//SanityCheck();
 
 		return (Result == EIntersectionResult::Intersects);
diff --git a/Engine/Plugins/Experimental/GeometryProcessing/Source/GeometricObjects/Public/MathUtil.h b/Engine/Plugins/Experimental/GeometryProcessing/Source/GeometricObjects/Public/MathUtil.h
--- a/Engine/Plugins/Experimental/GeometryProcessing/Source/GeometricObjects/Public/MathUtil.h
+++ b/Engine/Plugins/Experimental/GeometryProcessing/Source/GeometricObjects/Public/MathUtil.h
@@ -127,6 +127,8 @@ public:
 	static inline RealType Max3(const RealType A, const RealType B, const RealType C);
 	static inline RealType Min(const RealType A, const RealType B);
 	static inline RealType Min3(const RealType A, const RealType B, const RealType C);
+	/** compute min and max of a,b,c with max 3 comparisons (sometimes 2) */
+	static inline void MinMax(RealType A, RealType B, RealType C, RealType& MinOut, RealType& MaxOut);
 	static inline RealType Sqrt(const RealType Value);
 	static inline RealType Atan2(const RealType ValueY, const RealType ValueX);
 	static inline RealType Sin(const RealType Value);
@@ -214,6 +216,28 @@ RealType TMathUtil<RealType>::Min3(const RealType A, const RealType B, const Rea
 	return Min(Min(A, B), C);
 }
 
+// compute min and max of a,b,c with max 3 comparisons (sometimes 2)
+template<typename RealType>
+void TMathUtil<RealType>::MinMax(RealType A, RealType B, RealType C, RealType& MinOut, RealType& MaxOut)
+{
+	if (A < B) {
+		if (A < C) {
+			MinOut = A; MaxOut = TMathUtil<RealType>::Max(B, C);
+		}
+		else {
+			MinOut = C; MaxOut = B;
+		}
+	}
+	else {
+		if (A > C) {
+			MaxOut = A; MinOut = TMathUtil<RealType>::Min(B, C);
+		}
+		else {
+			MinOut = B; MaxOut = C;
+		}
+	}
+}
+
 template<typename RealType>
 RealType TMathUtil<RealType>::Sqrt(const RealType Value)
 {
diff --git a/Engine/Plugins/Experimental/GeometryProcessing/Source/GeometricObjects/Public/MatrixTypes.h b/Engine/Plugins/Experimental/GeometryProcessing/Source/GeometricObjects/Public/MatrixTypes.h
--- a/Engine/Plugins/Experimental/GeometryProcessing/Source/GeometricObjects/Public/MatrixTypes.h
+++ b/Engine/Plugins/Experimental/GeometryProcessing/Source/GeometricObjects/Public/MatrixTypes.h
@@ -123,12 +123,12 @@ struct TMatrix3
 		return TMatrix3<RealType>(M00, M01, M02, M10, M11, M12, M20, M21, M22);
 	}
 
-	TMatrix3<RealType> operator+(const TMatrix3<RealType>& Mat2)
+	TMatrix3<RealType> operator+(const TMatrix3<RealType>& Mat2) const 
 	{
 		return TMatrix3<RealType>(Row0 + Mat2.Row0, Row1 + Mat2.Row1, Row2 + Mat2.Row2, true);
 	}
 
-	TMatrix3<RealType> operator-(const TMatrix3<RealType>& Mat2)
+	TMatrix3<RealType> operator-(const TMatrix3<RealType>& Mat2) const 
 	{
 		return TMatrix3<RealType>(Row0 - Mat2.Row0, Row1 - Mat2.Row1, Row2 - Mat2.Row2, true);
 	}
@@ -154,6 +154,11 @@ struct TMatrix3
 		return Row0.Dot(Mat2.Row0) + Row1.Dot(Mat2.Row1) + Row2.Dot(Mat2.Row2);
 	}
 
+	RealType Trace() const
+	{
+		return Row0.X + Row1.Y + Row2.Z;
+	}
+
 	RealType Determinant() const
 	{
 		RealType a11 = Row0.X, a12 = Row0.Y, a13 = Row0.Z, a21 = Row1.X, a22 = Row1.Y, a23 = Row1.Z, a31 = Row2.X, a32 = Row2.Y, a33 = Row2.Z;
@@ -353,6 +358,11 @@ struct TMatrix2
 		return Row0.Dot(Mat2.Row0) + Row1.Dot(Mat2.Row1);
 	}
 
+	RealType Trace() const
+	{
+		return Row0.X + Row1.Y;
+	}
+
 	RealType Determinant() const
 	{
 		return Row0.X * Row1.Y - Row0.Y * Row1.X;
@@ -418,6 +428,19 @@ inline TMatrix2<RealType> operator*(RealType Scale, const TMatrix2<RealType>& Ma
 		Mat.Row1.X * Scale, Mat.Row1.Y * Scale);
 }
 
+// Skew-Symmetric matrix such that A X B = CrossProductMatrix(A) * B;
+template <typename RealType>
+inline static TMatrix3<RealType> CrossProductMatrix(const FVector3<RealType>& A)
+{
+	RealType Zero(0);
+
+	return TMatrix3<RealType>( Zero, -A[2],  A[1],
+		                       A[2],  Zero, -A[0],
+		                      -A[1],  A[0],  Zero);
+}
+
+
+
 typedef TMatrix3<float> FMatrix3f;
 typedef TMatrix3<double> FMatrix3d;
 typedef TMatrix2<float> FMatrix2f;
diff --git a/Engine/Plugins/Experimental/GeometryProcessing/Source/GeometricObjects/Public/MeshQueries.h b/Engine/Plugins/Experimental/GeometryProcessing/Source/GeometricObjects/Public/MeshQueries.h
--- a/Engine/Plugins/Experimental/GeometryProcessing/Source/GeometricObjects/Public/MeshQueries.h
+++ b/Engine/Plugins/Experimental/GeometryProcessing/Source/GeometricObjects/Public/MeshQueries.h
@@ -4,8 +4,11 @@
 
 #include "Distance/DistPoint3Triangle3.h"
 #include "Intersection/IntrRay3Triangle3.h"
+#include "Intersection/IntrTriangle3Triangle3.h"
 #include "BoxTypes.h"
 #include "IndexTypes.h"
+#include "Algo/Accumulate.h"
+#include "Async/ParallelFor.h"
 
 template <class TriangleMeshType>
 class TMeshQueries
@@ -215,6 +218,38 @@ public:
 		});
 	}
 
+	/**
+	 * brute force search for any intersecting triangles on two meshes
+	 * @return Index pair of IDs of first intersecting triangles found, or InvalidID if no intersection found
+	 */
+	static FIndex2i FindIntersectingTriangles_LinearSearch(const TriangleMeshType& Mesh1, const TriangleMeshType& Mesh2)
+	{
+		for (int TI = 0; TI < Mesh1.MaxTriangleID(); TI++)
+		{
+			if (!Mesh1.IsTriangle(TI))
+			{
+				continue;
+			}
+			FVector3d a, b, c;
+			FTriangle3d Tri1;
+			Mesh1.GetTriVertices(TI, Tri1.V[0], Tri1.V[1], Tri1.V[2]);
+			for (int TJ = 0; TJ < Mesh2.MaxTriangleID(); TJ++)
+			{
+				if (!Mesh2.IsTriangle(TJ))
+				{
+					continue;
+				}
+				FTriangle3d Tri2;
+				Mesh2.GetTriVertices(TJ, Tri2.V[0], Tri2.V[1], Tri2.V[2]);
+				if (FIntrTriangle3Triangle3d::Intersects(Tri1, Tri2))
+				{
+					return FIndex2i(TI, TJ);
+				}
+			}
+		}
+		return FIndex2i::Invalid();
+	}
+
 	/**
 	 * convenience function to construct a IntrRay3Triangle3 object for a Mesh triangle
 	 */
@@ -227,4 +262,232 @@ public:
 		Query.Find();
 		return Query;
 	}
+
+	/// Compute the mean edge length for the given mesh
+	static double AverageEdgeLength(const TriangleMeshType& Mesh)
+	{
+		if (Mesh.EdgeCount() == 0) 
+		{ 
+			return 0.0; 
+		}
+
+		double SumLengths = Algo::TransformAccumulate(Mesh.EdgeIndicesItr(), [&Mesh](int EdgeIndex) -> double
+		{
+			FIndex4i Edge = Mesh.GetEdge(EdgeIndex);
+			FVector3d vA = Mesh.GetVertex(Edge[0]);
+			FVector3d vB = Mesh.GetVertex(Edge[1]);
+			return (vA - vB).Length();
+		}, 0.0);
+
+		return SumLengths / Mesh.EdgeCount();
+	}
+
+	/// Compute the longest edge length for the given mesh
+	static double MaxEdgeLength(const TriangleMeshType& Mesh)
+	{
+		if (Mesh.EdgeCount() == 0)
+		{
+			return 0.0;
+		}
+
+		double MaxLength = -BIG_NUMBER;
+		for (auto EdgeID : Mesh.EdgeIndicesItr())
+		{
+			FIndex4i Edge = Mesh.GetEdge(EdgeID);
+			FVector3d vA = Mesh.GetVertex(Edge[0]);
+			FVector3d vB = Mesh.GetVertex(Edge[1]);
+			MaxLength = FMath::Max(MaxLength, (vA - vB).Length());
+		}
+
+		return MaxLength;
+	}
+
+	/// Compute the shortest edge length for the given mesh
+	static double MinEdgeLength(const TriangleMeshType& Mesh)
+	{
+		if (Mesh.EdgeCount() == 0)
+		{
+			return 0.0;
+		}
+
+		double MinLength = BIG_NUMBER;
+		for (auto EdgeID : Mesh.EdgeIndicesItr())
+		{
+			FIndex4i Edge = Mesh.GetEdge(EdgeID);
+			FVector3d vA = Mesh.GetVertex(Edge[0]);
+			FVector3d vB = Mesh.GetVertex(Edge[1]);
+			MinLength = FMath::Min(MinLength, (vA - vB).Length());
+		}
+
+		return MinLength;
+	}
+
+	/// Given a mesh and a subset of mesh edges, compute the min, max, and mean edge lengths
+	static void EdgeLengthStatsFromEdges(const TriangleMeshType& Mesh, const TArray<int>& Edges, double& MinEdgeLength,
+		double& MaxEdgeLength, double& AverageEdgeLength)
+	{
+		if (Mesh.EdgeCount() == 0)
+		{
+			MinEdgeLength = 0.0;
+			MaxEdgeLength = 0.0;
+			AverageEdgeLength = 0.0;
+			return;
+		}
+
+		MinEdgeLength = BIG_NUMBER;
+		MaxEdgeLength = -BIG_NUMBER;
+		AverageEdgeLength = 0;
+		int EdgeCount = 0;
+
+		for (int EdgeID : Edges)
+		{
+			if (Mesh.IsEdge(EdgeID))
+			{
+				FVector3d A, B;
+				Mesh.GetEdgeV(EdgeID, A, B);
+				double Length = A.Distance(B);
+				if (Length < MinEdgeLength) { MinEdgeLength = Length; }
+				if (Length > MaxEdgeLength) { MaxEdgeLength = Length; }
+				AverageEdgeLength += Length;
+				++EdgeCount;
+			}
+		}
+
+		AverageEdgeLength /= (double)EdgeCount;
+	}
+
+	/// Compute the min, max, and mean edge lengths for the given mesh. Optionally, choose a subest of size NumSamples
+	/// and compute stats of that subset.
+	static void EdgeLengthStats(const TriangleMeshType& Mesh, double& MinEdgeLength, double& MaxEdgeLength,
+		double& AverageEdgeLength, int NumSamples = 0)
+	{
+		if (Mesh.EdgeCount() == 0)
+		{
+			MinEdgeLength = 0.0;
+			MaxEdgeLength = 0.0;
+			AverageEdgeLength = 0.0;
+			return;
+		}
+
+		MinEdgeLength = BIG_NUMBER;
+		MaxEdgeLength = -BIG_NUMBER;
+		AverageEdgeLength = 0;
+		int MaxID = Mesh.MaxEdgeID();
+
+		// if we are only taking some samples, use a prime-modulo-loop instead of random
+		int PrimeNumber = (NumSamples == 0) ? 1 : 31337;
+		int MaxCount = (NumSamples == 0) ? MaxID : NumSamples;
+
+		FVector3d A, B;
+		int EdgeID = 0;
+		int EdgeCount = 0;
+		do
+		{
+			if (Mesh.IsEdge(EdgeID))
+			{
+				Mesh.GetEdgeV(EdgeID, A, B);
+				double Length = A.Distance(B);
+				if (Length < MinEdgeLength) MinEdgeLength = Length;
+				if (Length > MaxEdgeLength) MaxEdgeLength = Length;
+				AverageEdgeLength += Length;
+				++EdgeCount;
+			}
+			EdgeID = (EdgeID + PrimeNumber) % MaxID;
+		} while (EdgeID != 0 && EdgeCount < MaxCount);
+
+		AverageEdgeLength /= (double)EdgeCount;
+	}
+
+	/// For each vertex on MeshA, compute the distance to the nearest point on the surface contained in SpatialB.
+	/// @param MeshA The mesh whose vertex distances should be computed.
+	/// @param SpatialB The target surface's acceleration structure.
+	/// @param Distances For each vertex in MeshA, the distance to the closest point in SpatialB.
+	template<typename MeshSpatialType>
+	static void VertexToSurfaceDistances(const TriangleMeshType& MeshA, const MeshSpatialType& SpatialB, TArray<double>& Distances)
+	{
+		check(SpatialB.SupportsNearestTriangle());
+		Distances.SetNumZeroed(MeshA.VertexCount());
+
+		ParallelFor(MeshA.VertexCount(), [&MeshA, &SpatialB, &Distances](int VertexID)
+		{
+			if (!MeshA.IsVertex(VertexID)) { return; }
+
+			FVector3d VertexPosition = MeshA.GetVertex(VertexID);
+			double DistSqr;
+			SpatialB.FindNearestTriangle(VertexPosition, DistSqr);
+			Distances[VertexID] = sqrt(DistSqr);
+		});
+	}
+
+	/// Compute all vertex-to-surface distances in parallel. Serial raw loop to find max element (using Algo::MaxElement 
+	/// was never faster in initial benchmarking.)
+	/// @param MeshA The mesh whose maximum vertex distance should be computed.
+	/// @param SpatialB The target surface's acceleration structure.
+	/// @return The maximum distance to the surface contained in SpatialB over all vertices in MeshA.
+	template<typename MeshSpatialType>
+	static double HausdorffDistance(const TriangleMeshType& MeshA, const MeshSpatialType& SpatialB)
+	{
+		TArray<double> Distances;
+		VertexToSurfaceDistances(MeshA, SpatialB, Distances);
+
+		double MaxDist = -BIG_NUMBER;
+		for (auto& Dist : Distances)
+		{
+			MaxDist = FMath::Max(Dist, MaxDist);
+		}
+
+		return MaxDist;
+	}
+
+	/// Because Hausdorff distance is not symmetric, we compute the maximum of the distances between two surfaces.
+	template<typename MeshSpatialType>
+	static double TwoSidedHausdorffDistance(const TriangleMeshType& MeshA, const MeshSpatialType& SpatialA,
+											const TriangleMeshType& MeshB, const MeshSpatialType& SpatialB)
+	{
+		return FMath::Max(HausdorffDistance(MeshA, SpatialB), HausdorffDistance(MeshB, SpatialA));
+	}
+
+
+	/// Compute all vertex-to-surface distances in serial. Should only be used for debugging the parallel version above!
+	template<typename MeshSpatialType>
+	static void VertexToSurfaceDistancesSerial(const TriangleMeshType& MeshA, const MeshSpatialType& SpatialB, TArray<double>& Distances)
+	{
+		check(SpatialB.SupportsNearestTriangle());
+		Distances.SetNumZeroed(MeshA.VertexCount());
+
+		for (int VertexID = 0; VertexID < MeshA.VertexCount(); ++VertexID)
+		{
+			if (!MeshA.IsVertex(VertexID)) { continue; }
+
+			FVector3d VertexPosition = MeshA.GetVertex(VertexID);
+			double DistSqr;
+			SpatialB.FindNearestTriangle(VertexPosition, DistSqr);
+			Distances[VertexID] = sqrt(DistSqr);
+		}
+	}
+
+	/// Compute all distances in serial, then a serial raw loop to find max. Should only be used for debugging the 
+	/// parallel version above!
+	template<typename MeshSpatialType>
+	static double HausdorffDistanceSerial(const TriangleMeshType& MeshA, const MeshSpatialType& SpatialB)
+	{
+		TArray<double> Distances;
+		VertexToSurfaceDistancesSerial(MeshA, SpatialB, Distances);
+
+		double MaxDist = -BIG_NUMBER;
+		for (auto& Dist : Distances)
+		{
+			MaxDist = FMath::Max(Dist, MaxDist);
+		}
+
+		return MaxDist;
+	}
+
+	template<typename MeshSpatialType>
+	static double TwoSidedHausdorffDistanceSerial(const TriangleMeshType& MeshA, const MeshSpatialType& SpatialA,
+												  const TriangleMeshType& MeshB, const MeshSpatialType& SpatialB)
+	{
+		return FMath::Max(HausdorffDistanceSerial(MeshA, SpatialB), HausdorffDistanceSerial(MeshB, SpatialA));
+	}
+
 };
diff --git a/Engine/Plugins/Experimental/GeometryProcessing/Source/GeometricObjects/Public/QuadricError.h b/Engine/Plugins/Experimental/GeometryProcessing/Source/GeometricObjects/Public/QuadricError.h
--- a/Engine/Plugins/Experimental/GeometryProcessing/Source/GeometricObjects/Public/QuadricError.h
+++ b/Engine/Plugins/Experimental/GeometryProcessing/Source/GeometricObjects/Public/QuadricError.h
@@ -143,7 +143,7 @@ struct TQuadricError
 		RealType a33 = Axx * Ayy - Axy * Axy;
 		RealType det = (Axx * a11) + (Axy * a12) + (Axz * a13);
 
-		// [RMS] not sure what we should be using for this threshold...have seen
+		// TODO: not sure what we should be using for this threshold...have seen
 		//  det less than 10^-9 on "normal" meshes.
 		if (FMath::Abs(det) > minThresh)
 		{
diff --git a/Engine/Plugins/Experimental/GeometryProcessing/Source/GeometricObjects/Public/Spatial/DenseGrid3.h b/Engine/Plugins/Experimental/GeometryProcessing/Source/GeometricObjects/Public/Spatial/DenseGrid3.h
new file mode 100644
--- /dev/null
+++ b/Engine/Plugins/Experimental/GeometryProcessing/Source/GeometricObjects/Public/Spatial/DenseGrid3.h
@@ -0,0 +1,220 @@
+// Copyright Epic Games, Inc. All Rights Reserved.
+
+// Port of geometry3Sharp DenseGrid3
+
+#pragma once
+
+#include "CoreMinimal.h"
+#include "BoxTypes.h"
+
+#include "HAL/PlatformAtomics.h"
+
+
+/**
+ * 3D dense grid of floating-point scalar values. 
+ */
+template<typename ElemType>
+class TDenseGrid3
+{
+protected:
+	/** grid of allocated elements */
+	TArray<ElemType> Buffer;
+
+	/** dimensions per axis */
+	FVector3i Dimensions;
+
+public:
+	/**
+	 * Create empty grid
+	 */
+	TDenseGrid3() : Dimensions(0,0,0)
+	{
+	}
+
+	TDenseGrid3(int DimX, int DimY, int DimZ, ElemType InitialValue)
+	{
+		Resize(DimX, DimY, DimZ);
+		Assign(InitialValue);
+	}
+
+	int Size() const
+	{
+		return Dimensions.X * Dimensions.Y * Dimensions.Z;
+	}
+
+	const FVector3i& GetDimensions() const
+	{
+		return Dimensions;
+	}
+
+	void Resize(int DimX, int DimY, int DimZ, bool bAllowShrinking = true)
+	{
+		Buffer.SetNumUninitialized(DimX * DimY * DimZ, bAllowShrinking);
+		Dimensions = FVector3i(DimX, DimY, DimZ);
+	}
+
+	void Assign(ElemType Value)
+	{
+		for (int i = 0; i < Buffer.Num(); ++i)
+		{
+			Buffer[i] = Value;
+		}
+	}
+
+	void SetMin(const FVector3i& IJK, ElemType F)
+	{
+		int Idx = IJK.X + Dimensions.X * (IJK.Y + Dimensions.Y * IJK.Z);
+		if (F < Buffer[Idx])
+		{
+			Buffer[Idx] = F;
+		}
+	}
+	void SetMax(const FVector3i& IJK, ElemType F)
+	{
+		int Idx = IJK.X + Dimensions.X * (IJK.Y + Dimensions.Y * IJK.Z);
+		if (F > Buffer[Idx])
+		{
+			Buffer[Idx] = F;
+		}
+	}
+
+	constexpr ElemType& operator[](int Idx)
+	{
+		return Buffer[Idx];
+	}
+	constexpr const ElemType& operator[](int Idx) const
+	{
+		return Buffer[Idx];
+	}
+
+	constexpr ElemType& operator[](FVector3i Idx)
+	{
+		return Buffer[Idx.X + Dimensions.X * (Idx.Y + Dimensions.Y * Idx.Z)];
+	}
+	constexpr const ElemType& operator[](FVector3i Idx) const
+	{
+		return Buffer[Idx.X + Dimensions.X * (Idx.Y + Dimensions.Y * Idx.Z)];
+	}
+	constexpr ElemType& At(int I, int J, int K)
+	{
+		return Buffer[I + Dimensions.X * (J + Dimensions.Y * K)];
+	}
+	constexpr const ElemType& At(int I, int J, int K) const
+	{
+		return Buffer[I + Dimensions.X * (J + Dimensions.Y * K)];
+	}
+
+
+	void GetXPair(int X0, int Y, int Z, ElemType& AOut, ElemType& BOut) const
+	{
+		int Offset = Dimensions.X * (Y + Dimensions.Y * Z);
+		AOut = Buffer[Offset + X0];
+		BOut = Buffer[Offset + X0 + 1];
+	}
+
+	void Apply(TFunctionRef<ElemType(ElemType)> F)
+	{
+		for (int Idx = 0, Num = Size(); Idx < Num; Idx++)
+		{
+			Buffer[Idx] = F(Buffer[Idx]);
+		}
+	}
+
+	// TODO: implement TDenseGrid2<ElemType> and then implement this
+	//TDenseGrid2<ElemType> get_slice(int slice_i, int dimension)
+	//{
+	//	TDenseGrid2<ElemType> slice;
+	//	if (dimension == 0) {
+	//		slice = TDenseGrid2<ElemType>(Dimensions.Y, Dimensions.Z, 0);
+	//		for (int k = 0; k < Dimensions.Z; ++k)
+	//			for (int j = 0; j < Dimensions.Y; ++j)
+	//				slice[j, k] = Buffer[slice_i + Dimensions.X * (j + Dimensions.Y * k)];
+	//	} else if (dimension == 1) {
+	//		slice = TDenseGrid2<ElemType>(Dimensions.X, Dimensions.Z, 0);
+	//		for (int k = 0; k < Dimensions.Z; ++k)
+	//			for (int i = 0; i < Dimensions.X; ++i)
+	//				slice[i, k] = Buffer[i + Dimensions.X * (slice_i + Dimensions.Y * k)];
+	//	} else {
+	//		slice = TDenseGrid2<ElemType>(Dimensions.X, Dimensions.Y, 0);
+	//		for (int j = 0; j < Dimensions.Y; ++j)
+	//			for (int i = 0; i < Dimensions.X; ++i)
+	//				slice[i, j] = Buffer[i + Dimensions.X * (j + Dimensions.Y * slice_i)];
+	//	}
+	//	return slice;
+	//}
+	//void set_slice(TDenseGrid2<ElemType> slice, int slice_i, int dimension)
+	//{
+	//	if (dimension == 0) {
+	//		for (int k = 0; k < Dimensions.Z; ++k)
+	//			for (int j = 0; j < Dimensions.Y; ++j)
+	//				Buffer[slice_i + Dimensions.X * (j + Dimensions.Y * k)] = slice[j, k];
+	//	} else if (dimension == 1) {
+	//		for (int k = 0; k < Dimensions.Z; ++k)
+	//			for (int i = 0; i < Dimensions.X; ++i)
+	//				Buffer[i + Dimensions.X * (slice_i + Dimensions.Y * k)] = slice[i, k];
+	//	} else {
+	//		for (int j = 0; j < Dimensions.Y; ++j)
+	//			for (int i = 0; i < Dimensions.X; ++i)
+	//				Buffer[i + Dimensions.X * (j + Dimensions.Y * slice_i)] = slice[i, j];
+	//	}
+	//}
+
+
+	FAxisAlignedBox3i Bounds() const
+	{
+		return FAxisAlignedBox3i({0, 0, 0},{Dimensions.X, Dimensions.Y, Dimensions.Z});
+	}
+	FAxisAlignedBox3i BoundsInclusive() const
+	{
+		return FAxisAlignedBox3i({0, 0, 0},{Dimensions.X-1, Dimensions.Y-1, Dimensions.Z-1});
+	}
+
+	FVector3i ToIndex(int Idx) const
+	{
+		int x = Idx % Dimensions.X;
+		int y = (Idx / Dimensions.X) % Dimensions.Y;
+		int z = Idx / (Dimensions.X * Dimensions.Y);
+		return FVector3i(x, y, z);
+	}
+	int ToLinear(int X, int Y, int Z) const
+	{
+		return X + Dimensions.X * (Y + Dimensions.Y * Z);
+	}
+	int ToLinear(const FVector3i& IJK) const
+	{
+		return IJK.X + Dimensions.X * (IJK.Y + Dimensions.Y * IJK.Z);
+	}
+};
+
+
+typedef TDenseGrid3<float> FDenseGrid3f;
+typedef TDenseGrid3<double> FDenseGrid3d;
+typedef TDenseGrid3<int> FDenseGrid3i;
+
+// additional utility functions
+namespace DenseGrid
+{
+
+	inline void AtomicIncrement(FDenseGrid3i& Grid, int i, int j, int k)
+	{
+		FPlatformAtomics::InterlockedIncrement(&Grid.At(i,j,k));
+	}
+
+	inline void AtomicDecrement(FDenseGrid3i& Grid, int i, int j, int k)
+	{
+		FPlatformAtomics::InterlockedDecrement(&Grid.At(i, j, k));
+	}
+
+	inline void AtomicIncDec(FDenseGrid3i& Grid, int i, int j, int k, bool bDecrement = false)
+	{
+		if (bDecrement)
+		{
+			FPlatformAtomics::InterlockedDecrement(&Grid.At(i, j, k));
+		}
+		else
+		{
+			FPlatformAtomics::InterlockedIncrement(&Grid.At(i, j, k));
+		}
+	}
+}
+
diff --git a/Engine/Plugins/Experimental/GeometryProcessing/Source/GeometricObjects/Public/Spatial/FastWinding.h b/Engine/Plugins/Experimental/GeometryProcessing/Source/GeometricObjects/Public/Spatial/FastWinding.h
--- a/Engine/Plugins/Experimental/GeometryProcessing/Source/GeometricObjects/Public/Spatial/FastWinding.h
+++ b/Engine/Plugins/Experimental/GeometryProcessing/Source/GeometricObjects/Public/Spatial/FastWinding.h
@@ -214,6 +214,11 @@ public:
 		}
 	}
 
+	TMeshAABBTree3<TriangleMeshType>* GetTree() const
+	{
+		return Tree;
+	}
+
 	void Build(bool bAlwaysBuildRegardlessOfTimestamp = true)
 	{
 		check(Tree);
