Edigrate Dev-Editor updates to InteractiveToolsFramework, GeometryProcessing, MeshModelingToolset, and ModelingTools Editor Mode Plugins. Update a few files in new MeshPaint Plugin for UInteractiveToo

diff --git a/Engine/Plugins/Experimental/GeometryProcessing/Source/GeometricObjects/Public/Spatial/MeshAABBTree3.h b/Engine/Plugins/Experimental/GeometryProcessing/Source/GeometricObjects/Public/Spatial/MeshAABBTree3.h
--- a/Engine/Plugins/Experimental/GeometryProcessing/Source/GeometricObjects/Public/Spatial/MeshAABBTree3.h
+++ b/Engine/Plugins/Experimental/GeometryProcessing/Source/GeometricObjects/Public/Spatial/MeshAABBTree3.h
@@ -6,6 +6,7 @@
 
 #include "Util/DynamicVector.h"
 #include "Intersection/IntrRay3AxisAlignedBox3.h"
+#include "Intersection/IntrTriangle3Triangle3.h"
 #include "MeshQueries.h"
 #include "Spatial/SpatialInterfaces.h"
 #include "Distance/DistTriangle3Triangle3.h"
@@ -13,6 +14,35 @@
 template <class TriangleMeshType>
 class TFastWindingTree;
 
+namespace MeshIntersection
+{
+	/**
+	 * Intersection query result types for triangle mesh intersections
+	 */
+	struct FPointIntersection
+	{
+		int TriangleID[2];
+		FVector3d Point;
+	};
+	struct FSegmentIntersection
+	{
+		int TriangleID[2];
+		FVector3d Point[2];
+	};
+	struct FPolygonIntersection
+	{
+		int TriangleID[2];
+		FVector3d Point[6]; // Coplanar tri-tri intersection forms a polygon of at most six points
+		int Quantity; // number of points actually used
+	};
+	struct FIntersectionsQueryResult
+	{
+		TArray<FPointIntersection> Points;
+		TArray<FSegmentIntersection> Segments;
+		TArray<FPolygonIntersection> Polygons;
+	};
+}
+
 template <class TriangleMeshType>
 class TMeshAABBTree3 : public IMeshSpatial
 {
@@ -26,11 +56,6 @@ protected:
 public:
 	static constexpr double DOUBLE_MAX = TNumericLimits<double>::Max();
 
-	/**
-	 * If non-null, only triangle IDs that pass this filter (ie filter is true) are considered
-	 */
-	TFunction<bool(int)> TriangleFilterF = nullptr;
-
 	TMeshAABBTree3()
 	{
 		Mesh = nullptr;
@@ -69,13 +94,24 @@ public:
 		return MeshTimestamp == Mesh->GetShapeTimestamp();
 	}
 
+	void SetBuildOptions(int32 MaxBoxTriCount)
+	{
+		TopDownLeafMaxTriCount = MaxBoxTriCount;
+	}
+
 	void Build()
 	{
 		BuildTopDown(false);
 		MeshTimestamp = Mesh->GetShapeTimestamp();
 	}
 
-	virtual bool SupportsNearestTriangle() override
+	void Build(const TArray<int32>& TriangleList)
+	{
+		BuildTopDown(false, TriangleList, TriangleList.Num());
+		MeshTimestamp = Mesh->GetShapeTimestamp();
+	}
+
+	virtual bool SupportsNearestTriangle() const override
 	{
 		return true;
 	}
@@ -84,7 +120,10 @@ public:
 	 * Find the triangle closest to P, and distance to it, within distance MaxDist, or return InvalidID
 	 * Use MeshQueries.TriangleDistance() to get more information
 	 */
-	virtual int FindNearestTriangle(const FVector3d& P, double& NearestDistSqr, double MaxDist = TNumericLimits<double>::Max()) override
+	virtual int FindNearestTriangle(
+		const FVector3d& P, double& NearestDistSqr,
+		const FQueryOptions& Options = FQueryOptions()
+	) const override
 	{
 		check(MeshTimestamp == Mesh->GetShapeTimestamp());
 		check(RootIndex >= 0);
@@ -93,20 +132,37 @@ public:
 			return IndexConstants::InvalidID;
 		}
 
-		NearestDistSqr = (MaxDist < DOUBLE_MAX) ? MaxDist * MaxDist : DOUBLE_MAX;
+		NearestDistSqr = (Options.MaxDistance < DOUBLE_MAX) ? Options.MaxDistance * Options.MaxDistance : DOUBLE_MAX;
 		int tNearID = IndexConstants::InvalidID;
-		find_nearest_tri(RootIndex, P, NearestDistSqr, tNearID);
+		find_nearest_tri(RootIndex, P, NearestDistSqr, tNearID, Options);
 		return tNearID;
 	}
 
+	/**
+	 * Get the overall bounding box of the whole tree
+	 */
+	FAxisAlignedBox3d GetBoundingBox()
+	{
+		if (!ensure(RootIndex >= 0))
+		{
+			return FAxisAlignedBox3d::Empty();
+		}
+		else
+		{
+			return GetBox(RootIndex);
+		}
+	}
+
 	/**
 	 * Convenience function that calls FindNearestTriangle and then finds nearest point
 	 * @return nearest point to Point, or Point itself if a nearest point was not found
 	 */
-	virtual FVector3d FindNearestPoint(const FVector3d& Point, double MaxDist = TNumericLimits<double>::Max())
+	virtual FVector3d FindNearestPoint(
+		const FVector3d& Point, const FQueryOptions& Options = FQueryOptions()
+	) const
 	{
 		double NearestDistSqr;
-		int32 NearTriID = FindNearestTriangle(Point, NearestDistSqr, MaxDist);
+		int32 NearTriID = FindNearestTriangle(Point, NearestDistSqr, Options);
 		if (NearTriID >= 0)
 		{
 			FDistPoint3Triangle3d Query = TMeshQueries<TriangleMeshType>::TriangleDistance(*Mesh, NearTriID, Point);
@@ -116,7 +172,7 @@ public:
 	}
 
 protected:
-	void find_nearest_tri(int IBox, const FVector3d& P, double& NearestDistSqr, int& TID)
+	void find_nearest_tri(int IBox, const FVector3d& P, double& NearestDistSqr, int& TID, const FQueryOptions& Options) const
 	{
 		int idx = BoxToIndex[IBox];
 		if (idx < TrianglesEnd)
@@ -125,7 +181,7 @@ protected:
 			for (int i = 1; i <= num_tris; ++i)
 			{
 				int ti = IndexList[idx + i];
-				if (TriangleFilterF != nullptr && TriangleFilterF(ti) == false)
+				if (Options.TriangleFilterF != nullptr && Options.TriangleFilterF(ti) == false)
 				{
 					continue;
 				}
@@ -146,7 +202,7 @@ protected:
 				double fChild1DistSqr = BoxDistanceSqr(iChild1, P);
 				if (fChild1DistSqr <= NearestDistSqr)
 				{
-					find_nearest_tri(iChild1, P, NearestDistSqr, TID);
+					find_nearest_tri(iChild1, P, NearestDistSqr, TID, Options);
 				}
 			}
 			else
@@ -160,10 +216,10 @@ protected:
 				{
 					if (fChild1DistSqr < NearestDistSqr)
 					{
-						find_nearest_tri(iChild1, P, NearestDistSqr, TID);
+						find_nearest_tri(iChild1, P, NearestDistSqr, TID, Options);
 						if (fChild2DistSqr < NearestDistSqr)
 						{
-							find_nearest_tri(iChild2, P, NearestDistSqr, TID);
+							find_nearest_tri(iChild2, P, NearestDistSqr, TID, Options);
 						}
 					}
 				}
@@ -171,10 +227,10 @@ protected:
 				{
 					if (fChild2DistSqr < NearestDistSqr)
 					{
-						find_nearest_tri(iChild2, P, NearestDistSqr, TID);
+						find_nearest_tri(iChild2, P, NearestDistSqr, TID, Options);
 						if (fChild1DistSqr < NearestDistSqr)
 						{
-							find_nearest_tri(iChild1, P, NearestDistSqr, TID);
+							find_nearest_tri(iChild1, P, NearestDistSqr, TID, Options);
 						}
 					}
 				}
@@ -189,7 +245,8 @@ public:
 	/**
 	 * Find the Vertex closest to P, and distance to it, within distance MaxDist, or return InvalidID
 	 */
-	virtual int FindNearestVertex(const FVector3d& P, double& NearestDistSqr, double MaxDist = TNumericLimits<double>::Max())
+	virtual int FindNearestVertex(const FVector3d& P, double& NearestDistSqr,
+		double MaxDist = TNumericLimits<double>::Max(), const FQueryOptions& Options = FQueryOptions())
 	{
 		check(MeshTimestamp == Mesh->GetShapeTimestamp());
 		check(RootIndex >= 0);
@@ -200,13 +257,14 @@ public:
 
 		NearestDistSqr = (MaxDist < DOUBLE_MAX) ? MaxDist * MaxDist : DOUBLE_MAX;
 		int NearestVertexID = IndexConstants::InvalidID;
-		find_nearest_vertex(RootIndex, P, NearestDistSqr, NearestVertexID);
+		find_nearest_vertex(RootIndex, P, NearestDistSqr, NearestVertexID, Options);
 		return NearestVertexID;
 	}
 
 
 protected:
-	void find_nearest_vertex(int IBox, const FVector3d& P, double& NearestDistSqr, int& NearestVertexID)
+	void find_nearest_vertex(int IBox, const FVector3d& P, double& NearestDistSqr,
+		int& NearestVertexID, const FQueryOptions& Options)
 	{
 		int idx = BoxToIndex[IBox];
 		if (idx < TrianglesEnd)
@@ -215,7 +273,7 @@ protected:
 			for (int i = 1; i <= num_tris; ++i)
 			{
 				int ti = IndexList[idx + i];
-				if (TriangleFilterF != nullptr && TriangleFilterF(ti) == false)
+				if (Options.TriangleFilterF != nullptr && Options.TriangleFilterF(ti) == false)
 				{
 					continue;
 				}
@@ -240,7 +298,7 @@ protected:
 				double fChild1DistSqr = BoxDistanceSqr(iChild1, P);
 				if (fChild1DistSqr <= NearestDistSqr)
 				{
-					find_nearest_vertex(iChild1, P, NearestDistSqr, NearestVertexID);
+					find_nearest_vertex(iChild1, P, NearestDistSqr, NearestVertexID, Options);
 				}
 			}
 			else
@@ -254,10 +312,10 @@ protected:
 				{
 					if (fChild1DistSqr < NearestDistSqr)
 					{
-						find_nearest_vertex(iChild1, P, NearestDistSqr, NearestVertexID);
+						find_nearest_vertex(iChild1, P, NearestDistSqr, NearestVertexID, Options);
 						if (fChild2DistSqr < NearestDistSqr)
 						{
-							find_nearest_vertex(iChild2, P, NearestDistSqr, NearestVertexID);
+							find_nearest_vertex(iChild2, P, NearestDistSqr, NearestVertexID, Options);
 						}
 					}
 				}
@@ -265,10 +323,10 @@ protected:
 				{
 					if (fChild2DistSqr < NearestDistSqr)
 					{
-						find_nearest_vertex(iChild2, P, NearestDistSqr, NearestVertexID);
+						find_nearest_vertex(iChild2, P, NearestDistSqr, NearestVertexID, Options);
 						if (fChild1DistSqr < NearestDistSqr)
 						{
-							find_nearest_vertex(iChild1, P, NearestDistSqr, NearestVertexID);
+							find_nearest_vertex(iChild1, P, NearestDistSqr, NearestVertexID, Options);
 						}
 					}
 				}
@@ -281,23 +339,30 @@ protected:
 
 
 public:
-	virtual bool SupportsTriangleRayIntersection() override
+	virtual bool SupportsTriangleRayIntersection() const override
 	{
 		return true;
 	}
 
-	inline virtual int FindNearestHitTriangle(const FRay3d& Ray, double MaxDistance = TNumericLimits<double>::Max()) override
+	inline virtual int FindNearestHitTriangle(
+		const FRay3d& Ray, const FQueryOptions& Options = FQueryOptions()) const override
 	{
 		double NearestT;
 		int TID;
-		FindNearestHitTriangle(Ray, NearestT, TID, MaxDistance);
+		FindNearestHitTriangle(Ray, NearestT, TID, Options);
 		return TID;
 	}
 
-	virtual bool FindNearestHitTriangle(const FRay3d& Ray, double& NearestT, int& TID, double MaxDist = TNumericLimits<double>::Max()) override
+	virtual bool FindNearestHitTriangle(
+		const FRay3d& Ray, double& NearestT, int& TID, 
+		const FQueryOptions& Options = FQueryOptions()) const override
 	{
 		TID = IndexConstants::InvalidID;
-		NearestT = (MaxDist < TNumericLimits<double>::Max()) ? MaxDist : TNumericLimits<float>::Max();
+
+		// Note: using TNumericLimits<float>::Max() here because we need to use <= to compare Box hit
+		//   to NearestT, and Box hit returns TNumericLimits<double>::Max() on no-hit. So, if we set
+		//   nearestT to TNumericLimits<double>::Max(), then we will test all boxes (!)
+		NearestT = (Options.MaxDistance < TNumericLimits<double>::Max()) ? Options.MaxDistance : TNumericLimits<float>::Max();
 
 		check(MeshTimestamp == Mesh->GetShapeTimestamp());
 		check(RootIndex >= 0);
@@ -307,14 +372,14 @@ public:
 		}
 		// TODO: check( ray_is_normalized)
 
-		// [RMS] note: using float.MaxValue here because we need to use <= to compare Box hit
-		//   to NearestT, and Box hit returns double.MaxValue on no-hit. So, if we set
-		//   nearestT to double.MaxValue, then we will test all boxes (!)
-		FindHitTriangle(RootIndex, Ray, NearestT, TID);
+
+		FindHitTriangle(RootIndex, Ray, NearestT, TID, Options);
 		return TID != IndexConstants::InvalidID;
 	}
 
-	void FindHitTriangle(int IBox, const FRay3d& Ray, double& NearestT, int& TID)
+	void FindHitTriangle(
+		int IBox, const FRay3d& Ray, double& NearestT, int& TID,
+		const FQueryOptions& Options = FQueryOptions()) const
 	{
 		int idx = BoxToIndex[IBox];
 		if (idx < TrianglesEnd)
@@ -324,7 +389,7 @@ public:
 			for (int i = 1; i <= num_tris; ++i)
 			{
 				int ti = IndexList[idx + i];
-				if (TriangleFilterF != nullptr && TriangleFilterF(ti) == false)
+				if (Options.TriangleFilterF != nullptr && Options.TriangleFilterF(ti) == false)
 				{
 					continue;
 				}
@@ -352,7 +417,7 @@ public:
 				double fChild1T = box_ray_intersect_t(iChild1, Ray);
 				if (fChild1T <= NearestT + e)
 				{
-					FindHitTriangle(iChild1, Ray, NearestT, TID);
+					FindHitTriangle(iChild1, Ray, NearestT, TID, Options);
 				}
 			}
 			else
@@ -366,10 +431,10 @@ public:
 				{
 					if (fChild1T <= NearestT + e)
 					{
-						FindHitTriangle(iChild1, Ray, NearestT, TID);
+						FindHitTriangle(iChild1, Ray, NearestT, TID, Options);
 						if (fChild2T <= NearestT + e)
 						{
-							FindHitTriangle(iChild2, Ray, NearestT, TID);
+							FindHitTriangle(iChild2, Ray, NearestT, TID, Options);
 						}
 					}
 				}
@@ -377,10 +442,10 @@ public:
 				{
 					if (fChild2T <= NearestT + e)
 					{
-						FindHitTriangle(iChild2, Ray, NearestT, TID);
+						FindHitTriangle(iChild2, Ray, NearestT, TID, Options);
 						if (fChild1T <= NearestT + e)
 						{
-							FindHitTriangle(iChild1, Ray, NearestT, TID);
+							FindHitTriangle(iChild1, Ray, NearestT, TID, Options);
 						}
 					}
 				}
@@ -389,12 +454,16 @@ public:
 	}
 
 	/**
-	 * Find nearest pair of triangles on this tree with otherTree, within max_dist.
-	 * TransformF transforms vertices of otherTree into our coordinates. can be null.
-	 * returns triangle-id pair (my_tri,other_tri), or Index2i::Invalid if not found within max_dist
+	 * Find nearest pair of triangles on this tree with OtherTree, within Options.MaxDistance.
+	 * TransformF transforms vertices of OtherTree into our coordinates. can be null.
+	 * returns triangle-id pair (my_tri,other_tri), or FIndex2i::Invalid if not found within max_dist
 	 * Use MeshQueries.TrianglesDistance() to get more information
+	 * Note: Only uses MaxDistance from Options; OtherTreeOptions.MaxDistance is not used
 	 */
-	virtual FIndex2i FindNearestTriangles(TMeshAABBTree3& OtherTree, TFunction<FVector3d(const FVector3d&)> TransformF, double& Distance, double MaxDist = FMathd::MaxReal)
+	virtual FIndex2i FindNearestTriangles(
+		TMeshAABBTree3& OtherTree, const TFunction<FVector3d(const FVector3d&)>& TransformF,
+		double& Distance, const FQueryOptions& Options = FQueryOptions(), const FQueryOptions& OtherTreeOptions = FQueryOptions()
+	)
 	{
 		check(MeshTimestamp == Mesh->GetShapeTimestamp());
 		check(RootIndex >= 0);
@@ -404,25 +473,25 @@ public:
 		}
 
 		double NearestSqr = FMathd::MaxReal;
-		if (MaxDist < FMathd::MaxReal)
+		if (Options.MaxDistance < FMathd::MaxReal)
 		{
-			NearestSqr = MaxDist * MaxDist;
+			NearestSqr = Options.MaxDistance * Options.MaxDistance;
 		}
 		FIndex2i NearestPair = FIndex2i::Invalid();
 
-		find_nearest_triangles(RootIndex, OtherTree, TransformF, OtherTree.RootIndex, 0, NearestSqr, NearestPair);
+		find_nearest_triangles(RootIndex, OtherTree, TransformF, OtherTree.RootIndex, 0, NearestSqr, NearestPair, Options, OtherTreeOptions);
 		Distance = (NearestSqr < FMathd::MaxReal) ? FMathd::Sqrt(NearestSqr) : FMathd::MaxReal;
 		return NearestPair;
 	}
 
 
 
-	virtual bool SupportsPointContainment() override
+	virtual bool SupportsPointContainment() const override
 	{
 		return false;
 	}
 
-	virtual bool IsInside(const FVector3d& P) override
+	virtual bool IsInside(const FVector3d& P) const override
 	{
 		return false;
 	}
@@ -433,13 +502,18 @@ public:
 		// return false to terminate this branch
 		// arguments are Box and Depth in tree
 		TFunction<bool(const FAxisAlignedBox3d&, int)> NextBoxF = [](const FAxisAlignedBox3d& Box, int Depth) { return true; };
+
+		TFunction<void(int)> BeginBoxTrianglesF = [](int32 BoxID) {};
+
 		TFunction<void(int)> NextTriangleF = [](int TID) {};
+
+		TFunction<void(int)> EndBoxTrianglesF = [](int32 BoxID) {};
 	};
 
 	/**
 	 * Hierarchically descend through the tree Nodes, calling the TreeTrversal functions at each level
 	 */
-	virtual void DoTraversal(FTreeTraversal& Traversal)
+	virtual void DoTraversal(FTreeTraversal& Traversal, const FQueryOptions& Options = FQueryOptions()) const
 	{
 		check(MeshTimestamp == Mesh->GetShapeTimestamp());
 		check(RootIndex >= 0);
@@ -448,27 +522,33 @@ public:
 			return;
 		}
 
-		TreeTraversalImpl(RootIndex, 0, Traversal);
+		TreeTraversalImpl(RootIndex, 0, Traversal, Options);
 	}
 
 	// Traversal implementation. you can override to customize this if necessary.
-	virtual void TreeTraversalImpl(int IBox, int Depth, FTreeTraversal& Traversal)
+	virtual void TreeTraversalImpl(
+		int IBox, int Depth, FTreeTraversal& Traversal, const FQueryOptions& Options
+	) const
 	{
 		int idx = BoxToIndex[IBox];
 
 		if (idx < TrianglesEnd)
 		{
+			Traversal.BeginBoxTrianglesF(IBox);
+
 			// triangle-list case, array is [N t1 t2 ... tN]
 			int n = IndexList[idx];
 			for (int i = 1; i <= n; ++i)
 			{
 				int ti = IndexList[idx + i];
-				if (TriangleFilterF != nullptr && TriangleFilterF(ti) == false)
+				if (Options.TriangleFilterF != nullptr && Options.TriangleFilterF(ti) == false)
 				{
 					continue;
 				}
 				Traversal.NextTriangleF(ti);
 			}
+
+			Traversal.EndBoxTrianglesF(IBox);
 		}
 		else
 		{
@@ -479,7 +559,7 @@ public:
 				i0 = (-i0) - 1;
 				if (Traversal.NextBoxF(GetBox(i0), Depth + 1))
 				{
-					TreeTraversalImpl(i0, Depth + 1, Traversal);
+					TreeTraversalImpl(i0, Depth + 1, Traversal, Options);
 				}
 			}
 			else
@@ -488,17 +568,154 @@ public:
 				i0 = i0 - 1;
 				if (Traversal.NextBoxF(GetBox(i0), Depth + 1))
 				{
-					TreeTraversalImpl(i0, Depth + 1, Traversal);
+					TreeTraversalImpl(i0, Depth + 1, Traversal, Options);
 				}
 				int i1 = IndexList[idx + 1] - 1;
 				if (Traversal.NextBoxF(GetBox(i1), Depth + 1))
 				{
-					TreeTraversalImpl(i1, Depth + 1, Traversal);
+					TreeTraversalImpl(i1, Depth + 1, Traversal, Options);
 				}
 			}
 		}
 	}
 
+
+	/**
+	 * return true if *any* triangle of TestMesh intersects with our tree.
+	 * If TestMeshBounds is not empty, only test collision if the provided bounding box intersects the root AABB box
+	 * Use TransformF to transform vertices of TestMesh into space of this tree.
+	 */
+	virtual bool TestIntersection(
+		const TriangleMeshType* TestMesh, FAxisAlignedBox3d TestMeshBounds = FAxisAlignedBox3d::Empty(),
+		const TFunction<FVector3d(const FVector3d&)>& TransformF = nullptr,
+		const FQueryOptions& Options = FQueryOptions()
+	) const
+	{
+		check(MeshTimestamp == Mesh->GetShapeTimestamp());
+
+		if (!TestMeshBounds.IsEmpty())
+		{
+			if (TransformF != nullptr)
+			{
+				TestMeshBounds = FAxisAlignedBox3d(TestMeshBounds, TransformF);
+			}
+			if (box_box_intersect(RootIndex, TestMeshBounds) == false)
+			{
+				return false;
+			}
+		}
+
+		FTriangle3d TestTri;
+		for (int TID = 0, N = TestMesh->MaxTriangleID(); TID < N; TID++)
+		{
+			if (TransformF != nullptr)
+			{
+				FIndex3i Tri = TestMesh->GetTriangle(TID);
+				TestTri.V[0] = TransformF(TestMesh->GetVertex(Tri.A));
+				TestTri.V[1] = TransformF(TestMesh->GetVertex(Tri.B));
+				TestTri.V[2] = TransformF(TestMesh->GetVertex(Tri.C));
+			}
+			else
+			{
+				TestMesh->GetTriVertices(TID, TestTri.V[0], TestTri.V[1], TestTri.V[2]);
+			}
+			if (TestIntersection(TestTri, Options))
+			{
+				return true;
+			}
+		}
+		return false;
+	}
+
+
+	/**
+	 * Returns true if there is *any* intersection between our mesh and 'other' mesh.
+	 * TransformF takes vertices of OtherTree into our tree - can be null if in same coord space
+	 */
+	virtual bool TestIntersection(
+		const TMeshAABBTree3& OtherTree,
+		const TFunction<FVector3d(const FVector3d&)>& TransformF = nullptr,
+		const FQueryOptions& Options = FQueryOptions(), const FQueryOptions& OtherTreeOptions = FQueryOptions()
+	) const
+	{
+		check(MeshTimestamp == Mesh->GetShapeTimestamp());
+
+		if (find_any_intersection(RootIndex, OtherTree, TransformF, OtherTree.RootIndex, 0, Options, OtherTreeOptions))
+		{
+			return true;
+		}
+
+		return false;
+	}
+
+	/**
+	 * Returns true if triangle intersects any triangle of our mesh
+	 */
+	virtual bool TestIntersection(
+		const FTriangle3d& Triangle,
+		const FQueryOptions& Options = FQueryOptions()
+	) const
+	{
+		check(MeshTimestamp == Mesh->GetShapeTimestamp());
+
+		FAxisAlignedBox3d triBounds(Triangle.V[0], Triangle.V[1], Triangle.V[2]);
+		int interTri = find_any_intersection(RootIndex, Triangle, triBounds, Options);
+		return (interTri >= 0);
+	}
+
+
+	/**
+	 * Compute all intersections between two meshes.
+	 * TransformF argument transforms vertices of OtherTree to our tree (can be null if in same coord space)
+	 * Returns pairs of intersecting triangles, which could intersect in either point or segment
+	 * Currently *does not* return coplanar intersections.
+	 */
+	virtual MeshIntersection::FIntersectionsQueryResult FindAllIntersections(
+		const TMeshAABBTree3& OtherTree, const TFunction<FVector3d(const FVector3d&)>& TransformF = nullptr,
+		const FQueryOptions& Options = FQueryOptions(), const FQueryOptions& OtherTreeOptions = FQueryOptions()
+	) const
+	{
+		check(MeshTimestamp == Mesh->GetShapeTimestamp());
+
+		MeshIntersection::FIntersectionsQueryResult result;
+
+		FIntrTriangle3Triangle3d intr;
+		find_intersections(RootIndex, OtherTree, TransformF, OtherTree.RootIndex, 0, intr, result, Options, OtherTreeOptions);
+
+		return result;
+	}
+
+	/**
+	 * Compute self intersections on our mesh.
+	 * Returns pairs of intersecting triangles, which could intersect in either point or segment
+	 * Currently *does not* return coplanar intersections.
+	 *
+	 * @param bIgnoreTopoConnected Ignore intersections between triangles that share a vertex (if false a lot of triangles that simply share an edge will be counted as self intersecting)
+	 */
+	virtual MeshIntersection::FIntersectionsQueryResult FindAllSelfIntersections(
+		bool bIgnoreTopoConnected = true,
+		const FQueryOptions& Options = FQueryOptions()
+	) const
+	{
+		check(MeshTimestamp == Mesh->GetShapeTimestamp());
+
+		MeshIntersection::FIntersectionsQueryResult Result;
+
+		FIntrTriangle3Triangle3d Intr;
+		find_self_intersections(Intr, &Result, bIgnoreTopoConnected, Options);
+
+		return Result;
+	}
+
+	virtual bool TestSelfIntersection(bool bIgnoreTopoConnected = true, const FQueryOptions& Options = FQueryOptions()) const
+	{
+		check(MeshTimestamp == Mesh->GetShapeTimestamp());
+
+		FIntrTriangle3Triangle3d Intr;
+		return find_self_intersections(Intr, nullptr, bIgnoreTopoConnected, Options);
+	}
+
+
 protected:
 	//
 	// Internals - data structures, construction, etc
@@ -511,7 +728,7 @@ protected:
 		FVector3d Min = c - e, Max = c + e;
 		return FAxisAlignedBox3d(Min, Max);
 	}
-	FAxisAlignedBox3d GetBox(int iBox, TFunction<FVector3d(const FVector3d&)> TransformF)
+	FAxisAlignedBox3d GetBox(int iBox, const TFunction<FVector3d(const FVector3d&)>& TransformF) const
 	{
 		if (TransformF != nullptr)
 		{
@@ -572,6 +789,15 @@ protected:
 		}
 	}
 
+	bool box_box_intersect(int IBox, const FAxisAlignedBox3d& TestBox) const
+	{
+		// [TODO] could compute this w/o constructing box
+		FAxisAlignedBox3d Box = GetBoxEps(IBox, BoxEps);
+
+		return Box.Intersects(TestBox);
+	}
+
+
 	// storage for Box Nodes.
 	//   - BoxToIndex is a pointer into IndexList
 	//   - BoxCenters and BoxExtents are the Centers/extents of the bounding boxes
@@ -612,6 +838,7 @@ protected:
 		}
 	};
 
+
 	void BuildTopDown(bool bSorted)
 	{
 		// build list of valid Triangles & Centers. We skip any
@@ -623,7 +850,40 @@ protected:
 		Centers.SetNumUninitialized(Mesh->TriangleCount());
 		for (int ti = 0; ti < Mesh->MaxTriangleID(); ti++)
 		{
-			if (!Mesh->IsTriangle(ti))
+			if (Mesh->IsTriangle(ti) == false)
+			{
+				continue;
+			}
+			FVector3d centroid = TMeshQueries<TriangleMeshType>::GetTriCentroid(*Mesh, ti);
+			double d2 = centroid.SquaredLength();
+			bool bInvalid = FMathd::IsNaN(d2) || (FMathd::IsFinite(d2) == false);
+			check(bInvalid == false);
+			if (bInvalid == false)
+			{
+				Triangles[i] = ti;
+				Centers[i] = TMeshQueries<TriangleMeshType>::GetTriCentroid(*Mesh, ti);
+				i++;
+			} // otherwise skip this Tri
+		}
+
+		// todo: is passing TriangleCount() correct here? what if we skipped some elements above?
+		BuildTopDown(Triangles, Centers, Mesh->TriangleCount());
+	}
+
+
+	template<typename TriIndexEnumerable>
+	void BuildTopDown(bool bSorted, TriIndexEnumerable TriangleList, int32 NumTriangles)
+	{
+		// build list of valid Triangles & Centers. We skip any
+		// Triangles that have infinite/garbage vertices...
+		int32 i = 0;
+		TArray<int32> Triangles;
+		Triangles.SetNumUninitialized(NumTriangles);
+		TArray<FVector3d> Centers;
+		Centers.SetNumUninitialized(NumTriangles);
+		for (int32 ti : TriangleList)
+		{
+			if ( Mesh->IsTriangle(ti) == false)
 			{
 				continue;
 			}
@@ -636,15 +896,22 @@ protected:
 				Triangles[i] = ti;
 				Centers[i] = TMeshQueries<TriangleMeshType>::GetTriCentroid(*Mesh, ti);
 				i++;
-			} // otherwise skip this tri
+			} // otherwise skip this Tri
 		}
 
+		// todo: is passing NumTriangles correct here? what if we skipped some elements above?
+		BuildTopDown(Triangles, Centers, NumTriangles);
+	}
+
+
+	void BuildTopDown(TArray<int>& Triangles, TArray<FVector3d>& Centers, int32 NumTriangles)
+	{
 		FBoxesSet Tris;
 		FBoxesSet Nodes;
 		FAxisAlignedBox3d rootBox;
 		int rootnode =
-			//(bSorted) ? split_tri_set_sorted(Triangles, Centers, 0, Mesh->TriangleCount, 0, TopDownLeafMaxTriCount, Tris, Nodes, out rootBox) :
-			SplitTriSetMidpoint(Triangles, Centers, 0, Mesh->TriangleCount(), 0, TopDownLeafMaxTriCount, Tris, Nodes, rootBox);
+			//(bSorted) ? split_tri_set_sorted(Triangles, Centers, 0, NumTriangles, 0, TopDownLeafMaxTriCount, Tris, Nodes, out rootBox) :
+			SplitTriSetMidpoint(Triangles, Centers, 0, NumTriangles, 0, TopDownLeafMaxTriCount, Tris, Nodes, rootBox);
 
 		BoxToIndex = Tris.BoxToIndex;
 		BoxCenters = Tris.BoxCenters;
@@ -655,7 +922,7 @@ protected:
 		int iBoxShift = Tris.IBoxCur;
 
 		// ok now append internal node boxes & index ptrs
-		for (i = 0; i < Nodes.IBoxCur; ++i)
+		for (int32 i = 0; i < Nodes.IBoxCur; ++i)
 		{
 			FVector3d NodeBoxCenter = Nodes.BoxCenters[i];		// cannot pass as argument in case a resize happens
 			BoxCenters.InsertAt(NodeBoxCenter, iBoxShift + i);
@@ -667,7 +934,7 @@ protected:
 		}
 
 		// now append index list
-		for (i = 0; i < Nodes.IIndicesCur; ++i)
+		for (int32 i = 0; i < Nodes.IIndicesCur; ++i)
 		{
 			int child_box = Nodes.IndexList[i];
 			if (child_box < 0)
@@ -685,6 +952,7 @@ protected:
 		RootIndex = rootnode + iBoxShift;
 	}
 
+
 	int SplitTriSetMidpoint(
 		TArray<int>& Triangles,
 		TArray<FVector3d>& Centers,
@@ -732,7 +1000,7 @@ protected:
 			int r = ICount - 1;
 			while (l < r)
 			{
-				// [RMS] is <= right here? if V.axis == midpoint, then this loop
+				// TODO: is <= right here? if V.axis == midpoint, then this loop
 				//   can get stuck unless one of these has an equality test. But
 				//   I did not think enough about if this is the right thing to do...
 				while (Centers[IStart + l][axis] <= midpoint)
@@ -787,28 +1055,32 @@ protected:
 	}
 
 
-	void find_nearest_triangles(int iBox, TMeshAABBTree3& otherTree, TFunction<FVector3d(const FVector3d&)> TransformF, int oBox, int depth, double &nearest_sqr, FIndex2i &nearest_pair)
+	void find_nearest_triangles(
+		int iBox, TMeshAABBTree3& OtherTree, const TFunction<FVector3d(const FVector3d&)>& TransformF,
+		int oBox, int depth, double &nearest_sqr, FIndex2i &nearest_pair,
+		const FQueryOptions& Options, const FQueryOptions& OtherTreeOptions
+	) const
 	{
 		int idx = BoxToIndex[iBox];
-		int odx = otherTree.BoxToIndex[oBox];
+		int odx = OtherTree.BoxToIndex[oBox];
 
-		if (idx < TrianglesEnd && odx < otherTree.TrianglesEnd)
+		if (idx < TrianglesEnd && odx < OtherTree.TrianglesEnd)
 		{
 			// ok we are at triangles for both trees, do triangle-level testing
-			FTriangle3d tri, otri;
-			int num_tris = IndexList[idx], onum_tris = otherTree.IndexList[odx];
+			FTriangle3d Tri, otri;
+			int num_tris = IndexList[idx], onum_tris = OtherTree.IndexList[odx];
 
 			FDistTriangle3Triangle3d dist;
 
 			// outer iteration is "other" tris that need to be transformed (more expensive)
 			for (int j = 1; j <= onum_tris; ++j)
 			{
-				int tj = otherTree.IndexList[odx + j];
-				if (otherTree.TriangleFilterF != nullptr && otherTree.TriangleFilterF(tj) == false)
+				int tj = OtherTree.IndexList[odx + j];
+				if (OtherTreeOptions.TriangleFilterF != nullptr && OtherTreeOptions.TriangleFilterF(tj) == false)
 				{
 					continue;
 				}
-				otherTree.Mesh->GetTriVertices(tj, otri.V[0], otri.V[1], otri.V[2]);
+				OtherTree.Mesh->GetTriVertices(tj, otri.V[0], otri.V[1], otri.V[2]);
 				if (TransformF != nullptr)
 				{
 					otri.V[0] = TransformF(otri.V[0]);
@@ -821,12 +1093,12 @@ protected:
 				for (int i = 1; i <= num_tris; ++i)
 				{
 					int ti = IndexList[idx + i];
-					if (TriangleFilterF != nullptr && TriangleFilterF(ti) == false)
+					if (Options.TriangleFilterF != nullptr && Options.TriangleFilterF(ti) == false)
 					{
 						continue;
 					}
-					Mesh->GetTriVertices(ti, tri.V[0], tri.V[1], tri.V[2]);
-					dist.Triangle[1] = tri;
+					Mesh->GetTriVertices(ti, Tri.V[0], Tri.V[1], Tri.V[2]);
+					dist.Triangle[1] = Tri;
 					double dist_sqr = dist.GetSquared();
 					if (dist_sqr < nearest_sqr)
 					{
@@ -845,7 +1117,7 @@ protected:
 		//   - otherwise, we alternate at each depth. This produces wider
 		//     branching but is significantly faster (~10x) for both hits and misses
 		bool bDescendOther = (idx < TrianglesEnd || depth % 2 == 0);
-		if (bDescendOther && odx < otherTree.TrianglesEnd)
+		if (bDescendOther && odx < OtherTree.TrianglesEnd)
 		{
 			bDescendOther = false;      // can't
 		}
@@ -856,23 +1128,23 @@ protected:
 			// the other side, so we descend "their" children
 			FAxisAlignedBox3d bounds = GetBox(iBox);
 
-			int oChild1 = otherTree.IndexList[odx];
+			int oChild1 = OtherTree.IndexList[odx];
 			if (oChild1 < 0)		// 1 child, descend if nearer than cur min-dist
 			{
 				oChild1 = (-oChild1) - 1;
-				FAxisAlignedBox3d oChild1Box = otherTree.GetBox(oChild1, TransformF);
+				FAxisAlignedBox3d oChild1Box = OtherTree.GetBox(oChild1, TransformF);
 				if (oChild1Box.DistanceSquared(bounds) < nearest_sqr)
 				{
-					find_nearest_triangles(iBox, otherTree, TransformF, oChild1, depth + 1, nearest_sqr, nearest_pair);
+					find_nearest_triangles(iBox, OtherTree, TransformF, oChild1, depth + 1, nearest_sqr, nearest_pair, Options, OtherTreeOptions);
 				}
 			}
 			else                            // 2 children
 			{
 				oChild1 = oChild1 - 1;
-				int oChild2 = otherTree.IndexList[odx + 1] - 1;
+				int oChild2 = OtherTree.IndexList[odx + 1] - 1;
 
-				FAxisAlignedBox3d oChild1Box = otherTree.GetBox(oChild1, TransformF);
-				FAxisAlignedBox3d oChild2Box = otherTree.GetBox(oChild2, TransformF);
+				FAxisAlignedBox3d oChild1Box = OtherTree.GetBox(oChild1, TransformF);
+				FAxisAlignedBox3d oChild2Box = OtherTree.GetBox(oChild2, TransformF);
 
 				// descend closer box first
 				double d1Sqr = oChild1Box.DistanceSquared(bounds);
@@ -881,22 +1153,22 @@ protected:
 				{
 					if (d2Sqr < nearest_sqr)
 					{
-						find_nearest_triangles(iBox, otherTree, TransformF, oChild2, depth + 1, nearest_sqr, nearest_pair);
+						find_nearest_triangles(iBox, OtherTree, TransformF, oChild2, depth + 1, nearest_sqr, nearest_pair, Options, OtherTreeOptions);
 					}
 					if (d1Sqr < nearest_sqr)
 					{
-						find_nearest_triangles(iBox, otherTree, TransformF, oChild1, depth + 1, nearest_sqr, nearest_pair);
+						find_nearest_triangles(iBox, OtherTree, TransformF, oChild1, depth + 1, nearest_sqr, nearest_pair, Options, OtherTreeOptions);
 					}
 				}
 				else
 				{
 					if (d1Sqr < nearest_sqr)
 					{
-						find_nearest_triangles(iBox, otherTree, TransformF, oChild1, depth + 1, nearest_sqr, nearest_pair);
+						find_nearest_triangles(iBox, OtherTree, TransformF, oChild1, depth + 1, nearest_sqr, nearest_pair, Options, OtherTreeOptions);
 					}
 					if (d2Sqr < nearest_sqr)
 					{
-						find_nearest_triangles(iBox, otherTree, TransformF, oChild2, depth + 1, nearest_sqr, nearest_pair);
+						find_nearest_triangles(iBox, OtherTree, TransformF, oChild2, depth + 1, nearest_sqr, nearest_pair, Options, OtherTreeOptions);
 					}
 				}
 			}
@@ -904,7 +1176,7 @@ protected:
 		else
 		{
 			// descend our tree nodes if they intersect w/ current bounds of other tree
-			FAxisAlignedBox3d oBounds = otherTree.GetBox(oBox, TransformF);
+			FAxisAlignedBox3d oBounds = OtherTree.GetBox(oBox, TransformF);
 
 			int iChild1 = IndexList[idx];
 			if (iChild1 < 0)                  // 1 child, descend if nearer than cur min-dist
@@ -912,7 +1184,7 @@ protected:
 				iChild1 = (-iChild1) - 1;
 				if (box_box_distsqr(iChild1, oBounds) < nearest_sqr)
 				{
-					find_nearest_triangles(iChild1, otherTree, TransformF, oBox, depth + 1, nearest_sqr, nearest_pair);
+					find_nearest_triangles(iChild1, OtherTree, TransformF, oBox, depth + 1, nearest_sqr, nearest_pair, Options, OtherTreeOptions);
 				}
 			}
 			else                             // 2 children
@@ -927,29 +1199,29 @@ protected:
 				{
 					if (d2Sqr < nearest_sqr)
 					{
-						find_nearest_triangles(iChild2, otherTree, TransformF, oBox, depth + 1, nearest_sqr, nearest_pair);
+						find_nearest_triangles(iChild2, OtherTree, TransformF, oBox, depth + 1, nearest_sqr, nearest_pair, Options, OtherTreeOptions);
 					}
 					if (d1Sqr < nearest_sqr)
 					{
-						find_nearest_triangles(iChild1, otherTree, TransformF, oBox, depth + 1, nearest_sqr, nearest_pair);
+						find_nearest_triangles(iChild1, OtherTree, TransformF, oBox, depth + 1, nearest_sqr, nearest_pair, Options, OtherTreeOptions);
 					}
 				}
 				else
 				{
 					if (d1Sqr < nearest_sqr)
 					{
-						find_nearest_triangles(iChild1, otherTree, TransformF, oBox, depth + 1, nearest_sqr, nearest_pair);
+						find_nearest_triangles(iChild1, OtherTree, TransformF, oBox, depth + 1, nearest_sqr, nearest_pair, Options, OtherTreeOptions);
 					}
 					if (d2Sqr < nearest_sqr)
 					{
-						find_nearest_triangles(iChild2, otherTree, TransformF, oBox, depth + 1, nearest_sqr, nearest_pair);
+						find_nearest_triangles(iChild2, OtherTree, TransformF, oBox, depth + 1, nearest_sqr, nearest_pair, Options, OtherTreeOptions);
 					}
 				}
 			}
 		}
 	}
 
-	double box_box_distsqr(int iBox, const FAxisAlignedBox3d& testBox)
+	double box_box_distsqr(int iBox, const FAxisAlignedBox3d& testBox) const
 	{
 		// [TODO] could compute this w/o constructing box
 		FAxisAlignedBox3d box = GetBoxEps(iBox, BoxEps);
@@ -957,9 +1229,607 @@ protected:
 	}
 
 
+	int find_any_intersection(
+		int iBox, const FTriangle3d& Triangle, const FAxisAlignedBox3d& triBounds, const FQueryOptions& Options
+	) const
+	{
+		int idx = BoxToIndex[iBox];
+		if (idx < TrianglesEnd)             // triangle-list case, array is [N t1 t2 ... tN]
+		{
+			FTriangle3d box_tri;
+			int num_tris = IndexList[idx];
+			for (int i = 1; i <= num_tris; ++i)
+			{
+				int ti = IndexList[idx + i];
+				if (Options.TriangleFilterF != nullptr && Options.TriangleFilterF(ti) == false)
+				{
+					continue;
+				}
+				Mesh->GetTriVertices(ti, box_tri.V[0], box_tri.V[1], box_tri.V[2]);
+				if (FIntrTriangle3Triangle3d::Intersects(Triangle, box_tri))
+				{
+					return ti;
+				}
+			}
+		}
+		else                                 // internal node, either 1 or 2 child boxes
+		{
+			int iChild1 = IndexList[idx];
+			if (iChild1 < 0)                  // 1 child, descend if nearer than cur min-dist
+			{
+				iChild1 = (-iChild1) - 1;
+				if (box_box_intersect(iChild1, triBounds))
+				{
+					return find_any_intersection(iChild1, Triangle, triBounds, Options);
+				}
+
+			}
+			else                             // 2 children, descend closest first
+			{
+				iChild1 = iChild1 - 1;
+				int iChild2 = IndexList[idx + 1] - 1;
+
+				int interTri = -1;
+				if (box_box_intersect(iChild1, triBounds))
+				{
+					interTri = find_any_intersection(iChild1, Triangle, triBounds, Options);
+				}
+				if (interTri == -1 && box_box_intersect(iChild2, triBounds))
+				{
+					interTri = find_any_intersection(iChild2, Triangle, triBounds, Options);
+				}
+				return interTri;
+			}
+		}
+
+		return -1;
+	}
+
+
+	bool find_any_intersection(
+		int iBox, const TMeshAABBTree3& OtherTree, const TFunction<FVector3d(const FVector3d&)>& TransformF, int oBox, int depth,
+		const FQueryOptions& Options, const FQueryOptions& OtherTreeOptions
+	) const
+	{
+		int idx = BoxToIndex[iBox];
+		int odx = OtherTree.BoxToIndex[oBox];
+
+		if (idx < TrianglesEnd && odx < OtherTree.TrianglesEnd)
+		{
+			// ok we are at triangles for both trees, do triangle-level testing
+			FTriangle3d Tri, otri;
+			int num_tris = IndexList[idx], onum_tris = OtherTree.IndexList[odx];
+
+			// outer iteration is "other" tris that need to be transformed (more expensive)
+			for (int j = 1; j <= onum_tris; ++j)
+			{
+				int tj = OtherTree.IndexList[odx + j];
+				if (OtherTreeOptions.TriangleFilterF != nullptr && OtherTreeOptions.TriangleFilterF(tj) == false)
+				{
+					continue;
+				}
+				OtherTree.Mesh->GetTriVertices(tj, otri.V[0], otri.V[1], otri.V[2]);
+				if (TransformF != nullptr)
+				{
+					otri.V[0] = TransformF(otri.V[0]);
+					otri.V[1] = TransformF(otri.V[1]);
+					otri.V[2] = TransformF(otri.V[2]);
+				}
+
+				// inner iteration over "our" triangles
+				for (int i = 1; i <= num_tris; ++i)
+				{
+					int ti = IndexList[idx + i];
+					if (Options.TriangleFilterF != nullptr && Options.TriangleFilterF(ti) == false)
+					{
+						continue;
+					}
+					Mesh->GetTriVertices(ti, Tri.V[0], Tri.V[1], Tri.V[2]);
+					if (TIntrTriangle3Triangle3<double>::Intersects(otri, Tri))
+					{
+						return true;
+					}
+				}
+			}
+			return false;
+		}
+
+		// we either descend "our" tree or the other tree
+		//   - if we have hit triangles on "our" tree, we have to descend other
+		//   - if we hit triangles on "other", we have to descend ours
+		//   - otherwise, we alternate at each depth. This produces wider
+		//     branching but is significantly faster (~10x) for both hits and misses
+		bool bDescendOther = (idx < TrianglesEnd || depth % 2 == 0);
+		if (bDescendOther && odx < OtherTree.TrianglesEnd)
+		{
+			bDescendOther = false;      // can't
+		}
+
+		if (bDescendOther)
+		{
+			// ok we hit triangles on our side but we need to still reach triangles on
+			// the other side, so we descend "their" children
+
+			// [TODO] could we do efficient box.intersects(transform(box)) test?
+			//   ( Contains() on each xformed point? )
+			FAxisAlignedBox3d bounds = GetBox(iBox);
+
+			int oChild1 = OtherTree.IndexList[odx];
+			if (oChild1 < 0)                  // 1 child, descend if nearer than cur min-dist
+			{
+				oChild1 = (-oChild1) - 1;
+				FAxisAlignedBox3d oChild1Box = OtherTree.GetBox(oChild1, TransformF);
+				if (oChild1Box.Intersects(bounds))
+				{
+					return find_any_intersection(iBox, OtherTree, TransformF, oBox, depth + 1, Options, OtherTreeOptions);
+				}
+
+			}
+			else                             // 2 children
+			{
+				oChild1 = oChild1 - 1;          // [TODO] could descend one w/ larger overlap volume first??
+				int oChild2 = OtherTree.IndexList[odx + 1] - 1;
+
+				bool intersects = false;
+				FAxisAlignedBox3d oChild1Box = OtherTree.GetBox(oChild1, TransformF);
+				if (oChild1Box.Intersects(bounds))
+				{
+					intersects = find_any_intersection(iBox, OtherTree, TransformF, oChild1, depth + 1, Options, OtherTreeOptions);
+				}
+
+				if (intersects == false)
+				{
+					FAxisAlignedBox3d oChild2Box = OtherTree.GetBox(oChild2, TransformF);
+					if (oChild2Box.Intersects(bounds))
+					{
+						intersects = find_any_intersection(iBox, OtherTree, TransformF, oChild2, depth + 1, Options, OtherTreeOptions);
+					}
+				}
+				return intersects;
+			}
+		}
+		else
+		{
+			// descend our tree nodes if they intersect w/ current bounds of other tree
+			FAxisAlignedBox3d oBounds = OtherTree.GetBox(oBox, TransformF);
+
+			int iChild1 = IndexList[idx];
+			if (iChild1 < 0)                  // 1 child, descend if nearer than cur min-dist
+			{
+				iChild1 = (-iChild1) - 1;
+				if (box_box_intersect(iChild1, oBounds))
+				{
+					return find_any_intersection(iChild1, OtherTree, TransformF, oBox, depth + 1, Options, OtherTreeOptions);
+				}
+			}
+			else                             // 2 children
+			{
+				iChild1 = iChild1 - 1;          // [TODO] could descend one w/ larger overlap volume first??
+				int iChild2 = IndexList[idx + 1] - 1;
+
+				bool intersects = false;
+				if (box_box_intersect(iChild1, oBounds))
+				{
+					intersects = find_any_intersection(iChild1, OtherTree, TransformF, oBox, depth + 1, Options, OtherTreeOptions);
+				}
+				if (intersects == false && box_box_intersect(iChild2, oBounds))
+				{
+					intersects = find_any_intersection(iChild2, OtherTree, TransformF, oBox, depth + 1, Options, OtherTreeOptions);
+				}
+				return intersects;
+			}
+
+		}
+		return false;
+	}
+
+	
+	// helper for find_self_intersections that checks intersections between triangles from separate boxes
+	bool find_self_intersections_acrossboxes(
+		int Box1, int Box2, FIntrTriangle3Triangle3d& Intr, MeshIntersection::FIntersectionsQueryResult* Result,
+		bool bIgnoreTopoConnected, int depth, const FQueryOptions& Options
+	) const
+	{
+		bool bFound = false;
+		if (Box1 < 0)
+		{
+			return false;
+		}
+
+		int Box1Idx = BoxToIndex[Box1];
+		int Box2Idx = Box2 > -1 ? BoxToIndex[Box2] : -1; // Box2Idx allowed to be invalid
+
+		// at leaf
+		if (Box1Idx < TrianglesEnd && Box2Idx < TrianglesEnd)
+		{
+			if (Box2 == -1)
+			{
+				return false;
+			}
+
+			for (int IdxA = Box1Idx + 1, Box1End = Box1Idx + 1 + IndexList[Box1Idx]; IdxA < Box1End; IdxA++)
+			{
+				int TID_A = IndexList[IdxA];
+				bFound = find_tri_tri_intersections(TID_A, Box2Idx + 1, Box2Idx + 1 + IndexList[Box2Idx], Intr, Result, bIgnoreTopoConnected, Options) || bFound;
+				if (!Result && bFound)
+				{
+					return true;
+				}
+			}
+			return bFound;
+		}
+
+		// alternate which box we descend, while still making sure Box1 is not a leaf
+		// (the alternating part is just a heuristic that I guess may help performance; it's not necessary)
+		if (Box1Idx < TrianglesEnd || (Box2Idx >= TrianglesEnd && depth % 2 == 1))
+		{
+			Swap(Box1, Box2);
+			Swap(Box1Idx, Box2Idx);
+		}
+		
+		
+		// if Box2 is present, we *only* care about the intersection of Box1's children w/ Box2
+		// if Box2 is not present, we handle all the other cases (descending into the children of box1, and checking if the children intersect)
+		if (Box2 > -1)
+		{
+			FAxisAlignedBox3d Box2Box = GetBox(Box2);
+
+			// Descend through Box1
+			int iChild1 = IndexList[Box1Idx];
+			if (iChild1 < 0)                  // 1 child
+			{
+				iChild1 = (-iChild1) - 1;
+
+				if (Box2 > -1 && box_box_intersect(iChild1, Box2Box))
+				{
+					bFound = find_self_intersections_acrossboxes(iChild1, Box2, Intr, Result, bIgnoreTopoConnected, depth + 1, Options) || bFound;
+					if (!Result && bFound)
+					{
+						return true;
+					}
+				}
+			}
+			else                             // 2 children
+			{
+				iChild1 = iChild1 - 1;
+				int iChild2 = IndexList[Box1Idx + 1] - 1;
+
+				if (Box2 > -1 && box_box_intersect(iChild1, Box2Box))
+				{
+					bFound = find_self_intersections_acrossboxes(iChild1, Box2, Intr, Result, bIgnoreTopoConnected, depth + 1, Options) || bFound;
+					if (!Result && bFound)
+					{
+						return true;
+					}
+				}
+				if (Box2 > -1 && box_box_intersect(iChild2, Box2Box))
+				{
+					bFound = find_self_intersections_acrossboxes(iChild2, Box2, Intr, Result, bIgnoreTopoConnected, depth + 1, Options) || bFound;
+					if (!Result && bFound)
+					{
+						return true;
+					}
+				}
+			}
+		}
+		else // Box2 is not present
+		{
+			// Descend through Box1
+			int iChild1 = IndexList[Box1Idx];
+			if (iChild1 < 0)                  // 1 child
+			{
+				iChild1 = (-iChild1) - 1;
+
+				bFound = find_self_intersections_acrossboxes(iChild1, -1, Intr, Result, bIgnoreTopoConnected, depth + 1, Options) || bFound;
+				if (!Result && bFound)
+				{
+					return true;
+				}
+			}
+			else                             // 2 children
+			{
+				iChild1 = iChild1 - 1;
+				int iChild2 = IndexList[Box1Idx + 1] - 1;
+
+				bFound = find_self_intersections_acrossboxes(iChild1, -1, Intr, Result, bIgnoreTopoConnected, depth + 1, Options) || bFound;
+				if (!Result && bFound)
+				{
+					return true;
+				}
+				bFound = find_self_intersections_acrossboxes(iChild2, -1, Intr, Result, bIgnoreTopoConnected, depth + 1, Options) || bFound;
+				if (!Result && bFound)
+				{
+					return true;
+				}
+
+				FAxisAlignedBox3d Child2Box = GetBox(iChild2);
+
+				if (box_box_intersect(iChild1, Child2Box))
+				{
+					bFound = find_self_intersections_acrossboxes(iChild1, iChild2, Intr, Result, bIgnoreTopoConnected, depth + 1, Options) || bFound;
+					if (!Result && bFound)
+					{
+						return true;
+					}
+				}
+			}
+		}
+
+		return bFound;
+	}
+
+	// helper for find_self_intersections that intersects a single triangle vs a range of triangles
+	bool find_tri_tri_intersections(int TID_A, int IdxRangeStart, int IdxRangeEnd, 
+		FIntrTriangle3Triangle3d& Intr,
+		MeshIntersection::FIntersectionsQueryResult* Result,
+		bool bIgnoreTopoConnected, const FQueryOptions& Options
+	) const
+	{
+		if (Options.TriangleFilterF != nullptr && Options.TriangleFilterF(TID_A) == false)
+		{
+			return false;
+		}
+
+		bool bFound = false;
+		FTriangle3d TriA, TriB;
+		Mesh->GetTriVertices(TID_A, TriA.V[0], TriA.V[1], TriA.V[2]);
+		FIndex3i TriA_VID = Mesh->GetTriangle(TID_A);
+		Intr.SetTriangle0(TriA);
+		for (int IdxB = IdxRangeStart; IdxB < IdxRangeEnd; IdxB++)
+		{
+			int TID_B = IndexList[IdxB];
+			if (Options.TriangleFilterF != nullptr && Options.TriangleFilterF(TID_B) == false)
+			{
+				continue;
+			}
+
+			Mesh->GetTriVertices(TID_B, TriB.V[0], TriB.V[1], TriB.V[2]);
+			FIndex3i TriB_VID = Mesh->GetTriangle(TID_B);
+
+			if (bIgnoreTopoConnected)
+			{
+				// ignore triangles that share a vertex
+				bool bTopoConnected = false;
+				for (int AIdx = 0; AIdx < 3 && !bTopoConnected; AIdx++)
+				{
+					int VID_A = TriA_VID[AIdx];
+					for (int BIdx = 0; BIdx < 3; BIdx++)
+					{
+						if (VID_A == TriB_VID[BIdx])
+						{
+							bTopoConnected = true;
+							break;
+						}
+					}
+				}
+				if (bTopoConnected)
+				{
+					continue;
+				}
+			}
+
+			Intr.SetTriangle1(TriB);
+
+			// Note: Test() is much faster than Find() so it makes sense to call it first, as most
+			// triangles will not intersect (right?)
+			if (Intr.Test())
+			{
+				bFound = true;
+				if (!Result)
+				{
+					return true;
+				}
+				if (Intr.Find())
+				{
+					if (Intr.Quantity == 1)
+					{
+						Result->Points.Add(MeshIntersection::FPointIntersection{ {TID_A, TID_B}, Intr.Points[0] });
+					}
+					else if (Intr.Quantity == 2)
+					{
+						Result->Segments.Add(MeshIntersection::FSegmentIntersection{ {TID_A, TID_B}, {Intr.Points[0], Intr.Points[1]} });
+					}
+					else
+					{
+						// this can't happen unless tri-tri intersection is set to detect coplanar intersection, which it currently is not!
+						// TODO: change this check to test whether coplanar intersection is enabled, once the option to enable it exists
+						checkf(false, TEXT("Coplanar intersection not enabled but found intr.Quanitity > 2?  Quantity=%d"), Intr.Quantity);
+
+						//result.Polygons.Add(MeshIntersection::FPolygonIntersection{ {ti, tj}, intr.Points, intr.Quantity });
+					}
+				}
+			}
+		}
+
+		return bFound;
+	}
+
+
+	bool find_self_intersections(
+		FIntrTriangle3Triangle3d& Intr, MeshIntersection::FIntersectionsQueryResult* Result, bool bIgnoreTopoConnected,
+		const FQueryOptions & Options
+	) const
+	{
+		// Check each leaf-box for intersecting triangles within the same box
+		bool bFound = false;
+		for (int StartIdx = 0; StartIdx < TrianglesEnd; )
+		{
+			int NumTris = IndexList[StartIdx];
+			int EndIdx = StartIdx + NumTris + 1;
+			for (int IdxA = StartIdx + 1; IdxA + 1 < EndIdx; IdxA++)
+			{
+				int TID_A = IndexList[IdxA];
+				bFound = find_tri_tri_intersections(TID_A, IdxA + 1, EndIdx, Intr, Result, bIgnoreTopoConnected, Options) || bFound;
+				if (!Result && bFound) // early out if not filling Result
+				{
+					return bFound;
+				}
+			}
+			StartIdx = EndIdx;
+		}
+
+		// Recursively check across boxes
+		return find_self_intersections_acrossboxes(RootIndex, -1, Intr, Result, bIgnoreTopoConnected, 0, Options) || bFound;
+	}
+
+
+	void find_intersections(
+		int iBox, const TMeshAABBTree3& OtherTree, const TFunction<FVector3d(const FVector3d&)>& TransformF,
+		int oBox, int depth, FIntrTriangle3Triangle3d& intr, MeshIntersection::FIntersectionsQueryResult& result,
+		const FQueryOptions& Options, const FQueryOptions& OtherTreeOptions
+	) const
+	{
+		int idx = BoxToIndex[iBox];
+		int odx = OtherTree.BoxToIndex[oBox];
+
+		if (idx < TrianglesEnd && odx < OtherTree.TrianglesEnd)
+		{
+			// ok we are at triangles for both trees, do triangle-level testing
+			FTriangle3d Tri, otri;
+			int num_tris = IndexList[idx], onum_tris = OtherTree.IndexList[odx];
+
+			// outer iteration is "other" tris that need to be transformed (more expensive)
+			for (int j = 1; j <= onum_tris; ++j)
+			{
+				int tj = OtherTree.IndexList[odx + j];
+				if (OtherTreeOptions.TriangleFilterF != nullptr && OtherTreeOptions.TriangleFilterF(tj) == false)
+				{
+					continue;
+				}
+				OtherTree.Mesh->GetTriVertices(tj, otri.V[0], otri.V[1], otri.V[2]);
+				if (TransformF != nullptr)
+				{
+					otri.V[0] = TransformF(otri.V[0]);
+					otri.V[1] = TransformF(otri.V[1]);
+					otri.V[2] = TransformF(otri.V[2]);
+				}
+				intr.SetTriangle0(otri);
+
+				// inner iteration over "our" triangles
+				for (int i = 1; i <= num_tris; ++i)
+				{
+					int ti = IndexList[idx + i];
+					if (Options.TriangleFilterF != nullptr && Options.TriangleFilterF(ti) == false)
+					{
+						continue;
+					}
+					Mesh->GetTriVertices(ti, Tri.V[0], Tri.V[1], Tri.V[2]);
+					intr.SetTriangle1(Tri);
+
+					// Note: Test() is much faster than Find() so it makes sense to call it first, as most
+					// triangles will not intersect (right?)
+					if (intr.Test())
+					{
+						if (intr.Find())
+						{
+							if (intr.Quantity == 1)
+							{
+								result.Points.Add(MeshIntersection::FPointIntersection{ {ti, tj}, intr.Points[0] });
+							}
+							else if (intr.Quantity == 2)
+							{
+								result.Segments.Add(MeshIntersection::FSegmentIntersection{ {ti, tj}, {intr.Points[0], intr.Points[1]} });
+							}
+							else
+							{
+								// this can't happen unless tri-tri intersection is set to detect coplanar intersection, which it currently is not!
+								// TODO: change this check to test whether coplanar intersection is enabled, once the option to enable it exists
+								checkf(false, TEXT("Coplanar intersection not enabled but found intr.Quanitity > 2?  Quantity=%d"), intr.Quantity);
+
+								//result.Polygons.Add(MeshIntersection::FPolygonIntersection{ {ti, tj}, intr.Points, intr.Quantity });
+							}
+						}
+					}
+				}
+			}
+
+			// done these nodes
+			return;
+		}
+
+		// we either descend "our" tree or the other tree
+		//   - if we have hit triangles on "our" tree, we have to descend other
+		//   - if we hit triangles on "other", we have to descend ours
+		//   - otherwise, we alternate at each depth. This produces wider
+		//     branching but is significantly faster (~10x) for both hits and misses
+		bool bDescendOther = (idx < TrianglesEnd || depth % 2 == 0);
+		if (bDescendOther && odx < OtherTree.TrianglesEnd)
+			bDescendOther = false;      // can't
+
+		if (bDescendOther) {
+			// ok we hit triangles on our side but we need to still reach triangles on
+			// the other side, so we descend "their" children
+
+			// [TODO] could we do efficient box.intersects(transform(box)) test?
+			//   ( Contains() on each xformed point? )
+			FAxisAlignedBox3d bounds = GetBoxEps(iBox);
+
+			int oChild1 = OtherTree.IndexList[odx];
+			if (oChild1 < 0)                  // 1 child, descend if nearer than cur min-dist
+			{
+				oChild1 = (-oChild1) - 1;
+				FAxisAlignedBox3d oChild1Box = OtherTree.GetBox(oChild1, TransformF);
+				if (oChild1Box.Intersects(bounds))
+					find_intersections(iBox, OtherTree, TransformF, oChild1, depth + 1, intr, result, Options, OtherTreeOptions);
+
+			}
+			else                             // 2 children
+			{
+				oChild1 = oChild1 - 1;
+
+				FAxisAlignedBox3d oChild1Box = OtherTree.GetBox(oChild1, TransformF);
+				if (oChild1Box.Intersects(bounds))
+				{
+					find_intersections(iBox, OtherTree, TransformF, oChild1, depth + 1, intr, result, Options, OtherTreeOptions);
+				}
+
+				int oChild2 = OtherTree.IndexList[odx + 1] - 1;
+				FAxisAlignedBox3d oChild2Box = OtherTree.GetBox(oChild2, TransformF);
+				if (oChild2Box.Intersects(bounds))
+				{
+					find_intersections(iBox, OtherTree, TransformF, oChild2, depth + 1, intr, result, Options, OtherTreeOptions);
+				}
+			}
+
+		}
+		else
+		{
+			// descend our tree nodes if they intersect w/ current bounds of other tree
+			FAxisAlignedBox3d oBounds = OtherTree.GetBox(oBox, TransformF);
+
+			int iChild1 = IndexList[idx];
+			if (iChild1 < 0)                  // 1 child, descend if nearer than cur min-dist
+			{
+				iChild1 = (-iChild1) - 1;
+				if (box_box_intersect(iChild1, oBounds))
+				{
+					find_intersections(iChild1, OtherTree, TransformF, oBox, depth + 1, intr, result, Options, OtherTreeOptions);
+				}
+
+			}
+			else                             // 2 children
+			{
+				iChild1 = iChild1 - 1;
+				if (box_box_intersect(iChild1, oBounds))
+				{
+					find_intersections(iChild1, OtherTree, TransformF, oBox, depth + 1, intr, result, Options, OtherTreeOptions);
+				}
+
+				int iChild2 = IndexList[idx + 1] - 1;
+				if (box_box_intersect(iChild2, oBounds))
+				{
+					find_intersections(iChild2, OtherTree, TransformF, oBox, depth + 1, intr, result, Options, OtherTreeOptions);
+				}
+			}
+
+		}
+	}
+
+
+
 
 public:
-	// 1) make sure we can reach every tri in Mesh through tree (also demo of how to traverse tree...)
+	// 1) make sure we can reach every Tri in Mesh through tree (also demo of how to traverse tree...)
 	// 2) make sure that Triangles are contained in parent boxes
 	void TestCoverage()
 	{
@@ -1063,7 +1933,7 @@ private:
 				check(fabs(fTriDistSqr - fBoxDistSqr) <= FMathd::ZeroTolerance * 100);
 			}
 		};
-		TreeTraversalImpl(IBox, 0, t);
+		TreeTraversalImpl(IBox, 0, t, FQueryOptions());
 	}
 
 	// do full tree Traversal below IBox to make sure that all child Triangles are contained
@@ -1080,6 +1950,6 @@ private:
 				check(Box.Contains(V));
 			}
 		};
-		TreeTraversalImpl(IBox, 0, t);
+		TreeTraversalImpl(IBox, 0, t, FQueryOptions());
 	}
 };
diff --git a/Engine/Plugins/Experimental/GeometryProcessing/Source/GeometricObjects/Public/Spatial/PointHashGrid2.h b/Engine/Plugins/Experimental/GeometryProcessing/Source/GeometricObjects/Public/Spatial/PointHashGrid2.h
--- a/Engine/Plugins/Experimental/GeometryProcessing/Source/GeometricObjects/Public/Spatial/PointHashGrid2.h
+++ b/Engine/Plugins/Experimental/GeometryProcessing/Source/GeometricObjects/Public/Spatial/PointHashGrid2.h
@@ -12,40 +12,38 @@
  * Hash Grid for values associated with 2D points.
  *
  * This class addresses the situation where you have a list of (point, point_data) and you
- * would like to be able to do efficient proximity queries, ie find the nearest point_data
+ * would like to be able to do efficient proximity queries, i.e. find the nearest point_data
  * for a given query point.
  *
  * We don't store copies of the 2D points. You provide a point_data type. This could just be the
  * integer index into your list for example, a pointer to something more complex, etc.
  * Insert and Remove functions require you to pass in the 2D point for the point_data.
- * To Update a point you need to know it's old and new 2D coordinates.
- *
- * @todo also template on point realtype?
+ * To Update a point you need to know its old and new 2D coordinates.
  */
-template<typename PointDataType>
-class TPointHashGrid2d
+template<typename PointDataType, typename RealType>
+class TPointHashGrid2
 {
 private:
 	TMultiMap<FVector2i, PointDataType> Hash;
 	FCriticalSection CriticalSection;
-	FScaleGridIndexer2d Indexer;
-	PointDataType invalidValue;
+	TScaleGridIndexer2<RealType> Indexer;
+	PointDataType InvalidValue;
 
 public:
 
 	/**
-	 * Construct 2D hashgrid
+	 * Construct 2D hash grid
 	 * @param cellSize size of grid cells
-	 * @param invalidValue this value will be returned by queries if no valid result is found (eg bounded-distance query)
+	 * @param InvalidValue this value will be returned by queries if no valid result is found (e.g. bounded-distance query)
 	 */
-	TPointHashGrid2d(double cellSize, PointDataType InvalidValue) : Indexer(cellSize), invalidValue(InvalidValue)
+	TPointHashGrid2(RealType cellSize, PointDataType InvalidValue) : Indexer(cellSize), InvalidValue(InvalidValue)
 	{
 	}
 
 	/** Invalid grid value */
-	PointDataType InvalidValue() const
+	PointDataType GetInvalidValue() const
 	{
-		return invalidValue;
+		return InvalidValue;
 	}
 
 	/**
@@ -53,7 +51,7 @@ public:
 	 * @param Value the point/value to insert
 	 * @param Position the position associated with this value
 	 */
-	void InsertPoint(const PointDataType& Value, const FVector2d& Position)
+	void InsertPoint(const PointDataType& Value, const FVector2<RealType>& Position)
 	{
 		FVector2i idx = Indexer.ToGrid(Position);
 		{
@@ -67,7 +65,7 @@ public:
 	 * @param Value the point/value to insert
 	 * @param Position the position associated with this value
 	 */
-	void InsertPointUnsafe(const PointDataType& Value, const FVector2d& Position)
+	void InsertPointUnsafe(const PointDataType& Value, const FVector2<RealType>& Position)
 	{
 		FVector2i idx = Indexer.ToGrid(Position);
 		Hash.Add(idx, Value);
@@ -80,7 +78,7 @@ public:
 	 * @param Position the position associated with this value
 	 * @return true if the value existed at this position
 	 */
-	bool RemovePoint(const PointDataType& Value, const FVector2d& Position)
+	bool RemovePoint(const PointDataType& Value, const FVector2<RealType>& Position)
 	{
 		FVector2i idx = Indexer.ToGrid(Position);
 		{
@@ -95,7 +93,7 @@ public:
 	 * @param Position the position associated with this value
 	 * @return true if the value existed at this position
 	 */
-	bool RemovePointUnsafe(const PointDataType& Value, const FVector2d& Position)
+	bool RemovePointUnsafe(const PointDataType& Value, const FVector2<RealType>& Position)
 	{
 		FVector2i idx = Indexer.ToGrid(Position);
 		return Hash.RemoveSingle(idx, Value) > 0;
@@ -108,7 +106,7 @@ public:
 	 * @param OldPosition the current position associated with this value
 	 * @param NewPosition the new position for this value
 	 */
-	void UpdatePoint(const PointDataType& Value, const FVector2d& OldPosition, const FVector2d& NewPosition)
+	void UpdatePoint(const PointDataType& Value, const FVector2<RealType>& OldPosition, const FVector2<RealType>& NewPosition)
 	{
 		FVector2i old_idx = Indexer.ToGrid(OldPosition);
 		FVector2i new_idx = Indexer.ToGrid(NewPosition);
@@ -136,7 +134,7 @@ public:
 	 * @param OldPosition the current position associated with this value
 	 * @param NewPosition the new position for this value
 	 */
-	void UpdatePointUnsafe(const PointDataType& Value, const FVector2d& OldPosition, const FVector2d& NewPosition)
+	void UpdatePointUnsafe(const PointDataType& Value, const FVector2<RealType>& OldPosition, const FVector2<RealType>& NewPosition)
 	{
 		FVector2i old_idx = Indexer.ToGrid(OldPosition);
 		FVector2i new_idx = Indexer.ToGrid(NewPosition);
@@ -159,22 +157,22 @@ public:
 	 * @param IgnoreFunc optional Function you may provide which will result in a Value being ignored if IgnoreFunc(Value) returns true
 	 * @return the found pair (Value,DistanceFunc(Value)), or (InvalidValue,MaxDouble) if not found
 	 */
-	TPair<PointDataType, double> FindNearestInRadius(
-		const FVector2d& QueryPoint, double Radius, 
-		TFunction<double(const PointDataType&)> DistanceSqFunc,
-		TFunction<bool(const PointDataType&)> IgnoreFunc = [](const PointDataType& data) { return false; }) const
+	TPair<PointDataType, RealType> FindNearestInRadius(
+		const FVector2<RealType>& QueryPoint, RealType Radius, 
+		TFunctionRef<RealType(const PointDataType&)> DistanceSqFunc,
+		TFunctionRef<bool(const PointDataType&)> IgnoreFunc = [](const PointDataType& data) { return false; }) const
 	{
 		if (!Hash.Num())
 		{
-			return TPair<PointDataType, double>(InvalidValue(), TNumericLimits<double>::Max());
+			return TPair<PointDataType, RealType>(GetInvalidValue(), TNumericLimits<RealType>::Max());
 		}
 
-		FVector2i min_idx = Indexer.ToGrid(QueryPoint - Radius * FVector2d::One());
-		FVector2i max_idx = Indexer.ToGrid(QueryPoint + Radius * FVector2d::One());
+		FVector2i min_idx = Indexer.ToGrid(QueryPoint - Radius * FVector2<RealType>::One());
+		FVector2i max_idx = Indexer.ToGrid(QueryPoint + Radius * FVector2<RealType>::One());
 
-		double min_distsq = TNumericLimits<double>::Max();
-		PointDataType nearest = InvalidValue();
-		double RadiusSquared = Radius * Radius;
+		RealType min_distsq = TNumericLimits<RealType>::Max();
+		PointDataType nearest = GetInvalidValue();
+		RealType RadiusSquared = Radius * Radius;
 
 		TArray<PointDataType> Values;
 		for (int yi = min_idx.Y; yi <= max_idx.Y; yi++) 
@@ -190,7 +188,7 @@ public:
 					{
 						continue;
 					}
-					double distsq = DistanceSqFunc(Value);
+					RealType distsq = DistanceSqFunc(Value);
 					if (distsq < RadiusSquared && distsq < min_distsq)
 					{
 						nearest = Value;
@@ -200,7 +198,9 @@ public:
 			}
 		}
 
-		return TPair<PointDataType, double>(nearest, min_distsq);
+		return TPair<PointDataType, RealType>(nearest, min_distsq);
 	}
 };
 
+template <typename PointDataType> using TPointHashGrid2d = TPointHashGrid2<PointDataType, double>;
+template <typename PointDataType> using TPointHashGrid2f = TPointHashGrid2<PointDataType, float>;
diff --git a/Engine/Plugins/Experimental/GeometryProcessing/Source/GeometricObjects/Public/Spatial/SpatialInterfaces.h b/Engine/Plugins/Experimental/GeometryProcessing/Source/GeometricObjects/Public/Spatial/SpatialInterfaces.h
--- a/Engine/Plugins/Experimental/GeometryProcessing/Source/GeometricObjects/Public/Spatial/SpatialInterfaces.h
+++ b/Engine/Plugins/Experimental/GeometryProcessing/Source/GeometricObjects/Public/Spatial/SpatialInterfaces.h
@@ -14,17 +14,17 @@
 class ISpatial
 {
 public:
-	virtual ~ISpatial() {}
+	virtual ~ISpatial() = default;
 
 	/**
 	 * @return true if this object supports point-containment inside/outside queries
 	 */
-	virtual bool SupportsPointContainment() = 0;
+	virtual bool SupportsPointContainment() const = 0;
 
 	/**
 	 * @return true if the query point is inside the object
 	 */
-	virtual bool IsInside(const FVector3d & Point) = 0;
+	virtual bool IsInside(const FVector3d & Point) const = 0;
 };
 
 
@@ -35,12 +35,30 @@ public:
 class IMeshSpatial : public ISpatial
 {
 public:
-	virtual ~IMeshSpatial() {}
+	virtual ~IMeshSpatial() = default;
+
+	// standard shared options for all mesh spatial queries
+	struct FQueryOptions
+	{
+		/**
+		 * Maximum search distance / hit distance, where applicable
+		 */
+		double MaxDistance = TNumericLimits<double>::Max();
+
+		/**
+		 * If non-null, only triangle IDs that pass this filter (i.e. filter is true) are considered
+		 */
+		TFunction<bool(int)> TriangleFilterF = nullptr;
+
+		FQueryOptions() {}
+		FQueryOptions(TFunction<bool(int)> TriangleFilterF) : TriangleFilterF(TriangleFilterF) {}
+		FQueryOptions(double MaxDistance, TFunction<bool(int)> TriangleFilterF = nullptr) : MaxDistance(MaxDistance), TriangleFilterF(TriangleFilterF) {}
+	};
 
 	/**
 	 * @return true if this object supports nearest-triangle queries
 	 */
-	virtual bool SupportsNearestTriangle() = 0;
+	virtual bool SupportsNearestTriangle() const = 0;
 
 	/**
 	 * @param Query point
@@ -48,24 +66,24 @@ public:
 	 * @param MaxDistance maximum search distance
 	 * @return ID of triangle nearest to Point within MaxDistance, or InvalidID if not found
 	 */
-	virtual int FindNearestTriangle(const FVector3d& Point, double& NearestDistSqrOut, double MaxDistance = TNumericLimits<double>::Max()) = 0;
+	virtual int FindNearestTriangle(const FVector3d& Point, double& NearestDistSqrOut, const FQueryOptions& Options = FQueryOptions()) const = 0;
 
 
 	/**
 	 * @return true if this object supports ray-triangle intersection queries
 	 */
-	virtual bool SupportsTriangleRayIntersection() = 0;
+	virtual bool SupportsTriangleRayIntersection() const = 0;
 
 	/**
 	 * @param Ray query ray
 	 * @param MaxDistance maximum hit distance
 	 * @return ID of triangle intersected by ray within MaxDistance, or InvalidID if not found
 	 */
-	inline virtual int FindNearestHitTriangle(const FRay3d& Ray, double MaxDistance = TNumericLimits<double>::Max())
+	virtual int FindNearestHitTriangle(const FRay3d& Ray, const FQueryOptions& Options = FQueryOptions()) const
 	{
 		double NearestT;
 		int TID;
-		FindNearestHitTriangle(Ray, NearestT, TID, MaxDistance);
+		FindNearestHitTriangle(Ray, NearestT, TID, Options);
 		return TID;
 	}
 
@@ -77,7 +95,7 @@ public:
 	 * @param MaxDistance maximum hit distance
 	 * @return true if hit, false if no hit found
 	 */
-	virtual bool FindNearestHitTriangle(const FRay3d& Ray, double& NearestT, int& TID, double MaxDistance = TNumericLimits<double>::Max()) = 0;
+	virtual bool FindNearestHitTriangle(const FRay3d& Ray, double& NearestT, int& TID, const FQueryOptions& Options = FQueryOptions()) const = 0;
 
 
 };
diff --git a/Engine/Plugins/Experimental/GeometryProcessing/Source/GeometricObjects/Public/TriangleTypes.h b/Engine/Plugins/Experimental/GeometryProcessing/Source/GeometricObjects/Public/TriangleTypes.h
--- a/Engine/Plugins/Experimental/GeometryProcessing/Source/GeometricObjects/Public/TriangleTypes.h
+++ b/Engine/Plugins/Experimental/GeometryProcessing/Source/GeometricObjects/Public/TriangleTypes.h
@@ -57,6 +57,10 @@ struct TTriangle2
 		return BaryCoords[0] * V[0] + BaryCoords[1] * V[1] + BaryCoords[2] * V[2];
 	}
 
+	FVector3<RealType> GetBarycentricCoords(const FVector2<RealType>& Point) const
+	{
+		return VectorUtil::BarycentricCoords(Point, V[0], V[1], V[2]);
+	}
 
 	/**
 	 * @param A first vertex of triangle
@@ -97,6 +101,46 @@ struct TTriangle2
 		return IsInside(V[0], V[1], V[2], QueryPoint);
 	}
 
+	/**
+	 * More robust (because it doesn't multiply orientation test results) inside-triangle test for oriented triangles only
+	 *  (the code early-outs at the first 'outside' edge, which only works if the triangle is oriented as expected)
+	 * @return 1 if outside, -1 if inside, 0 if on boundary
+	 */
+	int IsInsideOrOn_Oriented(const FVector2<RealType>& QueryPoint) const
+	{
+		return IsInsideOrOn_Oriented(V[0], V[1], V[2], QueryPoint);
+	}
+
+	/**
+	 * More robust (because it doesn't multiply orientation test results) inside-triangle test for oriented triangles only
+	 *  (the code early-outs at the first 'outside' edge, which only works if the triangle is oriented as expected)
+	 * @return 1 if outside, -1 if inside, 0 if on boundary
+	 */
+	static int IsInsideOrOn_Oriented(const FVector2<RealType>& A, const FVector2<RealType>& B, const FVector2<RealType>& C, const FVector2<RealType>& QueryPoint)
+	{
+		checkSlow(FVector2<RealType>::Orient(A, B, C) <= 0); // TODO: remove this checkSlow; it's just to make sure the orientation is as expected
+
+		RealType Sign1 = FVector2<RealType>::Orient(A, B, QueryPoint);
+		if (Sign1 > 0)
+		{
+			return 1;
+		}
+		
+		RealType Sign2 = FVector2<RealType>::Orient(B, C, QueryPoint);
+		if (Sign2 > 0)
+		{
+			return 1;
+		}
+		
+		RealType Sign3 = FVector2<RealType>::Orient(A, C, QueryPoint);
+		if (Sign3 < 0) // note this edge is queried backwards so the sign test is also backwards
+		{
+			return 1;
+		}
+
+		return (Sign1 != 0 && Sign2 != 0 && Sign3 != 0) ? -1 : 0;
+	}
+
 
 };
 
@@ -140,6 +184,10 @@ struct TTriangle3
 		return BaryCoords[0]*V[0] + BaryCoords[1]*V[1] + BaryCoords[2]*V[2];
 	}
 
+	FVector3<RealType> GetBarycentricCoords(const FVector3<RealType>& Point) const
+	{
+		return VectorUtil::BarycentricCoords(Point, V[0], V[1], V[2]);
+	}
 
 	/** @return vector that is perpendicular to the plane of this triangle */
 	FVector3<RealType> Normal() const
diff --git a/Engine/Plugins/Experimental/GeometryProcessing/Source/GeometricObjects/Public/Util/IndexPriorityQueue.h b/Engine/Plugins/Experimental/GeometryProcessing/Source/GeometricObjects/Public/Util/IndexPriorityQueue.h
--- a/Engine/Plugins/Experimental/GeometryProcessing/Source/GeometricObjects/Public/Util/IndexPriorityQueue.h
+++ b/Engine/Plugins/Experimental/GeometryProcessing/Source/GeometricObjects/Public/Util/IndexPriorityQueue.h
@@ -11,7 +11,7 @@
  * Integer IDs must be provided by the user to identify unique nodes.
  * Internally an array is used to keep track of the mapping from ids to internal indices, so the max ID must also be provided.
  *
- * @todo based on C# code where TDynamicVector could not return a reference to internal element. In C++ we can and code could be updated to be more efficient in many places.
+ * @todo based on C# code where TDynamicVector could not return a reference to internal element. In C++ we can, and code could be updated to be more efficient in many places.
  * @todo id_to_index could be sparse in many situations...
  */
 class FIndexPriorityQueue
@@ -211,7 +211,7 @@ private:
 			return;
 		}
 
-		// [RMS] is there a better way to do this? seems random to move the last node to
+		// TODO: is there a better way to do this? seems random to move the last node to
 		// top of tree? But I guess otherwise we might have to shift entire branches??
 
 		//Swap the node with the last node
diff --git a/Engine/Plugins/Experimental/GeometryProcessing/Source/GeometricObjects/Public/Util/IndexUtil.h b/Engine/Plugins/Experimental/GeometryProcessing/Source/GeometricObjects/Public/Util/IndexUtil.h
--- a/Engine/Plugins/Experimental/GeometryProcessing/Source/GeometricObjects/Public/Util/IndexUtil.h
+++ b/Engine/Plugins/Experimental/GeometryProcessing/Source/GeometricObjects/Public/Util/IndexUtil.h
@@ -273,6 +273,20 @@ namespace IndexUtil
 		return true;
 	}
 
+	/**
+	 * integer indices offsets in x/y directions
+	 */
+	extern GEOMETRICOBJECTS_API const FVector2i GridOffsets4[4];
+
+	/**
+	 * integer indices offsets in x/y directions and diagonals
+	 */
+	extern GEOMETRICOBJECTS_API const FVector2i GridOffsets8[8];
+
+	/**
+	 * integer indices offsets in x/y/z directions, corresponds w/ BoxFaces directions
+	 */
+	extern GEOMETRICOBJECTS_API const FVector3i GridOffsets6[6];
 
 	/**
 	 * all permutations of (+-1, +-1, +-1), can be used to iterate over connected face/edge/corner neighbours of a grid cell
diff --git a/Engine/Plugins/Experimental/GeometryProcessing/Source/GeometricObjects/Public/Util/RefCountVector.h b/Engine/Plugins/Experimental/GeometryProcessing/Source/GeometricObjects/Public/Util/RefCountVector.h
--- a/Engine/Plugins/Experimental/GeometryProcessing/Source/GeometricObjects/Public/Util/RefCountVector.h
+++ b/Engine/Plugins/Experimental/GeometryProcessing/Source/GeometricObjects/Public/Util/RefCountVector.h
@@ -74,7 +74,7 @@ public:
 		UsedCount++;
 		if (FreeIndices.IsEmpty())
 		{
-			// [RMS] do we need this branch anymore?
+			// TODO: do we need this branch anymore?
 			RefCounts.Add(1);
 			return (int)RefCounts.GetLength() - 1;
 		}
diff --git a/Engine/Plugins/Experimental/GeometryProcessing/Source/GeometricObjects/Public/Util/SmallListSet.h b/Engine/Plugins/Experimental/GeometryProcessing/Source/GeometricObjects/Public/Util/SmallListSet.h
--- a/Engine/Plugins/Experimental/GeometryProcessing/Source/GeometricObjects/Public/Util/SmallListSet.h
+++ b/Engine/Plugins/Experimental/GeometryProcessing/Source/GeometricObjects/Public/Util/SmallListSet.h
@@ -158,6 +158,11 @@ public:
 	bool Replace(int32 ListIndex, const TFunction<bool(int32)>& PredicateFunc, int32 NewValue);
 
 
+	/**
+	 * Call ApplyFunc on each element of the list at ListIndex
+	 */
+	void Enumerate(int32 ListIndex, TFunctionRef<void(int32)> ApplyFunc) const;
+
 
 	//
 	// iterator support
diff --git a/Engine/Plugins/Experimental/GeometryProcessing/Source/GeometricObjects/Public/VectorTypes.h b/Engine/Plugins/Experimental/GeometryProcessing/Source/GeometricObjects/Public/VectorTypes.h
--- a/Engine/Plugins/Experimental/GeometryProcessing/Source/GeometricObjects/Public/VectorTypes.h
+++ b/Engine/Plugins/Experimental/GeometryProcessing/Source/GeometricObjects/Public/VectorTypes.h
@@ -4,6 +4,7 @@
 
 #include "Math/Vector.h"
 #include "MathUtil.h"
+#include <sstream>
 
 
 
@@ -286,11 +287,16 @@ struct FVector2
 
 
 
-	bool operator==(const FVector2<T>& Other) const
+	constexpr bool operator==(const FVector2<T>& Other) const
 	{
 		return X == Other.X && Y == Other.Y;
 	}
 
+	constexpr bool operator!=(const FVector2<T>& Other) const
+	{
+		return X != Other.X || Y != Other.Y;
+	}
+
 
 	static FVector2<T> Lerp(const FVector2<T>& A, const FVector2<T>& B, T Alpha)
 	{
@@ -551,6 +557,12 @@ struct FVector3
 			X * V2.Y - Y * V2.X);
 	}
 
+	FVector3<T> UnitCross(const FVector3<T>& V2) const
+	{
+		FVector3<T> N = Cross(V2);
+		return N.Normalized();
+	}
+
 	/**
 	 * Assumes this vector and V2 are both already normalized
 	 * @return the angle to vector V2 in degrees
@@ -614,6 +626,19 @@ struct FVector3
 		return TMathUtil<T>::Min3(TMathUtil<T>::Abs(X), TMathUtil<T>::Abs(Y), TMathUtil<T>::Abs(Z));
 	}
 
+	constexpr FVector2<T> XY() const
+	{
+		return FVector2<T>(X, Y);
+	}
+	constexpr FVector2<T> XZ() const
+	{
+		return FVector2<T>(X, Z);
+	}
+	constexpr FVector2<T> YZ() const
+	{
+		return FVector2<T>(Y, Z);
+	}
+
 	static FVector3<T> Lerp(const FVector3<T>& A, const FVector3<T>& B, T Alpha)
 	{
 		T OneMinusAlpha = (T)1 - Alpha;
@@ -626,6 +651,11 @@ struct FVector3
 	{
 		return X == Other.X && Y == Other.Y && Z == Other.Z;
 	}
+
+	constexpr bool operator!=(const FVector3<T>& Other) const
+	{
+		return X != Other.X || Y != Other.Y || Z != Other.Z;
+	}
 };
 
 template <typename RealType>
@@ -641,6 +671,13 @@ inline FVector3<RealType> operator*(RealType2 Scalar, const FVector3<RealType>&
 	return FVector3<RealType>((RealType)Scalar * V.X, (RealType)Scalar * V.Y, (RealType)Scalar * V.Z);
 }
 
+template <typename RealType>
+std::ostream& operator<<(std::ostream& os, const FVector3<RealType>& Vec)
+{
+	os << Vec.X << " " << Vec.Y << " " << Vec.Z;
+	return os;
+}
+
 typedef FVector3<float> FVector3f;
 typedef FVector3<double> FVector3d;
 typedef FVector3<int> FVector3i;
@@ -678,6 +715,13 @@ inline FVector2<RealType> operator*(RealType2 Scalar, const FVector2<RealType>&
 	return FVector2<RealType>((RealType)Scalar * V.X, (RealType)Scalar * V.Y);
 }
 
+template <typename RealType>
+std::ostream& operator<<(std::ostream& os, const FVector2<RealType>& Vec)
+{
+	os << Vec.X << " " << Vec.Y;
+	return os;
+}
+
 typedef FVector2<float> FVector2f;
 typedef FVector2<double> FVector2d;
 typedef FVector2<int> FVector2i;
diff --git a/Engine/Plugins/Experimental/GeometryProcessing/Source/GeometricObjects/Public/VectorUtil.h b/Engine/Plugins/Experimental/GeometryProcessing/Source/GeometricObjects/Public/VectorUtil.h
--- a/Engine/Plugins/Experimental/GeometryProcessing/Source/GeometricObjects/Public/VectorUtil.h
+++ b/Engine/Plugins/Experimental/GeometryProcessing/Source/GeometricObjects/Public/VectorUtil.h
@@ -316,7 +316,7 @@ namespace VectorUtil
 	* TODO: make robust to degenerate triangles?
 	*/
 	template <typename RealType>
-	FVector3<RealType> BarycentricCoords(const FVector3<RealType>& Point, const FVector2<RealType>& V0, const FVector2<RealType>& V1, const FVector2<RealType>& V2)
+	FVector3<RealType> BarycentricCoords(const FVector2<RealType>& Point, const FVector2<RealType>& V0, const FVector2<RealType>& V1, const FVector2<RealType>& V2)
 	{
 		FVector2<RealType> kV02 = V0 - V2;
 		FVector2<RealType> kV12 = V1 - V2;
@@ -421,7 +421,13 @@ namespace VectorUtil
 			NormalIn.X*TangentIn.Y - NormalIn.Y*TangentIn.X);
 	}
 
-
+	/// @return Aspect ratio of triangle 
+	inline double AspectRatio(const FVector3d& v1, const FVector3d& v2, const FVector3d& v3)
+	{
+		double a = v1.Distance(v2), b = v2.Distance(v3), c = v3.Distance(v1);
+		double s = (a + b + c) / 2.0;
+		return (a * b * c) / (8.0 * (s - a) * (s - b) * (s - c));
+	}
 
 
 }; // namespace VectorUtil
diff --git a/Engine/Plugins/Experimental/GeometryProcessing/Source/GeometryAlgorithms/Public/Arrangement2d.h b/Engine/Plugins/Experimental/GeometryProcessing/Source/GeometryAlgorithms/Public/Arrangement2d.h
--- a/Engine/Plugins/Experimental/GeometryProcessing/Source/GeometryAlgorithms/Public/Arrangement2d.h
+++ b/Engine/Plugins/Experimental/GeometryProcessing/Source/GeometryAlgorithms/Public/Arrangement2d.h
@@ -431,7 +431,7 @@ protected:
 		auto FuncIgnore = [&](int VID) { return VID == IgnoreVID; };
 		TPair<int, double> found = (IgnoreVID == -1) ? PointHash.FindNearestInRadius(Pt, SearchRadius, FuncDistSq)
 													 : PointHash.FindNearestInRadius(Pt, SearchRadius, FuncDistSq, FuncIgnore);
-		if (found.Key == PointHash.InvalidValue())
+		if (found.Key == PointHash.GetInvalidValue())
 		{
 			return -1;
 		}
@@ -447,7 +447,7 @@ protected:
 		auto FuncIgnore = [&](int VID) { return Graph.IsBoundaryVertex(VID) == false || VID == IgnoreVID; };
 		TPair<int, double> found =
 			PointHash.FindNearestInRadius(Pt, SearchRadius, FuncDistSq, FuncIgnore);
-		if (found.Key == PointHash.InvalidValue())
+		if (found.Key == PointHash.GetInvalidValue())
 		{
 			return -1;
 		}
diff --git a/Engine/Plugins/Experimental/GeometryProcessing/Source/MeshConversion/Private/MeshDescriptionToDynamicMesh.cpp b/Engine/Plugins/Experimental/GeometryProcessing/Source/MeshConversion/Private/MeshDescriptionToDynamicMesh.cpp
--- a/Engine/Plugins/Experimental/GeometryProcessing/Source/MeshConversion/Private/MeshDescriptionToDynamicMesh.cpp
+++ b/Engine/Plugins/Experimental/GeometryProcessing/Source/MeshConversion/Private/MeshDescriptionToDynamicMesh.cpp
@@ -5,6 +5,7 @@
 #include "DynamicMeshAttributeSet.h"
 #include "DynamicMeshOverlay.h"
 #include "MeshDescriptionBuilder.h"
+#include "Async/Async.h"
 
 
 struct FVertexUV
@@ -30,6 +31,10 @@ public:
 	TMap<FVertexUV, int> UniqueVertexUVs;
 	FDynamicMeshUVOverlay* UVOverlay;
 
+	FUVWelder() : UVOverlay(nullptr)
+	{
+	}
+
 	FUVWelder(FDynamicMeshUVOverlay* UVOverlayIn)
 	{
 		check(UVOverlayIn);
@@ -40,16 +45,14 @@ public:
 	{
 		FVertexUV VertUV = { VertexID, UV.X, UV.Y };
 
-		int NewIndex = -1;
-		if (UniqueVertexUVs.Contains(VertUV))
-		{
-			NewIndex = UniqueVertexUVs[VertUV];
-		}
-		else
+		const int32* FoundIndex = UniqueVertexUVs.Find(VertUV);
+		if (FoundIndex != nullptr)
 		{
-			NewIndex = UVOverlay->AppendElement(FVector2f(UV), VertexID);
-			UniqueVertexUVs.Add(VertUV, NewIndex);
+			return *FoundIndex;
 		}
+
+		int32 NewIndex = UVOverlay->AppendElement(FVector2f(UV));
+		UniqueVertexUVs.Add(VertUV, NewIndex);
 		return NewIndex;
 	}
 };
@@ -81,6 +84,10 @@ public:
 	TMap<FVertexNormal, int> UniqueVertexNormals;
 	FDynamicMeshNormalOverlay* NormalOverlay;
 
+	FNormalWelder() : NormalOverlay(nullptr)
+	{
+	}
+
 	FNormalWelder(FDynamicMeshNormalOverlay* NormalOverlayIn)
 	{
 		check(NormalOverlayIn);
@@ -91,16 +98,14 @@ public:
 	{
 		FVertexNormal VertNormal = { VertexID, Normal.X, Normal.Y, Normal.Z };
 
-		int NewIndex = -1;
-		if (UniqueVertexNormals.Contains(VertNormal))
-		{
-			NewIndex = UniqueVertexNormals[VertNormal];
-		}
-		else
+		const int32* FoundIndex = UniqueVertexNormals.Find(VertNormal);
+		if (FoundIndex != nullptr)
 		{
-			NewIndex = NormalOverlay->AppendElement(Normal, VertexID);
-			UniqueVertexNormals.Add(VertNormal, NewIndex);
+			return *FoundIndex;
 		}
+
+		int32 NewIndex = NormalOverlay->AppendElement(Normal);
+		UniqueVertexNormals.Add(VertNormal, NewIndex);
 		return NewIndex;
 	}
 };
@@ -139,22 +144,14 @@ void FMeshDescriptionToDynamicMesh::Convert(const FMeshDescription* MeshIn, FDyn
 	TVertexInstanceAttributesConstRef<FVector> InstanceNormals =
 		MeshIn->VertexInstanceAttributes().GetAttributesRef<FVector>(MeshAttribute::VertexInstance::Normal);
 
-	int NumUVLayers = InstanceUVs.GetNumIndices();
-
-	// enable attributes on output mesh
-	MeshOut.EnableAttributes();
-	MeshOut.Attributes()->SetNumUVLayers(NumUVLayers);
-	FDynamicMeshNormalOverlay* NormalOverlay = MeshOut.Attributes()->PrimaryNormals();
-
 	TPolygonAttributesConstRef<int> PolyGroups =
 		MeshIn->PolygonAttributes().GetAttributesRef<int>(ExtendedMeshAttribute::PolyTriGroups);
-
-	// base triangle groups will track polygons
 	if (bEnableOutputGroups)
 	{
 		MeshOut.EnableTriangleGroups(0);
 	}
 
+	// base triangle groups will track polygons
 	int NumVertices = MeshIn->Vertices().Num();
 	int NumPolygons = MeshIn->Polygons().Num();
 	int NumVtxInstances = MeshIn->VertexInstances().Num();
@@ -163,19 +160,19 @@ void FMeshDescriptionToDynamicMesh::Convert(const FMeshDescription* MeshIn, FDyn
 		UE_LOG(LogTemp, Warning, TEXT("FMeshDescriptionToDynamicMesh: MeshDescription verts %d polys %d instances %d"), NumVertices, NumPolygons, NumVtxInstances);
 	}
 
-	// reserve space in MeshOut?
+	FDateTime Time_AfterVertices = FDateTime::Now();
 
-	// used to merge coincident elements so that we get actual topology
-	TArray<FUVWelder> UVWelders;
-	for (int UVLayerIndex = 0; UVLayerIndex < NumUVLayers; UVLayerIndex++)
+	// Although it is slightly redundant, we will build up this list of MeshDescription data
+	// so that it is easier to index into it below (profile whether extra memory here hurts us?)
+	struct FTriData
 	{
-		UVWelders.Emplace(MeshOut.Attributes()->GetUVLayer(UVLayerIndex));
-	}
-	FNormalWelder NormalWelder(NormalOverlay);
-
-	// always enable material ID
-	MeshOut.Attributes()->EnableMaterialID();
-	FDynamicMeshMaterialAttribute* MaterialIDAttrib = MeshOut.Attributes()->GetMaterialID();
+		FPolygonID PolygonID;
+		int32 PolygonGroupID;
+		int32 TriIndex;
+		FVertexInstanceID TriInstances[3];
+	};
+	TArray<FTriData> AddedTriangles;
+	AddedTriangles.SetNum(MeshIn->Triangles().Num());
 
 	// NOTE: If you change the iteration order here, please update the corresponding iteration in FDynamicMeshToMeshDescription::UpdateAttributes, 
 	//	which assumes the iteration order here is polygons -> triangles, to correspond the triangles when writing updated attributes back!
@@ -186,11 +183,17 @@ void FMeshDescriptionToDynamicMesh::Convert(const FMeshDescription* MeshIn, FDyn
 
 		const TArray<FTriangleID>& TriangleIDs = MeshIn->GetPolygonTriangleIDs(PolygonID);
 		int NumTriangles = TriangleIDs.Num();
-		for ( int TriIdx = 0; TriIdx < NumTriangles; ++TriIdx)
+		for (int TriIdx = 0; TriIdx < NumTriangles; ++TriIdx)
 		{
-			const FTriangleID TriangleID = TriangleIDs[TriIdx];
-				
-			TArrayView<const FVertexInstanceID> InstanceTri = MeshIn->GetTriangleVertexInstances(TriangleID);
+			FTriData TriData;
+			TriData.PolygonID = PolygonID;
+			TriData.PolygonGroupID = PolygonGroupID;
+			TriData.TriIndex = TriIdx;
+
+			TArrayView<const FVertexInstanceID> InstanceTri = MeshIn->GetTriangleVertexInstances(TriangleIDs[TriIdx]);
+			TriData.TriInstances[0] = InstanceTri[0];
+			TriData.TriInstances[1] = InstanceTri[1];
+			TriData.TriInstances[2] = InstanceTri[2];
 
 			int GroupID = 0;
 			if (GroupMode == EPrimaryGroupMode::SetToPolyGroup)
@@ -208,19 +211,18 @@ void FMeshDescriptionToDynamicMesh::Convert(const FMeshDescription* MeshIn, FDyn
 			{
 				GroupID = PolygonGroupID;
 			}
-			
 
 			// append triangle
-			int VertexID0 = MeshIn->GetVertexInstanceVertex(InstanceTri[0]).GetValue();
-			int VertexID1 = MeshIn->GetVertexInstanceVertex(InstanceTri[1]).GetValue();
-			int VertexID2 = MeshIn->GetVertexInstanceVertex(InstanceTri[2]).GetValue();
+			int32 VertexID0 = MeshIn->GetVertexInstanceVertex(InstanceTri[0]).GetValue();
+			int32 VertexID1 = MeshIn->GetVertexInstanceVertex(InstanceTri[1]).GetValue();
+			int32 VertexID2 = MeshIn->GetVertexInstanceVertex(InstanceTri[2]).GetValue();
 			int NewTriangleID = MeshOut.AppendTriangle(VertexID0, VertexID1, VertexID2, GroupID);
 
 			if (NewTriangleID == FDynamicMesh3::DuplicateTriangleID)
 			{
 				continue;
 			}
-			
+
 			// if append failed due to non-manifold, duplicate verts
 			if (NewTriangleID == FDynamicMesh3::NonManifoldID)
 			{
@@ -272,49 +274,132 @@ void FMeshDescriptionToDynamicMesh::Convert(const FMeshDescription* MeshIn, FDyn
 			}
 
 			checkSlow(NewTriangleID >= 0);
+			AddedTriangles[NewTriangleID] = TriData;
+		}
+	}
+
+	FDateTime Time_AfterTriangles = FDateTime::Now();
+
+	//
+	// Enable relevant attributes and initialize UV/Normal welders
+	// 
+
+	int NumUVLayers = InstanceUVs.GetNumIndices();
+	TArray<FDynamicMeshUVOverlay*> UVOverlays;
+	TArray<FUVWelder> UVWelders;
+	FDynamicMeshNormalOverlay* NormalOverlay = nullptr;
+	FNormalWelder NormalWelder;
+	FDynamicMeshMaterialAttribute* MaterialIDAttrib = nullptr;
+	if (!bDisableAttributes)
+	{
+		MeshOut.EnableAttributes();
+
+		MeshOut.Attributes()->SetNumUVLayers(NumUVLayers);
+		UVOverlays.SetNum(NumUVLayers);
+		UVWelders.SetNum(NumUVLayers);
+		for (int j = 0; j < NumUVLayers; ++j)
+		{
+			UVOverlays[j] = MeshOut.Attributes()->GetUVLayer(j);
+			UVWelders[j].UVOverlay = UVOverlays[j];
+		}
+
+		NormalOverlay = MeshOut.Attributes()->PrimaryNormals();
+		NormalWelder.NormalOverlay = NormalOverlay;
+
+		// always enable Material ID if there are any attributes
+		MeshOut.Attributes()->EnableMaterialID();
+		MaterialIDAttrib = MeshOut.Attributes()->GetMaterialID();
+	}
+
 
-			FIndex3i Tri(VertexID0, VertexID1, VertexID2);
+	// we will weld/populate all the attributes simultaneously, hold on to futures in this array and then Wait for them at the end
+	TArray<TFuture<void>> Pending;
 
-			if (bCalculateMaps)
+	if (bCalculateMaps)
+	{
+		auto Future = Async(EAsyncExecution::ThreadPool, [&]()
+		{
+			for (int32 TriangleID : MeshOut.TriangleIndicesItr())
 			{
-				TriToPolyTriMap.Insert(FIndex2i(PolygonID.GetValue(), TriIdx), NewTriangleID);
+				const FTriData& TriData = AddedTriangles[TriangleID];
+				TriToPolyTriMap.Insert(FIndex2i(TriData.PolygonID.GetValue(), TriData.TriIndex), TriangleID);
 			}
+		});
+		Pending.Add(MoveTemp(Future));
+	}
 
-			for (int UVLayerIndex = 0; UVLayerIndex < NumUVLayers; UVLayerIndex++)
+	if (!bDisableAttributes)
+	{
+		for (int UVLayerIndex = 0; UVLayerIndex < NumUVLayers; UVLayerIndex++)
+		{
+			auto UVFuture = Async(EAsyncExecution::ThreadPool, [&, UVLayerIndex]() // must copy UVLayerIndex here!
 			{
-				FDynamicMeshUVOverlay* UVOverlay = MeshOut.Attributes()->GetUVLayer(UVLayerIndex);
-				FIndex3i TriUV;
-				for (int j = 0; j < 3; ++j)
+				for (int32 TriangleID : MeshOut.TriangleIndicesItr())
 				{
-					FVector2D UV = InstanceUVs.Get(InstanceTri[j], UVLayerIndex);
-					TriUV[j] = UVWelders[UVLayerIndex].FindOrAddUnique(UV, Tri[j]);
+					FIndex3i Tri = MeshOut.GetTriangle(TriangleID);
+					const FTriData& TriData = AddedTriangles[TriangleID];
+					FIndex3i TriUV;
+					for (int j = 0; j < 3; ++j)
+					{
+						FVector2D UV = InstanceUVs.Get(TriData.TriInstances[j], UVLayerIndex);
+						TriUV[j] = UVWelders[UVLayerIndex].FindOrAddUnique(UV, Tri[j]);
+					}
+					UVOverlays[UVLayerIndex]->SetTriangle(TriangleID, TriUV);
 				}
-				UVOverlay->SetTriangle(NewTriangleID, TriUV);
-			}
+			});
+			Pending.Add(MoveTemp(UVFuture));
+		}
 
-			if (NormalOverlay != nullptr)
+
+		if (NormalOverlay != nullptr)
+		{
+			auto NormalFuture = Async(EAsyncExecution::ThreadPool, [&]()
 			{
-				FIndex3i TriNormals;
-				for (int j = 0; j < 3; ++j)
+				for (int32 TriangleID : MeshOut.TriangleIndicesItr())
 				{
-					FVector Normal = InstanceNormals.Get(InstanceTri[j]);
-					TriNormals[j] = NormalWelder.FindOrAddUnique(Normal, Tri[j]);
+					FIndex3i Tri = MeshOut.GetTriangle(TriangleID);
+					const FTriData& TriData = AddedTriangles[TriangleID];
+					FIndex3i TriNormals;
+					for (int j = 0; j < 3; ++j)
+					{
+						FVector Normal = InstanceNormals.Get(TriData.TriInstances[j]);
+						TriNormals[j] = NormalWelder.FindOrAddUnique(Normal, Tri[j]);
+					}
+					NormalOverlay->SetTriangle(TriangleID, TriNormals);
 				}
-				NormalOverlay->SetTriangle(NewTriangleID, TriNormals);
-			}
+			});
+			Pending.Add(MoveTemp(NormalFuture));
+		}
 
-			// use PolygonGroup as MaterialID
-			if (MaterialIDAttrib != nullptr)
-			{
-				MaterialIDAttrib->SetValue(NewTriangleID, &PolygonGroupID);
-			}
 
+		if (MaterialIDAttrib != nullptr)
+		{
+			auto MaterialFuture = Async(EAsyncExecution::ThreadPool, [&]()
+			{
+				for (int32 TriangleID : MeshOut.TriangleIndicesItr())
+				{
+					const FTriData& TriData = AddedTriangles[TriangleID];
+					MaterialIDAttrib->SetValue(TriangleID, &TriData.PolygonGroupID);
+				}
+			});
+			Pending.Add(MoveTemp(MaterialFuture));
 		}
 	}
 
+	// wait for all work to be done
+	for (TFuture<void>& Future : Pending)
+	{
+		Future.Wait();
+	}
+
+	FDateTime Time_AfterAttribs = FDateTime::Now();
+
 	if (bPrintDebugMessages)
 	{
-		int NumUVs = (NumUVLayers > 0) ? MeshOut.Attributes()->PrimaryUV()->MaxElementID() : 0;
+		UE_LOG(LogTemp, Warning, TEXT("FMeshDescriptionToDynamicMesh:  Conversion Timing: Triangles %fs   Attributbes %fs"),
+			(Time_AfterTriangles - Time_AfterVertices).GetTotalSeconds(), (Time_AfterAttribs - Time_AfterTriangles).GetTotalSeconds());
+
+		int NumUVs = (MeshOut.HasAttributes() && NumUVLayers > 0) ? MeshOut.Attributes()->PrimaryUV()->MaxElementID() : 0;
 		int NumNormals = (NormalOverlay != nullptr) ? NormalOverlay->MaxElementID() : 0;
 		UE_LOG(LogTemp, Warning, TEXT("FMeshDescriptionToDynamicMesh:  FDynamicMesh verts %d triangles %d (primary) uvs %d normals %d"), MeshOut.MaxVertexID(), MeshOut.MaxTriangleID(), NumUVs, NumNormals);
 	}
diff --git a/Engine/Plugins/Experimental/GeometryProcessing/Source/MeshConversion/Public/MeshDescriptionToDynamicMesh.h b/Engine/Plugins/Experimental/GeometryProcessing/Source/MeshConversion/Public/MeshDescriptionToDynamicMesh.h
--- a/Engine/Plugins/Experimental/GeometryProcessing/Source/MeshConversion/Public/MeshDescriptionToDynamicMesh.h
+++ b/Engine/Plugins/Experimental/GeometryProcessing/Source/MeshConversion/Public/MeshDescriptionToDynamicMesh.h
@@ -27,6 +27,9 @@ public:
 	/** Should we calculate conversion index maps */
 	bool bCalculateMaps = true;
 
+	/** Ignore all mesh attributes (e.g. UV/Normal layers, material groups) */
+	bool bDisableAttributes = false;
+
 	/** map from triangle ID to (polygon,triangle) pair */
 	TArray<FIndex2i> TriToPolyTriMap;
 
diff --git a/Engine/Plugins/Experimental/GeometryProcessing/Source/MeshSolverUtilities/Private/LaplacianOperators.cpp b/Engine/Plugins/Experimental/GeometryProcessing/Source/MeshSolverUtilities/Private/LaplacianOperators.cpp
--- a/Engine/Plugins/Experimental/GeometryProcessing/Source/MeshSolverUtilities/Private/LaplacianOperators.cpp
+++ b/Engine/Plugins/Experimental/GeometryProcessing/Source/MeshSolverUtilities/Private/LaplacianOperators.cpp
@@ -193,9 +193,9 @@ public:
 				const double EdgeACSqLength = EdgeAC.SquaredLength();
 				const double EdgeBCSqLength = EdgeBC.SquaredLength();
 
-				VoronoiArea[0] = EdgeABSqLength * Cotangent[1] + EdgeACSqLength * Cotangent[2];
-				VoronoiArea[1] = EdgeABSqLength * Cotangent[0] + EdgeBCSqLength * Cotangent[2];
-				VoronoiArea[2] = EdgeACSqLength * Cotangent[0] + EdgeBCSqLength * Cotangent[1];
+				VoronoiArea[0] = EdgeABSqLength * Cotangent[2] + EdgeACSqLength * Cotangent[1];
+				VoronoiArea[1] = EdgeABSqLength * Cotangent[2] + EdgeBCSqLength * Cotangent[0];
+				VoronoiArea[2] = EdgeACSqLength * Cotangent[1] + EdgeBCSqLength * Cotangent[0];
 
 				const double Inv8 = .125; // 1/8
 
diff --git a/Engine/Plugins/Experimental/MeshModelingToolset/Source/MeshModelingTools/Private/AddPatchTool.cpp b/Engine/Plugins/Experimental/MeshModelingToolset/Source/MeshModelingTools/Private/AddPatchTool.cpp
--- a/Engine/Plugins/Experimental/MeshModelingToolset/Source/MeshModelingTools/Private/AddPatchTool.cpp
+++ b/Engine/Plugins/Experimental/MeshModelingToolset/Source/MeshModelingTools/Private/AddPatchTool.cpp
@@ -47,24 +47,6 @@ UAddPatchToolProperties::UAddPatchToolProperties()
 	Shift = 0.0;
 }
 
-void UAddPatchToolProperties::SaveProperties(UInteractiveTool* SaveFromTool)
-{
-	UAddPatchToolProperties* PropertyCache = GetPropertyCache<UAddPatchToolProperties>();
-	PropertyCache->Width = this->Width;
-	PropertyCache->Rotation = this->Rotation;
-	PropertyCache->Subdivisions = this->Subdivisions;
-}
-
-void UAddPatchToolProperties::RestoreProperties(UInteractiveTool* RestoreToTool)
-{
-	UAddPatchToolProperties* PropertyCache = GetPropertyCache<UAddPatchToolProperties>();
-	this->Width = PropertyCache->Width;
-	this->Rotation = PropertyCache->Rotation;
-	this->Subdivisions = PropertyCache->Subdivisions;
-}
-
-
-
 void UAddPatchTool::SetWorld(UWorld* World)
 {
 	this->TargetWorld = World;
@@ -176,7 +158,7 @@ void UAddPatchTool::OnEndHover()
 
 
 
-void UAddPatchTool::Tick(float DeltaTime)
+void UAddPatchTool::OnTick(float DeltaTime)
 {
 	if (bPreviewValid == false)
 	{
