fix for crash when finding boundary loops on mesh with weird boundaries

#rnx

diff --git a/Engine/Plugins/Experimental/GeometryProcessing/Source/DynamicMesh/Private/MeshBoundaryLoops.cpp b/Engine/Plugins/Experimental/GeometryProcessing/Source/DynamicMesh/Private/MeshBoundaryLoops.cpp
--- a/Engine/Plugins/Experimental/GeometryProcessing/Source/DynamicMesh/Private/MeshBoundaryLoops.cpp
+++ b/Engine/Plugins/Experimental/GeometryProcessing/Source/DynamicMesh/Private/MeshBoundaryLoops.cpp
@@ -488,6 +488,7 @@ bool FMeshBoundaryLoops::ExtractSubloops(TArray<int>& loopV, TArray<int>& loopE,
 			FEdgeSpan& NewSpan = subs.Spans[subs.Spans.Emplace()];
 			NewSpan.InitializeFromVertices(Mesh, VerticesTemp, false);
 			NewSpan.SetBowtieVertices(bowties);
+			return false;
 		}
 
 		if (bv != bv_shortest)
