add support for trim as an additional mesh boolean mode

diff --git a/Engine/Plugins/Experimental/GeometryProcessing/Source/DynamicMesh/Private/Operations/MeshBoolean.cpp b/Engine/Plugins/Experimental/GeometryProcessing/Source/DynamicMesh/Private/Operations/MeshBoolean.cpp
--- a/Engine/Plugins/Experimental/GeometryProcessing/Source/DynamicMesh/Private/Operations/MeshBoolean.cpp
+++ b/Engine/Plugins/Experimental/GeometryProcessing/Source/DynamicMesh/Private/Operations/MeshBoolean.cpp
@@ -49,6 +49,7 @@ bool FMeshBoolean::Compute()
 	// cut the meshes
 	FMeshMeshCut Cut(CutMesh[0], CutMesh[1]);
 	Cut.bTrackInsertedVertices = bCollapseDegenerateEdgesOnCut; // to collect candidates to collapse
+	Cut.bMutuallyCut = Operation != EBooleanOp::Trim;
 	Cut.Cut(Intersections);
 
 	if (Cancelled())
@@ -56,11 +57,13 @@ bool FMeshBoolean::Compute()
 		return false;
 	}
 
+	int NumMeshesToProcess = Operation == EBooleanOp::Trim ? 1 : 2;
+
 	// collapse tiny edges along cut boundary
 	if (bCollapseDegenerateEdgesOnCut)
 	{
 		double DegenerateEdgeTolSq = DegenerateEdgeTol * DegenerateEdgeTol;
-		for (int MeshIdx = 0; MeshIdx < 2; MeshIdx++)
+		for (int MeshIdx = 0; MeshIdx < NumMeshesToProcess; MeshIdx++)
 		{
 			// convert vertex chains to edge IDs to simplify logic of finding remaining candidate edges after collapses
 			TArray<int> EIDs;
@@ -122,7 +125,7 @@ bool FMeshBoolean::Compute()
 	{ // (just for scope)
 		// first decide what triangles to delete for both meshes (*before* deleting anything so winding doesn't get messed up!)
 		TArray<bool> KeepTri[2];
-		for (int MeshIdx = 0; MeshIdx < 2; MeshIdx++)
+		for (int MeshIdx = 0; MeshIdx < NumMeshesToProcess; MeshIdx++)
 		{
 			TFastWindingTree<FDynamicMesh3> Winding(&Spatial[1 - MeshIdx]);
 			FDynamicMeshAABBTree3& OtherSpatial = Spatial[1 - MeshIdx];
@@ -131,7 +134,7 @@ bool FMeshBoolean::Compute()
 			KeepTri[MeshIdx].SetNumUninitialized(MaxTriID);
 			bool bCoplanarKeepSameDir = Operation != EBooleanOp::Difference;
 			bool bRemoveInside = 1; // whether to remove the inside triangles (e.g. for union) or the outside ones (e.g. for intersection)
-			if (Operation == EBooleanOp::Intersect || (Operation == EBooleanOp::Difference && MeshIdx == 1))
+			if (Operation == EBooleanOp::Trim || Operation == EBooleanOp::Intersect || (Operation == EBooleanOp::Difference && MeshIdx == 1))
 			{
 				bRemoveInside = 0;
 			}
@@ -202,7 +205,7 @@ bool FMeshBoolean::Compute()
 			}
 		}
 		// now go ahead and delete from both meshes
-		for (int MeshIdx = 0; MeshIdx < 2; MeshIdx++)
+		for (int MeshIdx = 0; MeshIdx < NumMeshesToProcess; MeshIdx++)
 		{
 			FDynamicMesh3& ProcessMesh = *CutMesh[MeshIdx];
 
@@ -221,96 +224,100 @@ bool FMeshBoolean::Compute()
 		return false;
 	}
 
-	// Hash boundary verts for faster search
-	TArray<TPointHashGrid3d<int>> PointHashes;
-	for (int MeshIdx = 0; MeshIdx < 2; MeshIdx++)
+	// correspond vertices across both meshes (in cases where both meshes were processed)
+	TMap<int, int> AllVIDMatches; // mapping of matched vertex IDs from cutmesh 0 to cutmesh 1
+	if (NumMeshesToProcess == 2)
 	{
-		PointHashes.Emplace(CutMesh[MeshIdx]->GetCachedBounds().MaxDim() / 64, -1);
-		for (int BoundaryVID : PossUnmatchedBdryVerts[MeshIdx])
+		// Hash boundary verts for faster search
+		TArray<TPointHashGrid3d<int>> PointHashes;
+		for (int MeshIdx = 0; MeshIdx < 2; MeshIdx++)
 		{
-			PointHashes[MeshIdx].InsertPointUnsafe(BoundaryVID, CutMesh[MeshIdx]->GetVertex(BoundaryVID));
+			PointHashes.Emplace(CutMesh[MeshIdx]->GetCachedBounds().MaxDim() / 64, -1);
+			for (int BoundaryVID : PossUnmatchedBdryVerts[MeshIdx])
+			{
+				PointHashes[MeshIdx].InsertPointUnsafe(BoundaryVID, CutMesh[MeshIdx]->GetVertex(BoundaryVID));
+			}
 		}
-	}
 
-	// ensure segments that are now on boundaries have 1:1 vertex correspondence across meshes
-	TMap<int, int> AllVIDMatches; // mapping of matched vertex IDs from cutmesh 0 to cutmesh 1
-	for (int MeshIdx = 0; MeshIdx < 2; MeshIdx++)
-	{
-		int OtherMeshIdx = 1 - MeshIdx;
-		FDynamicMesh3& OtherMesh = *CutMesh[OtherMeshIdx];
-
-		// mapping from OtherMesh VIDs to ProcessMesh VIDs
-		// used to ensure we only keep the best match, in cases where multiple boundary vertices map to a given vertex on the other mesh boundary
-		TMap<int, int> FoundMatches;
-
-		for (int BoundaryVID : PossUnmatchedBdryVerts[MeshIdx])
+		// ensure segments that are now on boundaries have 1:1 vertex correspondence across meshes
+		for (int MeshIdx = 0; MeshIdx < 2; MeshIdx++)
 		{
-			FVector3d Pos = CutMesh[MeshIdx]->GetVertex(BoundaryVID);
-			TPair<int, double> VIDDist = PointHashes[OtherMeshIdx].FindNearestInRadius(Pos, SnapTolerance, [&Pos, &OtherMesh](int VID)
-				{
-					return Pos.DistanceSquared(OtherMesh.GetVertex(VID));
-				});
-			int NearestVID = VIDDist.Key; // ID of nearest vertex on other mesh
-			double DSq = VIDDist.Value;   // square distance to that vertex
+			int OtherMeshIdx = 1 - MeshIdx;
+			FDynamicMesh3& OtherMesh = *CutMesh[OtherMeshIdx];
 
-			if (NearestVID != FDynamicMesh3::InvalidID)
+			// mapping from OtherMesh VIDs to ProcessMesh VIDs
+			// used to ensure we only keep the best match, in cases where multiple boundary vertices map to a given vertex on the other mesh boundary
+			TMap<int, int> FoundMatches;
+
+			for (int BoundaryVID : PossUnmatchedBdryVerts[MeshIdx])
 			{
+				FVector3d Pos = CutMesh[MeshIdx]->GetVertex(BoundaryVID);
+				TPair<int, double> VIDDist = PointHashes[OtherMeshIdx].FindNearestInRadius(Pos, SnapTolerance, [&Pos, &OtherMesh](int VID)
+					{
+						return Pos.DistanceSquared(OtherMesh.GetVertex(VID));
+					});
+				int NearestVID = VIDDist.Key; // ID of nearest vertex on other mesh
+				double DSq = VIDDist.Value;   // square distance to that vertex
 
-				int* Match = FoundMatches.Find(NearestVID);
-				if (Match)
+				if (NearestVID != FDynamicMesh3::InvalidID)
 				{
-					double OldDSq = CutMesh[MeshIdx]->GetVertex(*Match).DistanceSquared(OtherMesh.GetVertex(NearestVID));
-					if (DSq < OldDSq) // new vertex is a better match than the old one
+
+					int* Match = FoundMatches.Find(NearestVID);
+					if (Match)
 					{
-						int OldVID = *Match; // copy old VID out of match before updating the TMap
-						FoundMatches.Add(NearestVID, BoundaryVID); // new VID is recorded as best match
-
-						// old VID is swapped in as the one to consider as unmatched
-						// it will now be matched below
-						BoundaryVID = OldVID;
-						Pos = CutMesh[MeshIdx]->GetVertex(BoundaryVID);
-						DSq = OldDSq;
+						double OldDSq = CutMesh[MeshIdx]->GetVertex(*Match).DistanceSquared(OtherMesh.GetVertex(NearestVID));
+						if (DSq < OldDSq) // new vertex is a better match than the old one
+						{
+							int OldVID = *Match; // copy old VID out of match before updating the TMap
+							FoundMatches.Add(NearestVID, BoundaryVID); // new VID is recorded as best match
+
+							// old VID is swapped in as the one to consider as unmatched
+							// it will now be matched below
+							BoundaryVID = OldVID;
+							Pos = CutMesh[MeshIdx]->GetVertex(BoundaryVID);
+							DSq = OldDSq;
+						}
+						NearestVID = FDynamicMesh3::InvalidID; // one of these vertices will be unmatched
+					}
+					else
+					{
+						FoundMatches.Add(NearestVID, BoundaryVID);
 					}
-					NearestVID = FDynamicMesh3::InvalidID; // one of these vertices will be unmatched
-				}
-				else
-				{
-					FoundMatches.Add(NearestVID, BoundaryVID);
 				}
-			}
 
-			// if we didn't find a valid match, try to split the nearest edge to create a match
-			if (NearestVID == FDynamicMesh3::InvalidID)
-			{
-				// vertex had no match -- try to split edge to match it
-				int OtherEID = FindNearestEdge(OtherMesh, CutBoundaryEdges[OtherMeshIdx], Pos);
-				if (OtherEID != FDynamicMesh3::InvalidID)
+				// if we didn't find a valid match, try to split the nearest edge to create a match
+				if (NearestVID == FDynamicMesh3::InvalidID)
 				{
-					FVector3d EdgePts[2];
-					OtherMesh.GetEdgeV(OtherEID, EdgePts[0], EdgePts[1]);
-					FSegment3d Seg(EdgePts[0], EdgePts[1]);
-					double Along = Seg.ProjectUnitRange(Pos);
-					FDynamicMesh3::FEdgeSplitInfo SplitInfo;
-					if (ensure(EMeshResult::Ok == OtherMesh.SplitEdge(OtherEID, SplitInfo, Along)))
+					// vertex had no match -- try to split edge to match it
+					int OtherEID = FindNearestEdge(OtherMesh, CutBoundaryEdges[OtherMeshIdx], Pos);
+					if (OtherEID != FDynamicMesh3::InvalidID)
 					{
-						FoundMatches.Add(SplitInfo.NewVertex, BoundaryVID);
-						OtherMesh.SetVertex(SplitInfo.NewVertex, Pos);
-						CutBoundaryEdges[OtherMeshIdx].Add(SplitInfo.NewEdges.A);
-						// Note: Do not update PossUnmatchedBdryVerts with the new vertex, because it is already matched by construction
-						// Likewise do not update the pointhash -- we don't want it to find vertices that were already perfectly matched
+						FVector3d EdgePts[2];
+						OtherMesh.GetEdgeV(OtherEID, EdgePts[0], EdgePts[1]);
+						FSegment3d Seg(EdgePts[0], EdgePts[1]);
+						double Along = Seg.ProjectUnitRange(Pos);
+						FDynamicMesh3::FEdgeSplitInfo SplitInfo;
+						if (ensure(EMeshResult::Ok == OtherMesh.SplitEdge(OtherEID, SplitInfo, Along)))
+						{
+							FoundMatches.Add(SplitInfo.NewVertex, BoundaryVID);
+							OtherMesh.SetVertex(SplitInfo.NewVertex, Pos);
+							CutBoundaryEdges[OtherMeshIdx].Add(SplitInfo.NewEdges.A);
+							// Note: Do not update PossUnmatchedBdryVerts with the new vertex, because it is already matched by construction
+							// Likewise do not update the pointhash -- we don't want it to find vertices that were already perfectly matched
+						}
 					}
 				}
 			}
-		}
 
-		// actually snap the positions together for final matches
-		for (TPair<int, int>& Match : FoundMatches)
-		{
-			CutMesh[MeshIdx]->SetVertex(Match.Value, OtherMesh.GetVertex(Match.Key));
+			// actually snap the positions together for final matches
+			for (TPair<int, int>& Match : FoundMatches)
+			{
+				CutMesh[MeshIdx]->SetVertex(Match.Value, OtherMesh.GetVertex(Match.Key));
 
-			// Copy match to AllVIDMatches; note this is always mapping from CutMesh 0 to 1
-			int VIDs[2]{ Match.Key, Match.Value }; // just so we can access by index
-			AllVIDMatches.Add(VIDs[1 - MeshIdx], VIDs[MeshIdx]);
+				// Copy match to AllVIDMatches; note this is always mapping from CutMesh 0 to 1
+				int VIDs[2]{ Match.Key, Match.Value }; // just so we can access by index
+				AllVIDMatches.Add(VIDs[1 - MeshIdx], VIDs[MeshIdx]);
+			}
 		}
 	}
 
@@ -331,13 +338,23 @@ bool FMeshBoolean::Compute()
 		return false;
 	}
 
-	FDynamicMeshEditor Editor(Result);
-	FMeshIndexMappings IndexMaps;
-	Editor.AppendMesh(CutMesh[1], IndexMaps);
+	bool bSuccess = true;
 
-	bool bWeldSuccess = MergeEdges(IndexMaps, CutMesh, CutBoundaryEdges, AllVIDMatches);
+	if (NumMeshesToProcess > 1)
+	{
+		FDynamicMeshEditor Editor(Result);
+		FMeshIndexMappings IndexMaps;
+		Editor.AppendMesh(CutMesh[1], IndexMaps);
+
+		bool bWeldSuccess = MergeEdges(IndexMaps, CutMesh, CutBoundaryEdges, AllVIDMatches);
+		bSuccess = bSuccess && bWeldSuccess;
+	}
+	else
+	{
+		CreatedBoundaryEdges = CutBoundaryEdges[0];
+	}
 
-	return bWeldSuccess;
+	return bSuccess;
 }
 
 
diff --git a/Engine/Plugins/Experimental/GeometryProcessing/Source/DynamicMesh/Public/Operations/MeshBoolean.h b/Engine/Plugins/Experimental/GeometryProcessing/Source/DynamicMesh/Public/Operations/MeshBoolean.h
--- a/Engine/Plugins/Experimental/GeometryProcessing/Source/DynamicMesh/Public/Operations/MeshBoolean.h
+++ b/Engine/Plugins/Experimental/GeometryProcessing/Source/DynamicMesh/Public/Operations/MeshBoolean.h
@@ -37,7 +37,8 @@ public:
 	{
 		Union,
 		Difference,
-		Intersect
+		Intersect,
+		Trim
 	};
 	EBooleanOp Operation;
 	
diff --git a/Engine/Plugins/Experimental/MeshModelingToolset/Source/ModelingOperators/Private/CompositionOps/BooleanMeshesOp.cpp b/Engine/Plugins/Experimental/MeshModelingToolset/Source/ModelingOperators/Private/CompositionOps/BooleanMeshesOp.cpp
--- a/Engine/Plugins/Experimental/MeshModelingToolset/Source/ModelingOperators/Private/CompositionOps/BooleanMeshesOp.cpp
+++ b/Engine/Plugins/Experimental/MeshModelingToolset/Source/ModelingOperators/Private/CompositionOps/BooleanMeshesOp.cpp
@@ -20,7 +20,7 @@ void FBooleanMeshesOp::CalculateResult(FProgressCancel* Progress)
 	check(Meshes.Num() == 2 && Transforms.Num() == 2);
 	
 	int FirstIdx = 0;
-	if (Operation == ECSGOperation::DifferenceBA)
+	if (Operation == ECSGOperation::DifferenceBA || Operation == ECSGOperation::TrimB)
 	{
 		FirstIdx = 1;
 	}
@@ -34,6 +34,10 @@ void FBooleanMeshesOp::CalculateResult(FProgressCancel* Progress)
 	case ECSGOperation::DifferenceBA:
 		Op = FMeshBoolean::EBooleanOp::Difference;
 		break;
+	case ECSGOperation::TrimA:
+	case ECSGOperation::TrimB:
+		Op = FMeshBoolean::EBooleanOp::Trim;
+		break;
 	case ECSGOperation::Union:
 		Op = FMeshBoolean::EBooleanOp::Union;
 		break;
@@ -60,7 +64,7 @@ void FBooleanMeshesOp::CalculateResult(FProgressCancel* Progress)
 		return;
 	}
 
-	if (!bSuccess && bAttemptFixHoles)
+	if (MeshBoolean.CreatedBoundaryEdges.Num() > 0 && bAttemptFixHoles)
 	{
 		FMeshBoundaryLoops OpenBoundary(MeshBoolean.Result, false);
 		TSet<int> ConsiderEdges(MeshBoolean.CreatedBoundaryEdges);
diff --git a/Engine/Plugins/Experimental/MeshModelingToolset/Source/ModelingOperators/Public/CompositionOps/BooleanMeshesOp.h b/Engine/Plugins/Experimental/MeshModelingToolset/Source/ModelingOperators/Public/CompositionOps/BooleanMeshesOp.h
--- a/Engine/Plugins/Experimental/MeshModelingToolset/Source/ModelingOperators/Public/CompositionOps/BooleanMeshesOp.h
+++ b/Engine/Plugins/Experimental/MeshModelingToolset/Source/ModelingOperators/Public/CompositionOps/BooleanMeshesOp.h
@@ -25,6 +25,13 @@ enum class ECSGOperation : uint8
 
 	/** union of two objects */
 	Union = 3 UMETA(DisplayName = "Union"),
+
+	/** Trim the first object using the second */
+	TrimA = 4 UMETA(DisplayName = "Trim A"),
+
+	/** Trim the second object using the first */
+	TrimB = 5 UMETA(DisplayName = "Trim B"),
+
 };
 
 
