GeometryProcessing: fix SparseDynamicOctree bug where if an object insertion spawned a new Root cell, object was always added to that cell regardless of whether it actually fit (which then causes spatial queries to incorrectly 'miss' the object outside the too-small Root box). Also changed some slightly overzealous safety check()'s to checkSlow()
#rnx

diff --git a/Engine/Plugins/Experimental/GeometryProcessing/Source/GeometricObjects/Public/Spatial/SparseDynamicOctree3.h b/Engine/Plugins/Experimental/GeometryProcessing/Source/GeometricObjects/Public/Spatial/SparseDynamicOctree3.h
--- a/Engine/Plugins/Experimental/GeometryProcessing/Source/GeometricObjects/Public/Spatial/SparseDynamicOctree3.h
+++ b/Engine/Plugins/Experimental/GeometryProcessing/Source/GeometricObjects/Public/Spatial/SparseDynamicOctree3.h
@@ -404,7 +404,7 @@ bool FSparseDynamicOctree3::ContainsObject(int32 ObjectID) const
 
 void FSparseDynamicOctree3::InsertObject(int32 ObjectID, const FAxisAlignedBox3d& Bounds)
 {
-	check(ContainsObject(ObjectID) == false);
+	checkSlow(ContainsObject(ObjectID) == false);
 
 	FSparseOctreeCell CurrentCell = FindCurrentContainingCell(Bounds);
 
@@ -414,6 +414,7 @@ void FSparseDynamicOctree3::InsertObject(int32 ObjectID, const FAxisAlignedBox3d
 		Insert_Spill(ObjectID, Bounds);
 		return;
 	}
+	checkSlow(CanFit(CurrentCell, Bounds));
 
 	// if we found a containing root cell but it doesn't exist, create it and insert
 	if (CurrentCell.Level == 0 && CurrentCell.IsExistingCell() == false)
@@ -448,7 +449,7 @@ void FSparseDynamicOctree3::Insert_NewChildCell(int32 ObjectID, const FAxisAlign
 	int ParentCellID, FSparseOctreeCell NewChildCell, int ChildIdx)
 {
 	FSparseOctreeCell& OrigParentCell = Cells[ParentCellID];
-	check(OrigParentCell.HasChild(ChildIdx) == false);
+	checkSlow(OrigParentCell.HasChild(ChildIdx) == false);
 
 	NewChildCell.CellID = CellRefCounts.Allocate();
 	Cells.InsertAt(NewChildCell, NewChildCell.CellID);
@@ -461,7 +462,7 @@ void FSparseDynamicOctree3::Insert_NewChildCell(int32 ObjectID, const FAxisAlign
 
 	OrigParentCell.SetChild(ChildIdx, NewChildCell);
 
-	check(CanFit(NewChildCell, Bounds));
+	checkSlow(CanFit(NewChildCell, Bounds));
 	// this check is unstable if the center point is within machine-epsilon of the cell border
 	//check(PointToIndex(NewChildCell.Level, Bounds.Center()) == NewChildCell.Index);
 }
@@ -469,14 +470,14 @@ void FSparseDynamicOctree3::Insert_NewChildCell(int32 ObjectID, const FAxisAlign
 
 void FSparseDynamicOctree3::Insert_ToCell(int32 ObjectID, const FAxisAlignedBox3d& Bounds, const FSparseOctreeCell& ExistingCell)
 {
-	check(CellRefCounts.IsValid(ExistingCell.CellID));
+	checkSlow(CellRefCounts.IsValid(ExistingCell.CellID));
 
 	ObjectIDToCellMap.InsertAt(ExistingCell.CellID, ObjectID);
 	ValidObjectIDs.Set(ObjectID, true);
 
 	CellObjectLists.Insert(ExistingCell.CellID, ObjectID);
 
-	check(CanFit(ExistingCell, Bounds));
+	checkSlow(CanFit(ExistingCell, Bounds));
 	// this check is unstable if the center point is within machine-epsilon of the cell border
 	//check(PointToIndex(ExistingCell.Level, Bounds.Center()) == ExistingCell.Index);
 }
@@ -484,7 +485,7 @@ void FSparseDynamicOctree3::Insert_ToCell(int32 ObjectID, const FAxisAlignedBox3
 
 void FSparseDynamicOctree3::Insert_NewRoot(int32 ObjectID, const FAxisAlignedBox3d& Bounds, FSparseOctreeCell NewRootCell)
 {
-	check(RootCells.Has(NewRootCell.Index) == false);
+	checkSlow(RootCells.Has(NewRootCell.Index) == false);
 
 	NewRootCell.CellID = CellRefCounts.Allocate();
 	Cells.InsertAt(NewRootCell, NewRootCell.CellID);
@@ -497,6 +498,8 @@ void FSparseDynamicOctree3::Insert_NewRoot(int32 ObjectID, const FAxisAlignedBox
 
 	CellObjectLists.AllocateAt(NewRootCell.CellID);
 	CellObjectLists.Insert(NewRootCell.CellID, ObjectID);
+
+	checkSlow(CanFit(NewRootCell, Bounds));
 }
 
 void FSparseDynamicOctree3::Insert_Spill(int32 ObjectID, const FAxisAlignedBox3d& Bounds)
@@ -520,7 +523,7 @@ bool FSparseDynamicOctree3::RemoveObject(int32 ObjectID)
 	if (CellID == SpillCellID)
 	{
 		int32 RemovedCount = SpillObjectSet.Remove(ObjectID);
-		check(RemovedCount > 0);
+		checkSlow(RemovedCount > 0);
 		ValidObjectIDs.Set(ObjectID, false);
 		return (RemovedCount > 0);
 	}
@@ -533,7 +536,7 @@ bool FSparseDynamicOctree3::RemoveObject(int32 ObjectID)
 	ValidObjectIDs.Set(ObjectID, false);
 
 	bool bInList = CellObjectLists.Remove(CellID, ObjectID);
-	check(bInList);
+	checkSlow(bInList);
 	return true;
 }
 
@@ -773,9 +776,19 @@ FSparseOctreeCell FSparseDynamicOctree3::FindCurrentContainingCell(const FAxisAl
 	const uint32* RootCellID = RootCells.Get(RootIndex);
 	if (RootCellID == nullptr)
 	{
-		return FSparseOctreeCell(0, RootIndex);
+		FSparseOctreeCell RootCell(0, RootIndex);
+		// have to make sure we can fit in this root cell
+		if (CanFit(RootCell, Bounds))
+		{
+			return RootCell;
+		}
+		else
+		{
+			return FSparseOctreeCell(FSparseOctreeCell::InvalidLevel, FVector3i::Zero());
+		}
+
 	}
-	check(CellRefCounts.IsValid(*RootCellID));
+	checkSlow(CellRefCounts.IsValid(*RootCellID));
 
 	// check if box is contained in root cell, if not we have to spill
 	// (should we do this before checking for existence? we can...)
@@ -792,7 +805,7 @@ FSparseOctreeCell FSparseDynamicOctree3::FindCurrentContainingCell(const FAxisAl
 		if (CurrentCell->HasChild(ChildIdx))
 		{
 			int32 ChildCellID = CurrentCell->GetChildCellID(ChildIdx);
-			check(CellRefCounts.IsValid(ChildCellID));
+			checkSlow(CellRefCounts.IsValid(ChildCellID));
 			const FSparseOctreeCell* ChildCell = &Cells[ChildCellID];
 			if (CanFit(*ChildCell, Bounds))
 			{
@@ -874,12 +887,12 @@ void FSparseDynamicOctree3::CheckValidity(
 			{
 				// this is a spill node...
 				SpillObjectCount++;
-				check(SpillObjectSet.Contains(ObjectID));
+				CheckOrFailF(SpillObjectSet.Contains(ObjectID));
 			}
 			else if (CellID == InvalidCellID)
 			{
 				MissingObjectCount++;
-				check(SpillObjectSet.Contains(ObjectID) == false);
+				CheckOrFailF(SpillObjectSet.Contains(ObjectID) == false);
 			}
 			else
 			{
