GeometryProcessing: merge MeshSolverUtilities module into DynamicMesh module. Update Build.cs references, test code.

Most of the code moved to /Private/Solvers/Internal, except:
- MeshElementLinearization.h split into /GeometricObjects/Util/ElementLinearization.h (not mesh-specific) and /DynamicMesh/Solvers/MeshLinearization.h for vtx/tri subclasses
- ELaplacianWeightScheme and a few util functions moved to /DynamicMesh/Solvers/MeshLaplacian.h
- IConstrainedMeshOperator renamed to IConstrainedMeshSolver and moved to /DynamicMesh/Solvers/
- MeshSolverUtilities.h API functions moved to /DynamicMesh/Solvers/MeshSmoothing.h, ConstrainedMeshDeformer.h, ConstrainedMeshSmoother.h
- Laplacian matrix assembly refactored out of LaplacianOperators.cpp and moved to /DynamicMesh/Solvers/LaplacianMatrixAssembly. LaplacianOperator functions are now wrappers around those functions.
- most of FSOAPositions refactored into TVector3Arrays<T>, which it now subclasses. The arrays are now TArray<T> and Eigen::Map is used to pass them to/from Eigen code. Moved to FSparseMatrixD.h

#rnx

diff --git a/Engine/Plugins/Experimental/GeometryProcessing/GeometryProcessing.uplugin b/Engine/Plugins/Experimental/GeometryProcessing/GeometryProcessing.uplugin
--- a/Engine/Plugins/Experimental/GeometryProcessing/GeometryProcessing.uplugin
+++ b/Engine/Plugins/Experimental/GeometryProcessing/GeometryProcessing.uplugin
@@ -33,11 +33,6 @@
 			"Name": "MeshConversion",
 			"Type": "Editor",
 			"LoadingPhase": "Default"
-		},
-		{
-			"Name": "MeshSolverUtilities",
-			"Type": "Editor",
-			"LoadingPhase": "Default"
 		}
 	]
 }
\ No newline at end of file
diff --git a/Engine/Plugins/Experimental/GeometryProcessing/Source/DynamicMesh/DynamicMesh.Build.cs b/Engine/Plugins/Experimental/GeometryProcessing/Source/DynamicMesh/DynamicMesh.Build.cs
--- a/Engine/Plugins/Experimental/GeometryProcessing/Source/DynamicMesh/DynamicMesh.Build.cs
+++ b/Engine/Plugins/Experimental/GeometryProcessing/Source/DynamicMesh/DynamicMesh.Build.cs
@@ -13,6 +13,13 @@ public class DynamicMesh : ModuleRules
 				"Core",
 				"GeometricObjects"
 			}
-			);
+		);
+
+		PrivateDependencyModuleNames.AddRange(
+			new string[] {
+				"Eigen"
+			}
+		);
+
 	}
 }
diff --git a/Engine/Plugins/Experimental/GeometryProcessing/Source/DynamicMesh/Private/Solvers/ConstrainedMeshDeformer.cpp b/Engine/Plugins/Experimental/GeometryProcessing/Source/DynamicMesh/Private/Solvers/ConstrainedMeshDeformer.cpp
new file mode 100644
--- /dev/null
+++ b/Engine/Plugins/Experimental/GeometryProcessing/Source/DynamicMesh/Private/Solvers/ConstrainedMeshDeformer.cpp
@@ -0,0 +1,12 @@
+// Copyright Epic Games, Inc. All Rights Reserved.
+
+#include "Solvers/ConstrainedMeshDeformer.h"
+#include "Solvers/Internal/LaplacianMeshSmoother.h"
+
+
+TUniquePtr<IConstrainedMeshSolver> UE::MeshDeformation::ConstructConstrainedMeshDeformer(const ELaplacianWeightScheme WeightScheme, const FDynamicMesh3& DynamicMesh)
+{
+	TUniquePtr<IConstrainedMeshSolver> Deformer(new FConstrainedMeshDeformer(DynamicMesh, WeightScheme));
+	return Deformer;
+}
+
diff --git a/Engine/Plugins/Experimental/GeometryProcessing/Source/DynamicMesh/Private/Solvers/ConstrainedMeshSmoother.cpp b/Engine/Plugins/Experimental/GeometryProcessing/Source/DynamicMesh/Private/Solvers/ConstrainedMeshSmoother.cpp
new file mode 100644
--- /dev/null
+++ b/Engine/Plugins/Experimental/GeometryProcessing/Source/DynamicMesh/Private/Solvers/ConstrainedMeshSmoother.cpp
@@ -0,0 +1,12 @@
+// Copyright Epic Games, Inc. All Rights Reserved.
+
+#include "Solvers/ConstrainedMeshSmoother.h"
+#include "Solvers/Internal/LaplacianMeshSmoother.h"
+
+
+TUniquePtr<IConstrainedMeshSolver> UE::MeshDeformation::ConstructConstrainedMeshSmoother(const ELaplacianWeightScheme WeightScheme, const FDynamicMesh3& DynamicMesh)
+{
+	TUniquePtr<IConstrainedMeshSolver> Deformer(new FBiHarmonicMeshSmoother(DynamicMesh, WeightScheme));
+	return Deformer;
+}
+
diff --git a/Engine/Plugins/Experimental/GeometryProcessing/Source/DynamicMesh/Private/Solvers/Internal/ConstrainedPoissonSolver.h b/Engine/Plugins/Experimental/GeometryProcessing/Source/DynamicMesh/Private/Solvers/Internal/ConstrainedPoissonSolver.h
new file mode 100644
--- /dev/null
+++ b/Engine/Plugins/Experimental/GeometryProcessing/Source/DynamicMesh/Private/Solvers/Internal/ConstrainedPoissonSolver.h
@@ -0,0 +1,350 @@
+// Copyright Epic Games, Inc. All Rights Reserved.
+
+#pragma once
+
+#include "CoreMinimal.h"
+#include "DynamicMesh3.h"
+#include "FSparseMatrixD.h"
+#include "MatrixSolver.h"
+
+
+// only needed if you enable logging by changing the typedef for FConstrainedSolverTimeLogger
+//#include "ProfilingDebugging/ScopedTimers.h"
+
+
+
+/**
+* Least Squares constrained poisson solver.
+*
+* Where L is a Laplacian type matrix
+*       {lambda_i, cvec_i} are k weights (lambda) and positional constraints.
+*
+* Solves in the lease squares sense the equation
+*         L p_vec = d_vec
+*         lambda_i * (p_vec)_i  = lambda_i (c_vec)_i  for some constraints.
+*
+* this amounts to
+*         ( Transpose(L) * L   + (0  0      )  ) p_vec = source_vec + ( 0              )
+*    	  (                      (0 lambda^2)  )                      ( lambda^2 c_vec )
+*
+* where source_vec = Transpose(L)*L p_vec 
+*
+* C.F. http://sites.fas.harvard.edu/~cs277/papers/deformation_survey.pdf  section V, subsection C
+*
+* @todo consider reworking to hold TUniquePtr<> to the various objects.
+*/
+
+class FConstrainedSolver
+{
+public:
+
+	typedef typename FSparseMatrixD::Scalar     ScalarType;
+	typedef typename FSOAPositions::VectorType  VectorType;
+
+	struct FDummyTimeLogger
+	{
+		FDummyTimeLogger(FString InMsg = TEXT("Scoped action"), FOutputDevice* InDevice = GLog) {}
+	};
+
+	// Change this to enable time logging for testing.
+	//typedef FScopedDurationTimeLogger FConstrainedSolverTimeLogger; 
+	typedef FDummyTimeLogger FConstrainedSolverTimeLogger;
+
+	struct FConstraintPosition
+	{
+		FConstraintPosition(const FVector3d& P, bool b) : Position(P), bPostFix(b) {}
+
+		FVector3d Position;
+		bool     bPostFix;
+	};
+
+
+	FConstrainedSolver(TUniquePtr<FSparseMatrixD>& SymmatrixMatrixOperator, const EMatrixSolverType MatrixSolverType)
+		: ConstraintPositions(SymmatrixMatrixOperator->cols())
+	{
+		SymmetricMatrixPtr.Reset(SymmatrixMatrixOperator.Release());
+
+		bMatrixSolverDirty = true;
+		MatrixSolver = ContructMatrixSolver(MatrixSolverType);
+	}
+
+
+	// Updates the diagonal weights matrix.
+	void SetConstraintWeights(const TMap<int32, double>& WeightMap)
+	{
+		typedef FSparseMatrixD::Scalar    ScalarT;
+		typedef Eigen::Triplet<ScalarT>   MatrixTripletT;
+
+		ClearWeights();
+
+		std::vector<MatrixTripletT> MatrixTripelList;
+		MatrixTripelList.reserve(WeightMap.Num());
+
+		for (const auto& WeightPair : WeightMap)
+		{
+			const int32 i = WeightPair.Key; // row id
+			double Weight = WeightPair.Value;
+
+			checkSlow(i < SymmetricMatrixPtr->cols());
+
+			// the soft constrained system uses the square of the weight.
+			Weight *= Weight;
+
+			//const int32 i = ToIndex[VertId];
+			MatrixTripelList.push_back(MatrixTripletT(i, i, Weight));
+
+		}
+
+		// Construct matrix with weights on the diagonal for the constrained verts ( and zero everywhere else)
+		WeightsSqrdMatrix.setFromTriplets(MatrixTripelList.begin(), MatrixTripelList.end());
+		WeightsSqrdMatrix.makeCompressed();
+
+		// The solver matrix will have to be updated and re-factored
+		UpdateSolverWithContraints();
+	}
+
+	// Updates the positional source term
+	void SetContraintPositions(const TMap<int32, FConstraintPosition>& PositionMap)
+	{
+		ClearConstraintPositions();
+
+		for (const auto& PositionPair : PositionMap)
+		{
+			const int32 i = PositionPair.Key; // row id
+			const FVector3d& Pos = PositionPair.Value.Position;
+
+			checkSlow(i < SymmetricMatrixPtr->cols());
+
+
+			// constrained source vector
+			ConstraintPositions.SetX(i, Pos.X);
+			ConstraintPositions.SetY(i, Pos.Y);
+			ConstraintPositions.SetZ(i, Pos.Z);
+
+		}
+	}
+
+	ScalarType GetWeightSqrd(int32 VtxIndex) const
+	{
+		return WeightsSqrdMatrix.coeff(VtxIndex, VtxIndex);
+	}
+
+	/**
+	* Access to the underlying solver
+	*/
+	IMatrixSolverBase* GetMatrixSolverBase() { return MatrixSolver.Get(); }
+	
+	const IMatrixSolverBase* GetMatrixSolverBase() const { return MatrixSolver.Get(); }
+
+	/**
+	* Access to the underlying solver in iterative form - will return null if the solver is not iterative.
+	*/
+	IIterativeMatrixSolverBase* GetMatrixSolverIterativeBase()
+	{
+		IIterativeMatrixSolverBase* IterativeBasePtr = NULL;
+		if (MatrixSolver->bIsIterative())
+		{
+			IterativeBasePtr = (IIterativeMatrixSolverBase*)MatrixSolver.Get();
+		}
+		return IterativeBasePtr;
+	}
+	const IIterativeMatrixSolverBase* GetMatrixSolverIterativeBase() const 
+	{
+		const IIterativeMatrixSolverBase* IterativeBasePtr = NULL;
+		if (MatrixSolver->bIsIterative())
+		{
+			IterativeBasePtr = (const IIterativeMatrixSolverBase*)MatrixSolver.Get();
+		}
+		return IterativeBasePtr;
+	}
+
+	/**
+	* @param SourceVector   - Required that each component vector has size equal to the the number of columns in the laplacian
+	* @param SolutionVector - the resulting solution to the least squares problem
+	*/
+	bool Solve(const FSOAPositions& SourceVector, FSOAPositions& SolutionVector) const
+	{
+		checkSlow(bMatrixSolverDirty == false);
+		FConstrainedSolverTimeLogger Timmer(TEXT("Post-setup solve time"));
+
+		FSparseMatrixD& SymmetricMatrix = *SymmetricMatrixPtr;
+		// Set up the source vector
+		FSOAPositions RHSVector(SymmetricMatrix.cols());
+		for (int32 Dir = 0; Dir < 3; ++Dir) 
+		{
+			RHSVector.Array(Dir) = SourceVector.Array(Dir) + WeightsSqrdMatrix * ConstraintPositions.Array(Dir);
+		}
+		
+		MatrixSolver->Solve(RHSVector, SolutionVector);
+
+		bool bSuccess = MatrixSolver->bSucceeded();
+
+		return bSuccess;
+	}
+
+
+	/**
+	* Special case when the source vector is identically zero.
+	* NB: this is used for region smoothing.
+	*/
+	bool Solve(FSOAPositions& SolutionVector) const
+	{
+		checkSlow(bMatrixSolverDirty == false);
+		
+		FConstrainedSolverTimeLogger Timmer(TEXT("Post-setup solve time"));
+
+		FSparseMatrixD& SymmetricMatrix = *SymmetricMatrixPtr;
+		FSOAPositions RHSVector(SymmetricMatrix.cols());
+		for (int32 Dir = 0; Dir < 3; ++Dir)
+		{
+			RHSVector.Array(Dir) = WeightsSqrdMatrix * ConstraintPositions.Array(Dir);
+		}
+	
+		MatrixSolver->Solve(RHSVector, SolutionVector);
+
+		bool bSuccess = MatrixSolver->bSucceeded();
+		return bSuccess;
+	}
+
+	/*
+	* For use with iterative solvers.
+	*
+	* Reverts to Solve(SolutionVector); 
+	* if the matrix solver type is direct.
+	*/
+	bool SolveWithGuess(const FSOAPositions& Guess, FSOAPositions& SolutionVector) const
+	{
+		checkSlow(bMatrixSolverDirty == false);
+
+		if (MatrixSolver->bIsIterative())
+		{
+			const IIterativeMatrixSolverBase* IterativeSolver = (IIterativeMatrixSolverBase*)MatrixSolver.Get();
+
+			FConstrainedSolverTimeLogger Timmer(TEXT("Post-setup solve time"));
+
+			FSparseMatrixD& SymmetricMatrix = *SymmetricMatrixPtr;
+			FSOAPositions RHSVector(SymmetricMatrix.cols());
+			for (int32 Dir = 0; Dir < 3; ++Dir)
+			{
+				RHSVector.Array(Dir) = WeightsSqrdMatrix * ConstraintPositions.Array(Dir);
+			}
+			
+			IterativeSolver->SolveWithGuess(Guess, RHSVector, SolutionVector);
+
+			bool bSuccess = IterativeSolver->bSucceeded();
+			return bSuccess;
+		}
+		else
+		{
+			return Solve(SolutionVector);
+		}
+	}
+
+	/**
+	* For use with iterative solvers.
+	*
+	* Reverts to Solve(SourceVector, SolutionVector);
+	* if the matrix solver type is direct
+	*/
+	bool SolveWithGuess(const FSOAPositions& Guess, const FSOAPositions& SourceVector, FSOAPositions& SolutionVector) const
+	{
+		checkSlow(bMatrixSolverDirty == false);
+
+		if (MatrixSolver->bIsIterative())
+		{
+			const IIterativeMatrixSolverBase* IterativeSolver = (IIterativeMatrixSolverBase*)MatrixSolver.Get();
+
+			FConstrainedSolverTimeLogger Timmer(TEXT("Post-setup solve time"));
+
+			FSparseMatrixD& SymmetricMatrix = *SymmetricMatrixPtr;
+			// Set up the source vector
+			FSOAPositions RHSVector(SymmetricMatrix.cols());
+			for (int32 Dir = 0; Dir < 3; ++Dir)
+			{
+				RHSVector.Array(Dir) = SourceVector.Array(Dir) + WeightsSqrdMatrix * ConstraintPositions.Array(Dir);
+			}
+		
+			IterativeSolver->SolveWithGuess(Guess, RHSVector, SolutionVector);
+
+			bool bSuccess = IterativeSolver->bSucceeded();
+
+			return bSuccess;
+		}
+		else
+		{
+			return Solve(SourceVector, SolutionVector);
+		}
+	}
+
+
+	const FSparseMatrixD& Biharmonic() const { return *SymmetricMatrixPtr; }
+
+
+protected:
+
+	// Zero the diagonal matrix that holds constraints
+	void ClearConstraints()
+	{
+		ClearConstraintPositions();
+		ClearWeights();
+	}
+
+	void ClearConstraintPositions()
+	{
+		const FSparseMatrixD& SymmetricMatrix = *SymmetricMatrixPtr;
+		const int32 NumColumns = SymmetricMatrix.cols();
+
+		ConstraintPositions.SetZero(NumColumns);
+
+	}
+	void ClearWeights()
+	{
+		WeightsSqrdMatrix.setZero();
+
+		WeightsSqrdMatrix.resize(SymmetricMatrixPtr->rows(), SymmetricMatrixPtr->cols());
+
+		// The constraints are part of the matrix
+		bMatrixSolverDirty = true;
+	}
+
+	// Note: If constraints haven't been set, or if constraints have been cleared, this must be called prior to solve.
+	void UpdateSolverWithContraints()
+	{
+		FConstrainedSolverTimeLogger Timmer(TEXT("Matrix setup time"));
+
+		// The constrained verts are part of the matrix
+
+		FSparseMatrixD& SymmetricMatrix = *SymmetricMatrixPtr;
+		
+
+		LHSMatrix =  SymmetricMatrix + WeightsSqrdMatrix;
+		
+		LHSMatrix.makeCompressed();
+		// This matrix by construction is symmetric
+		const bool bIsSymmetric = true;
+
+
+		MatrixSolver->Reset();
+
+		MatrixSolver->SetUp(LHSMatrix, bIsSymmetric);
+
+		bMatrixSolverDirty = false;
+	}
+
+private:
+
+
+	// Positional Constraint Arrays	
+	FSOAPositions                  ConstraintPositions;
+
+	TUniquePtr<FSparseMatrixD>     SymmetricMatrixPtr; // Transpose(Laplacian) * Laplacain
+	FSparseMatrixD                 WeightsSqrdMatrix; // Weights Squared Diagonal Matrix.
+
+	FSparseMatrixD                 LHSMatrix;
+
+	bool                           bMatrixSolverDirty;
+	TUniquePtr<IMatrixSolverBase>  MatrixSolver;
+
+
+};
+
diff --git a/Engine/Plugins/Experimental/GeometryProcessing/Source/DynamicMesh/Private/Solvers/Internal/FSparseMatrixD.h b/Engine/Plugins/Experimental/GeometryProcessing/Source/DynamicMesh/Private/Solvers/Internal/FSparseMatrixD.h
new file mode 100644
--- /dev/null
+++ b/Engine/Plugins/Experimental/GeometryProcessing/Source/DynamicMesh/Private/Solvers/Internal/FSparseMatrixD.h
@@ -0,0 +1,67 @@
+// Copyright Epic Games, Inc. All Rights Reserved.
+
+#pragma once
+
+#include "CoreMinimal.h"
+#include "Util/ElementLinearization.h"   // TVector3Arrays
+
+
+// According to http://eigen.tuxfamily.org/index.php?title=Main_Page 
+// SimplicialCholesky, AMD ordering, and constrained_cg are disabled.
+
+#ifndef EIGEN_MPL2_ONLY
+#define EIGEN_MPL2_ONLY
+#endif
+
+#if defined(_MSC_VER) && USING_CODE_ANALYSIS
+#pragma warning(push)
+#pragma warning(disable : 6011)
+#pragma warning(disable : 6387)
+#pragma warning(disable : 6313)
+#pragma warning(disable : 6294)
+#endif
+THIRD_PARTY_INCLUDES_START
+#include <Eigen/Sparse>
+THIRD_PARTY_INCLUDES_END
+#if defined(_MSC_VER) && USING_CODE_ANALYSIS
+#pragma warning(pop)
+#endif
+
+
+// NB: The LU solver likes ColMajor but the CG sovler likes RowMajor
+//     Also, to change everything to float / double just change the scalar type here
+
+typedef Eigen::SparseMatrix<double, Eigen::ColMajor>  FSparseMatrixD;
+
+
+
+/**
+* A struct of arrays representation used to hold vertex positions
+* in three vectors that can interface with the eigen library
+*/
+class FSOAPositions : public TVector3Arrays<double>
+{
+public:
+	typedef typename FSparseMatrixD::Scalar ScalarType;
+	typedef Eigen::Matrix<ScalarType, Eigen::Dynamic, 1>  RealVectorType;
+	typedef Eigen::Map<RealVectorType> VectorType;
+	typedef Eigen::Map<const RealVectorType> ConstVectorType;
+
+	FSOAPositions(int32 Size) : TVector3Arrays<double>(Size)
+	{ }
+
+	FSOAPositions()
+	{}
+
+	VectorType Array(int32 i)
+	{
+		TArray<ScalarType>& Column = (i == 0) ? XVector : (i == 1) ? YVector : ZVector;
+		return (Column.Num() > 0) ? VectorType(&Column[0], Column.Num()) : VectorType(nullptr, 0);
+	}
+	ConstVectorType Array(int32 i) const
+	{
+		const TArray<ScalarType>& Column = (i == 0) ? XVector : (i == 1) ? YVector : ZVector;
+		return (Column.Num() > 0) ? ConstVectorType(&Column[0], Column.Num()) : ConstVectorType(nullptr, 0);
+	}
+
+};
\ No newline at end of file
diff --git a/Engine/Plugins/Experimental/GeometryProcessing/Source/DynamicMesh/Private/Solvers/Internal/LaplacianMeshSmoother.cpp b/Engine/Plugins/Experimental/GeometryProcessing/Source/DynamicMesh/Private/Solvers/Internal/LaplacianMeshSmoother.cpp
new file mode 100644
--- /dev/null
+++ b/Engine/Plugins/Experimental/GeometryProcessing/Source/DynamicMesh/Private/Solvers/Internal/LaplacianMeshSmoother.cpp
@@ -0,0 +1,779 @@
+// Copyright Epic Games, Inc. All Rights Reserved.
+
+#include "LaplacianMeshSmoother.h"
+
+#include "LaplacianOperators.h"
+#include "MatrixSolver.h"
+//#include "Solvers/MeshSmoothing.h"
+
+#ifdef TIME_LAPLACIAN_SMOOTHERS
+
+#include "ProfilingDebugging/ScopedTimers.h"
+DEFINE_LOG_CATEGORY_STATIC(LogMeshSmoother, Log, All);
+
+#endif
+
+double ComputeDistSqrd(const FSOAPositions& VecA, const FSOAPositions& VecB)
+{
+	const int32 NumA = VecA.Num();
+	const int32 NumB = VecB.Num();
+	checkSlow(NumA == NumB);
+
+#if 0
+	// the eigne way?
+	double Tmp  = (VecA.Array(0) - VecB.Array(0)).dot(VecA.Array(0) - VecB.Array(0))
+				+ (VecA.Array(1) - VecB.Array(1)).dot(VecA.Array(1) - VecB.Array(1))
+				+ (VecA.Array(2) - VecB.Array(2)).dot(VecA.Array(2) - VecB.Array(2));
+		     
+#endif
+
+	// doing it by hand so we can break when the error is large
+	double DistSqrd = 0.;
+	{
+		const auto& AX = VecA.Array(0);
+		const auto& AY = VecA.Array(1);
+		const auto& AZ = VecA.Array(2);
+
+		const auto& BX = VecB.Array(0);
+		const auto& BY = VecB.Array(1);
+		const auto& BZ = VecB.Array(2);
+
+
+		for (int32 i = 0; i < NumA; ++i)
+		{
+			double TmpX = AX(i)-BX(i);
+			double TmpY = AY(i)-BY(i);
+			double TmpZ = AZ(i)-BZ(i);
+
+			TmpX *= TmpX;
+			TmpY *= TmpY;
+			TmpZ *= TmpZ;
+			double TmpT = TmpX + TmpY + TmpZ;
+			DistSqrd += TmpT;
+		}
+	}
+	return DistSqrd;
+}
+
+
+FConstrainedMeshOperator::FConstrainedMeshOperator(const FDynamicMesh3& DynamicMesh, const ELaplacianWeightScheme Scheme, const EMatrixSolverType MatrixSolverType)
+	: VertexCount(DynamicMesh.VertexCount())
+{
+
+	FSparseMatrixD LaplacianInternal;
+	FSparseMatrixD LaplacianBoundary;
+	ConstructLaplacian(Scheme, DynamicMesh, VtxLinearization, LaplacianInternal, LaplacianBoundary);
+
+	// Copy the original boundary vertex locations
+	const int32 BoundaryVertexCount = VtxLinearization.NumBoundaryVerts();
+	
+	// Number of vertices in the interior of the mesh
+	InternalVertexCount = VertexCount - BoundaryVertexCount;
+
+	// Copy the original boundary vertex locations
+	{
+		const auto ToVertId = VtxLinearization.ToId();
+		BoundaryPositions.SetZero(BoundaryVertexCount);
+		for (int32 i = 0; i < BoundaryVertexCount; ++i)
+		{
+			const int32 VtxId = ToVertId[i + InternalVertexCount];
+			const FVector3d Pos = DynamicMesh.GetVertex(VtxId);
+			BoundaryPositions.SetXYZ(i, Pos);
+		}
+	}
+
+	checkSlow(LaplacianInternal.rows() == LaplacianInternal.cols());
+
+	TUniquePtr<FSparseMatrixD> LTLPtr( new FSparseMatrixD(LaplacianInternal.rows(), LaplacianInternal.cols()) );
+	FSparseMatrixD& LTLMatrix = *(LTLPtr);
+
+	bool bIsLaplacianSymmetric = (Scheme == ELaplacianWeightScheme::Valence || Scheme == ELaplacianWeightScheme::Uniform);
+
+	if (bIsLaplacianSymmetric)
+	{
+		// Laplacian is symmetric, i.e. equal to its transpose
+		LTLMatrix        = LaplacianInternal * LaplacianInternal;
+		BoundaryOperator = -1. * LaplacianInternal * LaplacianBoundary;
+		
+		ConstrainedSolver.Reset(new FConstrainedSolver(LTLPtr, MatrixSolverType));
+	}
+	else
+	{
+		// the laplacian 
+		LTLMatrix        = LaplacianInternal.transpose() * LaplacianInternal;
+		BoundaryOperator = -1. * LaplacianInternal.transpose() * LaplacianBoundary;
+		
+		ConstrainedSolver.Reset(new FConstrainedSolver(LTLPtr, MatrixSolverType));
+	}
+
+}
+
+void FConstrainedMeshOperator::AddConstraint(const int32 VtxId, const double Weight, const FVector3d& Pos, const bool bPostFix)
+{
+	const auto& ToIndex = VtxLinearization.ToIndex();
+	
+	if (VtxId > ToIndex.Num())
+	{
+		return;
+	}
+
+	const int32 Index = ToIndex[VtxId];
+
+	// Only add the constraint if the vertex is actually in the interior.  We aren't solving for edge vertices.
+	if (Index != FDynamicMesh3::InvalidID && Index < InternalVertexCount)
+	{
+		bConstraintPositionsDirty = true;
+		bConstraintWeightsDirty = true;
+
+
+		ConstraintPositionMap.Add(TTuple<int32, FConstraintPosition>(Index, FConstraintPosition(Pos, bPostFix)));
+		ConstraintWeightMap.Add(TTuple<int32, double>(Index, Weight));
+	}
+}
+
+
+bool FConstrainedMeshOperator::UpdateConstraintPosition(const int32 VtxId, const FVector3d& Pos, const bool bPostFix)
+{
+	bool Result = false;
+	const auto& ToIndex = VtxLinearization.ToIndex();
+
+	if (VtxId > ToIndex.Num())
+	{
+		return Result;
+	}
+
+	const int32 Index = ToIndex[VtxId];
+	
+
+	if (Index != FDynamicMesh3::InvalidID && Index < InternalVertexCount)
+	{
+		bConstraintPositionsDirty = true;
+
+		// Add should over-write any existing value for this key
+		ConstraintPositionMap.Add(TTuple<int32, FConstraintPosition>(Index, FConstraintPosition(Pos, bPostFix)));
+
+		Result = ConstraintWeightMap.Contains(VtxId);
+	}
+	return Result;
+}
+
+bool FConstrainedMeshOperator::UpdateConstraintWeight(const int32 VtxId, const double Weight)
+{
+	bool Result = false;
+	const auto& ToIndex = VtxLinearization.ToIndex();
+
+	if (VtxId > ToIndex.Num())
+	{
+		return Result;
+	}
+
+	const int32 Index = ToIndex[VtxId];
+	if (Index != FDynamicMesh3::InvalidID && Index < InternalVertexCount)
+	{
+		bConstraintWeightsDirty = true;
+
+		// Add should over-write any existing value for this key
+		ConstraintWeightMap.Add(TTuple<int32, double>(Index, Weight));
+
+		Result = ConstraintPositionMap.Contains(VtxId);
+	}
+	return Result;
+}
+
+
+bool FConstrainedMeshOperator::IsConstrained(const int32 VtxId) const
+{
+	bool Result = false;
+
+	const auto& ToIndex = VtxLinearization.ToIndex();
+
+	if (VtxId > ToIndex.Num())
+	{
+		return Result;
+	}
+
+	const int32 Index = ToIndex[VtxId];
+	
+	if (Index != FDynamicMesh3::InvalidID && Index < InternalVertexCount)
+	{
+		Result = ConstraintWeightMap.Contains(Index);
+	}
+
+	return Result;
+}
+
+void FConstrainedMeshOperator::UpdateSolverConstraints()
+{
+	if (bConstraintWeightsDirty)
+	{
+		ConstrainedSolver->SetConstraintWeights(ConstraintWeightMap);
+		bConstraintWeightsDirty = false;
+	}
+	
+	if (bConstraintPositionsDirty)
+	{
+		ConstrainedSolver->SetContraintPositions(ConstraintPositionMap);
+		bConstraintPositionsDirty = false;
+	}	
+}
+
+bool FConstrainedMeshOperator::CopyInternalPositions(const FSOAPositions& PositionalVector, TArray<FVector3d>& LinearArray) const
+{
+	// Number of positions
+
+	const int32 Num = PositionalVector.Num();
+
+	// early out if the x,y,z arrays in the PositionalVector have different lengths
+	if (!PositionalVector.bHasSize(Num))
+	{
+		return false;
+	}
+
+	checkSlow(Num == InternalVertexCount);
+
+	// 
+	const auto& ToVtxId = VtxLinearization.ToId();
+	const int32 MaxVtxId = ToVtxId.Num(); // NB: this is really max_used + 1 in the mesh.  See  FDynamicMesh3::MaxVertexID()
+
+	if (LinearArray.Num() != MaxVtxId)
+	{
+		return false;
+	}
+
+	// Update the internal positions.
+	for (int32 i = 0; i < InternalVertexCount; ++i)
+	{
+		const int32 VtxId = ToVtxId[i];
+
+		LinearArray[VtxId] = FVector3d(PositionalVector.X(i), PositionalVector.Y(i), PositionalVector.Z(i));
+	}
+
+	return true;
+}
+
+bool FConstrainedMeshOperator::CopyBoundaryPositions(TArray<FVector3d>& LinearArray) const
+{
+	const auto& ToVtxId = VtxLinearization.ToId();
+	const int32 MaxVtxId = ToVtxId.Num();
+
+	if (LinearArray.Num() != MaxVtxId)
+	{
+		return false;
+	}
+
+	int32 BoundaryVertexCount = VertexCount - InternalVertexCount;
+
+	for (int32 i = 0; i < BoundaryVertexCount; ++i)
+	{
+		const int32 VtxId = ToVtxId[i + InternalVertexCount];
+		LinearArray[VtxId] = FVector3d(BoundaryPositions.X(i), BoundaryPositions.Y(i), BoundaryPositions.Z(i)); //BoundaryPositions[i];
+	}
+
+	return true;
+}
+
+
+void FConstrainedMeshOperator::ExtractInteriorVertexPositions(const FDynamicMesh3& DynamicMesh, FSOAPositions& VertexPositions) const
+{
+	VertexPositions.SetZero(InternalVertexCount);
+
+	const auto& ToVtxId = VtxLinearization.ToId();
+
+	for (int32 i = 0; i < InternalVertexCount; ++i)
+	{
+		const int32 VtxId = ToVtxId[i];
+		const FVector3d& Pos = DynamicMesh.GetVertex(VtxId);
+		VertexPositions.SetXYZ(i, Pos);
+	}
+}
+
+void FConstrainedMeshOperator::UpdateWithPostFixConstraints(FSOAPositions& PositionVector) const
+{
+	for (const auto& ConstraintPosition : ConstraintPositionMap)
+	{
+		const int32 Index = ConstraintPosition.Key;
+		const FConstraintPosition& Constraint = ConstraintPosition.Value;
+
+		checkSlow(Index < InternalVertexCount);
+
+		// we only care about post-fix constraints
+
+		if (Constraint.bPostFix)
+		{
+			const FVector3d& Pos = Constraint.Position;
+			PositionVector.SetXYZ(Index, Pos);
+		}
+	}
+}
+
+
+
+FConstrainedMeshDeformer::FConstrainedMeshDeformer(const FDynamicMesh3& DynamicMesh, const ELaplacianWeightScheme LaplacianType)
+	: FConstrainedMeshOperator(DynamicMesh, LaplacianType, EMatrixSolverType::LU)
+	, LaplacianVectors(FConstrainedMeshOperator::InternalVertexCount)
+{
+	
+
+	// The current vertex positions 
+	
+	// Note: the OriginalInteriorPositions are being stored as member data 
+	// for use if the solver is iterative.
+	// FSOAPositions OriginalInteriorPositions; 
+	ExtractInteriorVertexPositions(DynamicMesh, OriginalInteriorPositions);
+	
+	
+	// The biharmonic part of the constrained solver
+	//   Biharmonic := Laplacian^{T} * Laplacian
+	
+	const auto& Biharmonic = ConstrainedSolver->Biharmonic();
+
+	// Compute the Laplacian Vectors
+	//    := Biharmonic * VertexPostion
+	// In the case of the cotangent laplacian this can be identified as the mean curvature * normal.
+	checkSlow(LaplacianVectors.Num() == OriginalInteriorPositions.Num());
+
+	for (int32 i = 0; i < 3; ++i)
+	{
+		LaplacianVectors.Array(i) = Biharmonic * OriginalInteriorPositions.Array(i);
+	}
+}
+
+bool FConstrainedMeshDeformer::Deform(TArray<FVector3d>& PositionBuffer) 
+{
+
+	// Update constraints.  This only trigger solver rebuild if the weights were updated.
+	UpdateSolverConstraints();
+
+	// Allocate space for the result as a struct of arrays
+	FSOAPositions SolutionVector(InternalVertexCount);
+
+	// Solve the linear system
+	// NB: the original positions will only be used if the underlying solver type is iterative	
+	bool bSuccess = ConstrainedSolver->SolveWithGuess(OriginalInteriorPositions, LaplacianVectors, SolutionVector);
+	
+	// Move any vertices to match bPostFix constraints
+
+	UpdateWithPostFixConstraints(SolutionVector);
+
+	// Allocate Position Buffer for random access writes
+	int32 MaxVtxId = VtxLinearization.ToId().Num();
+	PositionBuffer.Empty(MaxVtxId);
+	PositionBuffer.AddUninitialized(MaxVtxId);
+
+	// Export the computed internal positions:
+	// Copy the results into the array of structs form.  
+	// NB: this re-indexes so the results can be looked up using VtxId
+
+	CopyInternalPositions(SolutionVector, PositionBuffer);
+
+	// Copy the boundary
+	// NB: this re-indexes so the results can be looked up using VtxId
+	CopyBoundaryPositions(PositionBuffer);
+
+	// the matrix solve state
+	return bSuccess;
+
+}
+
+
+bool FBiHarmonicMeshSmoother::ComputeSmoothedMeshPositions(TArray<FVector3d>& UpdatedPositions)
+{
+
+	UpdateSolverConstraints();
+
+	// Compute the source vector
+	FSOAPositions SourceVector(InternalVertexCount);
+
+	if (InternalVertexCount != VertexCount) // have boundary points
+	{
+		for (int32 dir = 0; dir < 3; ++dir)
+		{
+			SourceVector.Array(dir) = BoundaryOperator * BoundaryPositions.Array(dir);
+		}
+	}
+	else
+	{
+		SourceVector.SetZero(InternalVertexCount);
+	}
+
+	// Solves the constrained system and updates the mesh 
+
+	FSOAPositions SolutionVector(InternalVertexCount);
+
+	bool bSuccess = ConstrainedSolver->Solve(SourceVector, SolutionVector);
+
+
+	// Move any vertices to match bPostFix constraints
+
+	UpdateWithPostFixConstraints(SolutionVector);
+
+	// Allocate Position Buffer for random access writes
+	int32 MaxVtxId = VtxLinearization.ToId().Num();
+	UpdatedPositions.Empty(MaxVtxId);
+	UpdatedPositions.AddUninitialized(MaxVtxId);
+
+	// Export the computed internal positions:
+	// Copy the results into the array of structs form.  
+	// NB: this re-indexes so the results can be looked up using VtxId
+
+	CopyInternalPositions(SolutionVector, UpdatedPositions);
+
+	// Copy the boundary
+	// NB: this re-indexes so the results can be looked up using VtxId
+	CopyBoundaryPositions(UpdatedPositions);
+
+	return bSuccess;
+}
+
+bool FCGBiHarmonicMeshSmoother::ComputeSmoothedMeshPositions(TArray<FVector3d>& UpdatedPositions)
+{
+
+	UpdateSolverConstraints();
+
+	// Compute the source vector
+	FSOAPositions SourceVector(InternalVertexCount);
+
+	if (InternalVertexCount != VertexCount) // have boundary points
+	{
+		for (int32 dir = 0; dir < 3; ++dir)
+		{
+			SourceVector.Array(dir) = BoundaryOperator * BoundaryPositions.Array(dir);
+		}
+	}
+	else
+	{
+		SourceVector.SetZero(InternalVertexCount);
+	}
+
+	// Solves the constrained system and updates the mesh 
+
+	// Solves the constrained system
+
+	FSOAPositions SolutionVector(InternalVertexCount);
+
+	bool bSuccess = ConstrainedSolver->Solve(SourceVector, SolutionVector);
+
+
+	// Move any vertices to match bPostFix constraints
+
+	UpdateWithPostFixConstraints(SolutionVector);
+
+	// Allocate Position Buffer for random access writes
+	int32 MaxVtxId = VtxLinearization.ToId().Num();
+	UpdatedPositions.Empty(MaxVtxId);
+	UpdatedPositions.AddUninitialized(MaxVtxId);
+
+	// Export the computed internal postions:
+	// Copy the results into the array of structs form.  
+	// NB: this re-indexes so the results can be looked up using VtxId
+
+	CopyInternalPositions(SolutionVector, UpdatedPositions);
+
+	// Copy the boundary
+	// NB: this re-indexes so the results can be looked up using VtxId
+	CopyBoundaryPositions(UpdatedPositions);
+
+	return bSuccess;
+}
+
+
+
+FDiffusionIntegrator::FDiffusionIntegrator(const FDynamicMesh3& DynamicMesh, const ELaplacianWeightScheme Scheme)
+{
+	Id = 0;
+	bIsSymmetric = false;
+	MinDiagonalValue = 0;
+
+	VertexCount = DynamicMesh.VertexCount();
+	
+}
+
+void FDiffusionIntegrator::Initialize(const FDynamicMesh3& DynamicMesh, const ELaplacianWeightScheme Scheme)
+{
+	// Construct the laplacian, and extract the mapping for vertices (VtxLinearization)
+	//DiffusionOperator = ConstructDiffusionOperator(Scheme, DynamicMesh, bIsSymmetric, VtxLinearization, &EdgeVerts);
+
+	
+	ConstructOperators(Scheme, DynamicMesh, bIsSymmetric, VtxLinearization, DiffusionOperator, BoundaryOperator);
+
+	const int32 BoundaryVertexCount = VtxLinearization.NumBoundaryVerts();
+	InternalVertexCount = VertexCount - BoundaryVertexCount;
+
+	// Allocate the double buffers.
+	Tmp[0].SetZero(InternalVertexCount);
+	Tmp[1].SetZero(InternalVertexCount);
+
+	const auto& ToVertId = VtxLinearization.ToId();
+
+	// Extract current internal positions.
+	for (int32 i = 0; i < InternalVertexCount; ++i)
+	{
+		const int32 VtxId = ToVertId[i];
+		const FVector3d Pos = DynamicMesh.GetVertex(VtxId);
+		Tmp[0].SetXYZ(i, Pos);
+	}
+
+	// backup the locations of the boundary verts.
+	{
+		BoundaryPositions.SetZero(BoundaryVertexCount);
+		for (int32 i = 0; i < BoundaryVertexCount; ++i)
+		{
+			const int32 VtxId = ToVertId[i + InternalVertexCount];
+			const FVector3d Pos = DynamicMesh.GetVertex(VtxId);
+			BoundaryPositions.SetXYZ(i, Pos);
+		}
+	}
+
+	const FSparseMatrixD&  M = DiffusionOperator;
+
+	// Find the min diagonal entry (all should be negative).
+	int32 Rank = M.rows();
+	MinDiagonalValue = FSparseMatrixD::Scalar(0);
+	for (int32 i = 0; i < Rank; ++i)
+	{
+		auto Diag = M.coeff(i, i);
+		MinDiagonalValue = FMath::Min(Diag, MinDiagonalValue);
+	}
+	// The matrix should have a row for each internal vertex
+	checkSlow(Rank == InternalVertexCount);
+
+#if 0
+	// testing - how to print the matrix to debug output 
+
+	std::stringstream  ss;
+	ss << Eigen::MatrixXd(M) << std::endl;
+
+	FString Foo = ss.str().c_str();
+	FPlatformMisc::LowLevelOutputDebugStringf(*Foo);
+#endif
+
+}
+
+
+
+void FDiffusionIntegrator::Integrate_ForwardEuler(const int32 NumSteps, const double Speed)
+{
+	double Alpha = FMath::Clamp(Speed, 0., 1.);
+
+	FSparseMatrixD::Scalar TimeStep = -Alpha / MinDiagonalValue;
+	Id = 0;
+	for (int32 s = 0; s < NumSteps; ++s)
+	{
+
+		int32 SrcBuffer = Id;
+		Id = 1 - Id;
+		Tmp[Id].Array(0) = Tmp[SrcBuffer].Array(0) + TimeStep * ( DiffusionOperator * Tmp[SrcBuffer].Array(0) + BoundaryOperator * BoundaryPositions.Array(0));
+		Tmp[Id].Array(1) = Tmp[SrcBuffer].Array(1) + TimeStep * ( DiffusionOperator * Tmp[SrcBuffer].Array(1) + BoundaryOperator * BoundaryPositions.Array(1));
+		Tmp[Id].Array(2) = Tmp[SrcBuffer].Array(2) + TimeStep * ( DiffusionOperator * Tmp[SrcBuffer].Array(2) + BoundaryOperator * BoundaryPositions.Array(2));
+
+	}
+
+}
+
+
+
+void FDiffusionIntegrator::Integrate_BackwardEuler(const EMatrixSolverType MatrixSolverType, const int32 NumSteps, const double TimeStepSize)
+{
+
+	//typedef typename TMatrixSolverTrait<EMatrixSolverType::LU>::MatrixSolverType   MatrixSolverType;
+
+	// We solve 
+	// p^{n+1} - dt * L[p^{n+1}] = p^{n} + dt * B[boundaryPts]
+	// 
+	// i.e.
+	// [I - dt * L ] p^{n+1} = p^{n} + dt * B[boundaryPts]
+	//
+	// NB: in the case of the cotangent laplacian this would be better if we broke the L int
+	// L = (A^{-1}) H  where A is the "area matrix" (think "mass matrix"), then this would
+	// become
+	// [A - dt * H] p^{n+1} = Ap^{n}  dt * A *B[boundaryPts]
+	//  
+	// A - dt * H would be symmetric
+	//
+
+
+	
+	// Identity matrix
+	FSparseMatrixD Ident(DiffusionOperator.rows(), DiffusionOperator.cols());
+	Ident.setIdentity();
+
+	FSparseMatrixD::Scalar TimeStep = FMath::Abs(TimeStepSize);// Alpha * FMath::Min(Intensity, 1.e6);
+	
+	FSparseMatrixD SparseMatrix = Ident -TimeStep * DiffusionOperator;
+
+	
+	SparseMatrix.makeCompressed();
+
+	TUniquePtr<IMatrixSolverBase> MatrixSolver = ContructMatrixSolver(MatrixSolverType);
+
+	MatrixSolver->SetUp(SparseMatrix, bIsSymmetric);
+
+	// We are going to solve the system 
+	FSOAPositions Source(InternalVertexCount);
+
+	if (MatrixSolver->bIsIterative())
+	{
+		IIterativeMatrixSolverBase* IterativeSolver = (IIterativeMatrixSolverBase*)MatrixSolver.Get();
+
+		bool bForceSingleThreaded = false;
+		Id = 0;
+		for (int32 s = 0; s < NumSteps; ++s)
+		{
+
+			int32 SrcBuffer = Id;
+			Id = 1 - Id;
+
+			for (int32 i = 0; i < 3; ++i)
+			{
+				Source.Array(i) = Tmp[SrcBuffer].Array(i) + TimeStep *  BoundaryOperator * BoundaryPositions.Array(i);
+			}
+			// Old solution is the guess.
+			IterativeSolver->SolveWithGuess(Tmp[SrcBuffer], Source, Tmp[Id]);
+
+		}
+	}
+	else
+	{
+		bool bForceSingleThreaded = false;
+		Id = 0;
+		for (int32 s = 0; s < NumSteps; ++s)
+		{
+
+			int32 SrcBuffer = Id;
+			Id = 1 - Id;
+
+			for (int32 i = 0; i < 3; ++i)
+			{
+				Source.Array(i) = Tmp[SrcBuffer].Array(i) + TimeStep * BoundaryOperator * BoundaryPositions.Array(i);
+			}
+
+			MatrixSolver->Solve(Source, Tmp[Id]);
+		}
+	}
+	
+
+}
+
+void FDiffusionIntegrator::GetPositions(TArray<FVector3d>& PositionBuffer) const
+{
+	// Allocate Position Buffer for random access writes
+	int32 MaxVtxId = VtxLinearization.ToId().Num();
+	PositionBuffer.Empty(MaxVtxId);
+	PositionBuffer.AddUninitialized(MaxVtxId);
+
+	CopyInternalPositions(Tmp[Id], PositionBuffer);
+
+
+	// Copy the boundary
+	// NB: this re-indexes so the results can be looked up using VtxId
+	CopyBoundaryPositions(PositionBuffer);
+
+}
+
+bool FDiffusionIntegrator::CopyInternalPositions(const FSOAPositions& PositionalVector, TArray<FVector3d>& LinearArray) const
+{
+	// Number of positions
+
+	const int32 Num = PositionalVector.Num();
+
+	// early out if the x,y,z arrays in the PositionalVector have different lengths
+	if (!PositionalVector.bHasSize(Num))
+	{
+		return false;
+	}
+
+	// 
+	const auto& ToVtxId  = VtxLinearization.ToId();
+	const int32 MaxVtxId = ToVtxId.Num(); // NB: this is really max_used + 1 in the mesh.  See  FDynamicMesh3::MaxVertexID()
+	const int32 BoundaryVertexCount = VtxLinearization.NumBoundaryVerts();
+
+	if (MaxVtxId != LinearArray.Num())
+	{
+		return false;
+	}
+
+	// Copy the updated internal vertex locations over
+	for (int32 i = 0; i < InternalVertexCount; ++i)
+	{
+		const int32 VtxId = ToVtxId[i];
+
+		LinearArray[VtxId] = FVector3d(PositionalVector.X(i), PositionalVector.Y(i), PositionalVector.Z(i));
+	}
+
+	return true;
+}
+                       
+bool FDiffusionIntegrator::CopyBoundaryPositions(TArray<FVector3d>& LinearArray) const
+{
+	const auto& ToVtxId = VtxLinearization.ToId();
+	const int32 MaxVtxId = ToVtxId.Num();
+
+	if (LinearArray.Num() != MaxVtxId)
+	{
+		return false;
+	}
+
+	int32 BoundaryVertexCount = VertexCount - InternalVertexCount;
+
+	for (int32 i = 0; i < BoundaryVertexCount; ++i)
+	{
+		const int32 VtxId = ToVtxId[i + InternalVertexCount];
+		LinearArray[VtxId] = FVector3d(BoundaryPositions.X(i), BoundaryPositions.Y(i), BoundaryPositions.Z(i));
+	}
+
+	return true;
+}
+
+
+
+void FLaplacianDiffusionMeshSmoother::ConstructOperators( const ELaplacianWeightScheme Scheme,
+	                                                      const FDynamicMesh3& Mesh,
+	                                                      bool& bIsOperatorSymmetric,
+	                                                      FVertexLinearization& Linearization,
+	                                                      FSparseMatrixD& DiffusionOp,
+	                                                      FSparseMatrixD& BoundaryOp) 
+{
+	bIsOperatorSymmetric = bIsSymmetricLaplacian(Scheme);
+	ConstructLaplacian(Scheme, Mesh, VtxLinearization, DiffusionOp, BoundaryOp);
+}
+
+
+void FBiHarmonicDiffusionMeshSmoother::ConstructOperators( const ELaplacianWeightScheme Scheme,
+														   const FDynamicMesh3& Mesh,
+														   bool& bIsOperatorSymmetric,
+														   FVertexLinearization& Linearization,
+														   FSparseMatrixD& DiffusionOp,
+	                                                       FSparseMatrixD& BoundaryOp) 
+{
+	bIsOperatorSymmetric = true;
+	
+	FSparseMatrixD Laplacian; 
+	FSparseMatrixD BoundaryTerms;
+	ConstructLaplacian(Scheme, Mesh, VtxLinearization, Laplacian, BoundaryTerms);
+
+	bool bIsLaplacianSymmetric = bIsSymmetricLaplacian(Scheme);
+
+	// It is actually unclear the best way to approximate the boundary conditions in this case.  
+	// because we are repeatedly applying the operator ( for example thing about the way ( f(x+d)-f(x-d) )/ d will spread if you apply it twice
+	// as opposed to (f(x+d)-2f(x) + f(x-d) ) / d*d
+
+	// Anyhow here is a guess..
+
+	if (bIsLaplacianSymmetric)
+	{
+		DiffusionOp = -1. * Laplacian * Laplacian;
+		BoundaryOp  = -1. * Laplacian * BoundaryTerms;
+	}
+	else
+	{
+		FSparseMatrixD LTran = Laplacian.transpose();
+		DiffusionOp = -1. * LTran * Laplacian;
+		BoundaryOp  = -1. * LTran * BoundaryTerms;
+	}
+	
+	DiffusionOp.makeCompressed();
+	BoundaryOp.makeCompressed();
+
+}
+
+
diff --git a/Engine/Plugins/Experimental/GeometryProcessing/Source/DynamicMesh/Private/Solvers/Internal/LaplacianMeshSmoother.h b/Engine/Plugins/Experimental/GeometryProcessing/Source/DynamicMesh/Private/Solvers/Internal/LaplacianMeshSmoother.h
new file mode 100644
--- /dev/null
+++ b/Engine/Plugins/Experimental/GeometryProcessing/Source/DynamicMesh/Private/Solvers/Internal/LaplacianMeshSmoother.h
@@ -0,0 +1,272 @@
+// Copyright Epic Games, Inc. All Rights Reserved.
+
+#pragma once
+
+#include "CoreMinimal.h"
+#include "DynamicMesh3.h"
+#include "Solvers/MeshLaplacian.h"
+#include "Solvers/MeshLinearization.h"
+#include "Solvers/ConstrainedMeshSolver.h"
+#include "ConstrainedPoissonSolver.h"
+#include "MatrixSolver.h"
+
+
+
+class DYNAMICMESH_API FConstrainedMeshOperator : public IConstrainedMeshSolver
+{
+public:
+	typedef FConstrainedSolver::FConstraintPosition   FConstraintPosition;
+
+
+	FConstrainedMeshOperator(const FDynamicMesh3& DynamicMesh, const ELaplacianWeightScheme Scheme, const EMatrixSolverType MatrixSolverType);
+	virtual ~FConstrainedMeshOperator() override {}
+
+	// Add constraint associated with given vertex id.  Boundary vertices will be ignored
+	void AddConstraint(const int32 VtxId, const double Weight, const FVector3d& Pos, const bool bPostFix) override;
+
+	// Update the position of an existing constraint.  Bool return if a corresponding constraint weight exists. Boundary vertices will be ignored (and return false).  
+	bool UpdateConstraintPosition(const int32 VtxId, const FVector3d& Position, const bool bPostFix) override;
+
+	// The underlying solver will have to refactor the matrix if this is done. Bool return if a corresponding constraint position exists. Boundary vertices will be ignored (and return false).  
+	bool UpdateConstraintWeight(const int32 VtxId, const double Weight) override;
+
+	// Clear all constraints associated with this smoother
+	void ClearConstraints() override { FConstrainedMeshOperator::ClearConstraintPositions(); FConstrainedMeshOperator::ClearConstraintWeights(); }
+
+	void ClearConstraintWeights() override { ConstraintWeightMap.Empty();    bConstraintWeightsDirty = true; }
+
+	void ClearConstraintPositions() override { ConstraintPositionMap.Empty();  bConstraintPositionsDirty = true; }
+
+
+	// Test if for constraint associated with given vertex id. Will return false for any boundary vert.
+	bool IsConstrained(const int32 VtxId) const override;
+
+	virtual bool Deform(TArray<FVector3d>& PositionBuffer) override { return false;  }
+
+	// Sync constraints with internal solver.  If in the process any internal matrix factoring is dirty, it will be rebuilt.
+	// Note: this is called from within the Deform() method.   Only call this method if you want to trigger the matrix refactor yourself.
+	void UpdateSolverConstraints();
+
+protected:
+
+
+	void ExtractInteriorVertexPositions(const FDynamicMesh3& DynamicMesh, FSOAPositions& Positions) const ;
+
+	//void UpdateMeshWithConstraints();
+	// Respect any bPostFix constraints by moving those vertices to position defined by said constraint.
+	void UpdateWithPostFixConstraints(FSOAPositions& PositionalVector) const;
+
+
+	// converts the positionalvector to a TArray<FVector3d> where the offset in the array is implicitly the
+	// VtxId in the mesh, and not ness the matrix row id
+	// NB: the resulting array is treated as sparse and may have un-initialized elements.
+	bool CopyInternalPositions(const FSOAPositions& PositionalVector, TArray<FVector3d>& LinearArray) const;
+
+	bool CopyBoundaryPositions(TArray<FVector3d>& LinearArray) const;
+
+
+protected:
+
+	// The Key (int32) here is the vertex index not vertex ID
+	// making it the same as the matrix row
+	TMap<int32, FConstraintPosition> ConstraintPositionMap;
+	TMap<int32, double>              ConstraintWeightMap;
+
+	bool bConstraintPositionsDirty = true;
+	bool bConstraintWeightsDirty   = true;
+
+	// Cache the vertex count. boundary + internal
+	int32                 VertexCount;
+
+	// Cache the number of internal vertices
+	int32                 InternalVertexCount;
+
+	// Used to map between VtxId and vertex Index in linear vector..
+	FVertexLinearization  VtxLinearization;
+
+	// Boundary points, split into 3 arrays (x,y,z).
+	FSOAPositions         BoundaryPositions;
+
+
+	// Actual solver that manages the various linear algebra bits.
+	TUniquePtr<FConstrainedSolver>         ConstrainedSolver;
+
+    // Sparse Matrix that holds L^T * B where B has the boundary terms.
+	FSparseMatrixD                         BoundaryOperator;
+};
+
+
+class DYNAMICMESH_API FConstrainedMeshDeformer : public FConstrainedMeshOperator
+{
+public:
+	FConstrainedMeshDeformer(const FDynamicMesh3& DynamicMesh, const ELaplacianWeightScheme LaplacianType);
+	~FConstrainedMeshDeformer() override {}
+
+	bool Deform(TArray<FVector3d>& PositionBuffer) override;
+
+private:
+
+	FSOAPositions LaplacianVectors;
+	FSOAPositions OriginalInteriorPositions;
+};
+
+
+class FBiHarmonicMeshSmoother : public FConstrainedMeshOperator
+{
+public:
+	typedef FConstrainedMeshOperator         MyBaseType;
+
+	FBiHarmonicMeshSmoother(const FDynamicMesh3& DynamicMesh, const ELaplacianWeightScheme Scheme) :
+		MyBaseType(DynamicMesh, Scheme, EMatrixSolverType::LU)
+	{}
+
+	bool Deform(TArray<FVector3d>& UpdatedPositions) override 
+	{
+		return	ComputeSmoothedMeshPositions(UpdatedPositions);
+	}
+
+	// (Direct) Solve the constrained system and populate the UpdatedPositions with the result 
+	bool ComputeSmoothedMeshPositions(TArray<FVector3d>& UpdatedPositions);
+
+};
+
+
+// NB: This conjugate gradient solver could be updated to use  solveWithGuess() method on the iterative solver
+class FCGBiHarmonicMeshSmoother : public FConstrainedMeshOperator
+{
+public:
+	typedef FConstrainedMeshOperator      MyBaseType;
+
+	FCGBiHarmonicMeshSmoother(const FDynamicMesh3& DynamicMesh, const ELaplacianWeightScheme Scheme) :
+		MyBaseType(DynamicMesh, Scheme, EMatrixSolverType::BICGSTAB)
+	{}
+
+	bool Deform(TArray<FVector3d>& UpdatedPositions) override
+	{
+		return	ComputeSmoothedMeshPositions(UpdatedPositions);
+	}
+
+	void SetMaxIterations(int32 MaxIterations) 
+	{ 
+		IIterativeMatrixSolverBase* Solver = ConstrainedSolver->GetMatrixSolverIterativeBase(); 
+		if (Solver)
+		{
+			Solver->SetIterations(MaxIterations);
+		}
+	}
+
+	void SetTolerance(double Tol)
+	{
+		IIterativeMatrixSolverBase* Solver = ConstrainedSolver->GetMatrixSolverIterativeBase();
+		if (Solver)
+		{
+			Solver->SetTolerance(Tol);
+		}
+	}
+
+	// (Iterative) Solve the constrained system and populate the UpdatedPositions with the result 
+	bool ComputeSmoothedMeshPositions(TArray<FVector3d>& UpdatedPositions);
+
+
+};
+
+
+class FDiffusionIntegrator
+{
+public:
+
+	FDiffusionIntegrator(const FDynamicMesh3& DynamicMesh, const ELaplacianWeightScheme Scheme);
+	virtual ~FDiffusionIntegrator() {}
+
+	void Integrate_ForwardEuler(const int32 NumSteps, const double Speed);
+
+	// Note: 
+	void Integrate_BackwardEuler(const EMatrixSolverType MatrixSolverType, const int32 NumSteps, const double TimeStepSize);
+
+	void GetPositions(TArray<FVector3d>& PositionArray) const;
+
+protected:
+
+
+	// The derived class has to implement this.
+	// responsible for constructing the Diffusion Operator, the Boundary Operator etc
+	virtual void ConstructOperators(const ELaplacianWeightScheme Scheme,
+		                            const FDynamicMesh3& Mesh,
+	 	                            bool& bIsOperatorSymmetric,
+		                            FVertexLinearization& Linearization,
+									FSparseMatrixD& DiffusionOp,
+									FSparseMatrixD& BoundaryOp ) = 0;
+
+	void Initialize(const FDynamicMesh3& DynamicMesh, const ELaplacianWeightScheme Scheme);
+
+protected:
+
+	// Copy the 
+	bool CopyInternalPositions(const FSOAPositions& PositionalVector, TArray<FVector3d>& LinearArray) const;
+
+	bool CopyBoundaryPositions(TArray<FVector3d>& LinearArray) const;
+
+	// Cache the vertex count.
+	int32                 VertexCount;
+
+	// Cache the number of internal vertices
+	int32                 InternalVertexCount;
+
+	// Used to map between VtxId and vertex Index in linear vector..
+	FVertexLinearization  VtxLinearization;
+
+
+	// I don't know if we want to keep this after the constructor
+	bool                                   bIsSymmetric;
+	FSparseMatrixD                         DiffusionOperator;
+	FSparseMatrixD			               BoundaryOperator;
+	TArray<int32>                          EdgeVerts;
+	FSparseMatrixD::Scalar                 MinDiagonalValue;
+	
+	FSOAPositions                       BoundaryPositions;
+
+	FSOAPositions                       Tmp[2];
+	int32                               Id; // double buffer id
+};
+
+
+class FLaplacianDiffusionMeshSmoother : public FDiffusionIntegrator
+{
+public:
+
+	FLaplacianDiffusionMeshSmoother(const FDynamicMesh3& DynamicMesh, const ELaplacianWeightScheme Scheme)
+	: FDiffusionIntegrator(DynamicMesh, Scheme)
+	{ 
+		Initialize(DynamicMesh, Scheme);
+	}
+	
+protected:
+	
+	void ConstructOperators(const ELaplacianWeightScheme Scheme,
+		                    const FDynamicMesh3& Mesh,
+		                    bool& bIsOperatorSymmetric,
+		                    FVertexLinearization& Linearization,
+		                    FSparseMatrixD& DiffusionOp,
+		                    FSparseMatrixD& BoundaryOp) override;
+};
+
+class  FBiHarmonicDiffusionMeshSmoother : public FDiffusionIntegrator
+{
+public:
+
+	FBiHarmonicDiffusionMeshSmoother(const FDynamicMesh3& DynamicMesh, const ELaplacianWeightScheme Scheme)
+	: FDiffusionIntegrator(DynamicMesh, Scheme)
+	{
+		Initialize(DynamicMesh, Scheme);
+	}
+
+protected:
+
+	void ConstructOperators( const ELaplacianWeightScheme Scheme,
+		                     const FDynamicMesh3& Mesh,
+		                     bool& bIsOperatorSymmetric,
+		                     FVertexLinearization& Linearization,
+		                     FSparseMatrixD& DiffusionOp,
+		                     FSparseMatrixD& BoundaryOp) override;
+};
+
diff --git a/Engine/Plugins/Experimental/GeometryProcessing/Source/DynamicMesh/Private/Solvers/Internal/LaplacianOperators.cpp b/Engine/Plugins/Experimental/GeometryProcessing/Source/DynamicMesh/Private/Solvers/Internal/LaplacianOperators.cpp
new file mode 100644
--- /dev/null
+++ b/Engine/Plugins/Experimental/GeometryProcessing/Source/DynamicMesh/Private/Solvers/Internal/LaplacianOperators.cpp
@@ -0,0 +1,383 @@
+// Copyright Epic Games, Inc. All Rights Reserved.
+
+#include "LaplacianOperators.h"
+
+#include "Solvers/MatrixInterfaces.h"
+#include "Solvers/LaplacianMatrixAssembly.h"
+
+#include <cmath> // double version of sqrt
+#include <vector> // used by eigen to initialize sparse matrix
+
+#if defined(_MSC_VER) && USING_CODE_ANALYSIS
+#pragma warning(push)
+#pragma warning(disable : 6011)
+#pragma warning(disable : 6387)
+#pragma warning(disable : 6313)
+#pragma warning(disable : 6294)
+#endif
+THIRD_PARTY_INCLUDES_START
+#include <Eigen/Dense> // for Matrix4d in testing
+THIRD_PARTY_INCLUDES_END
+#if defined(_MSC_VER) && USING_CODE_ANALYSIS
+#pragma warning(pop)
+#endif
+
+
+
+//
+// Extension of TSparseMatrixAssembler suitable for eigen sparse matrix
+//
+class FEigenSparseMatrixAssembler : public TSparseMatrixAssembler<double>
+{
+public:
+	typedef FSparseMatrixD::Scalar    ScalarT;
+	typedef Eigen::Triplet<ScalarT>  MatrixTripletT;
+
+	TUniquePtr<FSparseMatrixD> Matrix;
+	std::vector<MatrixTripletT> EntryTriplets;
+
+	FEigenSparseMatrixAssembler(int32 RowsI, int32 ColsJ)
+	{
+		Matrix = MakeUnique<FSparseMatrixD>(RowsI, ColsJ);
+
+		ReserveEntriesFunc = [this](int32 NumElements)
+		{
+			EntryTriplets.reserve(NumElements);
+		};
+
+		AddEntryFunc = [this](int32 i, int32 j, double Value)
+		{
+			EntryTriplets.push_back(MatrixTripletT(i, j, Value));
+		};
+	}
+
+	TUniquePtr<FSparseMatrixD> ExtractResult()
+	{
+		Matrix->setFromTriplets(EntryTriplets.begin(), EntryTriplets.end());
+		Matrix->makeCompressed();
+	}
+
+	void ExtractResult(FSparseMatrixD& Result)
+	{
+		Matrix->setFromTriplets(EntryTriplets.begin(), EntryTriplets.end());
+		Matrix->makeCompressed();
+
+		Result.swap(*Matrix);
+	}
+};
+
+
+
+void ConstructUniformLaplacian(const FDynamicMesh3& DynamicMesh, FVertexLinearization& VertexMap, FSparseMatrixD& LaplacianInterior, FSparseMatrixD& LaplacianBoundary)
+{
+	// Sync the mapping between the mesh vertex ids and their offsets in a nominal linear array.
+	VertexMap.Reset(DynamicMesh);
+	const int32 NumVerts = VertexMap.NumVerts();
+	const int32 NumBoundaryVerts = VertexMap.NumBoundaryVerts();
+	const int32 NumInteriorVerts = NumVerts - NumBoundaryVerts;
+
+	FEigenSparseMatrixAssembler Interior(NumInteriorVerts, NumInteriorVerts);
+	FEigenSparseMatrixAssembler Boundary(NumInteriorVerts, NumBoundaryVerts);
+	UE::MeshDeformation::ConstructUniformLaplacian<double>(DynamicMesh, VertexMap, Interior, Boundary);
+	Interior.ExtractResult(LaplacianInterior);
+	Boundary.ExtractResult(LaplacianBoundary);
+}
+
+
+void ConstructUmbrellaLaplacian(const FDynamicMesh3& DynamicMesh, FVertexLinearization& VertexMap, FSparseMatrixD& LaplacianInterior, FSparseMatrixD& LaplacianBoundary)
+{
+	// Sync the mapping between the mesh vertex ids and their offsets in a nominal linear array.
+	VertexMap.Reset(DynamicMesh);
+	const int32 NumVerts = VertexMap.NumVerts();
+	const int32 NumBoundaryVerts = VertexMap.NumBoundaryVerts();
+	const int32 NumInteriorVerts = NumVerts - NumBoundaryVerts;
+
+	FEigenSparseMatrixAssembler Interior(NumInteriorVerts, NumInteriorVerts);
+	FEigenSparseMatrixAssembler Boundary(NumInteriorVerts, NumBoundaryVerts);
+	UE::MeshDeformation::ConstructUmbrellaLaplacian<double>(DynamicMesh, VertexMap, Interior, Boundary);
+	Interior.ExtractResult(LaplacianInterior);
+	Boundary.ExtractResult(LaplacianBoundary);
+}
+
+
+void ConstructValenceWeightedLaplacian(const FDynamicMesh3& DynamicMesh, FVertexLinearization& VertexMap, FSparseMatrixD& LaplacianInterior, FSparseMatrixD& LaplacianBoundary)
+{
+	// Sync the mapping between the mesh vertex ids and their offsets in a nominal linear array.
+	VertexMap.Reset(DynamicMesh);
+	const int32 NumVerts = VertexMap.NumVerts();
+	const int32 NumBoundaryVerts = VertexMap.NumBoundaryVerts();
+	const int32 NumInteriorVerts = NumVerts - NumBoundaryVerts;
+
+	FEigenSparseMatrixAssembler Interior(NumInteriorVerts, NumInteriorVerts);
+	FEigenSparseMatrixAssembler Boundary(NumInteriorVerts, NumBoundaryVerts);
+	UE::MeshDeformation::ConstructValenceWeightedLaplacian<double>(DynamicMesh, VertexMap, Interior, Boundary);
+	Interior.ExtractResult(LaplacianInterior);
+	Boundary.ExtractResult(LaplacianBoundary);
+}
+
+
+
+void ConstructMeanValueWeightLaplacian(const FDynamicMesh3& DynamicMesh, FVertexLinearization& VertexMap, FSparseMatrixD& LaplacianInterior, FSparseMatrixD& LaplacianBoundary)
+{
+	// Sync the mapping between the mesh vertex ids and their offsets in a nominal linear array.
+	VertexMap.Reset(DynamicMesh);
+	const int32 NumVerts = VertexMap.NumVerts();
+	const int32 NumBoundaryVerts = VertexMap.NumBoundaryVerts();
+	const int32 NumInteriorVerts = NumVerts - NumBoundaryVerts;
+
+	FEigenSparseMatrixAssembler Interior(NumInteriorVerts, NumInteriorVerts);
+	FEigenSparseMatrixAssembler Boundary(NumInteriorVerts, NumBoundaryVerts);
+	UE::MeshDeformation::ConstructMeanValueWeightLaplacian<double>(DynamicMesh, VertexMap, Interior, Boundary);
+	Interior.ExtractResult(LaplacianInterior);
+	Boundary.ExtractResult(LaplacianBoundary);
+}
+
+
+
+
+void ConstructCotangentLaplacian(const FDynamicMesh3& DynamicMesh, FVertexLinearization& VertexMap, FSparseMatrixD& AreaMatrix, FSparseMatrixD& LaplacianInterior, FSparseMatrixD& LaplacianBoundary)
+{
+	// Sync the mapping between the mesh vertex ids and their offsets in a nominal linear array.
+	VertexMap.Reset(DynamicMesh);
+	const int32 NumVerts = VertexMap.NumVerts();
+	const int32 NumBoundaryVerts = VertexMap.NumBoundaryVerts();
+	const int32 NumInteriorVerts = NumVerts - NumBoundaryVerts;
+
+	FEigenSparseMatrixAssembler Interior(NumInteriorVerts, NumInteriorVerts);
+	FEigenSparseMatrixAssembler Boundary(NumInteriorVerts, NumBoundaryVerts);
+	FEigenSparseMatrixAssembler Area(NumVerts, NumVerts);
+	UE::MeshDeformation::ConstructCotangentLaplacian<double>(DynamicMesh, VertexMap, Area, Interior, Boundary);
+	Interior.ExtractResult(LaplacianInterior);
+	Boundary.ExtractResult(LaplacianBoundary);
+	Area.ExtractResult(AreaMatrix);
+}
+
+
+void ConstructCotangentLaplacian(const FDynamicMesh3& DynamicMesh, FVertexLinearization& VertexMap, FSparseMatrixD& LaplacianInterior, FSparseMatrixD& LaplacianBoundary, const bool bClampWeights)
+{
+	// Sync the mapping between the mesh vertex ids and their offsets in a nominal linear array.
+	VertexMap.Reset(DynamicMesh);
+	const int32 NumVerts = VertexMap.NumVerts();
+	const int32 NumBoundaryVerts = VertexMap.NumBoundaryVerts();
+	const int32 NumInteriorVerts = NumVerts - NumBoundaryVerts;
+
+	FEigenSparseMatrixAssembler Interior(NumInteriorVerts, NumInteriorVerts);
+	FEigenSparseMatrixAssembler Boundary(NumInteriorVerts, NumBoundaryVerts);
+	UE::MeshDeformation::ConstructCotangentLaplacian<double>(DynamicMesh, VertexMap, Interior, Boundary, bClampWeights);
+	Interior.ExtractResult(LaplacianInterior);
+	Boundary.ExtractResult(LaplacianBoundary);
+}
+
+
+
+double ConstructScaledCotangentLaplacian(const FDynamicMesh3& DynamicMesh, FVertexLinearization& VertexMap, FSparseMatrixD& LaplacianInterior, FSparseMatrixD& LaplacianBoundary, const bool bClampAreas)
+{
+	typedef FSparseMatrixD::Scalar  ScalarT;
+	typedef Eigen::Triplet<ScalarT> MatrixTripletT;
+
+	// diagonal mass matrix.
+	FSparseMatrixD AreaMatrix;
+	FSparseMatrixD CotangentInterior;
+	FSparseMatrixD CotangentBoundary;
+	ConstructCotangentLaplacian(DynamicMesh, VertexMap, AreaMatrix, CotangentInterior, CotangentBoundary);
+
+	// Find average entry in the area matrix
+	const int32 Rank = AreaMatrix.cols();
+	double AveArea = 0.;
+	for (int32 i = 0; i < Rank; ++i)
+	{
+		double Area = AreaMatrix.coeff(i, i);
+		checkSlow(Area > 0.);  // Area must be positive.
+		AveArea += Area;
+	}
+	AveArea /= Rank;
+	
+	std::vector<MatrixTripletT> ScaledInvAreaTriplets;
+	ScaledInvAreaTriplets.reserve(Rank);
+	for (int32 i = 0; i < Rank; ++i)
+	{
+		double Area = AreaMatrix.coeff(i, i);
+		double ScaledInvArea = AveArea / Area;
+		if (bClampAreas)
+		{
+			ScaledInvArea = FMathd::Clamp(ScaledInvArea, 0.5, 5.); // when  squared this gives largest scales 100 x smallest
+		}
+
+		ScaledInvAreaTriplets.push_back(MatrixTripletT(i, i, ScaledInvArea));
+	}
+
+	FSparseMatrixD ScaledInvAreaMatrix(Rank, Rank);
+	ScaledInvAreaMatrix.setFromTriplets(ScaledInvAreaTriplets.begin(), ScaledInvAreaTriplets.end());
+	ScaledInvAreaMatrix.makeCompressed();
+
+	LaplacianBoundary = ScaledInvAreaMatrix * CotangentBoundary;
+	LaplacianBoundary.makeCompressed();
+	LaplacianInterior = ScaledInvAreaMatrix * CotangentInterior;
+	LaplacianInterior.makeCompressed();	
+
+	return AveArea;
+}
+
+
+
+
+
+void ConstructLaplacian(const ELaplacianWeightScheme Scheme, const FDynamicMesh3& DynamicMesh, FVertexLinearization& VertexMap, FSparseMatrixD& LaplacianInterior, FSparseMatrixD& LaplacianBoundary)
+{
+
+	switch (Scheme)
+	{
+	default:
+	case ELaplacianWeightScheme::Uniform:
+		ConstructUniformLaplacian(DynamicMesh, VertexMap, LaplacianInterior, LaplacianBoundary);
+		break;
+	case ELaplacianWeightScheme::Umbrella:
+		ConstructUmbrellaLaplacian(DynamicMesh, VertexMap, LaplacianInterior, LaplacianBoundary);
+		break;
+	case ELaplacianWeightScheme::Valence:
+		ConstructValenceWeightedLaplacian(DynamicMesh, VertexMap, LaplacianInterior, LaplacianBoundary);
+		break;
+	case ELaplacianWeightScheme::Cotangent:
+	{
+		bool bClampWeights = false;
+		ConstructScaledCotangentLaplacian(DynamicMesh, VertexMap, LaplacianInterior, LaplacianBoundary, bClampWeights);
+		break;
+	}
+	case ELaplacianWeightScheme::ClampedCotangent:
+	{
+		bool bClampWeights = true;
+		ConstructScaledCotangentLaplacian(DynamicMesh, VertexMap, LaplacianInterior, LaplacianBoundary, bClampWeights);
+		break;
+	}
+	case ELaplacianWeightScheme::MeanValue:
+		ConstructMeanValueWeightLaplacian(DynamicMesh, VertexMap, LaplacianInterior, LaplacianBoundary);
+		break;
+	}
+}
+
+
+
+
+
+
+//
+//
+//
+
+
+static void ExtractBoundaryVerts(const FVertexLinearization& VertexMap, TArray<int32>& BoundaryVerts)
+{
+	int32 NumBoundaryVerts = VertexMap.NumBoundaryVerts();
+	int32 NumInternalVerts = VertexMap.NumVerts() - NumBoundaryVerts;
+	BoundaryVerts.Empty(NumBoundaryVerts);
+
+	const auto& ToId = VertexMap.ToId();
+	for (int32 i = NumInternalVerts; i < VertexMap.NumVerts(); ++i)
+	{
+		int32 VtxId = ToId[i];
+		BoundaryVerts.Add(VtxId);
+	}
+}
+
+TUniquePtr<FSparseMatrixD> ConstructUniformLaplacian(const FDynamicMesh3& DynamicMesh, FVertexLinearization& VertexMap, TArray<int32>* BoundaryVerts)
+{
+	TUniquePtr<FSparseMatrixD> LaplacianMatrix(new FSparseMatrixD);
+	FSparseMatrixD BoundaryMatrix;
+
+	ConstructUniformLaplacian(DynamicMesh, VertexMap, *LaplacianMatrix, BoundaryMatrix);
+
+	if (BoundaryVerts)
+	{
+		ExtractBoundaryVerts(VertexMap, *BoundaryVerts);
+	}
+
+	return LaplacianMatrix;
+}
+
+
+
+TUniquePtr<FSparseMatrixD> ConstructUmbrellaLaplacian(const FDynamicMesh3& DynamicMesh, FVertexLinearization& VertexMap, TArray<int32>* BoundaryVerts)
+{
+	TUniquePtr<FSparseMatrixD> LaplacianMatrix(new FSparseMatrixD);
+	FSparseMatrixD BoundaryMatrix;
+
+	ConstructUmbrellaLaplacian(DynamicMesh, VertexMap, *LaplacianMatrix, BoundaryMatrix);
+
+	if (BoundaryVerts)
+	{
+		ExtractBoundaryVerts(VertexMap, *BoundaryVerts);
+	}
+
+	return LaplacianMatrix;
+}
+
+
+TUniquePtr<FSparseMatrixD> ConstructValenceWeightedLaplacian(const FDynamicMesh3& DynamicMesh, FVertexLinearization& VertexMap, TArray<int32>* BoundaryVerts)
+{
+	TUniquePtr<FSparseMatrixD> LaplacianMatrix(new FSparseMatrixD);
+	FSparseMatrixD BoundaryMatrix;
+
+	ConstructValenceWeightedLaplacian(DynamicMesh, VertexMap, *LaplacianMatrix, BoundaryMatrix);
+
+	if (BoundaryVerts)
+	{
+		ExtractBoundaryVerts(VertexMap, *BoundaryVerts);
+	}
+
+	return LaplacianMatrix;
+
+}
+
+
+TUniquePtr<FSparseMatrixD> ConstructCotangentLaplacian(const FDynamicMesh3& DynamicMesh, FVertexLinearization& VertexMap, FSparseMatrixD& AreaMatrix, TArray<int32>* BoundaryVerts)
+{
+	TUniquePtr<FSparseMatrixD> LaplacianMatrix(new FSparseMatrixD);
+	FSparseMatrixD BoundaryMatrix;
+
+	ConstructCotangentLaplacian(DynamicMesh, VertexMap, AreaMatrix, *LaplacianMatrix, BoundaryMatrix);
+
+	if (BoundaryVerts)
+	{
+		ExtractBoundaryVerts(VertexMap, *BoundaryVerts);
+	}
+
+	return LaplacianMatrix;
+}
+
+
+
+TUniquePtr<FSparseMatrixD> ConstructCotangentLaplacian(const FDynamicMesh3& DynamicMesh, FVertexLinearization& VertexMap, const bool bClampWeights, TArray<int32>* BoundaryVerts)
+{
+	TUniquePtr<FSparseMatrixD> LaplacianMatrix(new FSparseMatrixD);
+	FSparseMatrixD BoundaryMatrix;
+
+	ConstructCotangentLaplacian(DynamicMesh, VertexMap, *LaplacianMatrix, BoundaryMatrix, bClampWeights);
+
+	if (BoundaryVerts)
+	{
+		ExtractBoundaryVerts(VertexMap, *BoundaryVerts);
+	}
+
+	return LaplacianMatrix;
+}
+
+
+
+TUniquePtr<FSparseMatrixD> ConstructMeanValueWeightLaplacian(const FDynamicMesh3& DynamicMesh, FVertexLinearization& VertexMap, TArray<int32>* BoundaryVerts)
+{
+	TUniquePtr<FSparseMatrixD> LaplacianMatrix(new FSparseMatrixD);
+	FSparseMatrixD BoundaryMatrix;
+
+	ConstructMeanValueWeightLaplacian(DynamicMesh, VertexMap, *LaplacianMatrix, BoundaryMatrix);
+
+	if (BoundaryVerts)
+	{
+		ExtractBoundaryVerts(VertexMap, *BoundaryVerts);
+	}
+
+	return LaplacianMatrix;
+}
+
+
+
+
diff --git a/Engine/Plugins/Experimental/GeometryProcessing/Source/DynamicMesh/Private/Solvers/Internal/LaplacianOperators.h b/Engine/Plugins/Experimental/GeometryProcessing/Source/DynamicMesh/Private/Solvers/Internal/LaplacianOperators.h
new file mode 100644
--- /dev/null
+++ b/Engine/Plugins/Experimental/GeometryProcessing/Source/DynamicMesh/Private/Solvers/Internal/LaplacianOperators.h
@@ -0,0 +1,79 @@
+// Copyright Epic Games, Inc. All Rights Reserved.
+
+#pragma once
+
+#include "CoreMinimal.h"
+#include "DynamicMesh3.h"
+#include "Solvers/MeshLaplacian.h"
+#include "Solvers/MeshLinearization.h"
+#include "FSparseMatrixD.h"
+
+
+/**
+* Construct a sparse matrix representation of the requested "Laplacian" type operator.
+*
+* The mesh itself is assumed to have N interior vertices, and M boundary vertices.
+*
+*
+* @param Scheme             Selects the actual type of "laplacian" constructed.
+* @param DynamicMesh        The triangle mesh
+* @param VertexMap          On return, Additional arrays used to map between vertexID and offset in a linear array (i.e. the row).
+*                           The vertices are ordered so that last M ( = VertexMap.NumBoundaryVerts() )  correspond to those on the boundary.
+* @param LaplacianInterior  On return, the laplacian operator that acts on the interior vertices: sparse  N x N matrix
+* @param LaplacianBoundary  On return, the portion of the operator that acts on the boundary vertices: sparse  N x M matrix
+*
+*   LaplacianInterior * Vector_InteriorVerts + LaplacianBoundary * Vector_BoundaryVerts = Full Laplacian applied to interior vertices.
+*/
+void ConstructLaplacian(const ELaplacianWeightScheme Scheme,
+	                   const FDynamicMesh3& DynamicMesh,
+	                   FVertexLinearization& VertexMap,
+	                   FSparseMatrixD& LaplacianInterior,
+	                   FSparseMatrixD& LaplacianBoundary);
+
+
+
+
+
+
+//
+// Functions below are only used in current Unit Test code, and can be removed if/when 
+// those tests are rewritten against ConstructLaplacian() above
+//
+
+void ConstructCotangentLaplacian(
+	const FDynamicMesh3& DynamicMesh, 
+	FVertexLinearization& VertexMap, 
+	FSparseMatrixD& LaplacianInterior, 
+	FSparseMatrixD& LaplacianBoundary, 
+	const bool bClampWeights);
+
+TUniquePtr<FSparseMatrixD> ConstructLaplacian(const ELaplacianWeightScheme Scheme,
+	const FDynamicMesh3& DynamicMesh,
+	FVertexLinearization& VertexMap,
+	TArray<int32>* BoundaryVerts = NULL);
+
+TUniquePtr<FSparseMatrixD> ConstructUniformLaplacian(const FDynamicMesh3& DynamicMesh,
+	FVertexLinearization& VertexMap,
+	TArray<int32>* BoundaryVerts = NULL);
+
+TUniquePtr<FSparseMatrixD> ConstructUmbrellaLaplacian(const FDynamicMesh3& DynamicMesh,
+	FVertexLinearization& VertexMap,
+	TArray<int32>* BoundaryVerts = NULL);
+
+TUniquePtr<FSparseMatrixD>  ConstructValenceWeightedLaplacian(const FDynamicMesh3& DynamicMesh3,
+	FVertexLinearization& VertexMap,
+	TArray<int32>* BoundaryVerts = NULL);
+
+TUniquePtr<FSparseMatrixD> ConstructCotangentLaplacian(const FDynamicMesh3& DynamicMesh,
+	FVertexLinearization& VertexMap,
+	const bool bClampWeights = false,
+	TArray<int32>* BoundaryVerts = NULL);
+
+TUniquePtr<FSparseMatrixD> ConstructCotangentLaplacian(const FDynamicMesh3& DynamicMesh,
+	FVertexLinearization& VertexMap,
+	FSparseMatrixD& AreaMatrix,
+	TArray<int32>* BoundaryVerts = NULL);
+
+TUniquePtr<FSparseMatrixD> ConstructMeanValueWeightLaplacian(const FDynamicMesh3& DynamicMesh,
+	FVertexLinearization& VertexMap,
+	TArray<int32>* BoundaryVerts = NULL);
diff --git a/Engine/Plugins/Experimental/GeometryProcessing/Source/DynamicMesh/Private/Solvers/Internal/MatrixSolver.cpp b/Engine/Plugins/Experimental/GeometryProcessing/Source/DynamicMesh/Private/Solvers/Internal/MatrixSolver.cpp
new file mode 100644
--- /dev/null
+++ b/Engine/Plugins/Experimental/GeometryProcessing/Source/DynamicMesh/Private/Solvers/Internal/MatrixSolver.cpp
@@ -0,0 +1,36 @@
+// Copyright Epic Games, Inc. All Rights Reserved.
+
+#include "MatrixSolver.h"
+
+IMatrixSolverBase::~IMatrixSolverBase() {};
+IIterativeMatrixSolverBase::~IIterativeMatrixSolverBase() {};
+
+TUniquePtr<IMatrixSolverBase> ContructMatrixSolver(const EMatrixSolverType& MatrixSolverType)
+{
+	TUniquePtr<IMatrixSolverBase> ResultPtr;
+
+	switch (MatrixSolverType)
+	{
+	default:
+
+	case  EMatrixSolverType::LU:
+		ResultPtr.Reset(new FLUMatrixSolver());
+		break;
+	case EMatrixSolverType::QR:
+		ResultPtr.Reset(new FQRMatrixSolver());
+		break;
+	case EMatrixSolverType::PCG:
+		ResultPtr.Reset(new FPCGMatrixSolver());
+		break;	
+	case EMatrixSolverType::BICGSTAB:
+		ResultPtr.Reset(new FBiCGMatrixSolver());
+		break;
+#ifndef EIGEN_MPL2_ONLY
+	case EMatrixSolverType::LDLT:
+		ResultPtr.Reset(new FLDLTMatrixSolver());
+		break;
+#endif
+	}
+
+	return ResultPtr;
+}
diff --git a/Engine/Plugins/Experimental/GeometryProcessing/Source/DynamicMesh/Private/Solvers/Internal/MatrixSolver.h b/Engine/Plugins/Experimental/GeometryProcessing/Source/DynamicMesh/Private/Solvers/Internal/MatrixSolver.h
new file mode 100644
--- /dev/null
+++ b/Engine/Plugins/Experimental/GeometryProcessing/Source/DynamicMesh/Private/Solvers/Internal/MatrixSolver.h
@@ -0,0 +1,325 @@
+// Copyright Epic Games, Inc. All Rights Reserved.
+
+#pragma once 
+
+#include "CoreMinimal.h"
+#include "FSparseMatrixD.h"
+#include "Async/ParallelFor.h"
+
+#if defined(_MSC_VER) && USING_CODE_ANALYSIS
+#pragma warning(push)
+#pragma warning(disable : 6011)
+#pragma warning(disable : 6387)
+#pragma warning(disable : 6313)
+#pragma warning(disable : 6294)
+#endif
+THIRD_PARTY_INCLUDES_START
+#include <Eigen/Sparse>
+#include <Eigen/Core>
+#include <Eigen/SparseLU>
+#include <Eigen/SparseQR>
+#include <Eigen/OrderingMethods>
+#include <Eigen/IterativeLinearSolvers> // BiCGSTAB & IncompleteLUT
+#ifndef EIGEN_MPL2_ONLY
+#include <Eigen/SparseCholesky>
+#endif
+#include <Eigen/Dense>
+THIRD_PARTY_INCLUDES_END
+#if defined(_MSC_VER) && USING_CODE_ANALYSIS
+#pragma warning(pop)
+#endif
+
+#include "ProfilingDebugging/ScopedTimers.h"
+
+
+enum class EMatrixSolverType
+{
+	LU /* SparseLU*/,
+	QR /* SparseQR */,
+	BICGSTAB /* Iterative Bi-conjugate gradient */,
+	PCG /* Pre-conditioned conjugate gradient  - requires symmetric positive def.*/,
+#ifndef EIGEN_MPL2_ONLY
+	LDLT /**Not included due to MPL2 */
+#endif
+};
+
+static FString MatrixSolverName(const EMatrixSolverType SolverType)
+{
+	FString String;
+	switch (SolverType)
+	{
+	case EMatrixSolverType::LU:
+		String = FString(TEXT(" Direct LU "));
+		break;
+	case EMatrixSolverType::QR:
+		String = FString(TEXT(" Direct QR "));
+		break;
+	case EMatrixSolverType::BICGSTAB:
+		String = FString(TEXT(" Iterative BiConjugate Gradient "));
+		break;
+	case EMatrixSolverType::PCG:
+		String = FString(TEXT(" Iterative Preconditioned Conjugate Gradient "));
+		break;
+#ifndef EIGEN_MPL2_ONLY
+	case EMatrixSolverType::PCG:
+		String = FString(TEXT(" Direct Cholesky "));
+		break;
+#endif
+	default:
+		check(0);
+	}
+
+	return String;
+}
+
+class FMatrixSolverSettings
+{
+public:
+	EMatrixSolverType  MatrixSolverType;
+	
+	// Used by iterative solvers
+	int32  MaxIterations = 600;
+	double Tolerance = 1e-4;
+};
+
+
+
+// Pure Virtual Base Class used to wrap Eigen Solvers.
+class IMatrixSolverBase
+{
+public:
+	typedef typename FSparseMatrixD::Scalar    ScalarType;
+	typedef typename FSOAPositions::VectorType VectorType;
+
+	IMatrixSolverBase() {}
+	virtual ~IMatrixSolverBase() = 0;
+	virtual bool bIsIterative() const = 0;
+	virtual void Solve(const VectorType& BVector, VectorType& SolVector) const = 0;
+	virtual void Solve(const FSOAPositions& BVectors, FSOAPositions& SolVectors) const = 0;
+
+	virtual void SetUp(const FSparseMatrixD& Matrix, bool bIsSymmetric) = 0;
+	//virtual void SetIterations(int32 MaxIterations) = 0;
+	virtual void Reset() = 0;
+	virtual bool bSucceeded() const = 0;
+};
+
+// Additional methods particular to the iterative solvers
+class IIterativeMatrixSolverBase : public IMatrixSolverBase
+{
+public:
+	IIterativeMatrixSolverBase() {}
+
+	virtual ~IIterativeMatrixSolverBase() = 0;
+	virtual void SetIterations(int32 MaxIterations) = 0;
+	virtual void SetTolerance(double Tol) = 0;
+	virtual void SolveWithGuess(const VectorType& GuessVector, const VectorType& BVector, VectorType& SolVector)  const = 0;
+	virtual void SolveWithGuess(const FSOAPositions& GuessVector, const FSOAPositions& BVector, FSOAPositions& SolVector)  const = 0;
+
+};
+
+// Matrix Solver Factory
+TUniquePtr<IMatrixSolverBase> ContructMatrixSolver(const EMatrixSolverType& MatrixSolverType);
+
+template <typename DirectSolverType>
+class TMatrixSolver : public IMatrixSolverBase
+{
+public:
+	typedef DirectSolverType  SolverType;
+	TMatrixSolver() {}
+	~TMatrixSolver() override {};
+
+	bool bIsIterative() const override { return false; }
+
+	void Solve(const VectorType& BVector, VectorType& SolVector) const override
+	{
+		if (bSetup)
+		{
+			SolVector = MatrixSolver.solve(BVector);
+		}
+	}
+
+	void Solve(const FSOAPositions& BVectors, FSOAPositions& SolVectors) const override
+	{
+		const bool bForceSingleThreaded = false;
+		if (bSetup)
+		{
+			ParallelFor(3, [&](int Dir)
+			{ SolVectors.Array(Dir) = MatrixSolver.solve(BVectors.Array(Dir)); }, bForceSingleThreaded);
+		}
+	}
+
+	void SetUp(const FSparseMatrixD& SparseMatrix, bool bIsSymmetric) override
+	{
+		SetSymmetry(bIsSymmetric);
+		SetUp(SparseMatrix);
+	}
+
+	void SetSymmetry(bool bIsSymmetric);
+
+	//void SetIterations(int32 MaxIterations) override {}
+
+	void Reset() override
+	{
+		//MatrixSolver = FMatrixSolver();
+		bSetup = false;
+	}
+
+	bool bSucceeded() const override
+	{
+		const auto Info = MatrixSolver.info();
+		bool bSuccess = bSetup && (Info == Eigen::ComputationInfo::Success);
+		return bSuccess;
+	}
+
+
+private:
+
+	void SetUp(const FSparseMatrixD& SparseMatrix)
+	{
+		if (SparseMatrix.rows() == 0)
+		{
+			return;
+		}
+		
+		// The analyzePattern could be done just once if
+		// the sparsity pattern of the matrix is fixed.
+		// But testing indicates that takes little time compared 
+		// with factorizing:  e.g. dim 14508 matrix.  Analyze 0.145s, Factorize 0.935s 
+
+		MatrixSolver.analyzePattern(SparseMatrix);
+		MatrixSolver.factorize(SparseMatrix);
+
+		bSetup = true;	
+	}
+private:
+
+	bool bSetup = false;
+	DirectSolverType  MatrixSolver;
+};
+
+
+/**
+* Define wrappers for Direct Matrix Solver Types
+*/
+// Timing Tests:
+//    7k Tri  3.6k Verts.  10 sets of 3 backsolves  0.17s
+//    29k tri  14.5k Verts                           1.3s
+//   100k tri  51.5k Verts                           5.39s
+//   127k tri  63.4k Verts                           3.0s
+typedef TMatrixSolver<Eigen::SparseLU <FSparseMatrixD, Eigen::COLAMDOrdering<int>>>  FLUMatrixSolver;
+typedef TMatrixSolver<Eigen::SparseQR<FSparseMatrixD, Eigen::COLAMDOrdering<int>>>   FQRMatrixSolver;
+
+#ifndef EIGEN_MPL2_ONLY
+// Not included due to MPL2.  But in general much faster than standard LU
+
+// Timing info: 
+//              29k tris 14.5k verts    10 sets of 3 backsolves - 0.42s   Analyze 0.03s Factorize 0.28s
+//              45k tris 23k verts with 10 sets of 3 backsolves - 0.8s    Analyze 0.08s Factorize 0.54s
+//                       49k verts with 10 sets of 3 backsolves - 3.35s   Analyze 0.17s Factorize 2.67s
+//             101k tris 50k verts with 10 sets of 3 backsolves - 1.34s   Analyze 0.12s Factorize 0.84s 
+//             126k tris 63k verts with 10 sets of 3 backsolves - 0.8s    Analyze 0.17s Factorize 0.22s
+//             205k tris 102k verts with 10 sets of 3 backsolves - 3.5s    Analyze 0.41s Factorize 2.2s
+typedef TMatrixSolver< Eigen::SimplicialLDLT<FSparseMatrixD> >   FLDLTMatrixSolver
+template<>
+void TMatrixSolver<typename FLUMatrixSolver::SolverType>::SetSymmetry(bool bIsSymmetric) { check(bIsSymmetric); };
+
+#endif  //  EIGEN_MPL2_ONLY
+
+template<>
+inline void TMatrixSolver<typename FLUMatrixSolver::SolverType>::SetSymmetry(bool bIsSymmetric) { MatrixSolver.isSymmetric(bIsSymmetric); };
+
+template<>
+inline void TMatrixSolver<typename FQRMatrixSolver::SolverType>::SetSymmetry(bool bIsSymmetric) {};
+
+
+template <typename IterativeMatrixSolverType>
+class  TIterativeMatrixSolver : public  IIterativeMatrixSolverBase
+{
+public:
+	typedef IterativeMatrixSolverType  SolverType;
+
+	TIterativeMatrixSolver() { MatrixSolver.setMaxIterations(1000); MatrixSolver.setTolerance(1e-4); }
+	~TIterativeMatrixSolver() override {};
+
+	bool bIsIterative() const override { return true; }
+
+	void Solve(const VectorType& BVector, VectorType& SolVector) const override
+	{
+		if (bSetup)
+		{
+			SolVector = MatrixSolver.solve(BVector);
+		}
+	}
+	void Solve(const FSOAPositions& BVectors, FSOAPositions& SolVectors) const override
+	{
+		const bool bForceSingleThreaded = false;
+		if (bSetup)
+		{
+			ParallelFor(3, [&](int Dir)
+			{ SolVectors.Array(Dir) = MatrixSolver.solve(BVectors.Array(Dir)); }, bForceSingleThreaded);
+		}
+	}
+
+	void SolveWithGuess(const VectorType& GuessVector, const VectorType& BVector, VectorType& SolVector)  const override
+	{
+
+		if (bSetup)
+		{
+			SolVector = MatrixSolver.solveWithGuess(BVector, GuessVector);
+		}
+	}
+
+	void SolveWithGuess(const FSOAPositions& GuessVectors, const FSOAPositions& BVectors, FSOAPositions& SolVectors) const override
+	{
+		const bool bForceSingleThreaded = false;
+		if (bSetup)
+		{
+			ParallelFor(3, [&](int Dir)
+			{ SolVectors.Array(Dir) = MatrixSolver.solveWithGuess(BVectors.Array(Dir), GuessVectors.Array(Dir)); }, bForceSingleThreaded);
+		}
+	}
+	void SetUp(const FSparseMatrixD& SparseMatrix, bool bIsSymmetric) override
+	{
+		SetUp(SparseMatrix);
+	}
+
+	void SetIterations(int32 MaxIterations) override { MatrixSolver.setMaxIterations(MaxIterations); }
+	void SetTolerance(double Tol) override { MatrixSolver.setTolerance(Tol); }
+
+	void Reset() override
+	{
+		//MatrixSolver = FMatrixSolver();
+		bSetup = false;
+	}
+
+	bool bSucceeded() const override
+	{
+		const auto Info = MatrixSolver.info();
+		bool bSuccess = bSetup && (Info == Eigen::ComputationInfo::Success);
+		return bSuccess;
+	}
+
+private:
+	void SetUp(const FSparseMatrixD& SparseMatrix)
+	{
+		MatrixSolver.analyzePattern(SparseMatrix);
+		MatrixSolver.factorize(SparseMatrix);
+
+		bSetup = true;
+
+		int32 n = Eigen::nbThreads();
+	}
+
+private:
+	bool bSetup = false;
+
+	IterativeMatrixSolverType MatrixSolver;
+};
+
+/**
+* Define wrappers for Iterative Matrix Solver Types
+*/
+//typedef Eigen::IncompleteCholesky< double, Eigen::Lower | Eigen::Upper>  FICPreConditioner;
+typedef TIterativeMatrixSolver<Eigen::ConjugateGradient<FSparseMatrixD, Eigen::Lower | Eigen::Upper>>               FPCGMatrixSolver;
+typedef TIterativeMatrixSolver<Eigen::BiCGSTAB<FSparseMatrixD, Eigen::IncompleteLUT<FSparseMatrixD::Scalar, int> >> FBiCGMatrixSolver;
+
diff --git a/Engine/Plugins/Experimental/GeometryProcessing/Source/DynamicMesh/Private/Solvers/MeshLaplacian.cpp b/Engine/Plugins/Experimental/GeometryProcessing/Source/DynamicMesh/Private/Solvers/MeshLaplacian.cpp
new file mode 100644
--- /dev/null
+++ b/Engine/Plugins/Experimental/GeometryProcessing/Source/DynamicMesh/Private/Solvers/MeshLaplacian.cpp
@@ -0,0 +1,67 @@
+// Copyright Epic Games, Inc. All Rights Reserved.
+
+#include "Solvers/MeshLaplacian.h"
+
+
+
+FString LaplacianSchemeName(const ELaplacianWeightScheme Scheme)
+{
+	FString LaplacianName;
+	switch (Scheme)
+	{
+	case ELaplacianWeightScheme::ClampedCotangent:
+		LaplacianName = FString(TEXT("Clamped Cotangent Laplacian"));
+		break;
+	case ELaplacianWeightScheme::Cotangent:
+		LaplacianName = FString(TEXT("Cotangent Laplacian"));
+		break;
+	case ELaplacianWeightScheme::Umbrella:
+		LaplacianName = FString(TEXT("Umbrella Laplacian"));
+		break;
+	case ELaplacianWeightScheme::MeanValue:
+		LaplacianName = FString(TEXT("MeanValue Laplacian"));
+		break;
+	case ELaplacianWeightScheme::Uniform:
+		LaplacianName = FString(TEXT("Uniform Laplacian"));
+		break;
+	case ELaplacianWeightScheme::Valence:
+		LaplacianName = FString(TEXT("Valence Laplacian"));
+		break;
+	default:
+		check(0 && "Unknown Laplacian Weight Scheme Enum");
+	}
+
+	return LaplacianName;
+}
+
+
+
+
+bool bIsSymmetricLaplacian(const ELaplacianWeightScheme Scheme)
+{
+	bool bSymmetric = false;
+	switch (Scheme)
+	{
+	case ELaplacianWeightScheme::ClampedCotangent:
+		bSymmetric = false;
+		break;
+	case ELaplacianWeightScheme::Cotangent:
+		bSymmetric = false;
+		break;
+	case ELaplacianWeightScheme::Umbrella:
+		bSymmetric = false;
+		break;
+	case ELaplacianWeightScheme::MeanValue:
+		bSymmetric = false;
+		break;
+	case ELaplacianWeightScheme::Uniform:
+		bSymmetric = true;
+		break;
+	case ELaplacianWeightScheme::Valence:
+		bSymmetric = true;
+		break;
+	default:
+		check(0);
+	}
+	return bSymmetric;
+}
\ No newline at end of file
diff --git a/Engine/Plugins/Experimental/GeometryProcessing/Source/DynamicMesh/Private/Solvers/MeshSmoothing.cpp b/Engine/Plugins/Experimental/GeometryProcessing/Source/DynamicMesh/Private/Solvers/MeshSmoothing.cpp
new file mode 100644
--- /dev/null
+++ b/Engine/Plugins/Experimental/GeometryProcessing/Source/DynamicMesh/Private/Solvers/MeshSmoothing.cpp
@@ -0,0 +1,126 @@
+// Copyright Epic Games, Inc. All Rights Reserved.
+
+#include "Solvers/MeshSmoothing.h"
+#include "Solvers/Internal/LaplacianMeshSmoother.h"
+
+
+void UE::MeshDeformation::ComputeSmoothing_BiHarmonic(const ELaplacianWeightScheme WeightScheme, const FDynamicMesh3& OriginalMesh,
+	const double Speed, const double Intensity, const int32 NumIterations, TArray<FVector3d>& PositionArray)
+{
+	// This is equivalent to taking a single backward Euler time step of bi-harmonic diffusion
+	// where L is the Laplacian (Del^2) , and L^T L is an approximation of the Del^4.
+	// 
+	// dp/dt = - k*k L^T L[p]
+	// with 
+	// weight = 1 / (k * Sqrt[dt] )
+	//
+	// p^{n+1} + dt * k * k L^TL [p^{n+1}] = p^{n}
+	//
+	// re-write as
+	// L^TL[p^{n+1}] + weight * weight p^{n+1} = weight * weight p^{n}
+
+#ifndef EIGEN_MPL2_ONLY
+	const EMatrixSolverType MatrixSolverType = EMatrixSolverType::LTL;
+#else
+	// const EMatrixSolverType MatrixSolverType = EMatrixSolverType::LU;
+	// const EMatrixSolverType MatrixSolverType = EMatrixSolverType::PCG;
+
+	// The Symmetric Laplacians are SPD, and so are the LtL Operators
+	const EMatrixSolverType MatrixSolverType = (bIsSymmetricLaplacian(WeightScheme)) ? EMatrixSolverType::PCG : EMatrixSolverType::LU;
+
+#endif
+
+
+#ifdef TIME_LAPLACIAN_SMOOTHERS
+	FString DebugLogString = FString::Printf(TEXT("Biharmonic Smoothing of mesh with %d verts "), OriginalMesh.VertexCount()) + LaplacianSchemeName(WeightScheme) + MatrixSolverName(MatrixSolverType);
+
+	FScopedDurationTimeLogger Timer(DebugLogString);
+#endif
+
+	const double TimeStep = Speed * FMath::Min(Intensity, 1.e6);
+
+	FBiHarmonicDiffusionMeshSmoother BiHarmonicDiffusionSmoother(OriginalMesh, WeightScheme);
+
+	BiHarmonicDiffusionSmoother.Integrate_BackwardEuler(MatrixSolverType, NumIterations, TimeStep);
+
+	BiHarmonicDiffusionSmoother.GetPositions(PositionArray);
+
+}
+
+void UE::MeshDeformation::ComputeSmoothing_ImplicitBiHarmonicPCG(const ELaplacianWeightScheme WeightScheme, const FDynamicMesh3& OriginalMesh,
+	const double Speed, const double Weight, const int32 MaxIterations, TArray<FVector3d>& PositionArray)
+{
+	// This is equivalent to taking a single backward Euler time step of bi-harmonic diffusion
+	// where L is the Laplacian (Del^2) , and L^T L is an approximation of the Del^4.
+	// 
+	// dp/dt = - k*k L^T L[p]
+	// with 
+	// weight = 1 / (k * Sqrt[dt] )
+	//
+	// p^{n+1} + dt * k * k L^TL [p^{n+1}] = p^{n}
+	//
+	// re-write as
+	// L^TL[p^{n+1}] + weight * weight p^{n+1} = weight * weight p^{n}
+#ifdef TIME_LAPLACIAN_SMOOTHERS
+	FString DebugLogString = FString::Printf(TEXT("PCG Biharmonic Smoothing of mesh with %d verts "), OriginalMesh.VertexCount()) + LaplacianSchemeName(WeightScheme);
+
+	FScopedDurationTimeLogger Timer(DebugLogString);
+#endif 
+	if (MaxIterations < 1) return;
+
+	FCGBiHarmonicMeshSmoother Smoother(OriginalMesh, WeightScheme);
+
+	// Treat all vertices as constraints with the same weight
+	const bool bPostFix = false;
+
+	for (int32 VertId : OriginalMesh.VertexIndicesItr())
+	{
+		FVector3d Pos = OriginalMesh.GetVertex(VertId);
+
+		Smoother.AddConstraint(VertId, Weight, Pos, bPostFix);
+	}
+
+	Smoother.SetMaxIterations(MaxIterations);
+	Smoother.SetTolerance(1.e-4);
+
+	bool bSuccess = Smoother.ComputeSmoothedMeshPositions(PositionArray);
+
+}
+
+void UE::MeshDeformation::ComputeSmoothing_Diffusion(const ELaplacianWeightScheme WeightScheme, const FDynamicMesh3& OriginalMesh, bool bForwardEuler,
+	const double Speed, const double Intensity, const int32 IterationCount, TArray<FVector3d>& PositionArray)
+{
+#ifndef EIGEN_MPL2_ONLY
+	const EMatrixSolverType MatrixSolverType = EMatrixSolverType::LTL;
+#else
+	const EMatrixSolverType MatrixSolverType = EMatrixSolverType::LU;
+	//const EMatrixSolverType MatrixSolverType = EMatrixSolverType::PCG;
+	//const EMatrixSolverType MatrixSolverType = EMatrixSolverType::BICGSTAB;
+#endif
+
+#ifdef TIME_LAPLACIAN_SMOOTHERS
+	FString DebugLogString = FString::Printf(TEXT("Diffusion Smoothing of mesh with %d verts"), OriginalMesh.VertexCount());
+	if (!bForwardEuler)
+	{
+		DebugLogString += MatrixSolverName(MatrixSolverType);
+	}
+
+	FScopedDurationTimeLogger Timer(DebugLogString);
+#endif
+	if (IterationCount < 1) return;
+
+	FLaplacianDiffusionMeshSmoother Smoother(OriginalMesh, WeightScheme);
+
+	if (bForwardEuler)
+	{
+		Smoother.Integrate_ForwardEuler(IterationCount, Speed);
+	}
+	else
+	{
+		const double TimeStep = Speed * FMath::Min(Intensity, 1.e6);
+		Smoother.Integrate_BackwardEuler(MatrixSolverType, IterationCount, TimeStep);
+	}
+
+	Smoother.GetPositions(PositionArray);
+};
+
diff --git a/Engine/Plugins/Experimental/GeometryProcessing/Source/DynamicMesh/Public/Solvers/ConstrainedMeshDeformer.h b/Engine/Plugins/Experimental/GeometryProcessing/Source/DynamicMesh/Public/Solvers/ConstrainedMeshDeformer.h
new file mode 100644
--- /dev/null
+++ b/Engine/Plugins/Experimental/GeometryProcessing/Source/DynamicMesh/Public/Solvers/ConstrainedMeshDeformer.h
@@ -0,0 +1,60 @@
+// Copyright Epic Games, Inc. All Rights Reserved.
+
+#pragma once
+
+#include "CoreMinimal.h"
+#include "Templates/UniquePtr.h"
+#include "Solvers/ConstrainedMeshSolver.h"
+#include "Solvers/MeshLaplacian.h"
+#include "DynamicMesh3.h"
+
+namespace UE
+{
+	namespace MeshDeformation
+	{
+
+		/**
+		*  Solves the linear system for p_vec
+		*
+		*         ( Transpose(L) * L   + (0  0      )  ) p_vec = source_vec + ( 0              )
+		*		  (                      (0 lambda^2)  )                      ( lambda^2 c_vec )
+		*
+		*   where:  L := laplacian for the mesh,
+		*           source_vec := Transpose(L)*L mesh_vertex_positions
+		*           lambda := weights
+		*           c_vec := constrained positions
+		*
+		* Expected Use:
+		*
+		*   // Create Deformation Solver from Mesh
+		*   TUniquePtr<IConstrainedMeshSolver>  MeshDeformer = ConstructConstrainedMeshDeformer(ELaplacianWeightScheme::ClampedCotangent, DynamicMesh);
+		*
+		*   // Add constraints.
+		*   for..
+		*   {
+		*   	int32 VtxId = ..; double Weight = ..; FVector3d TargetPos = ..;  bool bPostFix = ...;
+		*   	MeshDeformer->AddConstraint(VtxId, Weight, TargetPos, bPostFix);
+		*   }
+		*
+		*   // Solve for new mesh vertex locations
+		*   TArray<FVector3d> PositionBuffer;
+		*   MeshDeformer->Deform(PositionBuffer);
+		*
+		*   // Update Mesh? for (int32 VtxId : DynamicMesh.VertexIndices()) DynamicMesh.SetVertex(VtxId, PositionBuffer[VtxId]);
+		*   ...
+		*
+		*   // Update constraint positions.
+		*   for ..
+		*   {
+		*   	int32 VtxId = ..;  FVector3d TargetPos = ..; bool bPostFix = ...;
+		*	    MeshDeformer->UpdateConstraintPosition(VtxId, TargetPos, bPostFix);
+		*   }
+		*
+		*   // Solve for new vertex locations.
+		*   MeshDeformer->Deform(PositionBuffer);
+		*   // Update Mesh?
+		*/
+		TUniquePtr<IConstrainedMeshSolver> DYNAMICMESH_API ConstructConstrainedMeshDeformer(const ELaplacianWeightScheme WeightScheme, const FDynamicMesh3& DynamicMesh);
+
+	}
+}
\ No newline at end of file
diff --git a/Engine/Plugins/Experimental/GeometryProcessing/Source/DynamicMesh/Public/Solvers/ConstrainedMeshSmoother.h b/Engine/Plugins/Experimental/GeometryProcessing/Source/DynamicMesh/Public/Solvers/ConstrainedMeshSmoother.h
new file mode 100644
--- /dev/null
+++ b/Engine/Plugins/Experimental/GeometryProcessing/Source/DynamicMesh/Public/Solvers/ConstrainedMeshSmoother.h
@@ -0,0 +1,31 @@
+// Copyright Epic Games, Inc. All Rights Reserved.
+
+#pragma once
+
+#include "CoreMinimal.h"
+#include "Templates/UniquePtr.h"
+#include "Solvers/ConstrainedMeshSolver.h"
+#include "Solvers/MeshLaplacian.h"
+#include "DynamicMesh3.h"
+
+namespace UE
+{
+	namespace MeshDeformation
+	{
+		/**
+		*  Solves the linear system for p_vec
+		*
+		*         ( Transpose(L) * L   + (0  0      )  ) p_vec = ( 0              )
+		*		  (                      (0 lambda^2)  )         ( lambda^2 c_vec )
+		*
+		*   where:  L := laplacian for the mesh,
+		*           lambda := weights
+		*           c_vec := constrained positions
+		*
+		* Expected Use: same as the ConstrainedMeshDeformer above.
+		*
+		*/
+		TUniquePtr<IConstrainedMeshSolver> DYNAMICMESH_API ConstructConstrainedMeshSmoother(const ELaplacianWeightScheme WeightScheme, const FDynamicMesh3& DynamicMesh);
+	}
+}
+
diff --git a/Engine/Plugins/Experimental/GeometryProcessing/Source/DynamicMesh/Public/Solvers/ConstrainedMeshSolver.h b/Engine/Plugins/Experimental/GeometryProcessing/Source/DynamicMesh/Public/Solvers/ConstrainedMeshSolver.h
new file mode 100644
--- /dev/null
+++ b/Engine/Plugins/Experimental/GeometryProcessing/Source/DynamicMesh/Public/Solvers/ConstrainedMeshSolver.h
@@ -0,0 +1,45 @@
+// Copyright Epic Games, Inc. All Rights Reserved.
+
+#pragma once
+
+#include "CoreMinimal.h"
+#include "VectorTypes.h"
+
+/**
+ * Interface to a index-based deformation solver for a 3D mesh vertex set that supports weighted point constraints.
+ */
+class DYNAMICMESH_API IConstrainedMeshSolver
+{
+public:
+
+	virtual ~IConstrainedMeshSolver() {}
+
+	// Add or update a weighted positional constraint associated with VtxId
+	// @param bPostFix if true, this constraint will be explicitly applied after the solve
+	virtual void AddConstraint(const int32 VtxId, const double Weight, const FVector3d& Position, const bool bPostFix) = 0;
+
+	// Update or Create a constraint position associated with @param VtxId
+	// @return true if a constraint weight is associated with @param VtxId
+	virtual bool UpdateConstraintPosition(const int32 VtxId, const FVector3d& Position, const bool bPostFix) = 0;
+
+	// Update or Create a constraint weight associated with @param VtxId
+	// @return true if a constraint position is associated with @param VtxId
+	virtual bool UpdateConstraintWeight(const int32 VtxId, const double Weight) = 0;
+
+	// Clear all constraints (Positions and Weights)
+	virtual void ClearConstraints() = 0;
+
+	// Clear all Constraint Weights
+	virtual void ClearConstraintWeights() = 0;
+
+	// Clear all Constraint Positions
+	virtual void ClearConstraintPositions() = 0;
+
+
+	// Test if a non-zero weighted constraint is associated with VtxId 
+	virtual bool IsConstrained(const int32 VtxId) const = 0;
+
+	// Returns the vertex locations of the deformed mesh. 
+	// Note the array may have empty elements as the index matches the Mesh based VtxId.  PositionBuffer[VtxId] = Pos.
+	virtual bool Deform(TArray<FVector3d>& PositionBuffer) = 0;
+};
\ No newline at end of file
diff --git a/Engine/Plugins/Experimental/GeometryProcessing/Source/DynamicMesh/Public/Solvers/LaplacianMatrixAssembly.h b/Engine/Plugins/Experimental/GeometryProcessing/Source/DynamicMesh/Public/Solvers/LaplacianMatrixAssembly.h
new file mode 100644
--- /dev/null
+++ b/Engine/Plugins/Experimental/GeometryProcessing/Source/DynamicMesh/Public/Solvers/LaplacianMatrixAssembly.h
@@ -0,0 +1,615 @@
+// Copyright Epic Games, Inc. All Rights Reserved.
+
+#pragma once
+
+#include "CoreMinimal.h"
+#include "DynamicMesh3.h"
+#include "Solvers/MatrixInterfaces.h"
+#include "Solvers/MeshLinearization.h"
+#include "Solvers/MeshLaplacian.h"
+#include "Solvers/PrecomputedMeshWeightData.h"
+
+
+namespace UE
+{
+	namespace MeshDeformation
+	{
+		/**
+		* Construct a sparse matrix representation of a uniform weighted Laplacian.
+		* The uniform weighted Laplacian is defined solely in terms of the connectivity
+		* of the mesh.  Note, by construction this should be a symmetric matrix.
+		*
+		* The mesh itself is assumed to have N interior vertices, and M boundary vertices.
+		*
+		* Row i represents the Laplacian at vert_i, the non-zero entries correspond to
+		* the incident one-ring vertices vert_j.
+		*
+		* L_{ij} = 1                      if vert_j is in the one-ring of vert_i
+		* L_{ii} = -Sum{ L_{ij}, j != i}
+		*
+		*
+		* @param DynamicMesh        The triangle mesh
+		* @param VertexMap          On return, Additional arrays used to map between vertexID and offset in a linear array (i.e. the row).
+		*                           The vertices are ordered so that last M ( = VertexMap.NumBoundaryVerts() )  correspond to those on the boundary.
+		* @param LaplacianInterior  On return, the laplacian operator that acts on the interior vertices: sparse  N x N matrix
+		* @param LaplacianBoundary  On return, the portion of the operator that acts on the boundary vertices: sparse  N x M matrix
+		*
+		*   LaplacianInterior * Vector_InteriorVerts + LaplacianBoundary * Vector_BoundaryVerts = Full Laplacian applied to interior vertices.
+		*/
+		template<typename RealType>
+		void ConstructUniformLaplacian(const FDynamicMesh3& DynamicMesh, const FVertexLinearization& VertexMap,
+			TSparseMatrixAssembler<RealType>& LaplacianInterior, TSparseMatrixAssembler<RealType>& LaplacianBoundary);
+
+		/**
+		* Construct a sparse matrix representation of an umbrella weighted Laplacian.
+		* This Laplacian is defined solely in terms of the connectivity
+		* of the mesh.  Note, there is no expectation that the resulting matrix will be symmetric.
+		*
+		* The mesh itself is assumed to have N interior vertices, and M boundary vertices.
+		*
+		* Row i represents the Laplacian at vert_i, the non-zero entries correspond to
+		* the incident one-ring vertices vert_j.
+		*
+		* L_{ij} = 1 / valence(of i)                      if vert_j is in the one-ring of vert_i
+		* L_{ii} = -Sum{ L_{ij}, j != i} = -1
+		*
+		*
+		* @param DynamicMesh        The triangle mesh
+		* @param VertexMap          Additional arrays used to map between vertexID and offset in a linear array (i.e. the row).
+		*                           The vertices are ordered so that last M ( = VertexMap.NumBoundaryVerts() )  correspond to those on the boundary.
+		* @param LaplacianInterior  On return, the laplacian operator that acts on the interior vertices: sparse  N x N matrix
+		* @param LaplacianBoundary  On return, the portion of the operator that acts on the boundary vertices: sparse  N x M matrix
+		*
+		*   LaplacianInterior * Vector_InteriorVerts + LaplacianBoundary * Vector_BoundaryVerts = Full Laplacian applied to interior vertices.
+		*/
+		template<typename RealType>
+		void ConstructUmbrellaLaplacian(const FDynamicMesh3& DynamicMesh, const FVertexLinearization& VertexMap,
+			TSparseMatrixAssembler<RealType>& LaplacianInterior, TSparseMatrixAssembler<RealType>& LaplacianBoundary);
+
+		/**
+		* Construct a sparse matrix representation of a valence-weighted Laplacian.
+		* The valence weighted Laplacian is defined solely in terms of the connectivity
+		* of the mesh.  Note, by construction this should be a symmetric matrix.
+		*
+		* The mesh itself is assumed to have N interior vertices, and M boundary vertices.
+		*
+		* Row i represents the Laplacian at vert_i, the non-zero entries correspond to
+		* the incident one-ring vertices vert_j.
+		*
+		* L_{ij} = 1/\sqrt(valence(i) + valence(j))   if vert_j is in the one-ring of vert_i
+		* L_{ii} = -Sum{ L_{ij}, j != i}
+		*
+		* @param DynamicMesh        The triangle mesh
+		* @param VertexMap          Additional arrays used to map between vertexID and offset in a linear array (i.e. the row).
+		*                           The vertices are ordered so that last M ( = VertexMap.NumBoundaryVerts() )  correspond to those on the boundary.
+		* @param LaplacianInterior  On return, the laplacian operator that acts on the interior vertices: sparse  N x N matrix
+		* @param LaplacianBoundary  On return, the portion of the operator that acts on the boundary vertices: sparse  N x M matrix
+		*
+		*   LaplacianInterior * Vector_InteriorVerts + LaplacianBoundary * Vector_BoundaryVerts = Full Laplacian applied to interior vertices.
+		*/
+		template<typename RealType>
+		void ConstructValenceWeightedLaplacian(const FDynamicMesh3& DynamicMesh, const FVertexLinearization& VertexMap,
+			TSparseMatrixAssembler<RealType>& LaplacianInterior, TSparseMatrixAssembler<RealType>& LaplacianBoundary);
+
+		/**
+		* Construct a sparse matrix representation using a meanvalue-weighted  Laplacian.
+		* NB: there is no reason to expect this to be a symmetric matrix.
+		*
+		* The mesh itself is assumed to have N interior vertices, and M boundary vertices.
+		*
+		* @param DynamicMesh        The triangle mesh
+		* @param VertexMap          Additional arrays used to map between vertexID and offset in a linear array (i.e. the row).
+		*                           The vertices are ordered so that last M ( = VertexMap.NumBoundaryVerts() )  correspond to those on the boundary.
+		* @param LaplacianInterior  On return, the laplacian operator that acts on the interior vertices: sparse  N x N matrix
+		* @param LaplacianBoundary  On return, the portion of the operator that acts on the boundary vertices: sparse  N x M matrix
+		*
+		*   LaplacianInterior * Vector_InteriorVerts + LaplacianBoundary * Vector_BoundaryVerts = Full Laplacian applied to interior vertices.
+		*/
+		template<typename RealType>
+		void ConstructMeanValueWeightLaplacian(const FDynamicMesh3& DynamicMesh, const FVertexLinearization& VertexMap,
+			TSparseMatrixAssembler<RealType>& LaplacianInterior, TSparseMatrixAssembler<RealType>& LaplacianBoundary);
+
+
+		/**
+		* Construct a sparse matrix representation using a cotangent-weighted  Laplacian.
+		* but returns the result in two symmetric parts.
+		*
+		* The mesh itself is assumed to have N interior vertices, and M boundary vertices.
+		*
+		*  (AreaMatrix^{-1}) * L_hat  = Cotangent weighted Laplacian.
+		*
+		*
+		* @param DynamicMesh        The triangle mesh
+		* @param VertexMap          Additional arrays used to map between vertexID and offset in a linear array (i.e. the row).
+		*                           The vertices are ordered so that last M ( = VertexMap.NumBoundaryVerts() )  correspond to those on the boundary.
+		* @param AreaMatrix         On return, the mass matrix for the internal vertices.  sparse N x N matrix
+		* @param LaplacianInterior  On return, the laplacian operator that acts on the interior vertices: sparse  N x N matrix - symmetric
+		* @param LaplacianBoundary  On return, the portion of the operator that acts on the boundary vertices: sparse  N x M matrix
+		*
+		*   AreaMatrix^{-1} * ( LaplacianInterior * Vector_InteriorVerts + LaplacianBoundary * Vector_BoundaryVerts) = Full Laplacian applied to interior vertices.
+		*/
+		template<typename RealType>
+		void ConstructCotangentLaplacian(const FDynamicMesh3& DynamicMesh, const FVertexLinearization& VertexMap,
+			TSparseMatrixAssembler<RealType>& AreaMatrix,
+			TSparseMatrixAssembler<RealType>& LaplacianInterior, TSparseMatrixAssembler<RealType>& LaplacianBoundary);
+
+
+		/**
+		* Construct a sparse matrix representation using a pre-multiplied cotangent-weighted Laplacian.
+		* NB: there is no reason to expect this to be a symmetric matrix.
+		*
+		* This computes the laplacian scaled by the average area A_ave:  ie.  LScaled =   A_ave/(2A_i) ( Cot alpha_ij + Cot beta_ij )
+		*
+		* The mesh itself is assumed to have N interior vertices, and M boundary vertices.
+		*
+		* @param DynamicMesh        The triangle mesh
+		* @param VertexMap          Additional arrays used to map between vertexID and offset in a linear array (i.e. the row).
+		*                           The vertices are ordered so that last M ( = VertexMap.NumBoundaryVerts() )  correspond to those on the boundary.
+		* @param LaplacianInterior  On return, scaled laplacian operator that acts on the interior vertices: sparse  N x N matrix
+		* @param LaplacianBoundary  On return, scaled portion of the operator that acts on the boundary vertices: sparse  N x M matrix
+		* @param bClampAreas        Indicates if (A_ave / A_i) should be clamped to (0.5, 5) range.
+		*                           in practice this is desirable when creating the biharmonic operator, but not the mean curvature flow operator
+		*
+		*   LaplacianInterior * Vector_InteriorVerts + LaplacianBoundary * Vector_BoundaryVerts = Full Laplacian applied to interior vertices.
+		*/
+		template<typename RealType>
+		void ConstructCotangentLaplacian(const FDynamicMesh3& DynamicMesh, const FVertexLinearization& VertexMap,
+			TSparseMatrixAssembler<RealType>& LaplacianInterior, TSparseMatrixAssembler<RealType>& LaplacianBoundary,
+			const bool bClampWeights);
+
+	}
+}
+
+
+
+
+
+//
+//
+//  Implementations
+//
+//
+
+
+
+template<typename RealType>
+void UE::MeshDeformation::ConstructUniformLaplacian(const FDynamicMesh3& DynamicMesh, const FVertexLinearization& VertexMap,
+	TSparseMatrixAssembler<RealType>& LaplacianInterior, TSparseMatrixAssembler<RealType>& LaplacianBoundary)
+{
+	//check(VertexMap_is_good)
+	const TArray<int32>& ToMeshV = VertexMap.ToId();
+	const TArray<int32>& ToIndex = VertexMap.ToIndex();
+	const int32 NumVerts = VertexMap.NumVerts();
+	const int32 NumBoundaryVerts = VertexMap.NumBoundaryVerts();
+	const int32 NumInteriorVerts = NumVerts - NumBoundaryVerts;
+
+	// pre-allocate space when possible
+	int32 NumMatrixEntries = ComputeNumMatrixElements(DynamicMesh, ToMeshV);
+	LaplacianInterior.ReserveEntriesFunc(NumMatrixEntries);
+
+	// Construct Laplacian Matrix: loop over verts constructing the corresponding matrix row.
+	// NB: the vertices are ordered with the interior verts first.
+	for (int32 i = 0; i < NumInteriorVerts; ++i)
+	{
+		const int32 VertId = ToMeshV[i];
+		RealType CenterWeight = RealType(0); // equal and opposite the sum of the neighbor weights
+
+		checkSlow(!DynamicMesh.IsBoundaryVertex(VertId));  // we should only be looping over the internal verts
+
+		for (int32 NeighborVertId : DynamicMesh.VtxVerticesItr(VertId))
+		{
+			const int32 j = ToIndex[NeighborVertId];
+			RealType NeighborWeight = RealType(1);
+			CenterWeight += NeighborWeight;
+
+			if (j < NumInteriorVerts)
+			{
+				// add the neighbor 
+				LaplacianInterior.AddEntryFunc(i, j, NeighborWeight);
+			}
+			else
+			{
+				int32 jBoundary = j - NumInteriorVerts;
+				LaplacianBoundary.AddEntryFunc(i, jBoundary, NeighborWeight);
+			}
+		}
+		// add the center
+		LaplacianInterior.AddEntryFunc(i, i, -CenterWeight);
+	}
+}
+
+
+
+
+template<typename RealType>
+void UE::MeshDeformation::ConstructUmbrellaLaplacian(const FDynamicMesh3& DynamicMesh, const FVertexLinearization& VertexMap,
+	TSparseMatrixAssembler<RealType>& LaplacianInterior, TSparseMatrixAssembler<RealType>& LaplacianBoundary)
+{
+	//check(VertexMap_is_good)
+	const TArray<int32>& ToMeshV = VertexMap.ToId();
+	const TArray<int32>& ToIndex = VertexMap.ToIndex();
+	const int32 NumVerts = VertexMap.NumVerts();
+	const int32 NumBoundaryVerts = VertexMap.NumBoundaryVerts();
+	const int32 NumInteriorVerts = NumVerts - NumBoundaryVerts;
+
+	// pre-allocate space when possible
+	int32 NumMatrixEntries = ComputeNumMatrixElements(DynamicMesh, ToMeshV);
+	LaplacianInterior.ReserveEntriesFunc(NumMatrixEntries);
+
+	// Cache valency of each vertex.
+	// Number of non-zero elements in the i'th row = 1 + OneRingSize(i)
+	TArray<int32> OneRingSize;
+	{
+		OneRingSize.SetNumUninitialized(NumVerts);
+		for (int32 i = 0; i < NumVerts; ++i)
+		{
+			const int32 VertId = ToMeshV[i];
+			OneRingSize[i] = DynamicMesh.GetVtxEdgeCount(VertId);
+		}
+	}
+
+	// Construct Laplacian Matrix: loop over verts constructing the corresponding matrix row.
+
+	for (int32 i = 0; i < NumInteriorVerts; ++i)
+	{
+		const int32 VertId = ToMeshV[i];
+		const int32 Valence = OneRingSize[i];
+		RealType InvValence = (Valence != 0) ? (RealType)1.0 / RealType(Valence) : 0.;
+
+		checkSlow(!DynamicMesh.IsBoundaryVertex(VertId));
+
+		for (int32 NeighborVertId : DynamicMesh.VtxVerticesItr(VertId))
+		{
+			const int32 j = ToIndex[NeighborVertId];
+
+			// add the neighbor 
+			if (j < NumInteriorVerts)
+			{
+				LaplacianInterior.AddEntryFunc(i, j, InvValence);
+			}
+			else
+			{
+				int32 jBoundary = j - NumInteriorVerts;
+				LaplacianBoundary.AddEntryFunc(i, jBoundary, InvValence);
+			}
+		}
+		// add the center
+		LaplacianInterior.AddEntryFunc(i, i, -RealType(1));
+
+	}
+}
+
+
+
+
+
+
+
+template<typename RealType>
+void UE::MeshDeformation::ConstructValenceWeightedLaplacian(const FDynamicMesh3& DynamicMesh, const FVertexLinearization& VertexMap,
+	TSparseMatrixAssembler<RealType>& LaplacianInterior, TSparseMatrixAssembler<RealType>& LaplacianBoundary)
+{
+	const TArray<int32>& ToMeshV = VertexMap.ToId();
+	const TArray<int32>& ToIndex = VertexMap.ToIndex();
+	const int32 NumVerts = VertexMap.NumVerts();
+	const int32 NumBoundaryVerts = VertexMap.NumBoundaryVerts();
+	const int32 NumInteriorVerts = NumVerts - NumBoundaryVerts;
+
+	// pre-allocate space when possible
+	int32 NumMatrixEntries = ComputeNumMatrixElements(DynamicMesh, ToMeshV);
+	LaplacianInterior.ReserveEntriesFunc(NumMatrixEntries);
+
+	// Cache valency of each vertex.
+	// Number of non-zero elements in the i'th row = 1 + OneRingSize(i)
+	TArray<int32> OneRingSize;
+	{
+		OneRingSize.SetNumUninitialized(NumVerts);
+
+		for (int32 i = 0; i < NumVerts; ++i)
+		{
+			const int32 VertId = ToMeshV[i];
+			OneRingSize[i] = DynamicMesh.GetVtxEdgeCount(VertId);
+		}
+	}
+
+	// Construct Laplacian Matrix: loop over verts constructing the corresponding matrix row.
+	for (int32 i = 0; i < NumInteriorVerts; ++i)
+	{
+		const int32 VertId = ToMeshV[i];
+		const int32 IOneRingSize = OneRingSize[i];
+
+		RealType CenterWeight = RealType(0); // equal and opposite the sum of the neighbor weights
+		for (int32 NeighborVertId : DynamicMesh.VtxVerticesItr(VertId))
+		{
+			const int32 j = ToIndex[NeighborVertId];
+			const int32 JOneRingSize = OneRingSize[j];
+
+			RealType NeighborWeight = RealType(1) / TMathUtil<RealType>::Sqrt(IOneRingSize + JOneRingSize);
+			CenterWeight += NeighborWeight;
+
+			if (j < NumInteriorVerts)
+			{
+				// add the neighbor 
+				LaplacianInterior.AddEntryFunc(i, j, NeighborWeight);
+			}
+			else
+			{
+				int32 jBoundary = j - NumInteriorVerts;
+				LaplacianBoundary.AddEntryFunc(i, jBoundary, NeighborWeight);
+			}
+		}
+		// add the center
+		LaplacianInterior.AddEntryFunc(i, i, -CenterWeight);
+	}
+}
+
+
+
+
+template<typename RealType>
+void UE::MeshDeformation::ConstructMeanValueWeightLaplacian(const FDynamicMesh3& DynamicMesh, const FVertexLinearization& VertexMap,
+	TSparseMatrixAssembler<RealType>& LaplacianInterior, TSparseMatrixAssembler<RealType>& LaplacianBoundary)
+{
+	const TArray<int32>& ToMeshV = VertexMap.ToId();
+	const TArray<int32>& ToIndex = VertexMap.ToIndex();
+	const int32 NumVerts = VertexMap.NumVerts();
+	const int32 NumBoundaryVerts = VertexMap.NumBoundaryVerts();
+	const int32 NumInteriorVerts = NumVerts - NumBoundaryVerts;
+
+	// pre-allocate space when possible
+	int32 NumMatrixEntries = ComputeNumMatrixElements(DynamicMesh, ToMeshV);
+	LaplacianInterior.ReserveEntriesFunc(NumMatrixEntries);
+
+	// Map the triangles.
+	FTriangleLinearization TriangleMap(DynamicMesh);
+	const TArray<int32>& ToMeshTri = TriangleMap.ToId();
+	const TArray<int32>& ToTriIdx = TriangleMap.ToIndex();
+	const int32 NumTris = TriangleMap.NumTris();
+
+	// Create an array that holds all the geometric information we need for each triangle.
+	TArray<MeanValueTriangleData> TriangleDataArray;
+	ConstructTriangleDataArray<MeanValueTriangleData>(DynamicMesh, TriangleMap, TriangleDataArray);
+
+	// Construct Laplacian Matrix: loop over verts constructing the corresponding matrix row.
+	//                             skipping the boundary verts for later use.
+	for (int32 i = 0; i < NumInteriorVerts; ++i)
+	{
+		const int32 IVertId = ToMeshV[i]; // I - the row
+		double WeightII = 0.; // accumulate to equal and opposite the sum of the neighbor weights
+
+		// for each connecting edge
+		for (int32 EdgeId : DynamicMesh.VtxEdgesItr(IVertId))
+		{
+			// [v0, v1, t0, t1]:  NB: both t0 & t1 exist since IVert isn't a boundary vert.
+			FIndex4i Edge = DynamicMesh.GetEdge(EdgeId);
+
+			// the other vert in the edge - identifies the matrix column
+			const int32 JVertId = (Edge[0] == IVertId) ? Edge[1] : Edge[0];  // J - the column
+
+			// Get the cotangents for this edge.
+			const int32 Tri0Idx = ToTriIdx[Edge[2]];
+			const auto& Tri0Data = TriangleDataArray[Tri0Idx];
+			double TanHalfAngleSum = Tri0Data.GetTanHalfAngle(IVertId);
+			double EdgeLength = FMathd::Max(1.e-5, Tri0Data.GetEdgeLenght(EdgeId)); // Clamp the length
+
+			// The second triangle will be invalid if this is an edge!
+			TanHalfAngleSum += (Edge[3] != FDynamicMesh3::InvalidID) ? TriangleDataArray[ToTriIdx[Edge[3]]].GetTanHalfAngle(IVertId) : 0.;
+
+			double WeightIJ = TanHalfAngleSum / EdgeLength;
+			WeightII += WeightIJ;
+
+			const int32 j = ToIndex[JVertId];
+			if (j < NumInteriorVerts)
+			{
+				LaplacianInterior.AddEntryFunc(i, j, (RealType)WeightIJ);
+			}
+			else
+			{
+				int32 jBoundary = j - NumInteriorVerts;
+				LaplacianBoundary.AddEntryFunc(i, jBoundary, (RealType)WeightIJ);
+			}
+		}
+		LaplacianInterior.AddEntryFunc(i, i, (RealType)-WeightII);
+	}
+}
+
+
+
+
+
+template<typename RealType>
+void UE::MeshDeformation::ConstructCotangentLaplacian(const FDynamicMesh3& DynamicMesh, const FVertexLinearization& VertexMap,
+	TSparseMatrixAssembler<RealType>& AreaMatrix,
+	TSparseMatrixAssembler<RealType>& LaplacianInterior, TSparseMatrixAssembler<RealType>& LaplacianBoundary)
+{
+	const TArray<int32>& ToMeshV = VertexMap.ToId();
+	const TArray<int32>& ToIndex = VertexMap.ToIndex();
+	const int32 NumVerts = VertexMap.NumVerts();
+	const int32 NumBoundaryVerts = VertexMap.NumBoundaryVerts();
+	const int32 NumInteriorVerts = NumVerts - NumBoundaryVerts;
+
+	// pre-allocate space when possible
+	int32 NumMatrixEntries = ComputeNumMatrixElements(DynamicMesh, ToMeshV);
+	LaplacianInterior.ReserveEntriesFunc(NumMatrixEntries);
+
+	// Create the mapping of triangles
+	FTriangleLinearization TriangleMap(DynamicMesh);
+	const TArray<int32>& ToMeshTri = TriangleMap.ToId();
+	const TArray<int32>& ToTriIdx = TriangleMap.ToIndex();
+	const int32 NumTris = TriangleMap.NumTris();
+
+	// Clear space for the areas
+	AreaMatrix.ReserveEntriesFunc(NumVerts);
+
+	// Create an array that holds all the geometric information we need for each triangle.
+	TArray<CotanTriangleData> CotangentTriangleDataArray;
+	ConstructTriangleDataArray<CotanTriangleData>(DynamicMesh, TriangleMap, CotangentTriangleDataArray);
+
+	// Construct Laplacian Matrix: loop over verts constructing the corresponding matrix row.
+	//                             store the id of the boundary verts for later use.
+	for (int32 i = 0; i < NumInteriorVerts; ++i)
+	{
+		const int32 IVertId = ToMeshV[i]; // I - the row
+
+		// Compute the Voronoi area for this vertex.
+		double WeightArea = 0.;
+		for (int32 TriId : DynamicMesh.VtxTrianglesItr(IVertId))
+		{
+			const int32 TriIdx = ToTriIdx[TriId];
+			const CotanTriangleData& TriData = CotangentTriangleDataArray[TriIdx];
+
+			// The three VertIds for this triangle.
+			const FIndex3i TriVertIds = DynamicMesh.GetTriangle(TriId);
+
+			// Which of the corners is IVertId?
+			int32 Offset = 0;
+			while (TriVertIds[Offset] != IVertId)
+			{
+				Offset++;
+				checkSlow(Offset < 3);
+			}
+
+			WeightArea += TriData.VoronoiArea[Offset];
+		}
+
+		double WeightII = 0.; // accumulate to equal and opposite the sum of the neighbor weights
+
+		// for each connecting edge
+		for (int32 EdgeId : DynamicMesh.VtxEdgesItr(IVertId))
+		{
+			// [v0, v1, t0, t1]:  NB: both t0 & t1 exist since IVert isn't a boundary vert.
+			FIndex4i Edge = DynamicMesh.GetEdge(EdgeId);
+
+
+			// the other vert in the edge - identifies the matrix column
+			const int32 JVertId = (Edge[0] == IVertId) ? Edge[1] : Edge[0];  // J - the column
+
+			checkSlow(JVertId != IVertId);
+
+			// Get the cotangents for this edge.
+			const int32 Tri0Idx = ToTriIdx[Edge[2]];
+			const CotanTriangleData& Tri0Data = CotangentTriangleDataArray[Tri0Idx];
+			const double CotanAlpha = Tri0Data.GetOpposingCotangent(EdgeId);
+
+			// The second triangle will be invalid if this is an edge!
+			const double CotanBeta = (Edge[3] != FDynamicMesh3::InvalidID) ? CotangentTriangleDataArray[ToTriIdx[Edge[3]]].GetOpposingCotangent(EdgeId) : 0.;
+
+			double WeightIJ = 0.5 * (CotanAlpha + CotanBeta);
+			WeightII += WeightIJ;
+
+			const int32 j = ToIndex[JVertId];
+			if (j < NumInteriorVerts)
+			{
+				LaplacianInterior.AddEntryFunc(i, j, (RealType)WeightIJ);
+			}
+			else
+			{
+				int32 jBoundary = j - NumInteriorVerts;
+				LaplacianBoundary.AddEntryFunc(i, jBoundary, (RealType)WeightIJ);
+			}
+		}
+		LaplacianInterior.AddEntryFunc(i, i, (RealType)-WeightII);
+		AreaMatrix.AddEntryFunc(i, i, (RealType)WeightArea);
+	}
+}
+
+
+
+
+
+template<typename RealType>
+void UE::MeshDeformation::ConstructCotangentLaplacian(const FDynamicMesh3& DynamicMesh, const FVertexLinearization& VertexMap,
+	TSparseMatrixAssembler<RealType>& LaplacianInterior, TSparseMatrixAssembler<RealType>& LaplacianBoundary,
+	const bool bClampWeights)
+{
+	const TArray<int32>& ToMeshV = VertexMap.ToId();
+	const TArray<int32>& ToIndex = VertexMap.ToIndex();
+	const int32 NumVerts = VertexMap.NumVerts();
+	const int32 NumBoundaryVerts = VertexMap.NumBoundaryVerts();
+	const int32 NumInteriorVerts = NumVerts - NumBoundaryVerts;
+
+	// pre-allocate space when possible
+	int32 NumMatrixEntries = ComputeNumMatrixElements(DynamicMesh, ToMeshV);
+	LaplacianInterior.ReserveEntriesFunc(NumMatrixEntries);
+
+	// Map the triangles.
+	FTriangleLinearization TriangleMap(DynamicMesh);
+	const TArray<int32>& ToMeshTri = TriangleMap.ToId();
+	const TArray<int32>& ToTriIdx = TriangleMap.ToIndex();
+	const int32 NumTris = TriangleMap.NumTris();
+
+	// Create an array that holds all the geometric information we need for each triangle.
+	TArray<CotanTriangleData> CotangentTriangleDataArray;
+	ConstructTriangleDataArray<CotanTriangleData>(DynamicMesh, TriangleMap, CotangentTriangleDataArray);
+
+	// Construct Laplacian Matrix: loop over verts constructing the corresponding matrix row.
+	//                             skipping the boundary verts for later use.
+	for (int32 i = 0; i < NumInteriorVerts; ++i)
+	{
+		const int32 IVertId = ToMeshV[i]; // I - the row
+
+		// Compute the Voronoi area for this vertex.
+		double WeightArea = 0.;
+		for (int32 TriId : DynamicMesh.VtxTrianglesItr(IVertId))
+		{
+			const int32 TriIdx = ToTriIdx[TriId];
+			const CotanTriangleData& TriData = CotangentTriangleDataArray[TriIdx];
+
+			// The three VertIds for this triangle.
+			const FIndex3i TriVertIds = DynamicMesh.GetTriangle(TriId);
+
+			// Which of the corners is IVertId?
+			int32 Offset = 0;
+			while (TriVertIds[Offset] != IVertId)
+			{
+				Offset++;
+				checkSlow(Offset < 3);
+			}
+
+			WeightArea += TriData.VoronoiArea[Offset];
+		}
+
+		double WeightII = 0.; // accumulate to equal and opposite the sum of the neighbor weights
+
+		// for each connecting edge
+		for (int32 EdgeId : DynamicMesh.VtxEdgesItr(IVertId))
+		{
+			// [v0, v1, t0, t1]:  NB: both t0 & t1 exist since IVert isn't a boundary vert.
+			FIndex4i Edge = DynamicMesh.GetEdge(EdgeId);
+
+			// the other vert in the edge - identifies the matrix column
+			const int32 JVertId = (Edge[0] == IVertId) ? Edge[1] : Edge[0];  // J - the column
+			checkSlow(JVertId != IVertId);
+
+			// Get the cotangents for this edge.
+			const int32 Tri0Idx = ToTriIdx[Edge[2]];
+			const CotanTriangleData& Tri0Data = CotangentTriangleDataArray[Tri0Idx];
+			const double CotanAlpha = Tri0Data.GetOpposingCotangent(EdgeId);
+
+			// The second triangle will be invalid if this is an edge!
+			const double CotanBeta = (Edge[3] != FDynamicMesh3::InvalidID) ? CotangentTriangleDataArray[ToTriIdx[Edge[3]]].GetOpposingCotangent(EdgeId) : 0.;
+
+			double WeightIJ = 0.5 * (CotanAlpha + CotanBeta);
+			if (bClampWeights)
+			{
+				WeightIJ = FMathd::Clamp(WeightIJ, -1.e5 * WeightArea, 1.e5 * WeightArea);
+			}
+
+			WeightII += WeightIJ;
+
+			const int32 j = ToIndex[JVertId];
+			if (j < NumInteriorVerts)
+			{
+				LaplacianInterior.AddEntryFunc(i, j, (RealType)(WeightIJ / WeightArea));
+			}
+			else
+			{
+				int32 jBoundary = j - NumInteriorVerts;
+				LaplacianBoundary.AddEntryFunc(i, jBoundary, (RealType)(WeightIJ / WeightArea));
+			}
+		}
+		LaplacianInterior.AddEntryFunc(i, i, (RealType)(-WeightII / WeightArea));
+	}
+}
+
+
diff --git a/Engine/Plugins/Experimental/GeometryProcessing/Source/DynamicMesh/Public/Solvers/MeshLaplacian.h b/Engine/Plugins/Experimental/GeometryProcessing/Source/DynamicMesh/Public/Solvers/MeshLaplacian.h
new file mode 100644
--- /dev/null
+++ b/Engine/Plugins/Experimental/GeometryProcessing/Source/DynamicMesh/Public/Solvers/MeshLaplacian.h
@@ -0,0 +1,39 @@
+// Copyright Epic Games, Inc. All Rights Reserved.
+
+#pragma once
+
+#include "CoreMinimal.h"
+
+
+/**
+ * Types of Laplacian weights. 
+ * This enum is used to indicate which type of Mesh Laplacian should be used in various Mesh Deformation solvers.
+ * See LaplacianMatrixAssembly.h for specifics
+ */
+enum class ELaplacianWeightScheme
+{
+	/** Weight on each vertex is 1 */
+	Uniform,
+	/** Weight on each vertex is 1/N, where N is valence of vertex */
+	Umbrella,
+	/** Weight on each vertex is 1/ Sqrt(Valence + Nbr_Valence) */
+	Valence,
+	/** Mean-Value weight on each vertex, where Mean Value weight is defined by formula 9 in https://www.inf.usi.ch/hormann/papers/Floater.2006.AGC.pdf */
+	MeanValue,
+	/**  */
+	Cotangent,
+	/**  */
+	ClampedCotangent
+};
+
+
+/**
+* Utility to map the enum names for debuging etc.
+*/
+FString DYNAMICMESH_API LaplacianSchemeName(const ELaplacianWeightScheme Scheme);
+
+
+/**
+* @return true if Laplacian scheme produces a symmetric matrix
+*/
+bool DYNAMICMESH_API bIsSymmetricLaplacian(const ELaplacianWeightScheme Scheme);
diff --git a/Engine/Plugins/Experimental/GeometryProcessing/Source/DynamicMesh/Public/Solvers/MeshLinearization.h b/Engine/Plugins/Experimental/GeometryProcessing/Source/DynamicMesh/Public/Solvers/MeshLinearization.h
new file mode 100644
--- /dev/null
+++ b/Engine/Plugins/Experimental/GeometryProcessing/Source/DynamicMesh/Public/Solvers/MeshLinearization.h
@@ -0,0 +1,125 @@
+// Copyright Epic Games, Inc. All Rights Reserved.
+
+#pragma once
+
+#include "CoreMinimal.h"
+#include "Util/ElementLinearization.h"
+#include "DynamicMesh3.h"
+
+
+/**
+* Used linearize the VtxIds in a mesh as a single array and allow mapping from array offset to mesh VtxId.
+* Generally, the array offset will correspond to a matrix row when forming a Laplacian.
+* The last NumBoundaryVerts are the boundary verts.  This may be zero.
+*/
+class FVertexLinearization : public FElementLinearization
+{
+public:
+	FVertexLinearization() = default;
+	FVertexLinearization(const FDynamicMesh3& DynamicMesh)
+	{
+		Reset(DynamicMesh);
+	}
+
+	void Reset(const FDynamicMesh3& DynamicMesh)
+	{
+		Empty();
+		FElementLinearization::Populate(DynamicMesh.MaxVertexID(), DynamicMesh.VertexCount(), DynamicMesh.VertexIndicesItr());
+		RemapBoundaryVerts(DynamicMesh);
+	}
+
+	int32 NumVerts() const { return FElementLinearization::NumIds(); }
+
+	int32 NumBoundaryVerts() const { return NumBndryVerts; }
+
+private:
+
+	// Moves the boundary verts to the end of the arrays and records the number of boundary verts
+	void RemapBoundaryVerts(const FDynamicMesh3& DynamicMesh)
+	{
+		int32 VertCount = NumVerts();
+
+		// Collect the BoundaryVerts and the internal verts in two array
+		TArray<int32> BoundaryVertIds;
+		TArray<int32> TmpToIdMap;
+		TmpToIdMap.Reserve(VertCount);
+		for (int32 i = 0, I = ToIdMap.Num(); i < I; ++i)
+		{
+			int32 VtxId = ToIdMap[i];
+			// Test if the vertex has a one ring
+			
+			bool bEmptyOneRing = true;
+			for (int NeighborVertId : DynamicMesh.VtxVerticesItr(VtxId))
+			{
+				bEmptyOneRing = false;
+				break;
+			};
+				
+
+			if (bEmptyOneRing || DynamicMesh.IsBoundaryVertex(VtxId))
+			{
+				BoundaryVertIds.Add(VtxId);
+			}
+			else
+			{
+				TmpToIdMap.Add(VtxId);
+			}
+		}
+
+		// The number of boundary verts
+		NumBndryVerts = BoundaryVertIds.Num();
+
+		// Merge the boundary verts at the tail 
+		// Add the Boundary verts at the end of the array.
+		TmpToIdMap.Append(BoundaryVertIds);
+	
+		// rebuild the 'to' Index
+		for (int32 i = 0, I = ToIndexMap.Num(); i < I; ++i)
+		{
+			ToIndexMap[i] = FDynamicMesh3::InvalidID;
+		}
+		for (int32 i = 0, I = TmpToIdMap.Num(); i < I; ++i)
+		{
+			int32 Id = TmpToIdMap[i];
+			ToIndexMap[Id] = i;
+		}
+
+		// swap the temp
+		Swap(TmpToIdMap, ToIdMap);
+	}
+
+private:
+
+	FVertexLinearization(const FVertexLinearization&);
+
+	int32 NumBndryVerts = 0;
+};
+
+
+
+/**
+* Used linearize the TriIds in a mesh as a single array and allow mapping from array offset to mesh TriId.
+*
+*/
+class FTriangleLinearization : public FElementLinearization
+{
+public:
+	FTriangleLinearization() = default;
+
+	FTriangleLinearization(const FDynamicMesh3& DynamicMesh)
+	{
+		Reset(DynamicMesh);
+	}
+
+	void Reset(const FDynamicMesh3& DynamicMesh)
+	{
+		Empty();
+		FElementLinearization::Populate(DynamicMesh.MaxTriangleID(), DynamicMesh.TriangleCount(), DynamicMesh.TriangleIndicesItr());
+	}
+
+	int32 NumTris() const { return FElementLinearization::NumIds(); }
+
+
+private:
+	FTriangleLinearization(const FTriangleLinearization&);
+};
diff --git a/Engine/Plugins/Experimental/GeometryProcessing/Source/DynamicMesh/Public/Solvers/MeshSmoothing.h b/Engine/Plugins/Experimental/GeometryProcessing/Source/DynamicMesh/Public/Solvers/MeshSmoothing.h
new file mode 100644
--- /dev/null
+++ b/Engine/Plugins/Experimental/GeometryProcessing/Source/DynamicMesh/Public/Solvers/MeshSmoothing.h
@@ -0,0 +1,58 @@
+// Copyright Epic Games, Inc. All Rights Reserved.
+
+#pragma once
+
+#include "CoreMinimal.h"
+#include "DynamicMesh3.h"
+#include "Solvers/MeshLaplacian.h"
+
+namespace UE
+{
+	namespace MeshDeformation
+	{
+		/**
+		*
+		* Note: for discussion of implicit / explicit integration of diffusion and biharmonic equations
+		*       see "Implicit Fairing of Irregular Meshes using Diffusion and Curvature Flow" - M Desbrun 99.
+		*       although the following suggests an additional source term could be included in the implicit solve for better accuracy.
+		*       or "Generalized Surface Flows for Mesh Processing" Eckstein et al. 2007
+		*/
+
+		/**
+		* This is equivalent to taking a single backward Euler time step of bi-harmonic diffusion
+		* where L is the Laplacian (Del^2) , and L^T L is an approximation of the Del^4.
+		*
+		* dp/dt = - k*k L^T L[p]
+		*
+		* p^{n+1} + dt * k * k L^TL [p^{n+1}] = p^{n}
+		*
+		* re-write as
+		* L^TL[p^{n+1}] + weight * weight p^{n+1} = weight * weight p^{n}
+		* with
+		* weight = 1 / (k * Sqrt[dt] )
+		*
+		* The result is returned in the PositionArray
+		*/
+		void DYNAMICMESH_API ComputeSmoothing_BiHarmonic(const ELaplacianWeightScheme WeightingScheme, const FDynamicMesh3& OriginalMesh,
+			const double Speed, const double Weight, const int32 NumIterations, TArray<FVector3d>& PositionArray);
+
+		void  DYNAMICMESH_API ComputeSmoothing_ImplicitBiHarmonicPCG(const ELaplacianWeightScheme WeightScheme, const FDynamicMesh3& OriginalMesh,
+			const double Speed, const double Weight, const int32 MaxIterations, TArray<FVector3d>& PositionArray);
+
+		/**
+		* This is equivalent to forward or backward Euler time steps of the diffusion equation
+		*
+		* dp/dt = L[p]
+		*
+		* p^{n+1} = p^{n} + dt L[p^{n}]
+		*
+		* with dt = Speed / Max(|w_ii|)
+		*
+		* here w_ii are the diagonal values of L
+		*/
+		void  DYNAMICMESH_API ComputeSmoothing_Diffusion(const ELaplacianWeightScheme WeightScheme, const FDynamicMesh3& OriginalMesh, bool bForwardEuler,
+			const double Speed, double Weight, const int32 NumIterations, TArray<FVector3d>& PositionArray);
+
+
+	}
+}
\ No newline at end of file
diff --git a/Engine/Plugins/Experimental/GeometryProcessing/Source/DynamicMesh/Public/Solvers/PrecomputedMeshWeightData.h b/Engine/Plugins/Experimental/GeometryProcessing/Source/DynamicMesh/Public/Solvers/PrecomputedMeshWeightData.h
new file mode 100644
--- /dev/null
+++ b/Engine/Plugins/Experimental/GeometryProcessing/Source/DynamicMesh/Public/Solvers/PrecomputedMeshWeightData.h
@@ -0,0 +1,396 @@
+// Copyright Epic Games, Inc. All Rights Reserved.
+
+#include "DynamicMesh3.h"
+
+namespace UE
+{
+	namespace MeshDeformation
+	{
+
+		// Utility to compute the number of elements in the sparse laplacian matrix
+		int32 ComputeNumMatrixElements(const FDynamicMesh3& DynamicMesh, const TArray<int32>& ToVtxId)
+		{
+			const int32 NumVerts = ToVtxId.Num();
+			TArray<int32> OneRingSize;
+			{
+				OneRingSize.SetNumUninitialized(NumVerts);
+
+				for (int32 i = 0; i < NumVerts; ++i)
+				{
+					const int32 VertId = ToVtxId[i];
+					OneRingSize[i] = DynamicMesh.GetVtxEdgeCount(VertId);
+				}
+			}
+
+			// Compute the total number of entries in the sparse matrix
+			int32 NumMatrixEntries = 0;
+			{
+				for (int32 i = 0; i < NumVerts; ++i)
+				{
+					NumMatrixEntries += 1 + OneRingSize[i]; // myself plus my neighbors
+				}
+
+			}
+
+			return NumMatrixEntries;
+		}
+
+
+
+
+		/**
+		* The per-triangle data used in constructing the cotangent weighted laplacian.
+		*
+		*/
+		class CotanTriangleData
+		{
+		public:
+
+			typedef FVector3d TriangleVertices[3];
+
+			CotanTriangleData() = default;
+
+			CotanTriangleData(const CotanTriangleData& Other)
+			{
+
+				Cotangent[0] = Other.Cotangent[0];
+				Cotangent[1] = Other.Cotangent[1];
+				Cotangent[2] = Other.Cotangent[2];
+
+				VoronoiArea[0] = Other.VoronoiArea[0];
+				VoronoiArea[1] = Other.VoronoiArea[1];
+				VoronoiArea[2] = Other.VoronoiArea[2];
+
+				OppositeEdge[0] = Other.OppositeEdge[0];
+				OppositeEdge[1] = Other.OppositeEdge[1];
+				OppositeEdge[2] = Other.OppositeEdge[2];
+
+			}
+
+
+			CotanTriangleData(const FDynamicMesh3& DynamicMesh, int32 TriId)
+			{
+				Initialize(DynamicMesh, TriId);
+			}
+
+			void Initialize(const FDynamicMesh3& DynamicMesh, int32 SrcTriId)
+			{
+				TriId = SrcTriId;
+
+				// edges: ab, bc, ca
+				FIndex3i EdgeIds = DynamicMesh.GetTriEdges(TriId);
+
+				FVector3d VertA, VertB, VertC;
+				DynamicMesh.GetTriVertices(TriId, VertA, VertB, VertC);
+
+				const FVector3d EdgeAB(VertB - VertA);
+				const FVector3d EdgeAC(VertC - VertA);
+				const FVector3d EdgeBC(VertC - VertB);
+
+
+				OppositeEdge[0] = EdgeIds[1]; // EdgeBC is opposite vert A
+				OppositeEdge[1] = EdgeIds[2]; // EdgeAC is opposite vert B
+				OppositeEdge[2] = EdgeIds[0]; // EdgeAB is opposite vert C
+
+
+				// NB: didn't use VectorUtil::Area() so we can re-use the Edges.
+				//     also this formulation of area is always positive.
+
+				const double TwiceArea = EdgeAB.Cross(EdgeAC).Length();
+
+				// NB: Area = 1/2 || EdgeA X EdgeB ||  where EdgeA and EdgeB are any two edges in the triangle.
+
+				// Compute the Voronoi areas
+				// 
+				// From Discrete Differential-Geometry Operators for Triangulated 2-Manifolds (Meyer, Desbrun, Schroder, Barr)
+				// http://www.geometry.caltech.edu/pubs/DMSB_III.pdf
+				//    Given triangle P,Q, R the voronoi area at P is given by
+				//    Area = (1/8) * ( |PR|**2 Cot<Q  + |PQ|**2 Cot<R ) 
+
+				if (TwiceArea > 2. * SmallTriangleArea)
+				{
+
+
+
+					// Compute the cotangent of the angle between V1 and V2 
+					// as the ratio  V1.Dot.V2 / || V1 Cross V2 ||
+
+					// Cotangent[i] is cos(theta)/sin(theta) at the i'th vertex.
+
+					Cotangent[0] = EdgeAB.Dot(EdgeAC) / TwiceArea;
+					Cotangent[1] = -EdgeAB.Dot(EdgeBC) / TwiceArea;
+					Cotangent[2] = EdgeAC.Dot(EdgeBC) / TwiceArea;
+
+					if (bIsObtuse())
+					{
+						const double Area = 0.5 * TwiceArea;
+
+						// Voronoi inappropriate case.  Instead use Area(T)/2 at obtuse corner
+						// and Area(T) / 4 at the other corners.
+
+						VoronoiArea[0] = 0.25 * Area;
+						VoronoiArea[1] = 0.25 * Area;
+						VoronoiArea[2] = 0.25 * Area;
+
+						for (int i = 0; i < 3; ++i)
+						{
+							if (Cotangent[i] < 0.)
+							{
+								VoronoiArea[i] = 0.5 * Area;
+							}
+						}
+					}
+					else
+					{
+						// If T is non-obtuse.
+
+						const double EdgeABSqLength = EdgeAB.SquaredLength();
+						const double EdgeACSqLength = EdgeAC.SquaredLength();
+						const double EdgeBCSqLength = EdgeBC.SquaredLength();
+
+						VoronoiArea[0] = EdgeABSqLength * Cotangent[2] + EdgeACSqLength * Cotangent[1];
+						VoronoiArea[1] = EdgeABSqLength * Cotangent[2] + EdgeBCSqLength * Cotangent[0];
+						VoronoiArea[2] = EdgeACSqLength * Cotangent[1] + EdgeBCSqLength * Cotangent[0];
+
+						const double Inv8 = .125; // 1/8
+
+						VoronoiArea[0] *= Inv8;
+						VoronoiArea[1] *= Inv8;
+						VoronoiArea[2] *= Inv8;
+					}
+				}
+				else
+				{
+					// default small triangle - equilateral 
+					double CotOf60 = 1. / FMathd::Sqrt(3.f);
+					Cotangent[0] = CotOf60;
+					Cotangent[1] = CotOf60;
+					Cotangent[2] = CotOf60;
+
+					VoronoiArea[0] = SmallTriangleArea / 3.;
+					VoronoiArea[1] = SmallTriangleArea / 3.;
+					VoronoiArea[2] = SmallTriangleArea / 3.;
+				}
+			}
+
+			int32 GetLocalEdgeIdx(const int32 DynamicsMeshEdgeId) const
+			{
+				int32 Result = -1;
+				if (DynamicsMeshEdgeId == OppositeEdge[0])
+				{
+					Result = 0;
+				}
+				else if (DynamicsMeshEdgeId == OppositeEdge[1])
+				{
+					Result = 1;
+				}
+				else if (DynamicsMeshEdgeId == OppositeEdge[2])
+				{
+					Result = 2;
+				}
+				return Result;
+			}
+
+			/** helper to return the cotangent of the angle opposite the given edge
+			*
+			*   @param DynamicsMeshEdgeId is the id used by FDynamicMesh3 for this edge.
+			*   @param bValid will false on return if the requested edge is not part of this triangle
+			*   @return Cotangent of the opposite angle.
+			*/
+			double GetOpposingCotangent(const int32 DynamicsMeshEdgeId, bool& bValid) const
+			{
+
+				double OpposingCotangent = -1.;
+				bValid = false;
+				int32 LocalEdgeIdx = GetLocalEdgeIdx(DynamicsMeshEdgeId);
+				if (LocalEdgeIdx > -1)
+				{
+					OpposingCotangent = Cotangent[LocalEdgeIdx];
+					bValid = true;
+				}
+
+				return OpposingCotangent;
+			}
+
+			double GetOpposingCotangent(const int32 DynamicsMeshEdgeId) const
+			{
+				bool bValid;
+				double OpposingCotangent = GetOpposingCotangent(DynamicsMeshEdgeId, bValid);
+				checkSlow(bValid);
+				return OpposingCotangent;
+			}
+
+			bool bIsObtuse() const
+			{
+				return (Cotangent[0] < 0.f || Cotangent[1] < 0.f || Cotangent[2] < 0.f);
+			}
+
+
+			// The "floor" for triangle area.
+			// NB: the cotan laplacian has terms ~ 1/TriArea
+			//     and the deformation matrix has terms ~ 1/TriArea**2
+
+			static constexpr double SmallTriangleArea = 1.e-4;
+
+			// testing
+			int32 TriId = -1;
+
+			/** Total byte count: 6 double + 3 int32 = 60 bytes.   */
+
+			// Cotangent[i] is cos(theta)/sin(theta) at the i'th vertex.
+
+			double Cotangent[3] = { 0. };
+
+			// VoronoiArea[i] is the voronoi area about the i'th vertex in this triangle.
+
+			double VoronoiArea[3] = { 0. };
+
+			// OppositeEdge[i] = Corresponding DynamicsMesh3::EdgeId for the edge that is opposite
+			// the i'th vertex in this triangle
+
+			int32 OppositeEdge[3] = { -1 };
+
+		};
+
+
+		/**
+		* The per-triangle data used in constructing the mean-value weighted laplacian.
+		*
+		*/
+		class MeanValueTriangleData
+		{
+
+		public:
+
+			MeanValueTriangleData(const MeanValueTriangleData& Other)
+				: TriId(Other.TriId)
+				, TriVtxIds(Other.TriVtxIds)
+				, TriEdgeIds(Other.TriEdgeIds)
+				, bDegenerate(Other.bDegenerate)
+			{
+				EdgeLength[0] = Other.EdgeLength[0];
+				EdgeLength[1] = Other.EdgeLength[1];
+				EdgeLength[2] = Other.EdgeLength[2];
+
+				TanHalfAngle[0] = Other.TanHalfAngle[0];
+				TanHalfAngle[1] = Other.TanHalfAngle[1];
+				TanHalfAngle[2] = Other.TanHalfAngle[2];
+			}
+
+			void Initialize(const FDynamicMesh3& DynamicMesh, int32 SrcTriId)
+			{
+				TriId = SrcTriId;
+
+				// VertAId, VertBId, VertCId
+				TriVtxIds = DynamicMesh.GetTriangle(TriId);
+				TriEdgeIds = DynamicMesh.GetTriEdges(TriId);
+
+				FVector3d VertA, VertB, VertC;
+				DynamicMesh.GetTriVertices(TriId, VertA, VertB, VertC);
+
+				const FVector3d EdgeAB(VertB - VertA);
+				const FVector3d EdgeAC(VertC - VertA);
+				const FVector3d EdgeBC(VertC - VertB);
+
+				EdgeLength[0] = EdgeAB.Length();
+				EdgeLength[1] = EdgeAC.Length();
+				EdgeLength[2] = EdgeBC.Length();
+
+				constexpr double SmallEdge = 1e-4;
+
+				bDegenerate = (EdgeLength[0] < SmallEdge || EdgeLength[1] < SmallEdge || EdgeLength[2] < SmallEdge);
+
+				// Compute tan(angle/2) = Sqrt[ (1-cos) / (1 + cos)]
+
+				const double ABdotAC = EdgeAB.Dot(EdgeAC);
+				const double BCdotBA = -EdgeBC.Dot(EdgeAB);
+				const double CAdotCB = EdgeAC.Dot(EdgeBC);
+
+
+				const double RegularizingConst = 1.e-6; // keeps us from dividing by zero when making tan[180/2] = sin[90]/cos[90] = inf
+				TanHalfAngle[0] = (EdgeLength[0] * EdgeLength[1] - ABdotAC) / (EdgeLength[0] * EdgeLength[1] + ABdotAC + RegularizingConst);
+				TanHalfAngle[1] = (EdgeLength[0] * EdgeLength[2] - BCdotBA) / (EdgeLength[0] * EdgeLength[2] + BCdotBA + RegularizingConst);
+				TanHalfAngle[2] = (EdgeLength[1] * EdgeLength[2] - CAdotCB) / (EdgeLength[1] * EdgeLength[2] + CAdotCB + RegularizingConst);
+
+				// The ABS is just a precaution.. mathematically these should all be positive, but very small angles may result in negative values.
+
+				TanHalfAngle[0] = FMathd::Sqrt(FMathd::Abs(TanHalfAngle[0])); // at vertA
+				TanHalfAngle[1] = FMathd::Sqrt(FMathd::Abs(TanHalfAngle[1])); // at vertB
+				TanHalfAngle[2] = FMathd::Sqrt(FMathd::Abs(TanHalfAngle[2])); // at vertC
+#if 0
+		// testing
+				double totalAngle = 2. * (std::atan(TanHalfAngle[0]) + std::atan(TanHalfAngle[1]) + std::atan(TanHalfAngle[2]));
+
+				double angleError = M_PI - totalAngle;
+#endif
+			}
+
+			// return Tan(angle / 2) for the corner indicated by this vert id.
+			double GetTanHalfAngle(int32 VtxId) const
+			{
+				int32 Offset = 0;
+
+				while (VtxId != TriVtxIds[Offset])
+				{
+					Offset++;
+					checkSlow(Offset < 3);
+				}
+
+				return TanHalfAngle[Offset];
+			}
+
+			// return the length of the indicated edge
+			double GetEdgeLenght(int32 EdgeId) const
+			{
+				int32 Offset = 0;
+
+				while (EdgeId != TriEdgeIds[Offset])
+				{
+					Offset++;
+					checkSlow(Offset < 3);
+				}
+
+				return EdgeLength[Offset];
+			}
+
+
+			int32 TriId = -1;
+			FIndex3i TriVtxIds;
+			FIndex3i TriEdgeIds;
+
+			bool bDegenerate = true;
+
+			double EdgeLength[3] = { 0. };
+			double TanHalfAngle[3] = { 0. };
+		};
+
+
+
+
+		/**
+		* Return and array in triangle order that holds the per-triangle derived data needed
+		*/
+		template <typename TriangleDataType>
+		void ConstructTriangleDataArray(const FDynamicMesh3& DynamicMesh, const FTriangleLinearization& TriangleLinearization, TArray<TriangleDataType>& TriangleDataArray)
+		{
+			const int32 NumTris = TriangleLinearization.NumTris();
+			TriangleDataArray.SetNumUninitialized(NumTris);
+
+			const auto& ToTriIdx = TriangleLinearization.ToIndex();
+
+			for (int32 i = 0; i < NumTris; ++i)
+			{
+				// Current triangle
+
+				const int32 TriIdx = ToTriIdx[i];
+
+				// Compute all the geometric data needed for this triangle.
+
+				TriangleDataArray[i].Initialize(DynamicMesh, TriIdx);
+			}
+		}
+
+
+
+	}
+}
\ No newline at end of file
diff --git a/Engine/Plugins/Experimental/GeometryProcessing/Source/GeometricObjects/Public/Solvers/MatrixInterfaces.h b/Engine/Plugins/Experimental/GeometryProcessing/Source/GeometricObjects/Public/Solvers/MatrixInterfaces.h
new file mode 100644
--- /dev/null
+++ b/Engine/Plugins/Experimental/GeometryProcessing/Source/GeometricObjects/Public/Solvers/MatrixInterfaces.h
@@ -0,0 +1,23 @@
+// Copyright Epic Games, Inc. All Rights Reserved.
+
+#pragma once
+
+#include "CoreMinimal.h"
+
+/**
+ * Generic adapter for building a Sparse Matrix. Generally sparse-matrix assembly
+ * code emits matrix entries as tuples of form (i,j,Value).
+ */
+template<typename RealType>
+class TSparseMatrixAssembler
+{
+public:
+	/** Hint to reserve space for at least this many entries */
+	TUniqueFunction<void(int32)> ReserveEntriesFunc;
+
+	/** Add matrix entry tuple (i,j,Value) */
+	TUniqueFunction<void(int32, int32, RealType)> AddEntryFunc;
+};
+
+typedef TSparseMatrixAssembler<double> FSparseMatrixAssemblerd;
+typedef TSparseMatrixAssembler<float> FSparseMatrixAssemblerf;
diff --git a/Engine/Plugins/Experimental/GeometryProcessing/Source/GeometricObjects/Public/Util/ElementLinearization.h b/Engine/Plugins/Experimental/GeometryProcessing/Source/GeometricObjects/Public/Util/ElementLinearization.h
new file mode 100644
--- /dev/null
+++ b/Engine/Plugins/Experimental/GeometryProcessing/Source/GeometricObjects/Public/Util/ElementLinearization.h
@@ -0,0 +1,152 @@
+// Copyright Epic Games, Inc. All Rights Reserved.
+
+#pragma once
+
+#include <CoreMinimal.h>
+#include "IndexTypes.h"
+#include "VectorTypes.h"
+
+
+/**
+* FElementLinearization maps a potentially-sparse index list into a linear array.
+* Used linearize things like VtxIds of a mesh as a single array and allow bidirectional mapping between array offset and mesh VtxId.
+* Linearized array offset can then be used for things like matrix row indexing when building a Laplacian matrix.
+*/
+class FElementLinearization
+{
+public:
+	FElementLinearization() = default;
+
+	// Lookup   ToVtxId(Index) = VtxId;
+	const TArray<int32>& ToId() const { return ToIdMap; }
+
+	// Lookup   ToIndex(VtxId) = Index;  may return FDynamicMesh3::InvalidID 
+	const TArray<int32>& ToIndex() const { return ToIndexMap; }
+
+	int32 NumIds() const { return ToIdMap.Num(); }
+
+	// Following the FDynamicMesh3 convention this is really MaxId + 1
+	int32 MaxId() const { return ToIndexMap.Num(); }
+
+	void Empty() { ToIdMap.Empty();  ToIndexMap.Empty(); }
+
+	template<typename IterableType>
+	void Populate(const int32 MaxId, const int32 Count, IterableType Iterable)
+	{
+		ToIndexMap.SetNumUninitialized(MaxId);
+		ToIdMap.SetNumUninitialized(Count);
+
+		for (int32 i = 0; i < MaxId; ++i)
+		{
+			ToIndexMap[i] = IndexConstants::InvalidID;
+		}
+
+		// create the mapping
+		{
+			int32 N = 0;
+			for (int32 Id : Iterable)
+			{
+				ToIdMap[N] = Id;
+				ToIndexMap[Id] = N;
+				N++;
+			}
+		}
+	}
+
+protected:
+	TArray<int32>  ToIdMap;
+	TArray<int32>  ToIndexMap;
+
+private:
+	FElementLinearization(const FElementLinearization&);
+};
+
+
+
+/**
+ * Structure-of-Array (SoA) storage for a list of 3-vectors
+ */
+template<typename RealType>
+class TVector3Arrays
+{
+protected:
+	TArray<RealType> XVector;
+	TArray<RealType> YVector;
+	TArray<RealType> ZVector;
+
+public:
+
+	TVector3Arrays(int32 Size)
+	{
+		XVector.SetNum(Size);
+		YVector.SetNum(Size);
+		ZVector.SetNum(Size);
+	}
+
+	TVector3Arrays()
+	{}
+
+	void SetZero(int32 NumElements)
+	{
+		XVector.Reset(0);
+		XVector.SetNumZeroed(NumElements, false);
+		YVector.Reset(0);
+		YVector.SetNumZeroed(NumElements, false);
+		ZVector.Reset(0);
+		ZVector.SetNumZeroed(NumElements, false);
+	}
+
+	// Test that all the arrays have the same given size.
+	bool bHasSize(int32 Size) const
+	{
+		return (XVector.Num() == Size && YVector.Num() == Size && ZVector.Num() == Size);
+	}
+
+	int32 Num() const
+	{
+		int32 Size = XVector.Num();
+		if (!bHasSize(Size))
+		{
+			Size = -1;
+		}
+		return Size;
+	}
+
+	RealType X(int32 i) const
+	{
+		return XVector[i];
+	}
+
+	RealType Y(int32 i) const
+	{
+		return YVector[i];
+	}
+
+	RealType Z(int32 i) const
+	{
+		return ZVector[i];
+	}
+
+	void SetX(int32 i, const RealType& Value)
+	{
+		XVector[i] = Value;
+	}
+
+	void SetY(int32 i, const RealType& Value)
+	{
+		YVector[i] = Value;
+	}
+
+	void SetZ(int32 i, const RealType& Value)
+	{
+		ZVector[i] = Value;
+	}
+
+	void SetXYZ(int32 i, const FVector3<RealType>& Value)
+	{
+		XVector[i] = Value.X;
+		YVector[i] = Value.Y;
+		ZVector[i] = Value.Z;
+	}
+};
+
diff --git a/Engine/Plugins/Experimental/GeometryProcessing/Source/MeshSolverUtilities/MeshSolverUtilities.Build.cs b/Engine/Plugins/Experimental/GeometryProcessing/Source/MeshSolverUtilities/MeshSolverUtilities.Build.cs
deleted file mode 100644
diff --git a/Engine/Plugins/Experimental/GeometryProcessing/Source/MeshSolverUtilities/Private/ConstrainedPoissonSolver.h b/Engine/Plugins/Experimental/GeometryProcessing/Source/MeshSolverUtilities/Private/ConstrainedPoissonSolver.h
deleted file mode 100644
diff --git a/Engine/Plugins/Experimental/GeometryProcessing/Source/MeshSolverUtilities/Private/FSOAPositions.h b/Engine/Plugins/Experimental/GeometryProcessing/Source/MeshSolverUtilities/Private/FSOAPositions.h
deleted file mode 100644
diff --git a/Engine/Plugins/Experimental/GeometryProcessing/Source/MeshSolverUtilities/Private/FSparseMatrixD.h b/Engine/Plugins/Experimental/GeometryProcessing/Source/MeshSolverUtilities/Private/FSparseMatrixD.h
deleted file mode 100644
diff --git a/Engine/Plugins/Experimental/GeometryProcessing/Source/MeshSolverUtilities/Private/LaplacianMeshSmoother.cpp b/Engine/Plugins/Experimental/GeometryProcessing/Source/MeshSolverUtilities/Private/LaplacianMeshSmoother.cpp
deleted file mode 100644
diff --git a/Engine/Plugins/Experimental/GeometryProcessing/Source/MeshSolverUtilities/Private/LaplacianMeshSmoother.h b/Engine/Plugins/Experimental/GeometryProcessing/Source/MeshSolverUtilities/Private/LaplacianMeshSmoother.h
deleted file mode 100644
diff --git a/Engine/Plugins/Experimental/GeometryProcessing/Source/MeshSolverUtilities/Private/LaplacianOperators.cpp b/Engine/Plugins/Experimental/GeometryProcessing/Source/MeshSolverUtilities/Private/LaplacianOperators.cpp
deleted file mode 100644
diff --git a/Engine/Plugins/Experimental/GeometryProcessing/Source/MeshSolverUtilities/Private/LaplacianOperators.h b/Engine/Plugins/Experimental/GeometryProcessing/Source/MeshSolverUtilities/Private/LaplacianOperators.h
deleted file mode 100644
diff --git a/Engine/Plugins/Experimental/GeometryProcessing/Source/MeshSolverUtilities/Private/MatrixSolver.cpp b/Engine/Plugins/Experimental/GeometryProcessing/Source/MeshSolverUtilities/Private/MatrixSolver.cpp
deleted file mode 100644
diff --git a/Engine/Plugins/Experimental/GeometryProcessing/Source/MeshSolverUtilities/Private/MatrixSolver.h b/Engine/Plugins/Experimental/GeometryProcessing/Source/MeshSolverUtilities/Private/MatrixSolver.h
deleted file mode 100644
diff --git a/Engine/Plugins/Experimental/GeometryProcessing/Source/MeshSolverUtilities/Private/MeshElementLinearizations.h b/Engine/Plugins/Experimental/GeometryProcessing/Source/MeshSolverUtilities/Private/MeshElementLinearizations.h
deleted file mode 100644
diff --git a/Engine/Plugins/Experimental/GeometryProcessing/Source/MeshSolverUtilities/Private/MeshSolverUtilitiesModule.cpp b/Engine/Plugins/Experimental/GeometryProcessing/Source/MeshSolverUtilities/Private/MeshSolverUtilitiesModule.cpp
deleted file mode 100644
diff --git a/Engine/Plugins/Experimental/GeometryProcessing/Source/MeshSolverUtilities/Public/MeshSmoothingUtilities.h b/Engine/Plugins/Experimental/GeometryProcessing/Source/MeshSolverUtilities/Public/MeshSmoothingUtilities.h
deleted file mode 100644
diff --git a/Engine/Plugins/Experimental/GeometryProcessing/Source/MeshSolverUtilities/Public/MeshSolverUtilitiesModule.h b/Engine/Plugins/Experimental/GeometryProcessing/Source/MeshSolverUtilities/Public/MeshSolverUtilitiesModule.h
deleted file mode 100644
diff --git a/Engine/Plugins/Experimental/MeshModelingToolset/Source/MeshModelingTools/MeshModelingTools.Build.cs b/Engine/Plugins/Experimental/MeshModelingToolset/Source/MeshModelingTools/MeshModelingTools.Build.cs
--- a/Engine/Plugins/Experimental/MeshModelingToolset/Source/MeshModelingTools/MeshModelingTools.Build.cs
+++ b/Engine/Plugins/Experimental/MeshModelingToolset/Source/MeshModelingTools/MeshModelingTools.Build.cs
@@ -42,7 +42,6 @@ public class MeshModelingTools : ModuleRules
 				"MeshConversion",
 				"MeshDescription",
                 "StaticMeshDescription",
-                "MeshSolverUtilities",
 				"ModelingComponents",
 				"ModelingOperators",
 
diff --git a/Engine/Plugins/Experimental/MeshModelingToolset/Source/MeshModelingTools/Private/DeformMeshPolygonsTool.cpp b/Engine/Plugins/Experimental/MeshModelingToolset/Source/MeshModelingTools/Private/DeformMeshPolygonsTool.cpp
--- a/Engine/Plugins/Experimental/MeshModelingToolset/Source/MeshModelingTools/Private/DeformMeshPolygonsTool.cpp
+++ b/Engine/Plugins/Experimental/MeshModelingToolset/Source/MeshModelingTools/Private/DeformMeshPolygonsTool.cpp
@@ -14,9 +14,264 @@
 #include "Async/ParallelFor.h"
 #include "Containers/BitArray.h"
 
+// WAT
+#include "DynamicMesh/Private/Solvers/Internal/LaplacianMeshSmoother.h"
+#include "ModelingOperators/Public/ModelingTaskTypes.h"
+
+
 #define LOCTEXT_NAMESPACE "UDeformMeshPolygonsTool"
 
 
+class FDeformTask;
+
+
+
+//Stores per-vertex data needed by the laplacian deformer object
+//TODO: May be a candidate for a subclass of the FGroupTopologyLaplacianDeformer
+struct FDeformerVertexConstraintData
+{
+	FDeformerVertexConstraintData& operator=(const FDeformerVertexConstraintData& other)
+	{
+		Position = other.Position;
+		Weight = other.Weight;
+		bPostFix = other.bPostFix;
+		return *this;
+	}
+
+	FVector3d Position;
+	double Weight{ 0.0 };
+	bool bPostFix{ false };
+};
+
+/**
+*	FDeformTask is an object which wraps an asynchronous task to be run multiple times on a separate thread.
+*	The Laplacian deformation process requires the use of potentially large sparse matrices and sparse multiplication.
+*
+*   Expected usage:
+*
+*
+*   // define constraints.  Need Constraints[VertID] to hold the constraints for the corresponding vertex.
+*   TArray<FDeformerVertexConstraintData> Constraints;
+*   ....
+*
+*   // populate with the VertexIDs of the vertices that are in the region you wish to deform.
+*   TArray<int32> SrcVertIDs;  //Basically a mini-index buffer.
+*   ...
+*
+*   // Create or reuse a laplacian deformation task.
+*   FDeformTask*   DeformTask = New FDeformTask(WeightScheme);
+*
+*   // the deformer will have to build a new mesh that represents the regions in SrcVertIDs;
+*   // but set this to false on subsequent calls to UpdateDeformer if the SrcVertIDs array hasn't changed.
+*   bool bRequiresRegion = true;
+
+*   DefTask->UpdateDeformer(WeightScheme, Mesh, Constraints, SrcVertIDs, bRequiresRegion);
+*
+*   DeformTask->DoWork();  or DeformTask->StartBackgroundTask(); //which calls DoWork on background thread.
+*
+*  // wheh DeformTask->IsDone == true; you can copy the results back to the mesh
+*  DeformTask->ExportResults(Mesh);
+*
+* Note: if only the positions in the Constraints change (e.g. handle positions) then subsequent calls
+* to UpdateDeformer() and DoWork() will be much faster as the matrix system will not be rebuilt or re-factored
+*/
+class FConstrainedMeshDeformerTask : public FNonAbandonableTask
+{
+	friend class FAsyncTask<FDeformTask>;
+public:
+
+	enum
+	{
+		INACTIVE_SUBSET_ID = -1
+	};
+
+	FConstrainedMeshDeformerTask(const ELaplacianWeightScheme SelectedWeightScheme)
+	{
+		LaplacianWeightScheme = SelectedWeightScheme;
+	}
+
+	virtual ~FConstrainedMeshDeformerTask() {};
+
+	//NO idea what this is meant to do. Performance analysis maybe? Scheduling data?
+	FORCEINLINE TStatId GetStatId() const
+	{
+		RETURN_QUICK_DECLARE_CYCLE_STAT(FConstrainedMeshDeformerTask, STATGROUP_ThreadPoolAsyncTasks);
+	}
+
+
+	/** Called by the main thread in the tool, this copies the Constraint buffer right before the task begins on another thread.
+	  * Ensures the FConstrainedMeshDeformer is using correct mesh subset and the selected settings, then updates on change in properties, i.e. weight scheme */
+	void UpdateDeformer(const ELaplacianWeightScheme SelectedWeightScheme, const FDynamicMesh3& Mesh, const TArray<FDeformerVertexConstraintData>& ConstraintArray, const TArray<int32>& SrcIDBufferSubset, bool bNewTransaction, const FRichCurve* Curve);
+
+	/** Required by the FAsyncTaskExecutor */
+	void SetAbortSource(bool* bAbort) { bAbortSource = bAbort; };
+
+	/** Called by the FAsyncTask<FDeformTask> object for background computation. */
+	void DoWork();
+
+
+	/** Updates the positions in the target mesh for regions that correspond to the subset mesh */
+	void ExportResults(FDynamicMesh3& TargetMesh) const;
+
+private:
+
+	/** Creates the mesh (i.e. SubsetMesh) that corresponds to the region of the SrcMesh defined by the partial index buffer SrcIDBufferSubset */
+	void InitializeSubsetMesh(const FDynamicMesh3& SrcMesh, const TArray<int32>& SrcIDBufferSubset);
+
+	/** Attenuates the weights of the constraints using the selected curve */
+	void ApplyAttenuation();
+
+	/** Denotes the weight scheme being used by the running background task. Changes when selected property changes in editor. */
+	ELaplacianWeightScheme LaplacianWeightScheme;
+
+	/** positions for each vertex in the subset mesh - for use in the deformer */
+	TArray<FVector3d> SubsetPositionBuffer;
+
+	/** constraint data for each vertex in subset mesh - for use by the deformer*/
+	TArray<FDeformerVertexConstraintData> SubsetConstraintBuffer;
+
+	FRichCurve WeightAttenuationCurve;
+
+	/** True only for the first update, and then false for the duration of the Input transaction
+	  * It's passed in and copied in UpdateDeformer() */
+	bool bIsNewTransaction = true;
+
+	/** When true, the constraint weights will be attenuated based on distance using the provided curve object*/
+	bool bAttenuateWeights = false;
+
+	/** The abort bool used by the Task Deleter */
+	bool* bAbortSource = nullptr;
+
+	/** Used to initialize the array mapping, updated during the UpdateDeformer() function */
+	int SrcMeshMaxVertexID;
+
+	/** A subset of the original mesh */
+	FDynamicMesh3 SubsetMesh;
+
+	/** Maps Subset Mesh VertexID to Src Mesh VertexID */
+	TArray<int32> SubsetVertexIDToSrcVertexIDMap;
+
+	/** Laplacian deformer object gets rebuilt each new transaction */
+	TUniquePtr<FConstrainedMeshDeformer> ConstrainedDeformer;
+
+private:
+	FConstrainedMeshDeformerTask();
+
+};
+
+class FGroupTopologyLaplacianDeformer : public FGroupTopologyDeformer
+{
+
+public:
+
+	FGroupTopologyLaplacianDeformer() = default;
+
+	virtual ~FGroupTopologyLaplacianDeformer();
+
+	/** Used to begin a procedural addition of modified vertices */
+	inline void ResetModifiedVertices()
+	{
+		ModifiedVertices.Empty();
+	};
+
+	/** Change tracking */
+	template <typename ValidSetAppendContainerType>
+	void RecordModifiedVertices(const ValidSetAppendContainerType& Container)
+	{
+		ModifiedVertices.Empty();
+		ModifiedVertices.Append(Container);
+	}
+
+	/** Used to iteratively add to the active change set (TSet<>)*/
+	inline void RecordModifiedVertex(int32 VertexID)
+	{
+		ModifiedVertices.Add(VertexID);
+	};
+
+
+	void SetActiveHandleFaces(const TArray<int>& FaceGroupIDs) override;
+	void SetActiveHandleEdges(const TArray<int>& TopologyEdgeIDs) override;
+	void SetActiveHandleCorners(const TArray<int>& TopologyCornerIDs) override;
+
+
+
+	/** Allocates shared storage for use in task synchronization */
+	void InitBackgroundWorker(const ELaplacianWeightScheme WeightScheme);
+
+
+	/** Coordinates the background tasks. Returns false if the worker was already running */
+	bool UpdateAndLaunchdWorker(const ELaplacianWeightScheme WeightScheme, const FRichCurve* Curve = nullptr);
+
+	/** Capture data about background task state.*/
+	bool IsTaskInFlight() const;
+
+
+
+	/** Sets the SrcMeshConstraintBuffer to have a size of MaxVertexID, and initializes with the current mesh positions, but weight zero*/
+	void InitializeConstraintBuffer();
+
+	/** Given an array of Group IDs, update the selection and record vertices */
+	void UpdateSelection(const FDynamicMesh3* TargetMesh, const TArray<int>& Groups, bool bLocalizeDeformation);
+
+	/** Updates the mesh preview and/or solvers upon user input, provided a deformation strategy */
+	void UpdateSolution(FDynamicMesh3* TargetMesh, const TFunction<FVector3d(FDynamicMesh3* Mesh, int)>& HandleVertexDeformFunc) override;
+
+	/** Updates the vertex positions of the mesh with the result from the last deformation solve. */
+	void ExportDeformedPositions(FDynamicMesh3* TargetMesh);
+
+	/** Returns true if the asynchronous task has finished. */
+	inline bool IsDone() { return AsyncMeshDeformTask == nullptr || AsyncMeshDeformTask->IsDone(); };
+
+	/** Triggers abort on task and passes off ownership to deleter object */
+	inline void Shutdown();
+
+	const TArray<FROIFace>& GetROIFaces() const { return ROIFaces; }
+
+	/** Stores the position of the vertex constraints and corresponding weights for the entire mesh.  This is used as a form of scratch space.*/
+	TArray<FDeformerVertexConstraintData> SrcMeshConstraintBuffer;
+
+	/** Array of vertex indices organized in groups of three - basically an index buffer - that defines the subset of the mesh that the deformation task will work on.*/
+	TArray<int32> SubsetIDBuffer;
+
+	/** Need to update the task with the current submesh */
+	bool bTaskSubmeshIsDirty = true;
+
+	/** Asynchronous task object. This object deals with expensive matrix functionality that computes the deformation of a local mesh. */
+	FAsyncTaskExecuterWithAbort<FConstrainedMeshDeformerTask>* AsyncMeshDeformTask = nullptr;
+
+
+	/** The weight which will be applied to the constraints corresponding to the handle vertices. */
+	double HandleWeights = 1.0;
+
+	/** This is set to true whenever the user interacts with the tool under laplacian deformation mode.
+	  * It is set to false immediately before beginning a background task and cannot be set to false again until the work is done. */
+	bool bDeformerNeedsToRun = false;
+
+
+	/** When true, tells the solver to attempt to postfix the actual position of the handles to the constrained position */
+	bool bPostfixHandles = false;
+
+	//This is set to false only after 
+	//	1) the asynchronous deformation task is complete
+	//	2) the main thread has seen it complete, and
+	//	3) the main thread updates the vertex positions of the mesh one last time
+	bool bVertexPositionsNeedSync = false;
+
+	bool bLocalize = true;
+
+};
+
+
+
+
+
+
+
+
+
+
+
+
 
 //////////////////////////////
 // DEBUG_SETTINGS
@@ -609,6 +864,8 @@ void UDeformMeshPolygonsTool::Setup()
 {
 	UMeshSurfacePointTool::Setup();
 
+	LaplacianDeformer = MakeShared<FGroupTopologyLaplacianDeformer>();
+
 	// create dynamic mesh component to use for live preview
 	DynamicMeshComponent = NewObject<USimpleDynamicMeshComponent>(ComponentTarget->GetOwnerActor(), "DynamicMesh");
 	DynamicMeshComponent->SetupAttachment(ComponentTarget->GetOwnerActor()->GetRootComponent());
@@ -666,7 +923,7 @@ void UDeformMeshPolygonsTool::Setup()
 	// Allocates buffers, sets up the asynchronous task
 	// Copies the source mesh positions.
 	const ELaplacianWeightScheme LaplacianWeightScheme = ConvertToLaplacianWeightScheme(TransformProps->SelectedWeightScheme);
-	LaplacianDeformer.InitBackgroundWorker(LaplacianWeightScheme);
+	LaplacianDeformer->InitBackgroundWorker(LaplacianWeightScheme);
 	
 
 	/**
@@ -695,7 +952,7 @@ void UDeformMeshPolygonsTool::Setup()
 void UDeformMeshPolygonsTool::Shutdown(EToolShutdownType ShutdownType)
 {
 	//Tell the background thread to cancel the rest of its jobs before we close;
-	LaplacianDeformer.Shutdown();
+	LaplacianDeformer->Shutdown();
 
 	if (DynamicMeshComponent != nullptr)
 	{
@@ -770,15 +1027,15 @@ void UDeformMeshPolygonsTool::OnDynamicMeshComponentChanged()
 	for (int VertexID : Mesh->VertexIndicesItr())
 	{
 		const FVector3d Position = Mesh->GetVertex(VertexID);
-		LaplacianDeformer.SrcMeshConstraintBuffer[VertexID].Position = Position;
+		LaplacianDeformer->SrcMeshConstraintBuffer[VertexID].Position = Position;
 	}
 
 	// a deform task could still be in flight.
-	if (LaplacianDeformer.AsyncMeshDeformTask != nullptr)
+	if (LaplacianDeformer->AsyncMeshDeformTask != nullptr)
 	{
-		LaplacianDeformer.AsyncMeshDeformTask->CancelAndDelete();
-		LaplacianDeformer.AsyncMeshDeformTask = nullptr;
-		LaplacianDeformer.bTaskSubmeshIsDirty = true;
+		LaplacianDeformer->AsyncMeshDeformTask->CancelAndDelete();
+		LaplacianDeformer->AsyncMeshDeformTask = nullptr;
+		LaplacianDeformer->bTaskSubmeshIsDirty = true;
 	}
 }
 
@@ -884,28 +1141,28 @@ void UDeformMeshPolygonsTool::OnBeginDrag(const FRay& WorldRay)
 
 	if (DeformationStrategy == EGroupTopologyDeformationStrategy::Laplacian)
 	{
-		LaplacianDeformer.bLocalize = true; // TransformProps->bLocalizeDeformation;
+		LaplacianDeformer->bLocalize = true; // TransformProps->bLocalizeDeformation;
 
 		//Determine which of the following (corners, edges or faces) has been selected by counting the associated feature's IDs
 		if (Selection.SelectedCornerIDs.Num() > 0)
 		{
 			//Add all the the Corner's adjacent poly-groups (NbrGroups) to the ongoing array of groups.
-			LaplacianDeformer.SetActiveHandleCorners(Selection.SelectedCornerIDs);
+			LaplacianDeformer->SetActiveHandleCorners(Selection.SelectedCornerIDs);
 		}
 		else if (Selection.SelectedEdgeIDs.Num() > 0)
 		{
 			//Add all the the edge's adjacent poly-groups (NbrGroups) to the ongoing array of groups.
-			LaplacianDeformer.SetActiveHandleEdges(Selection.SelectedEdgeIDs);
+			LaplacianDeformer->SetActiveHandleEdges(Selection.SelectedEdgeIDs);
 		}
 		else if (Selection.SelectedGroupIDs.Num() > 0)
 		{
-			LaplacianDeformer.SetActiveHandleFaces(Selection.SelectedGroupIDs);
+			LaplacianDeformer->SetActiveHandleFaces(Selection.SelectedGroupIDs);
 		}
 
 
 		// If there are actually no interior points, then we can't actually use the laplacian deformer. Need to fall back to the linear.
 		bool bHasInteriorVerts = false;
-		const auto& ROIFaces = LaplacianDeformer.GetROIFaces();
+		const auto& ROIFaces = LaplacianDeformer->GetROIFaces();
 		for (const auto& Face : ROIFaces)
 		{
 			bHasInteriorVerts = bHasInteriorVerts || ( Face.InteriorVerts.Num() != 0);
@@ -921,7 +1178,7 @@ void UDeformMeshPolygonsTool::OnBeginDrag(const FRay& WorldRay)
 		{		
 			// finalize the laplacian deformer : the task will need a new mesh that corresponds to the selected region.
 		
-			LaplacianDeformer.bTaskSubmeshIsDirty = true;
+			LaplacianDeformer->bTaskSubmeshIsDirty = true;
 
 		}
 		
@@ -1016,9 +1273,9 @@ void UDeformMeshPolygonsTool::UpdateChangeFromROI(bool bFinal)
 	const bool bIsLaplacian = (DeformationStrategy == EGroupTopologyDeformationStrategy::Laplacian);
 
 	FDynamicMesh3* Mesh = DynamicMeshComponent->GetMesh();
-	const TSet<int>& ModifiedVertices = (bIsLaplacian) ? LaplacianDeformer.GetModifiedVertices() : LinearDeformer.GetModifiedVertices();
+	const TSet<int>& ModifiedVertices = (bIsLaplacian) ? LaplacianDeformer->GetModifiedVertices() : LinearDeformer.GetModifiedVertices();
 	ActiveVertexChange->SaveVertices(Mesh, ModifiedVertices, !bFinal);
-	const TSet<int>& ModifiedNormals = (bIsLaplacian) ? LaplacianDeformer.GetModifiedOverlayNormals() : LinearDeformer.GetModifiedOverlayNormals();
+	const TSet<int>& ModifiedNormals = (bIsLaplacian) ? LaplacianDeformer->GetModifiedOverlayNormals() : LinearDeformer.GetModifiedOverlayNormals();
 	ActiveVertexChange->SaveOverlayNormals(Mesh, ModifiedNormals, !bFinal);
 }
 
@@ -1106,18 +1363,18 @@ void UDeformMeshPolygonsTool::ComputeUpdate()
 
 	if (DeformationStrategy == EGroupTopologyDeformationStrategy::Laplacian)
 	{
-		bool bIsWorking = LaplacianDeformer.IsTaskInFlight();
+		bool bIsWorking = LaplacianDeformer->IsTaskInFlight();
 
 		if (!bIsWorking)
 		{
 			// Sync update if we have new results.
-			if (LaplacianDeformer.bVertexPositionsNeedSync)
+			if (LaplacianDeformer->bVertexPositionsNeedSync)
 			{
 				
 				//Update the mesh with the provided solutions.
-				LaplacianDeformer.ExportDeformedPositions(DynamicMeshComponent->GetMesh());
+				LaplacianDeformer->ExportDeformedPositions(DynamicMeshComponent->GetMesh());
 
-				LaplacianDeformer.bVertexPositionsNeedSync = false;
+				LaplacianDeformer->bVertexPositionsNeedSync = false;
 
 				//Re-sync mesh, and flag the spatial data struct & topology for re-evaluation
 				DynamicMeshComponent->FastNotifyPositionsUpdated(true, false, false);
@@ -1128,14 +1385,14 @@ void UDeformMeshPolygonsTool::ComputeUpdate()
 			}
 
 			// emit end change if we are done with the drag
-			if (!LaplacianDeformer.bDeformerNeedsToRun && !bInDrag)
+			if (!LaplacianDeformer->bDeformerNeedsToRun && !bInDrag)
 			{
 				EndChange();
 			}
 
 			// Not working but we have more work for it to do..
 
-			if (LaplacianDeformer.bDeformerNeedsToRun)
+			if (LaplacianDeformer->bDeformerNeedsToRun)
 			{
 			
 				FRichCurve* Curve = NULL;
@@ -1149,7 +1406,7 @@ void UDeformMeshPolygonsTool::ComputeUpdate()
 				}
 				*/
 				const ELaplacianWeightScheme LaplacianWeightScheme = ConvertToLaplacianWeightScheme(TransformProps->SelectedWeightScheme);
-				LaplacianDeformer.UpdateAndLaunchdWorker(LaplacianWeightScheme, Curve);
+				LaplacianDeformer->UpdateAndLaunchdWorker(LaplacianWeightScheme, Curve);
 
 			}
 		}
@@ -1164,7 +1421,7 @@ void UDeformMeshPolygonsTool::ComputeUpdate()
 void UDeformMeshPolygonsTool::ComputeUpdate_Rotate()
 {
 	const bool bIsLaplacian = (DeformationStrategy == EGroupTopologyDeformationStrategy::Laplacian); 
-	FGroupTopologyDeformer& SelectedDeformer = (bIsLaplacian) ? LaplacianDeformer : LinearDeformer;
+	FGroupTopologyDeformer& SelectedDeformer = (bIsLaplacian) ? *LaplacianDeformer : LinearDeformer;
 
 	FDynamicMesh3* Mesh = DynamicMeshComponent->GetMesh();
 	FTransform3d Transform = FTransform3d(ComponentTarget->GetWorldTransform());
@@ -1254,7 +1511,7 @@ void UDeformMeshPolygonsTool::ComputeUpdate_Rotate()
 void UDeformMeshPolygonsTool::ComputeUpdate_Translate()
 {
 	const bool bIsLaplacian = (DeformationStrategy == EGroupTopologyDeformationStrategy::Laplacian);
-	FGroupTopologyDeformer& SelectedDeformer = (bIsLaplacian) ? LaplacianDeformer : LinearDeformer;
+	FGroupTopologyDeformer& SelectedDeformer = (bIsLaplacian) ? *LaplacianDeformer : LinearDeformer;
 
 	TFunction<FVector3d(const FVector3d&)> PointConstraintFunc = nullptr;
 	if (TransformProps->bSnapToWorldGrid 
@@ -1320,8 +1577,8 @@ void UDeformMeshPolygonsTool::ComputeUpdate_Translate()
 
 void UDeformMeshPolygonsTool::OnTick(float DeltaTime)
 {
-	LaplacianDeformer.HandleWeights   = TransformProps->HandleWeight;
-	LaplacianDeformer.bPostfixHandles = TransformProps->bPostFixHandles;
+	LaplacianDeformer->HandleWeights   = TransformProps->HandleWeight;
+	LaplacianDeformer->bPostfixHandles = TransformProps->bPostFixHandles;
 }
 
 
@@ -1332,10 +1589,10 @@ void UDeformMeshPolygonsTool::PrecomputeTopology()
 	Topology = FGroupTopology(Mesh, true);
 
 	LinearDeformer.Initialize(Mesh, &Topology);
-	LaplacianDeformer.Initialize(Mesh, &Topology);
+	LaplacianDeformer->Initialize(Mesh, &Topology);
 
 	// Make the Constraint Buffer, zero weights, but current pos
-	LaplacianDeformer.InitializeConstraintBuffer();
+	LaplacianDeformer->InitializeConstraintBuffer();
 }
 
 
@@ -1435,7 +1692,7 @@ void UDeformMeshPolygonsTool::OnPropertyModified(UObject* PropertySet, FProperty
 void UDeformMeshPolygonsTool::BeginChange()
 {
 	const bool bIsLaplacian = (DeformationStrategy == EGroupTopologyDeformationStrategy::Laplacian);
-	if (!bIsLaplacian || LaplacianDeformer.IsDone())
+	if (!bIsLaplacian || LaplacianDeformer->IsDone())
 	{
 		if (ActiveVertexChange == nullptr)
 		{
diff --git a/Engine/Plugins/Experimental/MeshModelingToolset/Source/MeshModelingTools/Public/DeformMeshPolygonsTool.h b/Engine/Plugins/Experimental/MeshModelingToolset/Source/MeshModelingTools/Public/DeformMeshPolygonsTool.h
--- a/Engine/Plugins/Experimental/MeshModelingToolset/Source/MeshModelingTools/Public/DeformMeshPolygonsTool.h
+++ b/Engine/Plugins/Experimental/MeshModelingToolset/Source/MeshModelingTools/Public/DeformMeshPolygonsTool.h
@@ -15,14 +15,12 @@
 #include "Spatial/GeometrySet3.h"
 #include "Selection/GroupTopologySelector.h"
 #include "Operations/GroupTopologyDeformer.h"
-#include "MeshSolverUtilities/Private/LaplacianMeshSmoother.h"
+#include "Solvers/MeshLaplacian.h"
 #include "Curves/CurveFloat.h"
-#include "ModelingOperators/Public/ModelingTaskTypes.h"
 #include "DeformMeshPolygonsTool.generated.h"
 
 class FMeshVertexChangeBuilder;
 class FGroupTopologyLaplacianDeformer;
-class FDeformTask;
 struct FDeformerVertexConstraintData;
 
 /**
@@ -141,241 +139,6 @@ public:
 */ 
 };
 
-//Stores per-vertex data needed by the laplacian deformer object
-//TODO: May be a candidate for a subclass of the FGroupTopologyLaplacianDeformer
-struct FDeformerVertexConstraintData
-{
-	FDeformerVertexConstraintData& operator=(const FDeformerVertexConstraintData& other)
-	{
-		Position = other.Position;
-		Weight   = other.Weight;
-		bPostFix = other.bPostFix;
-		return *this;
-	}
-
-	FVector3d Position;
-	double Weight{ 0.0 };
-	bool bPostFix{ false };
-};
-
-/**
-*	FDeformTask is an object which wraps an asynchronous task to be run multiple times on a separate thread.
-*	The Laplacian deformation process requires the use of potentially large sparse matrices and sparse multiplication.
-*	
-*   Expected usage:
-*
-*   
-*   // define constraints.  Need Constraints[VertID] to hold the constraints for the corresponding vertex.
-*   TArray<FDeformerVertexConstraintData> Constraints;
-*   ....
-*
-*   // populate with the VertexIDs of the vertices that are in the region you wish to deform.  
-*   TArray<int32> SrcVertIDs;  //Basically a mini-index buffer.
-*   ...
-* 
-*   // Create or reuse a laplacian deformation task.
-*   FDeformTask*   DeformTask = New FDeformTask(WeightScheme);
-*
-*   // the deformer will have to build a new mesh that represents the regions in SrcVertIDs;
-*   // but set this to false on subsequent calls to UpdateDeformer if the SrcVertIDs array hasn't changed.
-*   bool bRequiresRegion = true;
-
-*   DefTask->UpdateDeformer(WeightScheme, Mesh, Constraints, SrcVertIDs, bRequiresRegion);
-*
-*   DeformTask->DoWork();  or DeformTask->StartBackgroundTask(); //which calls DoWork on background thread.
-*
-*  // wheh DeformTask->IsDone == true; you can copy the results back to the mesh
-*  DeformTask->ExportResults(Mesh);
-*
-* Note: if only the positions in the Constraints change (e.g. handle positions) then subsequent calls 
-* to UpdateDeformer() and DoWork() will be much faster as the matrix system will not be rebuilt or re-factored
-*/
-class FConstrainedMeshDeformerTask : public FNonAbandonableTask
-{
-	friend class FAsyncTask<FDeformTask>;
-public:
-
-	enum
-	{
-		INACTIVE_SUBSET_ID = -1
-	};
-
-	FConstrainedMeshDeformerTask(const ELaplacianWeightScheme SelectedWeightScheme)
-	{
-		LaplacianWeightScheme = SelectedWeightScheme;
-	}
-
-	virtual ~FConstrainedMeshDeformerTask() {};
-
-	//NO idea what this is meant to do. Performance analysis maybe? Scheduling data?
-	FORCEINLINE TStatId GetStatId() const
-	{
-		RETURN_QUICK_DECLARE_CYCLE_STAT(FConstrainedMeshDeformerTask, STATGROUP_ThreadPoolAsyncTasks);
-	}
-
-	
-	/** Called by the main thread in the tool, this copies the Constraint buffer right before the task begins on another thread.
-	  * Ensures the FConstrainedMeshDeformer is using correct mesh subset and the selected settings, then updates on change in properties, i.e. weight scheme */
-	void UpdateDeformer(const ELaplacianWeightScheme SelectedWeightScheme, const FDynamicMesh3& Mesh, const TArray<FDeformerVertexConstraintData>& ConstraintArray, const TArray<int32>& SrcIDBufferSubset, bool bNewTransaction, const FRichCurve* Curve);
-
-	/** Required by the FAsyncTaskExecutor */
-	void SetAbortSource(bool* bAbort) { bAbortSource = bAbort; };
-
-	/** Called by the FAsyncTask<FDeformTask> object for background computation. */
-	void DoWork();
-
-
-	/** Updates the positions in the target mesh for regions that correspond to the subset mesh */
-	void ExportResults(FDynamicMesh3& TargetMesh) const;
-
-private:
-
-	/** Creates the mesh (i.e. SubsetMesh) that corresponds to the region of the SrcMesh defined by the partial index buffer SrcIDBufferSubset */
-	void InitializeSubsetMesh(const FDynamicMesh3& SrcMesh, const TArray<int32>& SrcIDBufferSubset);
-
-	/** Attenuates the weights of the constraints using the selected curve */
-	void ApplyAttenuation();
-
-	/** Denotes the weight scheme being used by the running background task. Changes when selected property changes in editor. */
-	ELaplacianWeightScheme LaplacianWeightScheme;
-
-	/** positions for each vertex in the subset mesh - for use in the deformer */
-	TArray<FVector3d> SubsetPositionBuffer;
-
-	/** constraint data for each vertex in subset mesh - for use by the deformer*/
-	TArray<FDeformerVertexConstraintData> SubsetConstraintBuffer;
-
-	FRichCurve WeightAttenuationCurve;
-
-	/** True only for the first update, and then false for the duration of the Input transaction
-	  * It's passed in and copied in UpdateDeformer() */
-	bool bIsNewTransaction = true;
-
-	/** When true, the constraint weights will be attenuated based on distance using the provided curve object*/
-	bool bAttenuateWeights = false;
-
-	/** The abort bool used by the Task Deleter */
-	bool* bAbortSource = nullptr;
-
-	/** Used to initialize the array mapping, updated during the UpdateDeformer() function */
-	int SrcMeshMaxVertexID;
-
-	/** A subset of the original mesh */
-	FDynamicMesh3 SubsetMesh;
-
-	/** Maps Subset Mesh VertexID to Src Mesh VertexID */
-	TArray<int32> SubsetVertexIDToSrcVertexIDMap;
-
-	/** Laplacian deformer object gets rebuilt each new transaction */
-	TUniquePtr<FConstrainedMeshDeformer> ConstrainedDeformer;
-
-private :
-	FConstrainedMeshDeformerTask();
-
-};
-
-class FGroupTopologyLaplacianDeformer : public FGroupTopologyDeformer
-{
-
-public:
-
-	FGroupTopologyLaplacianDeformer() = default;
-
-	virtual ~FGroupTopologyLaplacianDeformer();
-
-	/** Used to begin a procedural addition of modified vertices */
-	inline void ResetModifiedVertices()
-	{
-		ModifiedVertices.Empty();
-	};
-
-	/** Change tracking */
-	template <typename ValidSetAppendContainerType>
-	void RecordModifiedVertices(const ValidSetAppendContainerType& Container)
-	{
-		ModifiedVertices.Empty();
-		ModifiedVertices.Append(Container);
-	}
-
-	/** Used to iteratively add to the active change set (TSet<>)*/
-	inline void RecordModifiedVertex(int32 VertexID)
-	{
-		ModifiedVertices.Add(VertexID);
-	};
-
-
-	void SetActiveHandleFaces(const TArray<int>& FaceGroupIDs) override;
-	void SetActiveHandleEdges(const TArray<int>& TopologyEdgeIDs) override;
-	void SetActiveHandleCorners(const TArray<int>& TopologyCornerIDs) override;
-
-
-
-	/** Allocates shared storage for use in task synchronization */
-	void InitBackgroundWorker(const ELaplacianWeightScheme WeightScheme);
-
-
-	/** Coordinates the background tasks. Returns false if the worker was already running */
-	bool UpdateAndLaunchdWorker(const ELaplacianWeightScheme WeightScheme, const FRichCurve* Curve = nullptr);
-
-	/** Capture data about background task state.*/
-	bool IsTaskInFlight() const;
-
-
-
-	/** Sets the SrcMeshConstraintBuffer to have a size of MaxVertexID, and initializes with the current mesh positions, but weight zero*/
-	void InitializeConstraintBuffer();
-
-	/** Given an array of Group IDs, update the selection and record vertices */
-	void UpdateSelection(const FDynamicMesh3* TargetMesh, const TArray<int>& Groups, bool bLocalizeDeformation);
-
-	/** Updates the mesh preview and/or solvers upon user input, provided a deformation strategy */
-	void UpdateSolution(FDynamicMesh3* TargetMesh, const TFunction<FVector3d(FDynamicMesh3* Mesh, int)>& HandleVertexDeformFunc) override;
-
-	/** Updates the vertex positions of the mesh with the result from the last deformation solve. */
-	void ExportDeformedPositions(FDynamicMesh3* TargetMesh);
-
-	/** Returns true if the asynchronous task has finished. */
-	inline bool IsDone() { return AsyncMeshDeformTask == nullptr || AsyncMeshDeformTask->IsDone(); };
-
-	/** Triggers abort on task and passes off ownership to deleter object */
-	inline void Shutdown();
-
-	const TArray<FROIFace>& GetROIFaces() const { return ROIFaces; }
-
-	/** Stores the position of the vertex constraints and corresponding weights for the entire mesh.  This is used as a form of scratch space.*/
-	TArray<FDeformerVertexConstraintData> SrcMeshConstraintBuffer;
-
-	/** Array of vertex indices organized in groups of three - basically an index buffer - that defines the subset of the mesh that the deformation task will work on.*/
-	TArray<int32> SubsetIDBuffer;
-
-	/** Need to update the task with the current submesh */
-	bool bTaskSubmeshIsDirty = true;
-
-	/** Asynchronous task object. This object deals with expensive matrix functionality that computes the deformation of a local mesh. */
-	FAsyncTaskExecuterWithAbort<FConstrainedMeshDeformerTask>* AsyncMeshDeformTask = nullptr;
-
-
-	/** The weight which will be applied to the constraints corresponding to the handle vertices. */
-	double HandleWeights = 1.0;
-
-	/** This is set to true whenever the user interacts with the tool under laplacian deformation mode.
-	  * It is set to false immediately before beginning a background task and cannot be set to false again until the work is done. */
-	bool bDeformerNeedsToRun = false;
-
-
-	/** When true, tells the solver to attempt to postfix the actual position of the handles to the constrained position */
-	bool bPostfixHandles = false;
-
-	//This is set to false only after 
-	//	1) the asynchronous deformation task is complete
-	//	2) the main thread has seen it complete, and
-	//	3) the main thread updates the vertex positions of the mesh one last time
-	bool bVertexPositionsNeedSync = false;
-
-	bool bLocalize =  true;
-
-};
-
 
 /**
  *
@@ -481,7 +244,7 @@ protected:
 
 	// The two deformer type options.
 	FGroupTopologyDeformer LinearDeformer;
-	FGroupTopologyLaplacianDeformer LaplacianDeformer;
+	TSharedPtr<FGroupTopologyLaplacianDeformer> LaplacianDeformer;
 
 
 
diff --git a/Engine/Plugins/Experimental/MeshModelingToolset/Source/MeshModelingToolsEditorOnly/MeshModelingToolsEditorOnly.Build.cs b/Engine/Plugins/Experimental/MeshModelingToolset/Source/MeshModelingToolsEditorOnly/MeshModelingToolsEditorOnly.Build.cs
--- a/Engine/Plugins/Experimental/MeshModelingToolset/Source/MeshModelingToolsEditorOnly/MeshModelingToolsEditorOnly.Build.cs
+++ b/Engine/Plugins/Experimental/MeshModelingToolset/Source/MeshModelingToolsEditorOnly/MeshModelingToolsEditorOnly.Build.cs
@@ -33,7 +33,6 @@ public class MeshModelingToolsEditorOnly : ModuleRules
 				"DynamicMesh",
 				"MeshConversion",
 				"MeshModelingTools",
-                "MeshSolverUtilities",
 				"ModelingComponents",
 				"ModelingOperators",
 				"ModelingOperatorsEditorOnly",
diff --git a/Engine/Plugins/Experimental/MeshModelingToolset/Source/ModelingOperators/ModelingOperators.Build.cs b/Engine/Plugins/Experimental/MeshModelingToolset/Source/ModelingOperators/ModelingOperators.Build.cs
--- a/Engine/Plugins/Experimental/MeshModelingToolset/Source/ModelingOperators/ModelingOperators.Build.cs
+++ b/Engine/Plugins/Experimental/MeshModelingToolset/Source/ModelingOperators/ModelingOperators.Build.cs
@@ -30,7 +30,6 @@ public class ModelingOperators : ModuleRules
 				// dependent modules, please confirm by "launching" or packaging that you are not introducing
 				// editor dependencies
 				//
-				//"MeshSolverUtilities",
                 //"ProxyLODMeshReduction", // currently required to be public due to IVoxelBasedCSG API
                 //"MeshDescription",
 
@@ -40,7 +39,6 @@ public class ModelingOperators : ModuleRules
 				"DynamicMesh",
 				"MeshConversion",
                 "GeometryAlgorithms", // required for constrained Delaunay triangulation
-                "MeshSolverUtilities", // required by the smoothing operators
 				// ... add other public dependencies that you statically link with here ...
 			}
             );
diff --git a/Engine/Plugins/Experimental/MeshModelingToolset/Source/ModelingOperators/Private/SmoothingOps/CotanSmoothingOp.cpp b/Engine/Plugins/Experimental/MeshModelingToolset/Source/ModelingOperators/Private/SmoothingOps/CotanSmoothingOp.cpp
--- a/Engine/Plugins/Experimental/MeshModelingToolset/Source/ModelingOperators/Private/SmoothingOps/CotanSmoothingOp.cpp
+++ b/Engine/Plugins/Experimental/MeshModelingToolset/Source/ModelingOperators/Private/SmoothingOps/CotanSmoothingOp.cpp
@@ -1,8 +1,7 @@
 // Copyright Epic Games, Inc. All Rights Reserved.
 
 #include "SmoothingOps/CotanSmoothingOp.h"
-
-#include "MeshSmoothingUtilities.h"
+#include "Solvers/MeshSmoothing.h"
 
 
 FCotanSmoothingOp::FCotanSmoothingOp(const FDynamicMesh3* Mesh, float Speed, int32 Iterations) :
@@ -24,5 +23,5 @@ void FCotanSmoothingOp::CalculateResult(FProgressCancel* Progress)
 void FCotanSmoothingOp::Smooth()
 {
 	double Intensity = 1.;
-	MeshSmoothingOperators::ComputeSmoothing_BiHarmonic(ELaplacianWeightScheme::ClampedCotangent, *ResultMesh, SmoothSpeed, Intensity, SmoothIterations, PositionBuffer);
+	UE::MeshDeformation::ComputeSmoothing_BiHarmonic(ELaplacianWeightScheme::ClampedCotangent, *ResultMesh, SmoothSpeed, Intensity, SmoothIterations, PositionBuffer);
 }
diff --git a/Engine/Plugins/Experimental/MeshModelingToolset/Source/ModelingOperators/Private/SmoothingOps/MeanValueSmoothingOp.cpp b/Engine/Plugins/Experimental/MeshModelingToolset/Source/ModelingOperators/Private/SmoothingOps/MeanValueSmoothingOp.cpp
--- a/Engine/Plugins/Experimental/MeshModelingToolset/Source/ModelingOperators/Private/SmoothingOps/MeanValueSmoothingOp.cpp
+++ b/Engine/Plugins/Experimental/MeshModelingToolset/Source/ModelingOperators/Private/SmoothingOps/MeanValueSmoothingOp.cpp
@@ -1,8 +1,7 @@
 // Copyright Epic Games, Inc. All Rights Reserved.
 
 #include "SmoothingOps/MeanValueSmoothingOp.h"
-
-#include "MeshSmoothingUtilities.h"
+#include "Solvers/MeshSmoothing.h"
 
 
 FMeanValueSmoothingOp::FMeanValueSmoothingOp(const FDynamicMesh3* Mesh, float Speed, int32 Iterations) :
@@ -24,5 +23,5 @@ void FMeanValueSmoothingOp::CalculateResult(FProgressCancel* Progress)
 void FMeanValueSmoothingOp::Smooth()
 {
 	double Intensity = 1.;
-	MeshSmoothingOperators::ComputeSmoothing_BiHarmonic(ELaplacianWeightScheme::MeanValue, *ResultMesh, SmoothSpeed, Intensity, SmoothIterations, PositionBuffer);
+	UE::MeshDeformation::ComputeSmoothing_BiHarmonic(ELaplacianWeightScheme::MeanValue, *ResultMesh, SmoothSpeed, Intensity, SmoothIterations, PositionBuffer);
 }
\ No newline at end of file
diff --git a/Engine/Plugins/Experimental/MeshModelingToolset/Source/ModelingOperatorsEditorOnly/ModelingOperatorsEditorOnly.Build.cs b/Engine/Plugins/Experimental/MeshModelingToolset/Source/ModelingOperatorsEditorOnly/ModelingOperatorsEditorOnly.Build.cs
--- a/Engine/Plugins/Experimental/MeshModelingToolset/Source/ModelingOperatorsEditorOnly/ModelingOperatorsEditorOnly.Build.cs
+++ b/Engine/Plugins/Experimental/MeshModelingToolset/Source/ModelingOperatorsEditorOnly/ModelingOperatorsEditorOnly.Build.cs
@@ -33,7 +33,6 @@ public class ModelingOperatorsEditorOnly : ModuleRules
 				"GeometricObjects",
 				"DynamicMesh",
 				"MeshConversion",
-				"MeshSolverUtilities",
 				"GeometryAlgorithms", // required for constrained Delaunay triangulation
                 "ProxyLODMeshReduction", // currently required to be public due to IVoxelBasedCSG API
 				"MeshUtilitiesCommon", // required by uvlayoutop
