GeometryProcessing: remove unused function
#rnx

diff --git a/Engine/Plugins/Experimental/GeometryProcessing/Source/DynamicMesh/Private/Solvers/Internal/LaplacianOperators.cpp b/Engine/Plugins/Experimental/GeometryProcessing/Source/DynamicMesh/Private/Solvers/Internal/LaplacianOperators.cpp
--- a/Engine/Plugins/Experimental/GeometryProcessing/Source/DynamicMesh/Private/Solvers/Internal/LaplacianOperators.cpp
+++ b/Engine/Plugins/Experimental/GeometryProcessing/Source/DynamicMesh/Private/Solvers/Internal/LaplacianOperators.cpp
@@ -51,12 +51,6 @@ public:
 		};
 	}
 
-	TUniquePtr<FSparseMatrixD> ExtractResult()
-	{
-		Matrix->setFromTriplets(EntryTriplets.begin(), EntryTriplets.end());
-		Matrix->makeCompressed();
-	}
-
 	void ExtractResult(FSparseMatrixD& Result)
 	{
 		Matrix->setFromTriplets(EntryTriplets.begin(), EntryTriplets.end());
