GeometryProcessing: debug/improve FMeshWeights, add UE::MeshCurvature implementations
#rnx

diff --git a/Engine/Plugins/Experimental/GeometryProcessing/Source/DynamicMesh/Private/MeshCurvature.cpp b/Engine/Plugins/Experimental/GeometryProcessing/Source/DynamicMesh/Private/MeshCurvature.cpp
new file mode 100644
--- /dev/null
+++ b/Engine/Plugins/Experimental/GeometryProcessing/Source/DynamicMesh/Private/MeshCurvature.cpp
@@ -0,0 +1,138 @@
+// Copyright Epic Games, Inc. All Rights Reserved.
+
+#include "MeshCurvature.h"
+#include "MeshWeights.h"
+#include "VectorUtil.h"
+
+
+
+template<typename GetPositionFuncType>
+FVector3d TMeanCurvatureNormal(const FDynamicMesh3& mesh, int32 v_i, GetPositionFuncType GetPositionFunc, double CotClampRange = 1000 /* about 0.05 deg */)
+{
+	// based on equations in http://www.geometry.caltech.edu/pubs/DMSB_III.pdf
+
+	FVector3d vSum = FVector3d::Zero();
+	double wSum = 0;
+	FVector3d Vi = GetPositionFunc(v_i);
+
+	int32 v_j = FDynamicMesh3::InvalidID, opp_v1 = FDynamicMesh3::InvalidID, opp_v2 = FDynamicMesh3::InvalidID;
+	int32 t1 = FDynamicMesh3::InvalidID, t2 = FDynamicMesh3::InvalidID;
+	for (int32 eid : mesh.VtxEdgesItr(v_i))
+	{
+		opp_v2 = FDynamicMesh3::InvalidID;
+		mesh.GetVtxNbrhood(eid, v_i, v_j, opp_v1, opp_v2, t1, t2);
+		FVector3d Vj = GetPositionFunc(v_j);
+
+		FVector3d Vo1 = GetPositionFunc(opp_v1);
+		double cot_alpha_ij = VectorUtil::VectorCot((Vi - Vo1), (Vj - Vo1));
+		double w_ij = FMathd::Clamp(cot_alpha_ij, -CotClampRange, CotClampRange);
+
+		double cot_beta_ij = 0;
+		if (opp_v2 != FDynamicMesh3::InvalidID)
+		{
+			FVector3d Vo2 = GetPositionFunc(opp_v2);
+			cot_beta_ij = VectorUtil::VectorCot((Vi - Vo2), (Vj - Vo2));
+		}
+		w_ij += FMathd::Clamp(cot_beta_ij, -CotClampRange, CotClampRange);
+
+		vSum += w_ij * (Vi - Vj);
+	}
+	if (vSum.SquaredLength() < FMathd::ZeroTolerance)
+	{
+		return FVector3d::Zero();
+	}
+	double Area = FMeshWeights::VoronoiArea(mesh, v_i, GetPositionFunc);
+	if (Area < FMathd::ZeroTolerance)
+	{
+		return FVector3d::Zero();
+	}
+	return vSum / (2.0 * Area);
+}
+
+
+FVector3d UE::MeshCurvature::MeanCurvatureNormal(const FDynamicMesh3& Mesh, int32 VertexIndex)
+{
+	return TMeanCurvatureNormal(Mesh, VertexIndex, [&](int32 vid) { return Mesh.GetVertex(vid); });
+}
+
+FVector3d UE::MeshCurvature::MeanCurvatureNormal(const FDynamicMesh3& Mesh, int32 VertexIndex, TFunctionRef<FVector3d(int32)> VertexPositionFunc)
+{
+	return TMeanCurvatureNormal(Mesh, VertexIndex, VertexPositionFunc);
+}
+
+
+
+
+
+
+template<typename GetPositionFuncType>
+double TGaussianCurvature(const FDynamicMesh3& mesh, int32 v_i, GetPositionFuncType GetPositionFunc, double CotClampRange = 1000 /* about 0.05 deg */)
+{
+	// based on equation 9 in http://www.geometry.caltech.edu/pubs/DMSB_III.pdf
+	double AngleSum = 0;
+	double MixedAreaSum = 0;
+	FVector3d Vi = GetPositionFunc(v_i);
+
+	mesh.EnumerateVertexTriangles(v_i, [&](int32 tid)
+	{
+		FIndex3i t = mesh.GetTriangle(tid);
+		int ti = (t[0] == v_i) ? 0 : ((t[1] == v_i) ? 1 : 2);
+		FVector3d Vj = GetPositionFunc(t[(ti + 1) % 3]);
+		FVector3d Vk = GetPositionFunc(t[(ti + 2) % 3]);
+
+		FVector3d Vij = Vj - Vi;
+		FVector3d Nij(Vij);
+		double LenSqrVij = Nij.Normalize();
+		LenSqrVij *= LenSqrVij;
+
+		FVector3d Vik = Vk - Vi;
+		FVector3d Nik(Vik);
+		double LenSqrVik = Nik.Normalize();
+		LenSqrVik *= LenSqrVik;
+
+		double OpeningAngle = Nij.AngleR(Nik);
+		AngleSum += OpeningAngle;
+
+		if (VectorUtil::IsObtuse(Vi, Vj, Vk))
+		{
+			// if triangle is obtuse voronoi area is undefined and we just return portion of triangle area
+			double Dot = Vij.Dot(Vik);
+			double areaT = 0.5 * FMathd::Sqrt(LenSqrVij * LenSqrVik - Dot * Dot);
+			MixedAreaSum += (OpeningAngle > FMathd::HalfPi) ?    // obtuse at v_i ?
+				(areaT * 0.5) : (areaT * 0.25);
+		}
+		else
+		{
+			// voronoi area
+			FVector3d Vkj = Vj - Vk;
+			double cot_alpha_ij = VectorUtil::VectorCot(-Vik, Vkj);
+			cot_alpha_ij = FMathd::Clamp(cot_alpha_ij, -CotClampRange, CotClampRange);
+			double cot_beta_ik = VectorUtil::VectorCot(-Vij, -Vkj);
+			cot_beta_ik = FMathd::Clamp(cot_alpha_ij, -CotClampRange, CotClampRange);
+			MixedAreaSum += LenSqrVij * cot_alpha_ij * 0.125;
+			MixedAreaSum += LenSqrVik * cot_beta_ik * 0.125;
+		}
+	});
+
+	if (MixedAreaSum < FMathf::ZeroTolerance)
+	{
+		return 0.0f;
+	}
+	else
+	{
+		return (FMathd::TwoPi - AngleSum);// / MixedAreaSum;
+	}
+}
+
+
+
+double UE::MeshCurvature::GaussianCurvature(const FDynamicMesh3& Mesh, int32 VertexIndex)
+{
+	return TGaussianCurvature(Mesh, VertexIndex, [&](int32 vid) { return Mesh.GetVertex(vid); });
+}
+
+double UE::MeshCurvature::GaussianCurvature(const FDynamicMesh3& Mesh, int32 VertexIndex, TFunctionRef<FVector3d(int32)> VertexPositionFunc)
+{
+	return TGaussianCurvature(Mesh, VertexIndex, VertexPositionFunc);
+}
+
diff --git a/Engine/Plugins/Experimental/GeometryProcessing/Source/DynamicMesh/Private/MeshWeights.cpp b/Engine/Plugins/Experimental/GeometryProcessing/Source/DynamicMesh/Private/MeshWeights.cpp
--- a/Engine/Plugins/Experimental/GeometryProcessing/Source/DynamicMesh/Private/MeshWeights.cpp
+++ b/Engine/Plugins/Experimental/GeometryProcessing/Source/DynamicMesh/Private/MeshWeights.cpp
@@ -4,7 +4,7 @@
 #include "VectorUtil.h"
 
 
-FVector3d FMeshWeights::UniformCentroid(const FDynamicMesh3 & mesh, int VertexIndex)
+FVector3d FMeshWeights::UniformCentroid(const FDynamicMesh3 & mesh, int32 VertexIndex)
 {
 	FVector3d Centroid;
 	mesh.GetVtxOneRingCentroid(VertexIndex, Centroid);
@@ -12,14 +12,28 @@ FVector3d FMeshWeights::UniformCentroid(const FDynamicMesh3 & mesh, int VertexIn
 }
 
 
-FVector3d FMeshWeights::MeanValueCentroid(const FDynamicMesh3 & mesh, int v_i)
+FVector3d FMeshWeights::UniformCentroid(const FDynamicMesh3& mesh, int32 VertexIndex, TFunctionRef<FVector3d(int32)> VertexPositionFunc)
+{
+	FVector3d Centroid = FVector3d::Zero();
+	int32 Count = 0;
+	for (int32 nbrvid : mesh.VtxVerticesItr(VertexIndex))
+	{
+		Centroid += VertexPositionFunc(nbrvid);
+		Count++;
+	}
+	return (Count == 0) ? VertexPositionFunc(VertexIndex) : (Centroid / (double)Count);
+}
+
+
+template<typename GetPositionFuncType>
+FVector3d TMeanValueCentroid(const FDynamicMesh3& mesh, int32 v_i, GetPositionFuncType GetPositionFunc, double WeightClamp)
 {
 	// based on equations in https://www.inf.usi.ch/hormann/papers/Floater.2006.AGC.pdf (formula 9)
 	// refer to that paper for variable names/etc
 
 	FVector3d vSum = FVector3d::Zero();
 	double wSum = 0;
-	FVector3d Vi = mesh.GetVertex(v_i);
+	FVector3d Vi = GetPositionFunc(v_i);
 
 	int v_j = FDynamicMesh3::InvalidID, opp_v1 = FDynamicMesh3::InvalidID, opp_v2 = FDynamicMesh3::InvalidID;
 	int t1 = FDynamicMesh3::InvalidID, t2 = FDynamicMesh3::InvalidID;
@@ -28,44 +42,61 @@ FVector3d FMeshWeights::MeanValueCentroid(const FDynamicMesh3 & mesh, int v_i)
 		opp_v2 = FDynamicMesh3::InvalidID;
 		mesh.GetVtxNbrhood(eid, v_i, v_j, opp_v1, opp_v2, t1, t2);
 
-		FVector3d Vj = mesh.GetVertex(v_j);
+		FVector3d Vj = GetPositionFunc(v_j);
 		FVector3d vVj = (Vj - Vi);
 		double len_vVj = vVj.Normalize();
 		// TODO: is this the right thing to do? if vertices are coincident,
 		//   weight of this vertex should be very high!
-		if (len_vVj < FMathd::ZeroTolerance)
+		if (len_vVj < FMathf::ZeroTolerance)
 		{
+			double w_ij = FMathd::Clamp(10000.0, 0, WeightClamp);
+			vSum += w_ij * Vj;
+			wSum += w_ij;
 			continue;
 		}
-		FVector3d vVdelta = (mesh.GetVertex(opp_v1) - Vi).Normalized();
+		FVector3d vVdelta = (GetPositionFunc(opp_v1) - Vi).Normalized();
 		double w_ij = VectorUtil::VectorTanHalfAngle(vVj, vVdelta);
 
 		if (opp_v2 != FDynamicMesh3::InvalidID)
 		{
-			FVector3d vVgamma = (mesh.GetVertex(opp_v2) - Vi).Normalized();
+			FVector3d vVgamma = (GetPositionFunc(opp_v2) - Vi).Normalized();
 			w_ij += VectorUtil::VectorTanHalfAngle(vVj, vVgamma);
 		}
 
 		w_ij /= len_vVj;
+		w_ij = FMathd::Clamp(w_ij, 0, WeightClamp);
 
 		vSum += w_ij * Vj;
 		wSum += w_ij;
 	}
-	if (wSum < FMathd::ZeroTolerance)
-		return Vi;
+	if (wSum < FMathf::ZeroTolerance)
+	{
+		return FMeshWeights::UniformCentroid(mesh, v_i, GetPositionFunc);
+	}
 	return vSum / wSum;
 }
 
 
 
+FVector3d FMeshWeights::MeanValueCentroid(const FDynamicMesh3& mesh, int32 v_i, double WeightClamp)
+{
+	return TMeanValueCentroid(mesh, v_i, [&](int32 vid) { return mesh.GetVertex(vid); }, WeightClamp);
+}
+
+FVector3d FMeshWeights::MeanValueCentroid(const FDynamicMesh3& mesh, int32 v_i, TFunctionRef<FVector3d(int32)> VertexPositionFunc, double WeightClamp)
+{
+	return TMeanValueCentroid(mesh, v_i, VertexPositionFunc, WeightClamp);
+}
+
 
-FVector3d FMeshWeights::CotanCentroid(const FDynamicMesh3& mesh, int v_i)
+template<typename GetPositionFuncType>
+FVector3d TCotanCentroid(const FDynamicMesh3& mesh, int32 v_i, GetPositionFuncType GetPositionFunc)
 {
 	// based on equations in http://www.geometry.caltech.edu/pubs/DMSB_III.pdf
 
 	FVector3d vSum = FVector3d::Zero();
 	double wSum = 0;
-	FVector3d Vi = mesh.GetVertex(v_i);
+	FVector3d Vi = GetPositionFunc(v_i);
 
 	int v_j = FDynamicMesh3::InvalidID, opp_v1 = FDynamicMesh3::InvalidID, opp_v2 = FDynamicMesh3::InvalidID;
 	int t1 = FDynamicMesh3::InvalidID, t2 = FDynamicMesh3::InvalidID;
@@ -74,21 +105,21 @@ FVector3d FMeshWeights::CotanCentroid(const FDynamicMesh3& mesh, int v_i)
 	{
 		opp_v2 = FDynamicMesh3::InvalidID;
 		mesh.GetVtxNbrhood(eid, v_i, v_j, opp_v1, opp_v2, t1, t2);
-		FVector3d Vj = mesh.GetVertex(v_j);
+		FVector3d Vj = GetPositionFunc(v_j);
 
-		FVector3d Vo1 = mesh.GetVertex(opp_v1);
-		double cot_alpha_ij = VectorUtil::VectorCot(
-			(Vi - Vo1).Normalized(), (Vj - Vo1).Normalized());
-		if (cot_alpha_ij == 0) {
+		FVector3d Vo1 = GetPositionFunc(opp_v1);
+		double cot_alpha_ij = VectorUtil::VectorCot( (Vi - Vo1), (Vj - Vo1) );
+		if (cot_alpha_ij == 0) 
+		{
 			bAborted = true;
 			break;
 		}
 		double w_ij = cot_alpha_ij;
 
-		if (opp_v2 != FDynamicMesh3::InvalidID) {
-			FVector3d Vo2 = mesh.GetVertex(opp_v2);
-			double cot_beta_ij = VectorUtil::VectorCot(
-				(Vi - Vo2).Normalized(), (Vj - Vo2).Normalized());
+		if (opp_v2 != FDynamicMesh3::InvalidID) 
+		{
+			FVector3d Vo2 = GetPositionFunc(opp_v2);
+			double cot_beta_ij = VectorUtil::VectorCot( (Vi - Vo2), (Vj - Vo2) );
 			if (cot_beta_ij == 0) {
 				bAborted = true;
 				break;
@@ -100,51 +131,78 @@ FVector3d FMeshWeights::CotanCentroid(const FDynamicMesh3& mesh, int v_i)
 		wSum += w_ij;
 	}
 	if (bAborted || fabs(wSum) < FMathd::ZeroTolerance)
+	{
 		return Vi;
+	}
 	return vSum / wSum;
 }
 
 
 
-double FMeshWeights::VoronoiArea(const FDynamicMesh3& mesh, int v_i)
+FVector3d FMeshWeights::CotanCentroid(const FDynamicMesh3& mesh, int32 v_i)
+{
+	return TCotanCentroid(mesh, v_i, [&](int32 vid) { return mesh.GetVertex(vid); });
+}
+
+FVector3d FMeshWeights::CotanCentroid(const FDynamicMesh3& mesh, int32 v_i, TFunctionRef<FVector3d(int32)> VertexPositionFunc)
+{
+	return TCotanCentroid(mesh, v_i, VertexPositionFunc);
+}
+
+
+template<typename GetPositionFuncType>
+double TVoronoiArea(const FDynamicMesh3& mesh, int32 v_i, GetPositionFuncType GetPositionFunc, double CotClampRange = 1000 /* about 0.05 deg */ )
 {
 	// based on equations in http://www.geometry.caltech.edu/pubs/DMSB_III.pdf
 
 	double areaSum = 0;
-	FVector3d Vi = mesh.GetVertex(v_i);
+	FVector3d Vi = GetPositionFunc(v_i);
 
-	for (int tid : mesh.VtxTrianglesItr(v_i) ) 
+	for (int32 tid : mesh.VtxTrianglesItr(v_i))
 	{
 		FIndex3i t = mesh.GetTriangle(tid);
 		int ti = (t[0] == v_i) ? 0 : ((t[1] == v_i) ? 1 : 2);
-		FVector3d Vj = mesh.GetVertex(t[(ti + 1) % 3]);
-		FVector3d Vk = mesh.GetVertex(t[(ti + 2) % 3]);
+		FVector3d Vj = GetPositionFunc(t[(ti + 1) % 3]);
+		FVector3d Vk = GetPositionFunc(t[(ti + 2) % 3]);
 
-		if (VectorUtil::IsObtuse(Vi, Vj, Vk)) 
+		if (VectorUtil::IsObtuse(Vi, Vj, Vk))
 		{
 			// if triangle is obtuse voronoi area is undefind and we just return portion of triangle area
 			FVector3d Vij = Vj - Vi;
 			FVector3d Vik = Vk - Vi;
+			double Dot = Vij.Dot(Vik);
+			double areaT = 0.5 * FMathd::Sqrt(Vij.SquaredLength() * Vik.SquaredLength() - Dot * Dot);
 			Vij.Normalize(); Vik.Normalize();
-			double areaT = 0.5 * Vij.Cross(Vik).Length();
-			areaSum += ( Vij.AngleR(Vik) > FMathd::HalfPi ) ?    // obtuse at v_i ?
+			areaSum += (Vij.AngleR(Vik) > FMathd::HalfPi) ?    // obtuse at v_i ?
 				(areaT * 0.5) : (areaT * 0.25);
-
 		}
-		else 
+		else
 		{
 			// voronoi area
 			FVector3d Vji = Vi - Vj;
-			double dist_ji = Vji.Normalize();
 			FVector3d Vki = Vi - Vk;
-			double dist_ki = Vki.Normalize();
-			FVector3d Vkj = (Vj - Vk).Normalized();
-
+			FVector3d Vkj = Vj - Vk;
 			double cot_alpha_ij = VectorUtil::VectorCot(Vki, Vkj);
-			double cot_alpha_ik = VectorUtil::VectorCot(Vji, -Vkj);
-			areaSum += dist_ji * dist_ji * cot_alpha_ij * 0.125;
-			areaSum += dist_ki * dist_ki * cot_alpha_ik * 0.125;
+			cot_alpha_ij = FMathd::Clamp(cot_alpha_ij, -CotClampRange, CotClampRange);
+			double cot_beta_ik = VectorUtil::VectorCot(Vji, -Vkj);
+			cot_beta_ik = FMathd::Clamp(cot_alpha_ij, -CotClampRange, CotClampRange);
+			areaSum += Vji.SquaredLength() * cot_alpha_ij * 0.125;
+			areaSum += Vki.SquaredLength() * cot_beta_ik * 0.125;
 		}
 	}
 	return areaSum;
-}
\ No newline at end of file
+}
+
+
+
+
+double FMeshWeights::VoronoiArea(const FDynamicMesh3& mesh, int32 v_i)
+{
+	return TVoronoiArea(mesh, v_i, [&](int32 vid) { return mesh.GetVertex(vid); });
+}
+
+double FMeshWeights::VoronoiArea(const FDynamicMesh3& mesh, int32 v_i, TFunctionRef<FVector3d(int32)> VertexPositionFunc)
+{
+	return TVoronoiArea(mesh, v_i, VertexPositionFunc);
+}
+
diff --git a/Engine/Plugins/Experimental/GeometryProcessing/Source/DynamicMesh/Public/MeshCurvature.h b/Engine/Plugins/Experimental/GeometryProcessing/Source/DynamicMesh/Public/MeshCurvature.h
new file mode 100644
--- /dev/null
+++ b/Engine/Plugins/Experimental/GeometryProcessing/Source/DynamicMesh/Public/MeshCurvature.h
@@ -0,0 +1,46 @@
+// Copyright Epic Games, Inc. All Rights Reserved.
+
+#pragma once
+
+#include "DynamicMesh3.h"
+
+namespace UE
+{
+	/**
+	 * Implementations of various mesh curvatures and curvature-related metrics
+	 */
+	namespace MeshCurvature
+	{
+		/**
+		 * Calculate the Discrete Mean-Curvature Normal at a vertex as defined by discrete differential geometry.
+		 * Based on Eq 8 from "Discrete Differential-Geometry Operators for Triangulated 2-Manifolds", Meyer et al 2002
+		 * The Discrete Mean-Curvature Normal is (2.0 * MeanCurvature * SurfaceNormal)    ((unclear why it has the 2.0))
+		 */
+		DYNAMICMESH_API FVector3d MeanCurvatureNormal(const FDynamicMesh3& Mesh, int32 VertexIndex);
+
+
+		/**
+		 * Calculate the Discrete Mean-Curvature Normal at a vertex as defined by discrete differential geometry.
+		 * Based on Eq 8 from "Discrete Differential-Geometry Operators for Triangulated 2-Manifolds", Meyer et al 2002
+		 * The Discrete Mean-Curvature Normal is (2.0 * MeanCurvature * SurfaceNormal)    ((unclear why it has the 2.0))
+		 * @param VertexPositionFunc use positions returned by this function instead of mesh positions
+		 */
+		DYNAMICMESH_API FVector3d MeanCurvatureNormal(const FDynamicMesh3& Mesh, int32 VertexIndex, TFunctionRef<FVector3d(int32)> VertexPositionFunc);
+
+
+
+		/**
+		 * Calculate the Discrete Gaussian Curvature at a vertex as defined by discrete differential geometry
+		 * Based on Eq 9 from "Discrete Differential-Geometry Operators for Triangulated 2-Manifolds", Meyer et al 2002
+		 */
+		DYNAMICMESH_API double GaussianCurvature(const FDynamicMesh3& Mesh, int32 VertexIndex);
+
+		/**
+		 * Calculate the Discrete Gaussian Curvature at a vertex as defined by discrete differential geometry
+		 * Based on Eq 9 from "Discrete Differential-Geometry Operators for Triangulated 2-Manifolds", Meyer et al 2002
+		 * @param VertexPositionFunc use positions returned by this function instead of mesh positions
+		 */
+		DYNAMICMESH_API double GaussianCurvature(const FDynamicMesh3& Mesh, int32 VertexIndex, TFunctionRef<FVector3d(int32)> VertexPositionFunc);
+
+	}
+}
\ No newline at end of file
diff --git a/Engine/Plugins/Experimental/GeometryProcessing/Source/DynamicMesh/Public/MeshWeights.h b/Engine/Plugins/Experimental/GeometryProcessing/Source/DynamicMesh/Public/MeshWeights.h
--- a/Engine/Plugins/Experimental/GeometryProcessing/Source/DynamicMesh/Public/MeshWeights.h
+++ b/Engine/Plugins/Experimental/GeometryProcessing/Source/DynamicMesh/Public/MeshWeights.h
@@ -18,26 +18,53 @@ public:
 	 * Compute uniform centroid of a vertex one-ring.
 	 * These weights are strictly positive and all equal to 1 / valence
 	 */
-	static FVector3d UniformCentroid(const FDynamicMesh3 & mesh, int VertexIndex);
+	static FVector3d UniformCentroid(const FDynamicMesh3& Mesh, int32 VertexIndex);
 
+	/**
+	 * Compute uniform centroid of a vertex one-ring.
+	 * These weights are strictly positive and all equal to 1 / valence
+	 */
+	static FVector3d UniformCentroid(const FDynamicMesh3& Mesh, int32 VertexIndex, TFunctionRef<FVector3d(int32)> VertexPositionFunc );
+
+	/**
+	 * Compute mean-value centroid of a vertex one-ring.
+	 * These weights are strictly positive.
+	 */
+	static FVector3d MeanValueCentroid(const FDynamicMesh3& Mesh, int32 VertexIndex, double WeightClamp = FMathf::MaxReal);
 
 	/**
 	 * Compute mean-value centroid of a vertex one-ring.
 	 * These weights are strictly positive.
 	 */
-	static FVector3d MeanValueCentroid(const FDynamicMesh3 & mesh, int VertexIndex);
+	static FVector3d MeanValueCentroid(const FDynamicMesh3& Mesh, int32 VertexIndex, TFunctionRef<FVector3d(int32)> VertexPositionFunc, double WeightClamp = FMathf::MaxReal);
+
+	/**
+	 * Compute cotan-weighted centroid of a vertex one-ring.
+	 * These weights are numerically unstable if any of the triangles are degenerate.
+	 * We catch these problems and return input vertex as centroid
+	 */
+	static FVector3d CotanCentroid(const FDynamicMesh3& Mesh, int32 VertexIndex);
 
 	/**
 	 * Compute cotan-weighted centroid of a vertex one-ring.
 	 * These weights are numerically unstable if any of the triangles are degenerate.
 	 * We catch these problems and return input vertex as centroid
 	 */
-	static FVector3d CotanCentroid(const FDynamicMesh3 & mesh, int VertexIndex);
+	static FVector3d CotanCentroid(const FDynamicMesh3& Mesh, int32 VertexIndex, TFunctionRef<FVector3d(int32)> VertexPositionFunc);
+
 
 	/**
-	 * Compute the voronoi area associated with a vertex.
+	 * Compute the Mixed Voronoi Area associated with a vertex.
+	 * Based on Fig 4 from "Discrete Differential-Geometry Operators for Triangulated 2-Manifolds", Meyer et al 2002,
 	 */
-	static double VoronoiArea(const FDynamicMesh3& mesh, int VertexIndex);
+	static double VoronoiArea(const FDynamicMesh3& Mesh, int32 VertexIndex);
+
+	/**
+	 * Compute the Mixed Voronoi Area associated with a vertex.
+	 * Based on Fig 4 from "Discrete Differential-Geometry Operators for Triangulated 2-Manifolds", Meyer et al 2002,
+	 */
+	static double VoronoiArea(const FDynamicMesh3& Mesh, int32 VertexIndex, TFunctionRef<FVector3d(int32)> VertexPositionFunc);
+
 
 protected:
 	FMeshWeights() = delete;
diff --git a/Engine/Plugins/Experimental/GeometryProcessing/Source/GeometricObjects/Public/VectorUtil.h b/Engine/Plugins/Experimental/GeometryProcessing/Source/GeometricObjects/Public/VectorUtil.h
--- a/Engine/Plugins/Experimental/GeometryProcessing/Source/GeometricObjects/Public/VectorUtil.h
+++ b/Engine/Plugins/Experimental/GeometryProcessing/Source/GeometricObjects/Public/VectorUtil.h
@@ -262,7 +262,7 @@ namespace VectorUtil
 	}
 
 	/**
-	 * Fast cotangent of angle between two vectors.
+	 * Fast cotangent of angle between two vectors (*do not have to be normalized unit vectors*).
 	 * cot = cos/sin, both of which can be computed from vector identities
 	 * @return cotangent of angle between V1 and V2, or zero if result would be unstable (eg infinity)
 	 */ 
