GeometryProcessing: add MeshTransforms::Scale()
#rnx

diff --git a/Engine/Plugins/Experimental/GeometryProcessing/Source/DynamicMesh/Private/MeshTransforms.cpp b/Engine/Plugins/Experimental/GeometryProcessing/Source/DynamicMesh/Private/MeshTransforms.cpp
--- a/Engine/Plugins/Experimental/GeometryProcessing/Source/DynamicMesh/Private/MeshTransforms.cpp
+++ b/Engine/Plugins/Experimental/GeometryProcessing/Source/DynamicMesh/Private/MeshTransforms.cpp
@@ -21,6 +21,19 @@ void MeshTransforms::Translate(FDynamicMesh3& Mesh, const FVector3d& Translation
 }
 
 
+void MeshTransforms::Scale(FDynamicMesh3& Mesh, const FVector3d& Scale, const FVector3d& Origin)
+{
+	int NumVertices = Mesh.MaxVertexID();
+	ParallelFor(NumVertices, [&](int32 vid)
+	{
+		if (Mesh.IsVertex(vid))
+		{
+			Mesh.SetVertex(vid, (Mesh.GetVertex(vid) - Origin) * Scale + Origin);
+		}
+	});
+}
+
+
 
 void MeshTransforms::WorldToFrameCoords(FDynamicMesh3& Mesh, const FFrame3d& Frame)
 {
diff --git a/Engine/Plugins/Experimental/GeometryProcessing/Source/DynamicMesh/Public/MeshTransforms.h b/Engine/Plugins/Experimental/GeometryProcessing/Source/DynamicMesh/Public/MeshTransforms.h
--- a/Engine/Plugins/Experimental/GeometryProcessing/Source/DynamicMesh/Public/MeshTransforms.h
+++ b/Engine/Plugins/Experimental/GeometryProcessing/Source/DynamicMesh/Public/MeshTransforms.h
@@ -18,6 +18,11 @@ namespace MeshTransforms
 	 */
 	DYNAMICMESH_API void Translate(FDynamicMesh3& Mesh, const FVector3d& Translation);
 
+	/**
+	 * Apply Scale to vertex positions of Mesh, relative to given Origin. Does not modify any other attributes (normals/etc)
+	 */
+	DYNAMICMESH_API void Scale(FDynamicMesh3& Mesh, const FVector3d& Scale, const FVector3d& Origin);
+
 	/**
 	 * Transform Mesh into local coordinates of Frame
 	 */
