GeometryProcessing: add FMeshWeights::CotanCentroidSafe(), if returns uniform weights if magnitude of any weight is larger than threshold
ModelingTools: use above in sculpting to avoid vertices getting "stuck"
#rnx

diff --git a/Engine/Plugins/Experimental/GeometryProcessing/Source/DynamicMesh/Private/MeshWeights.cpp b/Engine/Plugins/Experimental/GeometryProcessing/Source/DynamicMesh/Private/MeshWeights.cpp
--- a/Engine/Plugins/Experimental/GeometryProcessing/Source/DynamicMesh/Private/MeshWeights.cpp
+++ b/Engine/Plugins/Experimental/GeometryProcessing/Source/DynamicMesh/Private/MeshWeights.cpp
@@ -150,6 +150,80 @@ FVector3d FMeshWeights::CotanCentroid(const FDynamicMesh3& mesh, int32 v_i, TFun
 }
 
 
+
+
+
+
+
+
+template<typename PositionFuncType>
+FVector3d TCotanCentroidSafe(const FDynamicMesh3& mesh, int32 v_i, PositionFuncType PositionFunc, double DegenerateTol = 100.0, bool* bOutputIsUniform = nullptr)
+{
+	// based on equations in http://www.geometry.caltech.edu/pubs/DMSB_III.pdf
+
+	FVector3d vSum = FVector3d::Zero();
+	double wSum = 0;
+	FVector3d Vi = PositionFunc(v_i);
+
+	bool bDegenerated = false;
+
+	int v_j = FDynamicMesh3::InvalidID, opp_v1 = FDynamicMesh3::InvalidID, opp_v2 = FDynamicMesh3::InvalidID;
+	int t1 = FDynamicMesh3::InvalidID, t2 = FDynamicMesh3::InvalidID;
+	for (int eid : mesh.VtxEdgesItr(v_i))
+	{
+		opp_v2 = FDynamicMesh3::InvalidID;
+		mesh.GetVtxNbrhood(eid, v_i, v_j, opp_v1, opp_v2, t1, t2);
+		FVector3d Vj = PositionFunc(v_j);
+
+		FVector3d Vo1 = PositionFunc(opp_v1);
+		double cot_alpha_ij = VectorUtil::VectorCot((Vi - Vo1), (Vj - Vo1));
+		double w_ij = cot_alpha_ij;
+
+		double cot_beta_ij = 0;
+		if (opp_v2 != FDynamicMesh3::InvalidID)
+		{
+			FVector3d Vo2 = PositionFunc(opp_v2);
+			cot_beta_ij = VectorUtil::VectorCot((Vi - Vo2), (Vj - Vo2));
+		}
+		w_ij += cot_beta_ij;
+
+		if (FMathd::Abs(w_ij) > DegenerateTol)
+		{
+			bDegenerated = true;
+			break;
+		}
+
+		vSum += w_ij * Vj;
+		wSum += w_ij;
+	}
+
+	if (bOutputIsUniform != nullptr)
+	{
+		*bOutputIsUniform = bDegenerated;
+	}
+
+	if (bDegenerated || fabs(wSum) < FMathd::ZeroTolerance)
+	{
+		return FMeshWeights::UniformCentroid(mesh, v_i);
+	}
+
+	return vSum / wSum;
+}
+
+
+
+FVector3d FMeshWeights::CotanCentroidSafe(const FDynamicMesh3& mesh, int32 v_i, double DegenerateTol, bool* bFailedToUniform)
+{
+	return TCotanCentroidSafe(mesh, v_i, [&](int32 vid) { return mesh.GetVertex(vid); });
+}
+
+FVector3d FMeshWeights::CotanCentroidSafe(const FDynamicMesh3& mesh, int32 v_i, TFunctionRef<FVector3d(int32)> VertexPositionFunc, double DegenerateTol, bool* bFailedToUniform)
+{
+	return TCotanCentroid(mesh, v_i, VertexPositionFunc);
+}
+
+
+
 template<typename GetPositionFuncType>
 double TVoronoiArea(const FDynamicMesh3& mesh, int32 v_i, GetPositionFuncType GetPositionFunc, double CotClampRange = 1000 /* about 0.05 deg */ )
 {
diff --git a/Engine/Plugins/Experimental/GeometryProcessing/Source/DynamicMesh/Public/MeshWeights.h b/Engine/Plugins/Experimental/GeometryProcessing/Source/DynamicMesh/Public/MeshWeights.h
--- a/Engine/Plugins/Experimental/GeometryProcessing/Source/DynamicMesh/Public/MeshWeights.h
+++ b/Engine/Plugins/Experimental/GeometryProcessing/Source/DynamicMesh/Public/MeshWeights.h
@@ -53,6 +53,25 @@ public:
 	static FVector3d CotanCentroid(const FDynamicMesh3& Mesh, int32 VertexIndex, TFunctionRef<FVector3d(int32)> VertexPositionFunc);
 
 
+
+	/**
+	 * Compute cotan-weighted centroid of a vertex one-ring, with some weight analysis/clamping to avoid vertices getting "stuck"
+	 * in explicit integration/iterations. If failure is detected, Uniform centroid is returned, which does cause some tangential flow
+	 * @param DegenerateTol if any weights are larger than this value, return uniform weights instead. Should be > 1.
+	 * @param bFailedToUniform will be set to true if non-null and result was clamped
+	 */
+	static FVector3d CotanCentroidSafe(const FDynamicMesh3& Mesh, int32 VertexIndex, double DegenerateTol = 100.0, bool* bFailedToUniform = nullptr);
+
+	/**
+	 * Compute cotan-weighted centroid of a vertex one-ring, with some weight analysis/clamping to avoid vertices getting "stuck"
+	 * in explicit integration/iterations. If failure is detected, Uniform centroid is returned, which does cause some tangential flow
+	 * @param DegenerateTol if any weights are larger than this value, return uniform weights instead. Should be > 1.
+	 * @param bFailedToUniform will be set to true if non-null and result was clamped
+	 */
+	static FVector3d CotanCentroidSafe(const FDynamicMesh3& Mesh, int32 VertexIndex, TFunctionRef<FVector3d(int32)> VertexPositionFunc, double DegenerateTol = 100.0, bool* bFailedToUniform = nullptr);
+
+
+
 	/**
 	 * Compute the Mixed Voronoi Area associated with a vertex.
 	 * Based on Fig 4 from "Discrete Differential-Geometry Operators for Triangulated 2-Manifolds", Meyer et al 2002,
diff --git a/Engine/Plugins/Experimental/MeshModelingToolset/Source/MeshModelingTools/Private/DynamicMeshSculptTool.cpp b/Engine/Plugins/Experimental/MeshModelingToolset/Source/MeshModelingTools/Private/DynamicMeshSculptTool.cpp
--- a/Engine/Plugins/Experimental/MeshModelingToolset/Source/MeshModelingTools/Private/DynamicMeshSculptTool.cpp
+++ b/Engine/Plugins/Experimental/MeshModelingToolset/Source/MeshModelingTools/Private/DynamicMeshSculptTool.cpp
@@ -558,7 +558,7 @@ void UDynamicMeshSculptTool::ApplySmoothBrush(const FRay& WorldRay)
 		double Falloff = CalculateBrushFalloff(OrigPos.Distance(NewBrushPosLocal));
 
 		FVector3d SmoothedPos = (SculptProperties->bPreserveUVFlow) ?
-			FMeshWeights::MeanValueCentroid(*Mesh, VertIdx) : FMeshWeights::UniformCentroid(*Mesh, VertIdx);
+			FMeshWeights::CotanCentroidSafe(*Mesh, VertIdx, 10.0) : FMeshWeights::UniformCentroid(*Mesh, VertIdx);
 
 		FVector3d NewPos = FVector3d::Lerp(OrigPos, SmoothedPos, Falloff*SculptProperties->SmoothBrushSpeed);
 
diff --git a/Engine/Plugins/Experimental/MeshModelingToolset/Source/MeshModelingTools/Public/Sculpting/MeshSmoothingBrushOps.h b/Engine/Plugins/Experimental/MeshModelingToolset/Source/MeshModelingTools/Public/Sculpting/MeshSmoothingBrushOps.h
--- a/Engine/Plugins/Experimental/MeshModelingToolset/Source/MeshModelingTools/Public/Sculpting/MeshSmoothingBrushOps.h
+++ b/Engine/Plugins/Experimental/MeshModelingToolset/Source/MeshModelingTools/Public/Sculpting/MeshSmoothingBrushOps.h
@@ -85,8 +85,15 @@ public:
 
 			double Falloff = GetFalloff().Evaluate(Stamp, OrigPos);
 
-			FVector3d SmoothedPos = (bPreserveUVFlow) ?
-				FMeshWeights::MeanValueCentroid(*Mesh, VertIdx) : FMeshWeights::UniformCentroid(*Mesh, VertIdx);
+			FVector3d SmoothedPos = OrigPos;
+			if (bPreserveUVFlow)
+			{
+				SmoothedPos = FMeshWeights::CotanCentroidSafe(*Mesh, VertIdx, 10.0);
+			}
+			else
+			{
+				SmoothedPos = FMeshWeights::UniformCentroid(*Mesh, VertIdx);
+			}
 
 			FVector3d NewPos = FVector3d::Lerp(OrigPos, SmoothedPos, Falloff * Stamp.Power);
 
@@ -122,6 +129,7 @@ public:
 };
 
 
+
 class FSmoothFillBrushOp : public FMeshSculptBrushOp
 {
 
@@ -142,8 +150,15 @@ public:
 
 			double Falloff = GetFalloff().Evaluate(Stamp, OrigPos);
 
-			FVector3d SmoothedPos = (bPreserveUVFlow) ?
-				FMeshWeights::MeanValueCentroid(*Mesh, VertIdx) : FMeshWeights::UniformCentroid(*Mesh, VertIdx);
+			FVector3d SmoothedPos = OrigPos;
+			if (bPreserveUVFlow)
+			{
+				SmoothedPos = FMeshWeights::CotanCentroidSafe(*Mesh, VertIdx, 10.0);
+			}
+			else
+			{
+				SmoothedPos = FMeshWeights::UniformCentroid(*Mesh, VertIdx);
+			}
 
 			FVector3d NewPos = FVector3d::Lerp(OrigPos, SmoothedPos, Falloff * Stamp.Power);
 
