GeometryProcessing: further refactoring of mesh solvers. Moved different parts of LaplacianMeshSmoother file into ConstarinedMeshSmoothers.h, ConstrainedMeshDeformers.h, ConstrainedMeshDeformationSolver.h, MeshDiffusionIntegrator, and MeshDiffusionSmoothing.h. DeformMeshPolygonsTool now uses IConstrainedMeshSolver interface, and no longer includes Private headers from GeometryProcessing.

Also made a small change to LaplacianOperators.cpp:ConstructScaledCotangentLaplacian(), which had checkSlow() that was failing on zero-elements Area matrix at boundary vertices. Code in this function does not appear to correctly handle boundary vertices (however things work out in the end because matrices are squared later?)

#rnx

diff --git a/Engine/Plugins/Experimental/GeometryProcessing/Source/DynamicMesh/Private/Solvers/ConstrainedMeshDeformer.cpp b/Engine/Plugins/Experimental/GeometryProcessing/Source/DynamicMesh/Private/Solvers/ConstrainedMeshDeformer.cpp
--- a/Engine/Plugins/Experimental/GeometryProcessing/Source/DynamicMesh/Private/Solvers/ConstrainedMeshDeformer.cpp
+++ b/Engine/Plugins/Experimental/GeometryProcessing/Source/DynamicMesh/Private/Solvers/ConstrainedMeshDeformer.cpp
@@ -1,12 +1,12 @@
 // Copyright Epic Games, Inc. All Rights Reserved.
 
 #include "Solvers/ConstrainedMeshDeformer.h"
-#include "Solvers/Internal/LaplacianMeshSmoother.h"
+#include "Solvers/Internal/ConstrainedMeshDeformers.h"
 
 
-TUniquePtr<IConstrainedMeshSolver> UE::MeshDeformation::ConstructConstrainedMeshDeformer(const ELaplacianWeightScheme WeightScheme, const FDynamicMesh3& DynamicMesh)
+TUniquePtr<UE::Solvers::IConstrainedMeshSolver> UE::MeshDeformation::ConstructConstrainedMeshDeformer(const ELaplacianWeightScheme WeightScheme, const FDynamicMesh3& DynamicMesh)
 {
-	TUniquePtr<IConstrainedMeshSolver> Deformer(new FConstrainedMeshDeformer(DynamicMesh, WeightScheme));
+	TUniquePtr<UE::Solvers::IConstrainedMeshSolver> Deformer(new FConstrainedMeshDeformer(DynamicMesh, WeightScheme));
 	return Deformer;
 }
 
diff --git a/Engine/Plugins/Experimental/GeometryProcessing/Source/DynamicMesh/Private/Solvers/ConstrainedMeshSmoother.cpp b/Engine/Plugins/Experimental/GeometryProcessing/Source/DynamicMesh/Private/Solvers/ConstrainedMeshSmoother.cpp
--- a/Engine/Plugins/Experimental/GeometryProcessing/Source/DynamicMesh/Private/Solvers/ConstrainedMeshSmoother.cpp
+++ b/Engine/Plugins/Experimental/GeometryProcessing/Source/DynamicMesh/Private/Solvers/ConstrainedMeshSmoother.cpp
@@ -1,12 +1,12 @@
 // Copyright Epic Games, Inc. All Rights Reserved.
 
 #include "Solvers/ConstrainedMeshSmoother.h"
-#include "Solvers/Internal/LaplacianMeshSmoother.h"
+#include "Solvers/Internal/ConstrainedMeshSmoothers.h"
 
 
-TUniquePtr<IConstrainedMeshSolver> UE::MeshDeformation::ConstructConstrainedMeshSmoother(const ELaplacianWeightScheme WeightScheme, const FDynamicMesh3& DynamicMesh)
+TUniquePtr<UE::Solvers::IConstrainedMeshSolver> UE::MeshDeformation::ConstructConstrainedMeshSmoother(const ELaplacianWeightScheme WeightScheme, const FDynamicMesh3& DynamicMesh)
 {
-	TUniquePtr<IConstrainedMeshSolver> Deformer(new FBiHarmonicMeshSmoother(DynamicMesh, WeightScheme));
+	TUniquePtr<UE::Solvers::IConstrainedMeshSolver> Deformer(new FBiHarmonicMeshSmoother(DynamicMesh, WeightScheme));
 	return Deformer;
 }
 
diff --git a/Engine/Plugins/Experimental/GeometryProcessing/Source/DynamicMesh/Private/Solvers/Internal/ConstrainedMeshDeformationSolver.cpp b/Engine/Plugins/Experimental/GeometryProcessing/Source/DynamicMesh/Private/Solvers/Internal/ConstrainedMeshDeformationSolver.cpp
new file mode 100644
--- /dev/null
+++ b/Engine/Plugins/Experimental/GeometryProcessing/Source/DynamicMesh/Private/Solvers/Internal/ConstrainedMeshDeformationSolver.cpp
@@ -0,0 +1,337 @@
+// Copyright Epic Games, Inc. All Rights Reserved.
+
+#include "ConstrainedMeshDeformationSolver.h"
+#include "LaplacianOperators.h"
+
+#include "MatrixSolver.h"
+#include "ConstrainedPoissonSolver.h"
+
+
+#ifdef TIME_LAPLACIAN_SMOOTHERS
+
+#include "ProfilingDebugging/ScopedTimers.h"
+DEFINE_LOG_CATEGORY_STATIC(LogMeshSmoother, Log, All);
+
+#endif
+
+double ComputeDistSqrd(const FSOAPositions& VecA, const FSOAPositions& VecB)
+{
+	const int32 NumA = VecA.Num();
+	const int32 NumB = VecB.Num();
+	checkSlow(NumA == NumB);
+
+#if 0
+	// the eigne way?
+	double Tmp = (VecA.Array(0) - VecB.Array(0)).dot(VecA.Array(0) - VecB.Array(0))
+		+ (VecA.Array(1) - VecB.Array(1)).dot(VecA.Array(1) - VecB.Array(1))
+		+ (VecA.Array(2) - VecB.Array(2)).dot(VecA.Array(2) - VecB.Array(2));
+
+#endif
+
+	// doing it by hand so we can break when the error is large
+	double DistSqrd = 0.;
+	{
+		const auto& AX = VecA.Array(0);
+		const auto& AY = VecA.Array(1);
+		const auto& AZ = VecA.Array(2);
+
+		const auto& BX = VecB.Array(0);
+		const auto& BY = VecB.Array(1);
+		const auto& BZ = VecB.Array(2);
+
+
+		for (int32 i = 0; i < NumA; ++i)
+		{
+			double TmpX = AX(i) - BX(i);
+			double TmpY = AY(i) - BY(i);
+			double TmpZ = AZ(i) - BZ(i);
+
+			TmpX *= TmpX;
+			TmpY *= TmpY;
+			TmpZ *= TmpZ;
+			double TmpT = TmpX + TmpY + TmpZ;
+			DistSqrd += TmpT;
+		}
+	}
+	return DistSqrd;
+}
+
+
+FConstrainedMeshDeformationSolver::~FConstrainedMeshDeformationSolver() 
+{
+}
+
+FConstrainedMeshDeformationSolver::FConstrainedMeshDeformationSolver(const FDynamicMesh3& DynamicMesh, const ELaplacianWeightScheme Scheme, const EMatrixSolverType MatrixSolverType)
+	: VertexCount(DynamicMesh.VertexCount())
+{
+
+	FSparseMatrixD LaplacianInternal;
+	FSparseMatrixD LaplacianBoundary;
+	ConstructLaplacian(Scheme, DynamicMesh, VtxLinearization, LaplacianInternal, LaplacianBoundary);
+
+	// Copy the original boundary vertex locations
+	const int32 BoundaryVertexCount = VtxLinearization.NumBoundaryVerts();
+
+	// Number of vertices in the interior of the mesh
+	InternalVertexCount = VertexCount - BoundaryVertexCount;
+
+	// Copy the original boundary vertex locations
+	{
+		const auto ToVertId = VtxLinearization.ToId();
+		BoundaryPositions.SetZero(BoundaryVertexCount);
+		for (int32 i = 0; i < BoundaryVertexCount; ++i)
+		{
+			const int32 VtxId = ToVertId[i + InternalVertexCount];
+			const FVector3d Pos = DynamicMesh.GetVertex(VtxId);
+			BoundaryPositions.SetXYZ(i, Pos);
+		}
+	}
+
+	checkSlow(LaplacianInternal.rows() == LaplacianInternal.cols());
+
+	TUniquePtr<FSparseMatrixD> LTLPtr(new FSparseMatrixD(LaplacianInternal.rows(), LaplacianInternal.cols()));
+	FSparseMatrixD& LTLMatrix = *(LTLPtr);
+
+	bool bIsLaplacianSymmetric = (Scheme == ELaplacianWeightScheme::Valence || Scheme == ELaplacianWeightScheme::Uniform);
+
+	if (bIsLaplacianSymmetric)
+	{
+		// Laplacian is symmetric, i.e. equal to its transpose
+		LTLMatrix = LaplacianInternal * LaplacianInternal;
+		BoundaryOperator = -1. * LaplacianInternal * LaplacianBoundary;
+
+		ConstrainedSolver.Reset(new FConstrainedSolver(LTLPtr, MatrixSolverType));
+	}
+	else
+	{
+		// the laplacian 
+		LTLMatrix = LaplacianInternal.transpose() * LaplacianInternal;
+		BoundaryOperator = -1. * LaplacianInternal.transpose() * LaplacianBoundary;
+
+		ConstrainedSolver.Reset(new FConstrainedSolver(LTLPtr, MatrixSolverType));
+	}
+
+}
+
+void FConstrainedMeshDeformationSolver::AddConstraint(const int32 VtxId, const double Weight, const FVector3d& Pos, const bool bPostFix)
+{
+	const auto& ToIndex = VtxLinearization.ToIndex();
+
+	if (VtxId > ToIndex.Num())
+	{
+		return;
+	}
+
+	const int32 Index = ToIndex[VtxId];
+
+	// Only add the constraint if the vertex is actually in the interior.  We aren't solving for edge vertices.
+	if (Index != FDynamicMesh3::InvalidID && Index < InternalVertexCount)
+	{
+		bConstraintPositionsDirty = true;
+		bConstraintWeightsDirty = true;
+
+
+		ConstraintPositionMap.Add(TTuple<int32, FConstraintPosition>(Index, FConstraintPosition(Pos, bPostFix)));
+		ConstraintWeightMap.Add(TTuple<int32, double>(Index, Weight));
+	}
+}
+
+
+bool FConstrainedMeshDeformationSolver::UpdateConstraintPosition(const int32 VtxId, const FVector3d& Pos, const bool bPostFix)
+{
+	bool Result = false;
+	const auto& ToIndex = VtxLinearization.ToIndex();
+
+	if (VtxId > ToIndex.Num())
+	{
+		return Result;
+	}
+
+	const int32 Index = ToIndex[VtxId];
+
+
+	if (Index != FDynamicMesh3::InvalidID && Index < InternalVertexCount)
+	{
+		bConstraintPositionsDirty = true;
+
+		// Add should over-write any existing value for this key
+		ConstraintPositionMap.Add(TTuple<int32, FConstraintPosition>(Index, FConstraintPosition(Pos, bPostFix)));
+
+		Result = ConstraintWeightMap.Contains(VtxId);
+	}
+	return Result;
+}
+
+bool FConstrainedMeshDeformationSolver::UpdateConstraintWeight(const int32 VtxId, const double Weight)
+{
+	bool Result = false;
+	const auto& ToIndex = VtxLinearization.ToIndex();
+
+	if (VtxId > ToIndex.Num())
+	{
+		return Result;
+	}
+
+	const int32 Index = ToIndex[VtxId];
+	if (Index != FDynamicMesh3::InvalidID && Index < InternalVertexCount)
+	{
+		bConstraintWeightsDirty = true;
+
+		// Add should over-write any existing value for this key
+		ConstraintWeightMap.Add(TTuple<int32, double>(Index, Weight));
+
+		Result = ConstraintPositionMap.Contains(VtxId);
+	}
+	return Result;
+}
+
+
+bool FConstrainedMeshDeformationSolver::IsConstrained(const int32 VtxId) const
+{
+	bool Result = false;
+
+	const auto& ToIndex = VtxLinearization.ToIndex();
+
+	if (VtxId > ToIndex.Num())
+	{
+		return Result;
+	}
+
+	const int32 Index = ToIndex[VtxId];
+
+	if (Index != FDynamicMesh3::InvalidID && Index < InternalVertexCount)
+	{
+		Result = ConstraintWeightMap.Contains(Index);
+	}
+
+	return Result;
+}
+
+void FConstrainedMeshDeformationSolver::UpdateSolverConstraints()
+{
+	if (bConstraintWeightsDirty)
+	{
+		ConstrainedSolver->SetConstraintWeights(ConstraintWeightMap);
+		bConstraintWeightsDirty = false;
+	}
+
+	if (bConstraintPositionsDirty)
+	{
+		ConstrainedSolver->SetContraintPositions(ConstraintPositionMap);
+		bConstraintPositionsDirty = false;
+	}
+}
+
+bool FConstrainedMeshDeformationSolver::CopyInternalPositions(const FSOAPositions& PositionalVector, TArray<FVector3d>& LinearArray) const
+{
+	// Number of positions
+
+	const int32 Num = PositionalVector.Num();
+
+	// early out if the x,y,z arrays in the PositionalVector have different lengths
+	if (!PositionalVector.bHasSize(Num))
+	{
+		return false;
+	}
+
+	checkSlow(Num == InternalVertexCount);
+
+	// 
+	const auto& ToVtxId = VtxLinearization.ToId();
+	const int32 MaxVtxId = ToVtxId.Num(); // NB: this is really max_used + 1 in the mesh.  See  FDynamicMesh3::MaxVertexID()
+
+	if (LinearArray.Num() != MaxVtxId)
+	{
+		return false;
+	}
+
+	// Update the internal positions.
+	for (int32 i = 0; i < InternalVertexCount; ++i)
+	{
+		const int32 VtxId = ToVtxId[i];
+
+		LinearArray[VtxId] = FVector3d(PositionalVector.X(i), PositionalVector.Y(i), PositionalVector.Z(i));
+	}
+
+	return true;
+}
+
+bool FConstrainedMeshDeformationSolver::CopyBoundaryPositions(TArray<FVector3d>& LinearArray) const
+{
+	const auto& ToVtxId = VtxLinearization.ToId();
+	const int32 MaxVtxId = ToVtxId.Num();
+
+	if (LinearArray.Num() != MaxVtxId)
+	{
+		return false;
+	}
+
+	int32 BoundaryVertexCount = VertexCount - InternalVertexCount;
+
+	for (int32 i = 0; i < BoundaryVertexCount; ++i)
+	{
+		const int32 VtxId = ToVtxId[i + InternalVertexCount];
+		LinearArray[VtxId] = FVector3d(BoundaryPositions.X(i), BoundaryPositions.Y(i), BoundaryPositions.Z(i)); //BoundaryPositions[i];
+	}
+
+	return true;
+}
+
+
+void FConstrainedMeshDeformationSolver::ExtractInteriorVertexPositions(const FDynamicMesh3& DynamicMesh, FSOAPositions& VertexPositions) const
+{
+	VertexPositions.SetZero(InternalVertexCount);
+
+	const auto& ToVtxId = VtxLinearization.ToId();
+
+	for (int32 i = 0; i < InternalVertexCount; ++i)
+	{
+		const int32 VtxId = ToVtxId[i];
+		const FVector3d& Pos = DynamicMesh.GetVertex(VtxId);
+		VertexPositions.SetXYZ(i, Pos);
+	}
+}
+
+void FConstrainedMeshDeformationSolver::UpdateWithPostFixConstraints(FSOAPositions& PositionVector) const
+{
+	for (const auto& ConstraintPosition : ConstraintPositionMap)
+	{
+		const int32 Index = ConstraintPosition.Key;
+		const FConstraintPosition& Constraint = ConstraintPosition.Value;
+
+		checkSlow(Index < InternalVertexCount);
+
+		// we only care about post-fix constraints
+
+		if (Constraint.bPostFix)
+		{
+			const FVector3d& Pos = Constraint.Position;
+			PositionVector.SetXYZ(Index, Pos);
+		}
+	}
+}
+
+
+
+
+bool FConstrainedMeshDeformationSolver::SetMaxIterations(int32 MaxIterations)
+{
+	IIterativeMatrixSolverBase* Solver = ConstrainedSolver->GetMatrixSolverIterativeBase();
+	if (Solver)
+	{
+		Solver->SetIterations(MaxIterations);
+		return true;
+	}
+	return false;
+}
+
+bool FConstrainedMeshDeformationSolver::SetTolerance(double Tol)
+{
+	IIterativeMatrixSolverBase* Solver = ConstrainedSolver->GetMatrixSolverIterativeBase();
+	if (Solver)
+	{
+		Solver->SetTolerance(Tol);
+		return true;
+	}
+	return false;
+}
diff --git a/Engine/Plugins/Experimental/GeometryProcessing/Source/DynamicMesh/Private/Solvers/Internal/ConstrainedMeshDeformationSolver.h b/Engine/Plugins/Experimental/GeometryProcessing/Source/DynamicMesh/Private/Solvers/Internal/ConstrainedMeshDeformationSolver.h
new file mode 100644
--- /dev/null
+++ b/Engine/Plugins/Experimental/GeometryProcessing/Source/DynamicMesh/Private/Solvers/Internal/ConstrainedMeshDeformationSolver.h
@@ -0,0 +1,113 @@
+// Copyright Epic Games, Inc. All Rights Reserved.
+
+#pragma once
+
+#include "CoreMinimal.h"
+#include "DynamicMesh3.h"
+#include "Solvers/MatrixInterfaces.h"
+#include "Solvers/MeshLaplacian.h"
+#include "Solvers/MeshLinearization.h"
+#include "Solvers/ConstrainedMeshSolver.h"
+
+#include "MatrixBase.h"
+
+
+class FConstrainedSolver;
+
+/**
+ * FConstrainedMeshDeformationSolver is an implmentation of IConstrainedMeshSolver that solves
+ * Mesh Deformation problems by using quadratic energy functions based on the vertex-graph Laplacians. 
+ *
+ * Both the Laplacian weighting scheme/type and The Sparse Matrix Solver type are configurable.
+ */
+class DYNAMICMESH_API FConstrainedMeshDeformationSolver : public UE::Solvers::IConstrainedMeshSolver
+{
+public:
+	typedef UE::Solvers::FPositionConstraint  FConstraintPosition;
+
+
+	FConstrainedMeshDeformationSolver(const FDynamicMesh3& DynamicMesh, const ELaplacianWeightScheme Scheme, const EMatrixSolverType MatrixSolverType);
+	virtual ~FConstrainedMeshDeformationSolver();
+
+	// Add constraint associated with given vertex id.  Boundary vertices will be ignored
+	void AddConstraint(const int32 VtxId, const double Weight, const FVector3d& Pos, const bool bPostFix) override;
+
+	// Update the position of an existing constraint.  Bool return if a corresponding constraint weight exists. Boundary vertices will be ignored (and return false).  
+	bool UpdateConstraintPosition(const int32 VtxId, const FVector3d& Position, const bool bPostFix) override;
+
+	// The underlying solver will have to refactor the matrix if this is done. Bool return if a corresponding constraint position exists. Boundary vertices will be ignored (and return false).  
+	bool UpdateConstraintWeight(const int32 VtxId, const double Weight) override;
+
+	// Clear all constraints associated with this smoother
+	void ClearConstraints() override { FConstrainedMeshDeformationSolver::ClearConstraintPositions(); FConstrainedMeshDeformationSolver::ClearConstraintWeights(); }
+
+	void ClearConstraintWeights() override { ConstraintWeightMap.Empty();    bConstraintWeightsDirty = true; }
+
+	void ClearConstraintPositions() override { ConstraintPositionMap.Empty();  bConstraintPositionsDirty = true; }
+
+
+	// Test if for constraint associated with given vertex id. Will return false for any boundary vert.
+	bool IsConstrained(const int32 VtxId) const override;
+
+	virtual bool Deform(TArray<FVector3d>& PositionBuffer) override { return false; }
+
+	// Sync constraints with internal solver.  If in the process any internal matrix factoring is dirty, it will be rebuilt.
+	// Note: this is called from within the Deform() method.   Only call this method if you want to trigger the matrix refactor yourself.
+	void UpdateSolverConstraints();
+
+
+
+public:
+	// these only apply to iterative solvers
+
+	bool SetMaxIterations(int32 MaxIterations);
+	bool SetTolerance(double Tol);
+
+
+protected:
+
+
+	void ExtractInteriorVertexPositions(const FDynamicMesh3& DynamicMesh, FSOAPositions& Positions) const;
+
+	//void UpdateMeshWithConstraints();
+	// Respect any bPostFix constraints by moving those vertices to position defined by said constraint.
+	void UpdateWithPostFixConstraints(FSOAPositions& PositionalVector) const;
+
+
+	// converts the positionalvector to a TArray<FVector3d> where the offset in the array is implicitly the
+	// VtxId in the mesh, and not ness the matrix row id
+	// NB: the resulting array is treated as sparse and may have un-initialized elements.
+	bool CopyInternalPositions(const FSOAPositions& PositionalVector, TArray<FVector3d>& LinearArray) const;
+
+	bool CopyBoundaryPositions(TArray<FVector3d>& LinearArray) const;
+
+
+protected:
+
+	// The Key (int32) here is the vertex index not vertex ID
+	// making it the same as the matrix row
+	TMap<int32, FConstraintPosition> ConstraintPositionMap;
+	TMap<int32, double>              ConstraintWeightMap;
+
+	bool bConstraintPositionsDirty = true;
+	bool bConstraintWeightsDirty = true;
+
+	// Cache the vertex count. boundary + internal
+	int32                 VertexCount;
+
+	// Cache the number of internal vertices
+	int32                 InternalVertexCount;
+
+	// Used to map between VtxId and vertex Index in linear vector..
+	FVertexLinearization  VtxLinearization;
+
+	// Boundary points, split into 3 arrays (x,y,z).
+	FSOAPositions         BoundaryPositions;
+
+
+	// Actual solver that manages the various linear algebra bits.
+	TUniquePtr<FConstrainedSolver>         ConstrainedSolver;
+
+	// Sparse Matrix that holds L^T * B where B has the boundary terms.
+	FSparseMatrixD                         BoundaryOperator;
+};
\ No newline at end of file
diff --git a/Engine/Plugins/Experimental/GeometryProcessing/Source/DynamicMesh/Private/Solvers/Internal/ConstrainedMeshDeformers.cpp b/Engine/Plugins/Experimental/GeometryProcessing/Source/DynamicMesh/Private/Solvers/Internal/ConstrainedMeshDeformers.cpp
new file mode 100644
--- /dev/null
+++ b/Engine/Plugins/Experimental/GeometryProcessing/Source/DynamicMesh/Private/Solvers/Internal/ConstrainedMeshDeformers.cpp
@@ -0,0 +1,75 @@
+
+// Copyright Epic Games, Inc. All Rights Reserved.
+
+#include "ConstrainedMeshDeformers.h"
+#include "MatrixSolver.h"
+#include "ConstrainedPoissonSolver.h"
+
+
+FConstrainedMeshDeformer::FConstrainedMeshDeformer(const FDynamicMesh3& DynamicMesh, const ELaplacianWeightScheme LaplacianType)
+	: FConstrainedMeshDeformationSolver(DynamicMesh, LaplacianType, EMatrixSolverType::LU)
+	, LaplacianVectors(FConstrainedMeshDeformationSolver::InternalVertexCount)
+{
+
+
+	// The current vertex positions 
+
+	// Note: the OriginalInteriorPositions are being stored as member data 
+	// for use if the solver is iterative.
+	// FSOAPositions OriginalInteriorPositions; 
+	ExtractInteriorVertexPositions(DynamicMesh, OriginalInteriorPositions);
+
+
+	// The biharmonic part of the constrained solver
+	//   Biharmonic := Laplacian^{T} * Laplacian
+
+	const auto& Biharmonic = ConstrainedSolver->Biharmonic();
+
+	// Compute the Laplacian Vectors
+	//    := Biharmonic * VertexPostion
+	// In the case of the cotangent laplacian this can be identified as the mean curvature * normal.
+	checkSlow(LaplacianVectors.Num() == OriginalInteriorPositions.Num());
+
+	for (int32 i = 0; i < 3; ++i)
+	{
+		LaplacianVectors.Array(i) = Biharmonic * OriginalInteriorPositions.Array(i);
+	}
+}
+
+
+bool FConstrainedMeshDeformer::Deform(TArray<FVector3d>& PositionBuffer)
+{
+
+	// Update constraints.  This only trigger solver rebuild if the weights were updated.
+	UpdateSolverConstraints();
+
+	// Allocate space for the result as a struct of arrays
+	FSOAPositions SolutionVector(InternalVertexCount);
+
+	// Solve the linear system
+	// NB: the original positions will only be used if the underlying solver type is iterative	
+	bool bSuccess = ConstrainedSolver->SolveWithGuess(OriginalInteriorPositions, LaplacianVectors, SolutionVector);
+
+	// Move any vertices to match bPostFix constraints
+
+	UpdateWithPostFixConstraints(SolutionVector);
+
+	// Allocate Position Buffer for random access writes
+	int32 MaxVtxId = VtxLinearization.ToId().Num();
+	PositionBuffer.Empty(MaxVtxId);
+	PositionBuffer.AddUninitialized(MaxVtxId);
+
+	// Export the computed internal positions:
+	// Copy the results into the array of structs form.  
+	// NB: this re-indexes so the results can be looked up using VtxId
+
+	CopyInternalPositions(SolutionVector, PositionBuffer);
+
+	// Copy the boundary
+	// NB: this re-indexes so the results can be looked up using VtxId
+	CopyBoundaryPositions(PositionBuffer);
+
+	// the matrix solve state
+	return bSuccess;
+
+}
diff --git a/Engine/Plugins/Experimental/GeometryProcessing/Source/DynamicMesh/Private/Solvers/Internal/ConstrainedMeshDeformers.h b/Engine/Plugins/Experimental/GeometryProcessing/Source/DynamicMesh/Private/Solvers/Internal/ConstrainedMeshDeformers.h
new file mode 100644
--- /dev/null
+++ b/Engine/Plugins/Experimental/GeometryProcessing/Source/DynamicMesh/Private/Solvers/Internal/ConstrainedMeshDeformers.h
@@ -0,0 +1,30 @@
+// Copyright Epic Games, Inc. All Rights Reserved.
+
+#pragma once
+
+#include "CoreMinimal.h"
+#include "DynamicMesh3.h"
+
+#include "ConstrainedMeshDeformationSolver.h"
+
+
+/**
+ * FConstrainedMeshDeformer solves detail-preserving mesh deformation problems with arbitrary position constraints.
+ * The initial Mesh Laplacians are defined as Biharmonic * VtxPositions
+ * A direct solver is used, currently LU decomposition.
+ *
+ * Boundary vertices are fixed to their input positions.
+ */
+class DYNAMICMESH_API FConstrainedMeshDeformer : public FConstrainedMeshDeformationSolver
+{
+public:
+	FConstrainedMeshDeformer(const FDynamicMesh3& DynamicMesh, const ELaplacianWeightScheme LaplacianType);
+	~FConstrainedMeshDeformer() override {}
+
+	bool Deform(TArray<FVector3d>& PositionBuffer) override;
+
+private:
+
+	FSOAPositions LaplacianVectors;
+	FSOAPositions OriginalInteriorPositions;
+};
diff --git a/Engine/Plugins/Experimental/GeometryProcessing/Source/DynamicMesh/Private/Solvers/Internal/ConstrainedMeshSmoothers.cpp b/Engine/Plugins/Experimental/GeometryProcessing/Source/DynamicMesh/Private/Solvers/Internal/ConstrainedMeshSmoothers.cpp
new file mode 100644
--- /dev/null
+++ b/Engine/Plugins/Experimental/GeometryProcessing/Source/DynamicMesh/Private/Solvers/Internal/ConstrainedMeshSmoothers.cpp
@@ -0,0 +1,115 @@
+// Copyright Epic Games, Inc. All Rights Reserved.
+
+#include "ConstrainedMeshSmoothers.h"
+#include "MatrixSolver.h"
+#include "ConstrainedPoissonSolver.h"
+
+
+
+
+bool FBiHarmonicMeshSmoother::ComputeSmoothedMeshPositions(TArray<FVector3d>& UpdatedPositions)
+{
+
+	UpdateSolverConstraints();
+
+	// Compute the source vector
+	FSOAPositions SourceVector(InternalVertexCount);
+
+	if (InternalVertexCount != VertexCount) // have boundary points
+	{
+		for (int32 dir = 0; dir < 3; ++dir)
+		{
+			SourceVector.Array(dir) = BoundaryOperator * BoundaryPositions.Array(dir);
+		}
+	}
+	else
+	{
+		SourceVector.SetZero(InternalVertexCount);
+	}
+
+	// Solves the constrained system and updates the mesh 
+
+	FSOAPositions SolutionVector(InternalVertexCount);
+
+	bool bSuccess = ConstrainedSolver->Solve(SourceVector, SolutionVector);
+
+
+	// Move any vertices to match bPostFix constraints
+
+	UpdateWithPostFixConstraints(SolutionVector);
+
+	// Allocate Position Buffer for random access writes
+	int32 MaxVtxId = VtxLinearization.ToId().Num();
+	UpdatedPositions.Empty(MaxVtxId);
+	UpdatedPositions.AddUninitialized(MaxVtxId);
+
+	// Export the computed internal positions:
+	// Copy the results into the array of structs form.  
+	// NB: this re-indexes so the results can be looked up using VtxId
+
+	CopyInternalPositions(SolutionVector, UpdatedPositions);
+
+	// Copy the boundary
+	// NB: this re-indexes so the results can be looked up using VtxId
+	CopyBoundaryPositions(UpdatedPositions);
+
+	return bSuccess;
+}
+
+
+
+
+
+
+bool FCGBiHarmonicMeshSmoother::ComputeSmoothedMeshPositions(TArray<FVector3d>& UpdatedPositions)
+{
+	// NB: This conjugate gradient solver could be updated to use  solveWithGuess() method on the iterative solver
+
+	UpdateSolverConstraints();
+
+	// Compute the source vector
+	FSOAPositions SourceVector(InternalVertexCount);
+
+	if (InternalVertexCount != VertexCount) // have boundary points
+	{
+		for (int32 dir = 0; dir < 3; ++dir)
+		{
+			SourceVector.Array(dir) = BoundaryOperator * BoundaryPositions.Array(dir);
+		}
+	}
+	else
+	{
+		SourceVector.SetZero(InternalVertexCount);
+	}
+
+	// Solves the constrained system and updates the mesh 
+
+	// Solves the constrained system
+
+	FSOAPositions SolutionVector(InternalVertexCount);
+
+	bool bSuccess = ConstrainedSolver->Solve(SourceVector, SolutionVector);
+
+
+	// Move any vertices to match bPostFix constraints
+
+	UpdateWithPostFixConstraints(SolutionVector);
+
+	// Allocate Position Buffer for random access writes
+	int32 MaxVtxId = VtxLinearization.ToId().Num();
+	UpdatedPositions.Empty(MaxVtxId);
+	UpdatedPositions.AddUninitialized(MaxVtxId);
+
+	// Export the computed internal postions:
+	// Copy the results into the array of structs form.  
+	// NB: this re-indexes so the results can be looked up using VtxId
+
+	CopyInternalPositions(SolutionVector, UpdatedPositions);
+
+	// Copy the boundary
+	// NB: this re-indexes so the results can be looked up using VtxId
+	CopyBoundaryPositions(UpdatedPositions);
+
+	return bSuccess;
+}
+
diff --git a/Engine/Plugins/Experimental/GeometryProcessing/Source/DynamicMesh/Private/Solvers/Internal/ConstrainedMeshSmoothers.h b/Engine/Plugins/Experimental/GeometryProcessing/Source/DynamicMesh/Private/Solvers/Internal/ConstrainedMeshSmoothers.h
new file mode 100644
--- /dev/null
+++ b/Engine/Plugins/Experimental/GeometryProcessing/Source/DynamicMesh/Private/Solvers/Internal/ConstrainedMeshSmoothers.h
@@ -0,0 +1,65 @@
+// Copyright Epic Games, Inc. All Rights Reserved.
+
+#pragma once
+
+#include "CoreMinimal.h"
+#include "DynamicMesh3.h"
+
+#include "ConstrainedMeshDeformationSolver.h"
+
+
+
+
+/**
+ * FConstrainedMeshDeformer solves implicit mesh smoothing problems with arbitrary position constraints.
+ * A direct solver is used, currently LU decomposition.
+ *
+ * Boundary vertices are fixed to their input positions.
+ */
+class FBiHarmonicMeshSmoother : public FConstrainedMeshDeformationSolver
+{
+public:
+	typedef FConstrainedMeshDeformationSolver         MyBaseType;
+
+	FBiHarmonicMeshSmoother(const FDynamicMesh3& DynamicMesh, const ELaplacianWeightScheme Scheme) :
+		MyBaseType(DynamicMesh, Scheme, EMatrixSolverType::LU)
+	{}
+
+	bool Deform(TArray<FVector3d>& UpdatedPositions) override
+	{
+		return	ComputeSmoothedMeshPositions(UpdatedPositions);
+	}
+
+	// (Direct) Solve the constrained system and populate the UpdatedPositions with the result 
+	bool ComputeSmoothedMeshPositions(TArray<FVector3d>& UpdatedPositions);
+
+};
+
+
+
+
+
+/**
+ * FConstrainedMeshDeformer solves implicit mesh smoothing problems with arbitrary position constraints.
+ * An iterative Conjugate Gradient solver is used to compute the solution.
+ *
+ * Boundary vertices are fixed to their input positions.
+ */
+class FCGBiHarmonicMeshSmoother : public FConstrainedMeshDeformationSolver
+{
+public:
+	typedef FConstrainedMeshDeformationSolver      MyBaseType;
+
+	FCGBiHarmonicMeshSmoother(const FDynamicMesh3& DynamicMesh, const ELaplacianWeightScheme Scheme) :
+		MyBaseType(DynamicMesh, Scheme, EMatrixSolverType::BICGSTAB)
+	{}
+
+	bool Deform(TArray<FVector3d>& UpdatedPositions) override
+	{
+		return	ComputeSmoothedMeshPositions(UpdatedPositions);
+	}
+
+	// (Iterative) Solve the constrained system and populate the UpdatedPositions with the result 
+	bool ComputeSmoothedMeshPositions(TArray<FVector3d>& UpdatedPositions);
+
+};
diff --git a/Engine/Plugins/Experimental/GeometryProcessing/Source/DynamicMesh/Private/Solvers/Internal/ConstrainedPoissonSolver.h b/Engine/Plugins/Experimental/GeometryProcessing/Source/DynamicMesh/Private/Solvers/Internal/ConstrainedPoissonSolver.h
--- a/Engine/Plugins/Experimental/GeometryProcessing/Source/DynamicMesh/Private/Solvers/Internal/ConstrainedPoissonSolver.h
+++ b/Engine/Plugins/Experimental/GeometryProcessing/Source/DynamicMesh/Private/Solvers/Internal/ConstrainedPoissonSolver.h
@@ -4,6 +4,7 @@
 
 #include "CoreMinimal.h"
 #include "DynamicMesh3.h"
+
 #include "FSparseMatrixD.h"
 #include "MatrixSolver.h"
 
@@ -50,14 +51,6 @@ public:
 	//typedef FScopedDurationTimeLogger FConstrainedSolverTimeLogger; 
 	typedef FDummyTimeLogger FConstrainedSolverTimeLogger;
 
-	struct FConstraintPosition
-	{
-		FConstraintPosition(const FVector3d& P, bool b) : Position(P), bPostFix(b) {}
-
-		FVector3d Position;
-		bool     bPostFix;
-	};
-
 
 	FConstrainedSolver(TUniquePtr<FSparseMatrixD>& SymmatrixMatrixOperator, const EMatrixSolverType MatrixSolverType)
 		: ConstraintPositions(SymmatrixMatrixOperator->cols())
@@ -104,7 +97,7 @@ public:
 	}
 
 	// Updates the positional source term
-	void SetContraintPositions(const TMap<int32, FConstraintPosition>& PositionMap)
+	void SetContraintPositions(const TMap<int32, UE::Solvers::FPositionConstraint>& PositionMap)
 	{
 		ClearConstraintPositions();
 
diff --git a/Engine/Plugins/Experimental/GeometryProcessing/Source/DynamicMesh/Private/Solvers/Internal/LaplacianMeshSmoother.cpp b/Engine/Plugins/Experimental/GeometryProcessing/Source/DynamicMesh/Private/Solvers/Internal/LaplacianMeshSmoother.cpp
deleted file mode 100644
diff --git a/Engine/Plugins/Experimental/GeometryProcessing/Source/DynamicMesh/Private/Solvers/Internal/LaplacianMeshSmoother.h b/Engine/Plugins/Experimental/GeometryProcessing/Source/DynamicMesh/Private/Solvers/Internal/LaplacianMeshSmoother.h
deleted file mode 100644
diff --git a/Engine/Plugins/Experimental/GeometryProcessing/Source/DynamicMesh/Private/Solvers/Internal/LaplacianOperators.cpp b/Engine/Plugins/Experimental/GeometryProcessing/Source/DynamicMesh/Private/Solvers/Internal/LaplacianOperators.cpp
--- a/Engine/Plugins/Experimental/GeometryProcessing/Source/DynamicMesh/Private/Solvers/Internal/LaplacianOperators.cpp
+++ b/Engine/Plugins/Experimental/GeometryProcessing/Source/DynamicMesh/Private/Solvers/Internal/LaplacianOperators.cpp
@@ -27,7 +27,7 @@ THIRD_PARTY_INCLUDES_END
 //
 // Extension of TSparseMatrixAssembler suitable for eigen sparse matrix
 //
-class FEigenSparseMatrixAssembler : public TSparseMatrixAssembler<double>
+class FEigenSparseMatrixAssembler : public UE::Solvers::TSparseMatrixAssembler<double>
 {
 public:
 	typedef FSparseMatrixD::Scalar    ScalarT;
@@ -175,22 +175,28 @@ double ConstructScaledCotangentLaplacian(const FDynamicMesh3& DynamicMesh, FVert
 	FSparseMatrixD CotangentBoundary;
 	ConstructCotangentLaplacian(DynamicMesh, VertexMap, AreaMatrix, CotangentInterior, CotangentBoundary);
 
+	// note: returned AreaMatrix here is full mesh size, ie Interior + Boundary, with Boundary vertices at the end
+	// Does the code below make sense in light of that? We are constructing ScaledInvAreaMatrix and multiplying
+	// CotangentInterior and CotangentBoundary separately, but they are each different sizes??
+
 	// Find average entry in the area matrix
+	const int32 InternalSize = CotangentInterior.cols();
 	const int32 Rank = AreaMatrix.cols();
 	double AveArea = 0.;
-	for (int32 i = 0; i < Rank; ++i)
+	for (int32 i = 0; i < InternalSize; ++i)		// only non-boundary vertices have area > 0!
 	{
 		double Area = AreaMatrix.coeff(i, i);
 		checkSlow(Area > 0.);  // Area must be positive.
 		AveArea += Area;
 	}
-	AveArea /= Rank;
+	AveArea /= (double)InternalSize;
 	
 	std::vector<MatrixTripletT> ScaledInvAreaTriplets;
 	ScaledInvAreaTriplets.reserve(Rank);
 	for (int32 i = 0; i < Rank; ++i)
 	{
-		double Area = AreaMatrix.coeff(i, i);
+		bool bInternal = i < InternalSize;
+		double Area = (bInternal) ? AreaMatrix.coeff(i, i) : AveArea;		// use average area for boundary vertices?
 		double ScaledInvArea = AveArea / Area;
 		if (bClampAreas)
 		{
@@ -204,6 +210,7 @@ double ConstructScaledCotangentLaplacian(const FDynamicMesh3& DynamicMesh, FVert
 	ScaledInvAreaMatrix.setFromTriplets(ScaledInvAreaTriplets.begin(), ScaledInvAreaTriplets.end());
 	ScaledInvAreaMatrix.makeCompressed();
 
+	// ScaledInvAreaMatrix is a different max-dimension than CotangentBoundary and CotangentInterior ??
 	LaplacianBoundary = ScaledInvAreaMatrix * CotangentBoundary;
 	LaplacianBoundary.makeCompressed();
 	LaplacianInterior = ScaledInvAreaMatrix * CotangentInterior;
diff --git a/Engine/Plugins/Experimental/GeometryProcessing/Source/DynamicMesh/Private/Solvers/Internal/MatrixBase.h b/Engine/Plugins/Experimental/GeometryProcessing/Source/DynamicMesh/Private/Solvers/Internal/MatrixBase.h
new file mode 100644
--- /dev/null
+++ b/Engine/Plugins/Experimental/GeometryProcessing/Source/DynamicMesh/Private/Solvers/Internal/MatrixBase.h
@@ -0,0 +1,94 @@
+// Copyright Epic Games, Inc. All Rights Reserved.
+
+#pragma once 
+
+#include "CoreMinimal.h"
+
+#include "FSparseMatrixD.h"
+
+
+
+enum class EMatrixSolverType
+{
+	LU /* SparseLU*/,
+	QR /* SparseQR */,
+	BICGSTAB /* Iterative Bi-conjugate gradient */,
+	PCG /* Pre-conditioned conjugate gradient  - requires symmetric positive def.*/,
+#ifndef EIGEN_MPL2_ONLY
+	LDLT /**Not included due to MPL2 */
+#endif
+};
+
+static FString MatrixSolverName(const EMatrixSolverType SolverType)
+{
+	FString String;
+	switch (SolverType)
+	{
+	case EMatrixSolverType::LU:
+		String = FString(TEXT(" Direct LU "));
+		break;
+	case EMatrixSolverType::QR:
+		String = FString(TEXT(" Direct QR "));
+		break;
+	case EMatrixSolverType::BICGSTAB:
+		String = FString(TEXT(" Iterative BiConjugate Gradient "));
+		break;
+	case EMatrixSolverType::PCG:
+		String = FString(TEXT(" Iterative Preconditioned Conjugate Gradient "));
+		break;
+#ifndef EIGEN_MPL2_ONLY
+	case EMatrixSolverType::PCG:
+		String = FString(TEXT(" Direct Cholesky "));
+		break;
+#endif
+	default:
+		check(0);
+	}
+
+	return String;
+}
+
+class FMatrixSolverSettings
+{
+public:
+	EMatrixSolverType  MatrixSolverType;
+
+	// Used by iterative solvers
+	int32  MaxIterations = 600;
+	double Tolerance = 1e-4;
+};
+
+
+
+// Pure Virtual Base Class used to wrap Eigen Solvers.
+class IMatrixSolverBase
+{
+public:
+	typedef typename FSparseMatrixD::Scalar    ScalarType;
+	typedef typename FSOAPositions::VectorType VectorType;
+
+	IMatrixSolverBase() {}
+	virtual ~IMatrixSolverBase() = 0;
+	virtual bool bIsIterative() const = 0;
+	virtual void Solve(const VectorType& BVector, VectorType& SolVector) const = 0;
+	virtual void Solve(const FSOAPositions& BVectors, FSOAPositions& SolVectors) const = 0;
+
+	virtual void SetUp(const FSparseMatrixD& Matrix, bool bIsSymmetric) = 0;
+	//virtual void SetIterations(int32 MaxIterations) = 0;
+	virtual void Reset() = 0;
+	virtual bool bSucceeded() const = 0;
+};
+
+// Additional methods particular to the iterative solvers
+class IIterativeMatrixSolverBase : public IMatrixSolverBase
+{
+public:
+	IIterativeMatrixSolverBase() {}
+
+	virtual ~IIterativeMatrixSolverBase() = 0;
+	virtual void SetIterations(int32 MaxIterations) = 0;
+	virtual void SetTolerance(double Tol) = 0;
+	virtual void SolveWithGuess(const VectorType& GuessVector, const VectorType& BVector, VectorType& SolVector)  const = 0;
+	virtual void SolveWithGuess(const FSOAPositions& GuessVector, const FSOAPositions& BVector, FSOAPositions& SolVector)  const = 0;
+
+};
\ No newline at end of file
diff --git a/Engine/Plugins/Experimental/GeometryProcessing/Source/DynamicMesh/Private/Solvers/Internal/MatrixSolver.h b/Engine/Plugins/Experimental/GeometryProcessing/Source/DynamicMesh/Private/Solvers/Internal/MatrixSolver.h
--- a/Engine/Plugins/Experimental/GeometryProcessing/Source/DynamicMesh/Private/Solvers/Internal/MatrixSolver.h
+++ b/Engine/Plugins/Experimental/GeometryProcessing/Source/DynamicMesh/Private/Solvers/Internal/MatrixSolver.h
@@ -3,9 +3,10 @@
 #pragma once 
 
 #include "CoreMinimal.h"
-#include "FSparseMatrixD.h"
 #include "Async/ParallelFor.h"
 
+#include "MatrixBase.h"
+
 #if defined(_MSC_VER) && USING_CODE_ANALYSIS
 #pragma warning(push)
 #pragma warning(disable : 6011)
@@ -32,90 +33,6 @@ THIRD_PARTY_INCLUDES_END
 #include "ProfilingDebugging/ScopedTimers.h"
 
 
-enum class EMatrixSolverType
-{
-	LU /* SparseLU*/,
-	QR /* SparseQR */,
-	BICGSTAB /* Iterative Bi-conjugate gradient */,
-	PCG /* Pre-conditioned conjugate gradient  - requires symmetric positive def.*/,
-#ifndef EIGEN_MPL2_ONLY
-	LDLT /**Not included due to MPL2 */
-#endif
-};
-
-static FString MatrixSolverName(const EMatrixSolverType SolverType)
-{
-	FString String;
-	switch (SolverType)
-	{
-	case EMatrixSolverType::LU:
-		String = FString(TEXT(" Direct LU "));
-		break;
-	case EMatrixSolverType::QR:
-		String = FString(TEXT(" Direct QR "));
-		break;
-	case EMatrixSolverType::BICGSTAB:
-		String = FString(TEXT(" Iterative BiConjugate Gradient "));
-		break;
-	case EMatrixSolverType::PCG:
-		String = FString(TEXT(" Iterative Preconditioned Conjugate Gradient "));
-		break;
-#ifndef EIGEN_MPL2_ONLY
-	case EMatrixSolverType::PCG:
-		String = FString(TEXT(" Direct Cholesky "));
-		break;
-#endif
-	default:
-		check(0);
-	}
-
-	return String;
-}
-
-class FMatrixSolverSettings
-{
-public:
-	EMatrixSolverType  MatrixSolverType;
-	
-	// Used by iterative solvers
-	int32  MaxIterations = 600;
-	double Tolerance = 1e-4;
-};
-
-
-
-// Pure Virtual Base Class used to wrap Eigen Solvers.
-class IMatrixSolverBase
-{
-public:
-	typedef typename FSparseMatrixD::Scalar    ScalarType;
-	typedef typename FSOAPositions::VectorType VectorType;
-
-	IMatrixSolverBase() {}
-	virtual ~IMatrixSolverBase() = 0;
-	virtual bool bIsIterative() const = 0;
-	virtual void Solve(const VectorType& BVector, VectorType& SolVector) const = 0;
-	virtual void Solve(const FSOAPositions& BVectors, FSOAPositions& SolVectors) const = 0;
-
-	virtual void SetUp(const FSparseMatrixD& Matrix, bool bIsSymmetric) = 0;
-	//virtual void SetIterations(int32 MaxIterations) = 0;
-	virtual void Reset() = 0;
-	virtual bool bSucceeded() const = 0;
-};
-
-// Additional methods particular to the iterative solvers
-class IIterativeMatrixSolverBase : public IMatrixSolverBase
-{
-public:
-	IIterativeMatrixSolverBase() {}
-
-	virtual ~IIterativeMatrixSolverBase() = 0;
-	virtual void SetIterations(int32 MaxIterations) = 0;
-	virtual void SetTolerance(double Tol) = 0;
-	virtual void SolveWithGuess(const VectorType& GuessVector, const VectorType& BVector, VectorType& SolVector)  const = 0;
-	virtual void SolveWithGuess(const FSOAPositions& GuessVector, const FSOAPositions& BVector, FSOAPositions& SolVector)  const = 0;
-
-};
 
 // Matrix Solver Factory
 TUniquePtr<IMatrixSolverBase> ContructMatrixSolver(const EMatrixSolverType& MatrixSolverType);
diff --git a/Engine/Plugins/Experimental/GeometryProcessing/Source/DynamicMesh/Private/Solvers/Internal/MeshDiffusionIntegrator.cpp b/Engine/Plugins/Experimental/GeometryProcessing/Source/DynamicMesh/Private/Solvers/Internal/MeshDiffusionIntegrator.cpp
new file mode 100644
--- /dev/null
+++ b/Engine/Plugins/Experimental/GeometryProcessing/Source/DynamicMesh/Private/Solvers/Internal/MeshDiffusionIntegrator.cpp
@@ -0,0 +1,251 @@
+// Copyright Epic Games, Inc. All Rights Reserved.
+
+#include "MeshDiffusionIntegrator.h"
+#include "MatrixSolver.h"
+
+
+FMeshDiffusionIntegrator::FMeshDiffusionIntegrator(const FDynamicMesh3& DynamicMesh, const ELaplacianWeightScheme Scheme)
+{
+	Id = 0;
+	bIsSymmetric = false;
+	MinDiagonalValue = 0;
+
+	VertexCount = DynamicMesh.VertexCount();
+
+}
+
+void FMeshDiffusionIntegrator::Initialize(const FDynamicMesh3& DynamicMesh, const ELaplacianWeightScheme Scheme)
+{
+	// Construct the laplacian, and extract the mapping for vertices (VtxLinearization)
+	//DiffusionOperator = ConstructDiffusionOperator(Scheme, DynamicMesh, bIsSymmetric, VtxLinearization, &EdgeVerts);
+
+
+	ConstructOperators(Scheme, DynamicMesh, bIsSymmetric, VtxLinearization, DiffusionOperator, BoundaryOperator);
+
+	const int32 BoundaryVertexCount = VtxLinearization.NumBoundaryVerts();
+	InternalVertexCount = VertexCount - BoundaryVertexCount;
+
+	// Allocate the double buffers.
+	Tmp[0].SetZero(InternalVertexCount);
+	Tmp[1].SetZero(InternalVertexCount);
+
+	const auto& ToVertId = VtxLinearization.ToId();
+
+	// Extract current internal positions.
+	for (int32 i = 0; i < InternalVertexCount; ++i)
+	{
+		const int32 VtxId = ToVertId[i];
+		const FVector3d Pos = DynamicMesh.GetVertex(VtxId);
+		Tmp[0].SetXYZ(i, Pos);
+	}
+
+	// backup the locations of the boundary verts.
+	{
+		BoundaryPositions.SetZero(BoundaryVertexCount);
+		for (int32 i = 0; i < BoundaryVertexCount; ++i)
+		{
+			const int32 VtxId = ToVertId[i + InternalVertexCount];
+			const FVector3d Pos = DynamicMesh.GetVertex(VtxId);
+			BoundaryPositions.SetXYZ(i, Pos);
+		}
+	}
+
+	const FSparseMatrixD& M = DiffusionOperator;
+
+	// Find the min diagonal entry (all should be negative).
+	int32 Rank = M.rows();
+	MinDiagonalValue = FSparseMatrixD::Scalar(0);
+	for (int32 i = 0; i < Rank; ++i)
+	{
+		auto Diag = M.coeff(i, i);
+		MinDiagonalValue = FMath::Min(Diag, MinDiagonalValue);
+	}
+	// The matrix should have a row for each internal vertex
+	checkSlow(Rank == InternalVertexCount);
+
+#if 0
+	// testing - how to print the matrix to debug output 
+
+	std::stringstream  ss;
+	ss << Eigen::MatrixXd(M) << std::endl;
+
+	FString Foo = ss.str().c_str();
+	FPlatformMisc::LowLevelOutputDebugStringf(*Foo);
+#endif
+
+}
+
+
+
+void FMeshDiffusionIntegrator::Integrate_ForwardEuler(const int32 NumSteps, const double Speed)
+{
+	double Alpha = FMath::Clamp(Speed, 0., 1.);
+
+	FSparseMatrixD::Scalar TimeStep = -Alpha / MinDiagonalValue;
+	Id = 0;
+	for (int32 s = 0; s < NumSteps; ++s)
+	{
+
+		int32 SrcBuffer = Id;
+		Id = 1 - Id;
+		Tmp[Id].Array(0) = Tmp[SrcBuffer].Array(0) + TimeStep * (DiffusionOperator * Tmp[SrcBuffer].Array(0) + BoundaryOperator * BoundaryPositions.Array(0));
+		Tmp[Id].Array(1) = Tmp[SrcBuffer].Array(1) + TimeStep * (DiffusionOperator * Tmp[SrcBuffer].Array(1) + BoundaryOperator * BoundaryPositions.Array(1));
+		Tmp[Id].Array(2) = Tmp[SrcBuffer].Array(2) + TimeStep * (DiffusionOperator * Tmp[SrcBuffer].Array(2) + BoundaryOperator * BoundaryPositions.Array(2));
+
+	}
+
+}
+
+
+
+void FMeshDiffusionIntegrator::Integrate_BackwardEuler(const EMatrixSolverType MatrixSolverType, const int32 NumSteps, const double TimeStepSize)
+{
+
+	//typedef typename TMatrixSolverTrait<EMatrixSolverType::LU>::MatrixSolverType   MatrixSolverType;
+
+	// We solve 
+	// p^{n+1} - dt * L[p^{n+1}] = p^{n} + dt * B[boundaryPts]
+	// 
+	// i.e.
+	// [I - dt * L ] p^{n+1} = p^{n} + dt * B[boundaryPts]
+	//
+	// NB: in the case of the cotangent laplacian this would be better if we broke the L int
+	// L = (A^{-1}) H  where A is the "area matrix" (think "mass matrix"), then this would
+	// become
+	// [A - dt * H] p^{n+1} = Ap^{n}  dt * A *B[boundaryPts]
+	//  
+	// A - dt * H would be symmetric
+	//
+
+
+	// Identity matrix
+	FSparseMatrixD Ident(DiffusionOperator.rows(), DiffusionOperator.cols());
+	Ident.setIdentity();
+
+	FSparseMatrixD::Scalar TimeStep = FMath::Abs(TimeStepSize);// Alpha * FMath::Min(Intensity, 1.e6);
+
+	FSparseMatrixD SparseMatrix = Ident - TimeStep * DiffusionOperator;
+
+
+	SparseMatrix.makeCompressed();
+
+	TUniquePtr<IMatrixSolverBase> MatrixSolver = ContructMatrixSolver(MatrixSolverType);
+
+	MatrixSolver->SetUp(SparseMatrix, bIsSymmetric);
+
+	// We are going to solve the system 
+	FSOAPositions Source(InternalVertexCount);
+
+	if (MatrixSolver->bIsIterative())
+	{
+		IIterativeMatrixSolverBase* IterativeSolver = (IIterativeMatrixSolverBase*)MatrixSolver.Get();
+
+		bool bForceSingleThreaded = false;
+		Id = 0;
+		for (int32 s = 0; s < NumSteps; ++s)
+		{
+
+			int32 SrcBuffer = Id;
+			Id = 1 - Id;
+
+			for (int32 i = 0; i < 3; ++i)
+			{
+				Source.Array(i) = Tmp[SrcBuffer].Array(i) + TimeStep * BoundaryOperator * BoundaryPositions.Array(i);
+			}
+			// Old solution is the guess.
+			IterativeSolver->SolveWithGuess(Tmp[SrcBuffer], Source, Tmp[Id]);
+
+		}
+	}
+	else
+	{
+		bool bForceSingleThreaded = false;
+		Id = 0;
+		for (int32 s = 0; s < NumSteps; ++s)
+		{
+
+			int32 SrcBuffer = Id;
+			Id = 1 - Id;
+
+			for (int32 i = 0; i < 3; ++i)
+			{
+				Source.Array(i) = Tmp[SrcBuffer].Array(i) + TimeStep * BoundaryOperator * BoundaryPositions.Array(i);
+			}
+
+			MatrixSolver->Solve(Source, Tmp[Id]);
+		}
+	}
+
+
+}
+
+void FMeshDiffusionIntegrator::GetPositions(TArray<FVector3d>& PositionBuffer) const
+{
+	// Allocate Position Buffer for random access writes
+	int32 MaxVtxId = VtxLinearization.ToId().Num();
+	PositionBuffer.Empty(MaxVtxId);
+	PositionBuffer.AddUninitialized(MaxVtxId);
+
+	CopyInternalPositions(Tmp[Id], PositionBuffer);
+
+
+	// Copy the boundary
+	// NB: this re-indexes so the results can be looked up using VtxId
+	CopyBoundaryPositions(PositionBuffer);
+
+}
+
+bool FMeshDiffusionIntegrator::CopyInternalPositions(const FSOAPositions& PositionalVector, TArray<FVector3d>& LinearArray) const
+{
+	// Number of positions
+
+	const int32 Num = PositionalVector.Num();
+
+	// early out if the x,y,z arrays in the PositionalVector have different lengths
+	if (!PositionalVector.bHasSize(Num))
+	{
+		return false;
+	}
+
+	// 
+	const auto& ToVtxId = VtxLinearization.ToId();
+	const int32 MaxVtxId = ToVtxId.Num(); // NB: this is really max_used + 1 in the mesh.  See  FDynamicMesh3::MaxVertexID()
+	const int32 BoundaryVertexCount = VtxLinearization.NumBoundaryVerts();
+
+	if (MaxVtxId != LinearArray.Num())
+	{
+		return false;
+	}
+
+	// Copy the updated internal vertex locations over
+	for (int32 i = 0; i < InternalVertexCount; ++i)
+	{
+		const int32 VtxId = ToVtxId[i];
+
+		LinearArray[VtxId] = FVector3d(PositionalVector.X(i), PositionalVector.Y(i), PositionalVector.Z(i));
+	}
+
+	return true;
+}
+
+bool FMeshDiffusionIntegrator::CopyBoundaryPositions(TArray<FVector3d>& LinearArray) const
+{
+	const auto& ToVtxId = VtxLinearization.ToId();
+	const int32 MaxVtxId = ToVtxId.Num();
+
+	if (LinearArray.Num() != MaxVtxId)
+	{
+		return false;
+	}
+
+	int32 BoundaryVertexCount = VertexCount - InternalVertexCount;
+
+	for (int32 i = 0; i < BoundaryVertexCount; ++i)
+	{
+		const int32 VtxId = ToVtxId[i + InternalVertexCount];
+		LinearArray[VtxId] = FVector3d(BoundaryPositions.X(i), BoundaryPositions.Y(i), BoundaryPositions.Z(i));
+	}
+
+	return true;
+}
+
diff --git a/Engine/Plugins/Experimental/GeometryProcessing/Source/DynamicMesh/Private/Solvers/Internal/MeshDiffusionIntegrator.h b/Engine/Plugins/Experimental/GeometryProcessing/Source/DynamicMesh/Private/Solvers/Internal/MeshDiffusionIntegrator.h
new file mode 100644
--- /dev/null
+++ b/Engine/Plugins/Experimental/GeometryProcessing/Source/DynamicMesh/Private/Solvers/Internal/MeshDiffusionIntegrator.h
@@ -0,0 +1,73 @@
+// Copyright Epic Games, Inc. All Rights Reserved.
+
+#pragma once
+
+#include "CoreMinimal.h"
+#include "DynamicMesh3.h"
+#include "Solvers/MeshLaplacian.h"
+#include "Solvers/MeshLinearization.h"
+
+#include "MatrixSolver.h"
+
+/**
+ * FMeshDiffusionIntegrator solves diffusion problems over vertices of a 3D triangle mesh.
+ *
+ * (Could likely be generalized to other graph diffusion problems)
+ */
+class FMeshDiffusionIntegrator
+{
+public:
+
+	FMeshDiffusionIntegrator(const FDynamicMesh3& DynamicMesh, const ELaplacianWeightScheme Scheme);
+	virtual ~FMeshDiffusionIntegrator() {}
+
+	void Integrate_ForwardEuler(const int32 NumSteps, const double Speed);
+
+	// Note: 
+	void Integrate_BackwardEuler(const EMatrixSolverType MatrixSolverType, const int32 NumSteps, const double TimeStepSize);
+
+	void GetPositions(TArray<FVector3d>& PositionArray) const;
+
+protected:
+
+
+	// The derived class has to implement this.
+	// responsible for constructing the Diffusion Operator, the Boundary Operator etc
+	virtual void ConstructOperators(const ELaplacianWeightScheme Scheme,
+		const FDynamicMesh3& Mesh,
+		bool& bIsOperatorSymmetric,
+		FVertexLinearization& Linearization,
+		FSparseMatrixD& DiffusionOp,
+		FSparseMatrixD& BoundaryOp) = 0;
+
+	void Initialize(const FDynamicMesh3& DynamicMesh, const ELaplacianWeightScheme Scheme);
+
+protected:
+
+	// Copy the 
+	bool CopyInternalPositions(const FSOAPositions& PositionalVector, TArray<FVector3d>& LinearArray) const;
+
+	bool CopyBoundaryPositions(TArray<FVector3d>& LinearArray) const;
+
+	// Cache the vertex count.
+	int32                 VertexCount;
+
+	// Cache the number of internal vertices
+	int32                 InternalVertexCount;
+
+	// Used to map between VtxId and vertex Index in linear vector..
+	FVertexLinearization  VtxLinearization;
+
+	// I don't know if we want to keep this after the constructor
+	bool                                   bIsSymmetric;
+	FSparseMatrixD                         DiffusionOperator;
+	FSparseMatrixD			               BoundaryOperator;
+	TArray<int32>                          EdgeVerts;
+	FSparseMatrixD::Scalar                 MinDiagonalValue;
+
+	FSOAPositions                       BoundaryPositions;
+
+	FSOAPositions                       Tmp[2];
+	int32                               Id; // double buffer id
+};
+
diff --git a/Engine/Plugins/Experimental/GeometryProcessing/Source/DynamicMesh/Private/Solvers/Internal/MeshDiffusionSmoothing.h b/Engine/Plugins/Experimental/GeometryProcessing/Source/DynamicMesh/Private/Solvers/Internal/MeshDiffusionSmoothing.h
new file mode 100644
--- /dev/null
+++ b/Engine/Plugins/Experimental/GeometryProcessing/Source/DynamicMesh/Private/Solvers/Internal/MeshDiffusionSmoothing.h
@@ -0,0 +1,90 @@
+// Copyright Epic Games, Inc. All Rights Reserved.
+
+#pragma once
+
+#include "MeshDiffusionIntegrator.h"
+#include "LaplacianOperators.h"
+
+
+/**
+ * Solve mesh diffusion using Laplacian matrix of various types
+ */
+class FLaplacianDiffusionMeshSmoother : public FMeshDiffusionIntegrator
+{
+public:
+
+	FLaplacianDiffusionMeshSmoother(const FDynamicMesh3& DynamicMesh, const ELaplacianWeightScheme Scheme)
+		: FMeshDiffusionIntegrator(DynamicMesh, Scheme)
+	{
+		Initialize(DynamicMesh, Scheme);
+	}
+
+protected:
+
+	void ConstructOperators(const ELaplacianWeightScheme Scheme,
+		const FDynamicMesh3& Mesh,
+		bool& bIsOperatorSymmetric,
+		FVertexLinearization& Linearization,
+		FSparseMatrixD& DiffusionOp,
+		FSparseMatrixD& BoundaryOp) override
+	{
+		bIsOperatorSymmetric = bIsSymmetricLaplacian(Scheme);
+		ConstructLaplacian(Scheme, Mesh, VtxLinearization, DiffusionOp, BoundaryOp);
+	}
+};
+
+
+
+/**
+ * Solve mesh diffusion using Biharmonic (ie squared Laplacian) matrix of various types
+ */
+class  FBiHarmonicDiffusionMeshSmoother : public FMeshDiffusionIntegrator
+{
+public:
+
+	FBiHarmonicDiffusionMeshSmoother(const FDynamicMesh3& DynamicMesh, const ELaplacianWeightScheme Scheme)
+		: FMeshDiffusionIntegrator(DynamicMesh, Scheme)
+	{
+		Initialize(DynamicMesh, Scheme);
+	}
+
+protected:
+
+	void ConstructOperators(const ELaplacianWeightScheme Scheme,
+		const FDynamicMesh3& Mesh,
+		bool& bIsOperatorSymmetric,
+		FVertexLinearization& Linearization,
+		FSparseMatrixD& DiffusionOp,
+		FSparseMatrixD& BoundaryOp) override
+	{
+		bIsOperatorSymmetric = true;
+
+		FSparseMatrixD Laplacian;
+		FSparseMatrixD BoundaryTerms;
+		ConstructLaplacian(Scheme, Mesh, VtxLinearization, Laplacian, BoundaryTerms);
+
+		bool bIsLaplacianSymmetric = bIsSymmetricLaplacian(Scheme);
+
+		// It is actually unclear the best way to approximate the boundary conditions in this case.  
+		// because we are repeatedly applying the operator ( for example thing about the way ( f(x+d)-f(x-d) )/ d will spread if you apply it twice
+		// as opposed to (f(x+d)-2f(x) + f(x-d) ) / d*d
+
+		// Anyhow here is a guess..
+
+		if (bIsLaplacianSymmetric)
+		{
+			DiffusionOp = -1. * Laplacian * Laplacian;
+			BoundaryOp = -1. * Laplacian * BoundaryTerms;
+		}
+		else
+		{
+			FSparseMatrixD LTran = Laplacian.transpose();
+			DiffusionOp = -1. * LTran * Laplacian;
+			BoundaryOp = -1. * LTran * BoundaryTerms;
+		}
+
+		DiffusionOp.makeCompressed();
+		BoundaryOp.makeCompressed();
+	}
+};
+
diff --git a/Engine/Plugins/Experimental/GeometryProcessing/Source/DynamicMesh/Private/Solvers/MeshSmoothing.cpp b/Engine/Plugins/Experimental/GeometryProcessing/Source/DynamicMesh/Private/Solvers/MeshSmoothing.cpp
--- a/Engine/Plugins/Experimental/GeometryProcessing/Source/DynamicMesh/Private/Solvers/MeshSmoothing.cpp
+++ b/Engine/Plugins/Experimental/GeometryProcessing/Source/DynamicMesh/Private/Solvers/MeshSmoothing.cpp
@@ -1,7 +1,8 @@
 // Copyright Epic Games, Inc. All Rights Reserved.
 
 #include "Solvers/MeshSmoothing.h"
-#include "Solvers/Internal/LaplacianMeshSmoother.h"
+#include "Solvers/Internal/MeshDiffusionSmoothing.h"
+#include "Solvers/Internal/ConstrainedMeshSmoothers.h"
 
 
 void UE::MeshDeformation::ComputeSmoothing_BiHarmonic(const ELaplacianWeightScheme WeightScheme, const FDynamicMesh3& OriginalMesh,
diff --git a/Engine/Plugins/Experimental/GeometryProcessing/Source/DynamicMesh/Public/Solvers/ConstrainedMeshDeformer.h b/Engine/Plugins/Experimental/GeometryProcessing/Source/DynamicMesh/Public/Solvers/ConstrainedMeshDeformer.h
--- a/Engine/Plugins/Experimental/GeometryProcessing/Source/DynamicMesh/Public/Solvers/ConstrainedMeshDeformer.h
+++ b/Engine/Plugins/Experimental/GeometryProcessing/Source/DynamicMesh/Public/Solvers/ConstrainedMeshDeformer.h
@@ -54,7 +54,7 @@ namespace UE
 		*   MeshDeformer->Deform(PositionBuffer);
 		*   // Update Mesh?
 		*/
-		TUniquePtr<IConstrainedMeshSolver> DYNAMICMESH_API ConstructConstrainedMeshDeformer(const ELaplacianWeightScheme WeightScheme, const FDynamicMesh3& DynamicMesh);
+		TUniquePtr<UE::Solvers::IConstrainedMeshSolver> DYNAMICMESH_API ConstructConstrainedMeshDeformer(const ELaplacianWeightScheme WeightScheme, const FDynamicMesh3& DynamicMesh);
 
 	}
 }
\ No newline at end of file
diff --git a/Engine/Plugins/Experimental/GeometryProcessing/Source/DynamicMesh/Public/Solvers/ConstrainedMeshSmoother.h b/Engine/Plugins/Experimental/GeometryProcessing/Source/DynamicMesh/Public/Solvers/ConstrainedMeshSmoother.h
--- a/Engine/Plugins/Experimental/GeometryProcessing/Source/DynamicMesh/Public/Solvers/ConstrainedMeshSmoother.h
+++ b/Engine/Plugins/Experimental/GeometryProcessing/Source/DynamicMesh/Public/Solvers/ConstrainedMeshSmoother.h
@@ -25,7 +25,7 @@ namespace UE
 		* Expected Use: same as the ConstrainedMeshDeformer above.
 		*
 		*/
-		TUniquePtr<IConstrainedMeshSolver> DYNAMICMESH_API ConstructConstrainedMeshSmoother(const ELaplacianWeightScheme WeightScheme, const FDynamicMesh3& DynamicMesh);
+		TUniquePtr<UE::Solvers::IConstrainedMeshSolver> DYNAMICMESH_API ConstructConstrainedMeshSmoother(const ELaplacianWeightScheme WeightScheme, const FDynamicMesh3& DynamicMesh);
 	}
 }
 
diff --git a/Engine/Plugins/Experimental/GeometryProcessing/Source/DynamicMesh/Public/Solvers/ConstrainedMeshSolver.h b/Engine/Plugins/Experimental/GeometryProcessing/Source/DynamicMesh/Public/Solvers/ConstrainedMeshSolver.h
--- a/Engine/Plugins/Experimental/GeometryProcessing/Source/DynamicMesh/Public/Solvers/ConstrainedMeshSolver.h
+++ b/Engine/Plugins/Experimental/GeometryProcessing/Source/DynamicMesh/Public/Solvers/ConstrainedMeshSolver.h
@@ -5,41 +5,50 @@
 #include "CoreMinimal.h"
 #include "VectorTypes.h"
 
-/**
- * Interface to a index-based deformation solver for a 3D mesh vertex set that supports weighted point constraints.
- */
-class DYNAMICMESH_API IConstrainedMeshSolver
+namespace UE
 {
-public:
+	namespace Solvers
+	{
 
-	virtual ~IConstrainedMeshSolver() {}
+		/**
+		 * Interface to a index-based deformation solver for a 3D mesh vertex set that supports weighted point constraints.
+		 */
+		class DYNAMICMESH_API IConstrainedMeshSolver
+		{
+		public:
 
-	// Add or update a weighted positional constraint associated with VtxId
-	// @param bPostFix if true, this constraint will be explicitly applied after the solve
-	virtual void AddConstraint(const int32 VtxId, const double Weight, const FVector3d& Position, const bool bPostFix) = 0;
+			virtual ~IConstrainedMeshSolver() {}
 
-	// Update or Create a constraint position associated with @param VtxId
-	// @return true if a constraint weight is associated with @param VtxId
-	virtual bool UpdateConstraintPosition(const int32 VtxId, const FVector3d& Position, const bool bPostFix) = 0;
+			// Add or update a weighted positional constraint associated with VtxId
+			// @param bPostFix if true, this constraint will be explicitly applied after the solve
+			virtual void AddConstraint(const int32 VtxId, const double Weight, const FVector3d& Position, const bool bPostFix) = 0;
 
-	// Update or Create a constraint weight associated with @param VtxId
-	// @return true if a constraint position is associated with @param VtxId
-	virtual bool UpdateConstraintWeight(const int32 VtxId, const double Weight) = 0;
+			// Update or Create a constraint position associated with @param VtxId
+			// @return true if a constraint weight is associated with @param VtxId
+			virtual bool UpdateConstraintPosition(const int32 VtxId, const FVector3d& Position, const bool bPostFix) = 0;
 
-	// Clear all constraints (Positions and Weights)
-	virtual void ClearConstraints() = 0;
+			// Update or Create a constraint weight associated with @param VtxId
+			// @return true if a constraint position is associated with @param VtxId
+			virtual bool UpdateConstraintWeight(const int32 VtxId, const double Weight) = 0;
 
-	// Clear all Constraint Weights
-	virtual void ClearConstraintWeights() = 0;
+			// Clear all constraints (Positions and Weights)
+			virtual void ClearConstraints() = 0;
 
-	// Clear all Constraint Positions
-	virtual void ClearConstraintPositions() = 0;
+			// Clear all Constraint Weights
+			virtual void ClearConstraintWeights() = 0;
 
+			// Clear all Constraint Positions
+			virtual void ClearConstraintPositions() = 0;
 
-	// Test if a non-zero weighted constraint is associated with VtxId 
-	virtual bool IsConstrained(const int32 VtxId) const = 0;
 
-	// Returns the vertex locations of the deformed mesh. 
-	// Note the array may have empty elements as the index matches the Mesh based VtxId.  PositionBuffer[VtxId] = Pos.
-	virtual bool Deform(TArray<FVector3d>& PositionBuffer) = 0;
-};
\ No newline at end of file
+			// Test if a non-zero weighted constraint is associated with VtxId 
+			virtual bool IsConstrained(const int32 VtxId) const = 0;
+
+			// Returns the vertex locations of the deformed mesh. 
+			// Note the array may have empty elements as the index matches the Mesh based VtxId.  PositionBuffer[VtxId] = Pos.
+			virtual bool Deform(TArray<FVector3d>& PositionBuffer) = 0;
+		};
+
+	}
+}
+
diff --git a/Engine/Plugins/Experimental/GeometryProcessing/Source/DynamicMesh/Public/Solvers/LaplacianMatrixAssembly.h b/Engine/Plugins/Experimental/GeometryProcessing/Source/DynamicMesh/Public/Solvers/LaplacianMatrixAssembly.h
--- a/Engine/Plugins/Experimental/GeometryProcessing/Source/DynamicMesh/Public/Solvers/LaplacianMatrixAssembly.h
+++ b/Engine/Plugins/Experimental/GeometryProcessing/Source/DynamicMesh/Public/Solvers/LaplacianMatrixAssembly.h
@@ -38,7 +38,8 @@ namespace UE
 		*/
 		template<typename RealType>
 		void ConstructUniformLaplacian(const FDynamicMesh3& DynamicMesh, const FVertexLinearization& VertexMap,
-			TSparseMatrixAssembler<RealType>& LaplacianInterior, TSparseMatrixAssembler<RealType>& LaplacianBoundary);
+			UE::Solvers::TSparseMatrixAssembler<RealType>& LaplacianInterior, 
+			UE::Solvers::TSparseMatrixAssembler<RealType>& LaplacianBoundary);
 
 		/**
 		* Construct a sparse matrix representation of an umbrella weighted Laplacian.
@@ -64,7 +65,8 @@ namespace UE
 		*/
 		template<typename RealType>
 		void ConstructUmbrellaLaplacian(const FDynamicMesh3& DynamicMesh, const FVertexLinearization& VertexMap,
-			TSparseMatrixAssembler<RealType>& LaplacianInterior, TSparseMatrixAssembler<RealType>& LaplacianBoundary);
+			UE::Solvers::TSparseMatrixAssembler<RealType>& LaplacianInterior, 
+			UE::Solvers::TSparseMatrixAssembler<RealType>& LaplacianBoundary);
 
 		/**
 		* Construct a sparse matrix representation of a valence-weighted Laplacian.
@@ -89,7 +91,8 @@ namespace UE
 		*/
 		template<typename RealType>
 		void ConstructValenceWeightedLaplacian(const FDynamicMesh3& DynamicMesh, const FVertexLinearization& VertexMap,
-			TSparseMatrixAssembler<RealType>& LaplacianInterior, TSparseMatrixAssembler<RealType>& LaplacianBoundary);
+			UE::Solvers::TSparseMatrixAssembler<RealType>& LaplacianInterior, 
+			UE::Solvers::TSparseMatrixAssembler<RealType>& LaplacianBoundary);
 
 		/**
 		* Construct a sparse matrix representation using a meanvalue-weighted  Laplacian.
@@ -107,7 +110,8 @@ namespace UE
 		*/
 		template<typename RealType>
 		void ConstructMeanValueWeightLaplacian(const FDynamicMesh3& DynamicMesh, const FVertexLinearization& VertexMap,
-			TSparseMatrixAssembler<RealType>& LaplacianInterior, TSparseMatrixAssembler<RealType>& LaplacianBoundary);
+			UE::Solvers::TSparseMatrixAssembler<RealType>& LaplacianInterior, 
+			UE::Solvers::TSparseMatrixAssembler<RealType>& LaplacianBoundary);
 
 
 		/**
@@ -130,8 +134,9 @@ namespace UE
 		*/
 		template<typename RealType>
 		void ConstructCotangentLaplacian(const FDynamicMesh3& DynamicMesh, const FVertexLinearization& VertexMap,
-			TSparseMatrixAssembler<RealType>& AreaMatrix,
-			TSparseMatrixAssembler<RealType>& LaplacianInterior, TSparseMatrixAssembler<RealType>& LaplacianBoundary);
+			UE::Solvers::TSparseMatrixAssembler<RealType>& AreaMatrix,
+			UE::Solvers::TSparseMatrixAssembler<RealType>& LaplacianInterior, 
+			UE::Solvers::TSparseMatrixAssembler<RealType>& LaplacianBoundary);
 
 
 		/**
@@ -154,7 +159,8 @@ namespace UE
 		*/
 		template<typename RealType>
 		void ConstructCotangentLaplacian(const FDynamicMesh3& DynamicMesh, const FVertexLinearization& VertexMap,
-			TSparseMatrixAssembler<RealType>& LaplacianInterior, TSparseMatrixAssembler<RealType>& LaplacianBoundary,
+			UE::Solvers::TSparseMatrixAssembler<RealType>& LaplacianInterior, 
+			UE::Solvers::TSparseMatrixAssembler<RealType>& LaplacianBoundary,
 			const bool bClampWeights);
 
 	}
@@ -174,7 +180,8 @@ namespace UE
 
 template<typename RealType>
 void UE::MeshDeformation::ConstructUniformLaplacian(const FDynamicMesh3& DynamicMesh, const FVertexLinearization& VertexMap,
-	TSparseMatrixAssembler<RealType>& LaplacianInterior, TSparseMatrixAssembler<RealType>& LaplacianBoundary)
+	UE::Solvers::TSparseMatrixAssembler<RealType>& LaplacianInterior, 
+	UE::Solvers::TSparseMatrixAssembler<RealType>& LaplacianBoundary)
 {
 	//check(VertexMap_is_good)
 	const TArray<int32>& ToMeshV = VertexMap.ToId();
@@ -223,7 +230,8 @@ void UE::MeshDeformation::ConstructUniformLaplacian(const FDynamicMesh3& Dynamic
 
 template<typename RealType>
 void UE::MeshDeformation::ConstructUmbrellaLaplacian(const FDynamicMesh3& DynamicMesh, const FVertexLinearization& VertexMap,
-	TSparseMatrixAssembler<RealType>& LaplacianInterior, TSparseMatrixAssembler<RealType>& LaplacianBoundary)
+	UE::Solvers::TSparseMatrixAssembler<RealType>& LaplacianInterior, 
+	UE::Solvers::TSparseMatrixAssembler<RealType>& LaplacianBoundary)
 {
 	//check(VertexMap_is_good)
 	const TArray<int32>& ToMeshV = VertexMap.ToId();
@@ -287,7 +295,8 @@ void UE::MeshDeformation::ConstructUmbrellaLaplacian(const FDynamicMesh3& Dynami
 
 template<typename RealType>
 void UE::MeshDeformation::ConstructValenceWeightedLaplacian(const FDynamicMesh3& DynamicMesh, const FVertexLinearization& VertexMap,
-	TSparseMatrixAssembler<RealType>& LaplacianInterior, TSparseMatrixAssembler<RealType>& LaplacianBoundary)
+	UE::Solvers::TSparseMatrixAssembler<RealType>& LaplacianInterior, 
+	UE::Solvers::TSparseMatrixAssembler<RealType>& LaplacianBoundary)
 {
 	const TArray<int32>& ToMeshV = VertexMap.ToId();
 	const TArray<int32>& ToIndex = VertexMap.ToIndex();
@@ -348,7 +357,8 @@ void UE::MeshDeformation::ConstructValenceWeightedLaplacian(const FDynamicMesh3&
 
 template<typename RealType>
 void UE::MeshDeformation::ConstructMeanValueWeightLaplacian(const FDynamicMesh3& DynamicMesh, const FVertexLinearization& VertexMap,
-	TSparseMatrixAssembler<RealType>& LaplacianInterior, TSparseMatrixAssembler<RealType>& LaplacianBoundary)
+	UE::Solvers::TSparseMatrixAssembler<RealType>& LaplacianInterior, 
+	UE::Solvers::TSparseMatrixAssembler<RealType>& LaplacianBoundary)
 {
 	const TArray<int32>& ToMeshV = VertexMap.ToId();
 	const TArray<int32>& ToIndex = VertexMap.ToIndex();
@@ -419,8 +429,9 @@ void UE::MeshDeformation::ConstructMeanValueWeightLaplacian(const FDynamicMesh3&
 
 template<typename RealType>
 void UE::MeshDeformation::ConstructCotangentLaplacian(const FDynamicMesh3& DynamicMesh, const FVertexLinearization& VertexMap,
-	TSparseMatrixAssembler<RealType>& AreaMatrix,
-	TSparseMatrixAssembler<RealType>& LaplacianInterior, TSparseMatrixAssembler<RealType>& LaplacianBoundary)
+	UE::Solvers::TSparseMatrixAssembler<RealType>& AreaMatrix,
+	UE::Solvers::TSparseMatrixAssembler<RealType>& LaplacianInterior, 
+	UE::Solvers::TSparseMatrixAssembler<RealType>& LaplacianBoundary)
 {
 	const TArray<int32>& ToMeshV = VertexMap.ToId();
 	const TArray<int32>& ToIndex = VertexMap.ToIndex();
@@ -519,7 +530,8 @@ void UE::MeshDeformation::ConstructCotangentLaplacian(const FDynamicMesh3& Dynam
 
 template<typename RealType>
 void UE::MeshDeformation::ConstructCotangentLaplacian(const FDynamicMesh3& DynamicMesh, const FVertexLinearization& VertexMap,
-	TSparseMatrixAssembler<RealType>& LaplacianInterior, TSparseMatrixAssembler<RealType>& LaplacianBoundary,
+	UE::Solvers::TSparseMatrixAssembler<RealType>& LaplacianInterior, 
+	UE::Solvers::TSparseMatrixAssembler<RealType>& LaplacianBoundary,
 	const bool bClampWeights)
 {
 	const TArray<int32>& ToMeshV = VertexMap.ToId();
diff --git a/Engine/Plugins/Experimental/GeometryProcessing/Source/GeometricObjects/Public/Solvers/MatrixInterfaces.h b/Engine/Plugins/Experimental/GeometryProcessing/Source/GeometricObjects/Public/Solvers/MatrixInterfaces.h
--- a/Engine/Plugins/Experimental/GeometryProcessing/Source/GeometricObjects/Public/Solvers/MatrixInterfaces.h
+++ b/Engine/Plugins/Experimental/GeometryProcessing/Source/GeometricObjects/Public/Solvers/MatrixInterfaces.h
@@ -3,21 +3,51 @@
 #pragma once
 
 #include "CoreMinimal.h"
+#include "VectorTypes.h"
 
-/**
- * Generic adapter for building a Sparse Matrix. Generally sparse-matrix assembly
- * code emits matrix entries as tuples of form (i,j,Value).
- */
-template<typename RealType>
-class TSparseMatrixAssembler
+
+namespace UE
 {
-public:
-	/** Hint to reserve space for at least this many entries */
-	TUniqueFunction<void(int32)> ReserveEntriesFunc;
 
-	/** Add matrix entry tuple (i,j,Value) */
-	TUniqueFunction<void(int32, int32, RealType)> AddEntryFunc;
-};
+	namespace Solvers
+	{
+
+		/**
+		 * Generic adapter for building a Sparse Matrix. Generally sparse-matrix assembly
+		 * code emits matrix entries as tuples of form (i,j,Value).
+		 */
+		template<typename RealType>
+		class TSparseMatrixAssembler
+		{
+		public:
+			/** Hint to reserve space for at least this many entries */
+			TUniqueFunction<void(int32)> ReserveEntriesFunc;
+
+			/** Add matrix entry tuple (i,j,Value) */
+			TUniqueFunction<void(int32, int32, RealType)> AddEntryFunc;
+		};
+
+		typedef TSparseMatrixAssembler<double> FSparseMatrixAssemblerd;
+		typedef TSparseMatrixAssembler<float> FSparseMatrixAssemblerf;
+
+
+
+		/**
+		 * Basic position constraint
+		 */
+		struct FPositionConstraint
+		{
+			FPositionConstraint(const FVector3d& P, bool b) : Position(P), bPostFix(b) {}
+
+			FVector3d Position;
+
+			/** If bPostFix is true, this position constraint should be explicitly enforced after a solve */
+			bool     bPostFix;
+		};
+
+
+
+	}
+}
+
 
-typedef TSparseMatrixAssembler<double> FSparseMatrixAssemblerd;
-typedef TSparseMatrixAssembler<float> FSparseMatrixAssemblerf;
diff --git a/Engine/Plugins/Experimental/MeshModelingToolset/Source/MeshModelingTools/Private/DeformMeshPolygonsTool.cpp b/Engine/Plugins/Experimental/MeshModelingToolset/Source/MeshModelingTools/Private/DeformMeshPolygonsTool.cpp
--- a/Engine/Plugins/Experimental/MeshModelingToolset/Source/MeshModelingTools/Private/DeformMeshPolygonsTool.cpp
+++ b/Engine/Plugins/Experimental/MeshModelingToolset/Source/MeshModelingTools/Private/DeformMeshPolygonsTool.cpp
@@ -14,8 +14,7 @@
 #include "Async/ParallelFor.h"
 #include "Containers/BitArray.h"
 
-// WAT
-#include "DynamicMesh/Private/Solvers/Internal/LaplacianMeshSmoother.h"
+#include "Solvers/ConstrainedMeshDeformer.h"
 #include "ModelingOperators/Public/ModelingTaskTypes.h"
 
 
@@ -152,7 +151,7 @@ private:
 	TArray<int32> SubsetVertexIDToSrcVertexIDMap;
 
 	/** Laplacian deformer object gets rebuilt each new transaction */
-	TUniquePtr<FConstrainedMeshDeformer> ConstrainedDeformer;
+	TUniquePtr<UE::Solvers::IConstrainedMeshSolver> ConstrainedDeformer;
 
 private:
 	FConstrainedMeshDeformerTask();
@@ -370,7 +369,7 @@ void FConstrainedMeshDeformerTask::DoWork()
 	{
 		
 		// Create a new deformation solver.
-		ConstrainedDeformer = MakeUnique<FConstrainedMeshDeformer>(SubsetMesh, LaplacianWeightScheme);
+		ConstrainedDeformer = UE::MeshDeformation::ConstructConstrainedMeshDeformer(LaplacianWeightScheme, SubsetMesh);
 
 		if (bAttenuateWeights)
 		{
diff --git a/Engine/Plugins/Experimental/MeshModelingToolset/Source/ModelingOperators/Private/SmoothingOps/CotanSmoothingOp.cpp b/Engine/Plugins/Experimental/MeshModelingToolset/Source/ModelingOperators/Private/SmoothingOps/CotanSmoothingOp.cpp
--- a/Engine/Plugins/Experimental/MeshModelingToolset/Source/ModelingOperators/Private/SmoothingOps/CotanSmoothingOp.cpp
+++ b/Engine/Plugins/Experimental/MeshModelingToolset/Source/ModelingOperators/Private/SmoothingOps/CotanSmoothingOp.cpp
@@ -32,7 +32,7 @@ void FCotanSmoothingOp::Smooth()
 		UseScheme = ELaplacianWeightScheme::Uniform;
 	}
 
-	TUniquePtr<IConstrainedMeshSolver> Smoother = UE::MeshDeformation::ConstructConstrainedMeshSmoother(
+	TUniquePtr<UE::Solvers::IConstrainedMeshSolver> Smoother = UE::MeshDeformation::ConstructConstrainedMeshSmoother(
 		UseScheme, *ResultMesh);
 
 	double Power = SmoothOptions.SmoothPower;
