ModelingTools: fix error in Laplacian area matrix sizing introduced by refactor.
#rnx

diff --git a/Engine/Plugins/Experimental/GeometryProcessing/Source/DynamicMesh/Private/Solvers/Internal/LaplacianOperators.cpp b/Engine/Plugins/Experimental/GeometryProcessing/Source/DynamicMesh/Private/Solvers/Internal/LaplacianOperators.cpp
--- a/Engine/Plugins/Experimental/GeometryProcessing/Source/DynamicMesh/Private/Solvers/Internal/LaplacianOperators.cpp
+++ b/Engine/Plugins/Experimental/GeometryProcessing/Source/DynamicMesh/Private/Solvers/Internal/LaplacianOperators.cpp
@@ -139,7 +139,7 @@ void ConstructCotangentLaplacian(const FDynamicMesh3& DynamicMesh, FVertexLinear
 
 	FEigenSparseMatrixAssembler Interior(NumInteriorVerts, NumInteriorVerts);
 	FEigenSparseMatrixAssembler Boundary(NumInteriorVerts, NumBoundaryVerts);
-	FEigenSparseMatrixAssembler Area(NumVerts, NumVerts);
+	FEigenSparseMatrixAssembler Area(NumInteriorVerts, NumInteriorVerts);
 	UE::MeshDeformation::ConstructCotangentLaplacian<double>(DynamicMesh, VertexMap, Area, Interior, Boundary);
 	Interior.ExtractResult(LaplacianInterior);
 	Boundary.ExtractResult(LaplacianBoundary);
@@ -175,28 +175,22 @@ double ConstructScaledCotangentLaplacian(const FDynamicMesh3& DynamicMesh, FVert
 	FSparseMatrixD CotangentBoundary;
 	ConstructCotangentLaplacian(DynamicMesh, VertexMap, AreaMatrix, CotangentInterior, CotangentBoundary);
 
-	// note: returned AreaMatrix here is full mesh size, ie Interior + Boundary, with Boundary vertices at the end
-	// Does the code below make sense in light of that? We are constructing ScaledInvAreaMatrix and multiplying
-	// CotangentInterior and CotangentBoundary separately, but they are each different sizes??
-
 	// Find average entry in the area matrix
-	const int32 InternalSize = CotangentInterior.cols();
 	const int32 Rank = AreaMatrix.cols();
 	double AveArea = 0.;
-	for (int32 i = 0; i < InternalSize; ++i)		// only non-boundary vertices have area > 0!
+	for (int32 i = 0; i < Rank; ++i)
 	{
 		double Area = AreaMatrix.coeff(i, i);
 		checkSlow(Area > 0.);  // Area must be positive.
 		AveArea += Area;
 	}
-	AveArea /= (double)InternalSize;
+	AveArea /= (double)Rank;
 	
 	std::vector<MatrixTripletT> ScaledInvAreaTriplets;
 	ScaledInvAreaTriplets.reserve(Rank);
 	for (int32 i = 0; i < Rank; ++i)
 	{
-		bool bInternal = i < InternalSize;
-		double Area = (bInternal) ? AreaMatrix.coeff(i, i) : AveArea;		// use average area for boundary vertices?
+		double Area = AreaMatrix.coeff(i, i);
 		double ScaledInvArea = AveArea / Area;
 		if (bClampAreas)
 		{
@@ -210,7 +204,6 @@ double ConstructScaledCotangentLaplacian(const FDynamicMesh3& DynamicMesh, FVert
 	ScaledInvAreaMatrix.setFromTriplets(ScaledInvAreaTriplets.begin(), ScaledInvAreaTriplets.end());
 	ScaledInvAreaMatrix.makeCompressed();
 
-	// ScaledInvAreaMatrix is a different max-dimension than CotangentBoundary and CotangentInterior ??
 	LaplacianBoundary = ScaledInvAreaMatrix * CotangentBoundary;
 	LaplacianBoundary.makeCompressed();
 	LaplacianInterior = ScaledInvAreaMatrix * CotangentInterior;
