MeshModelingTools: Fixed a bug in MeshDescriptionToDynamicMesh that was checking validity incorrectly.

#rnx

diff --git a/Engine/Plugins/Experimental/GeometryProcessing/Source/MeshConversion/Private/MeshDescriptionToDynamicMesh.cpp b/Engine/Plugins/Experimental/GeometryProcessing/Source/MeshConversion/Private/MeshDescriptionToDynamicMesh.cpp
--- a/Engine/Plugins/Experimental/GeometryProcessing/Source/MeshConversion/Private/MeshDescriptionToDynamicMesh.cpp
+++ b/Engine/Plugins/Experimental/GeometryProcessing/Source/MeshConversion/Private/MeshDescriptionToDynamicMesh.cpp
@@ -421,9 +421,9 @@ void FMeshDescriptionToDynamicMesh::CopyTangents(const FMeshDescription* SourceM
 	TArrayView<const FVector> InstanceTangents = Attributes.GetVertexInstanceTangents().GetRawArray();
 	TArrayView<const float> InstanceSigns = Attributes.GetVertexInstanceBinormalSigns().GetRawArray();
 
-	if (!ensureMsgf(InstanceNormals.Num() == 0, TEXT("Cannot CopyTangents from MeshDescription with invalid Instance Normals"))) return;
-	if (!ensureMsgf(InstanceTangents.Num() == 0, TEXT("Cannot CopyTangents from MeshDescription with invalid Instance Tangents"))) return;
-	if (!ensureMsgf(InstanceSigns.Num() == 0, TEXT("Cannot CopyTangents from MeshDescription with invalid Instance BinormalSigns"))) return;
+	if (!ensureMsgf(InstanceNormals.Num() != 0, TEXT("Cannot CopyTangents from MeshDescription with invalid Instance Normals"))) return;
+	if (!ensureMsgf(InstanceTangents.Num() != 0, TEXT("Cannot CopyTangents from MeshDescription with invalid Instance Tangents"))) return;
+	if (!ensureMsgf(InstanceSigns.Num() != 0, TEXT("Cannot CopyTangents from MeshDescription with invalid Instance BinormalSigns"))) return;
 
 	TangentsOut->SetMesh(TargetMesh);
 	TangentsOut->InitializeTriVertexTangents(false);
