Fix geometry unit tests

#rnx

diff --git a/Engine/Plugins/Experimental/GeometryProcessing/Source/DynamicMesh/Private/Solvers/Internal/ConstrainedMeshDeformers.cpp b/Engine/Plugins/Experimental/GeometryProcessing/Source/DynamicMesh/Private/Solvers/Internal/ConstrainedMeshDeformers.cpp
--- a/Engine/Plugins/Experimental/GeometryProcessing/Source/DynamicMesh/Private/Solvers/Internal/ConstrainedMeshDeformers.cpp
+++ b/Engine/Plugins/Experimental/GeometryProcessing/Source/DynamicMesh/Private/Solvers/Internal/ConstrainedMeshDeformers.cpp
@@ -42,6 +42,12 @@ bool FConstrainedMeshDeformer::Deform(TArray<FVector3d>& PositionBuffer)
 	// Update constraints.  This only trigger solver rebuild if the weights were updated.
 	UpdateSolverConstraints();
 
+	// If no internal vertices to solve for, fail out early (the solver expects at least one)
+	if (InternalVertexCount == 0)
+	{
+		return false;
+	}
+
 	// Allocate space for the result as a struct of arrays
 	FSOAPositions SolutionVector(InternalVertexCount);
 
