qhull algorithm for exact convex hull in GeometryAlgorithms, replacing naive algorithm

diff --git a/Engine/Plugins/Experimental/GeometryProcessing/Source/GeometryAlgorithms/Private/ConvexHull3.cpp b/Engine/Plugins/Experimental/GeometryProcessing/Source/GeometryAlgorithms/Private/ConvexHull3.cpp
--- a/Engine/Plugins/Experimental/GeometryProcessing/Source/GeometryAlgorithms/Private/ConvexHull3.cpp
+++ b/Engine/Plugins/Experimental/GeometryProcessing/Source/GeometryAlgorithms/Private/ConvexHull3.cpp
@@ -171,7 +171,431 @@ struct TExtremePoints3
 
 };
 
+template<typename RealType>
+struct FHullConnectivity
+{
+
+	struct FVisiblePoints
+	{
+		TArray<int32> Indices;
+		int32 MaxIdx = -1; // Note: an index into the source point array, not into the above Indices array
+		// Max value of the Orient3D predicate, indicating the farthest point outside; values less than zero are inside the hull so not considered
+		double MaxValue = 0;
+
+		void AddPt(int32 Idx, double Value)
+		{
+			if (Value > MaxValue)
+			{
+				MaxValue = Value;
+				MaxIdx = Idx;
+			}
+			Indices.Add(Idx);
+		}
+
+		int32 Num()
+		{
+			return Indices.Num();
+		}
+
+		int32 MaxPt()
+		{
+			return MaxIdx;
+		}
+
+		void Reset()
+		{
+			Indices.Reset();
+			MaxIdx = -1;
+			MaxValue = 0;
+		}
+	};
+
+	TArray<FIndex3i> TriNeighbors;
+	TArray<FVisiblePoints> VisiblePoints;
+	TSet<int32> TrisWithPoints;
+
+	/**
+	 * Fully build neighbor connectivity; only on the initial tet, 
+	 * as after that we can incrementally update connectivity w/ mesh changes
+	 */
+	void BuildNeighbors(const TArray<FIndex3i>& Triangles)
+	{
+		TMap<FIndex2i, int32> EdgeToTri;
+		for (int32 TriIdx = 0; TriIdx < Triangles.Num(); TriIdx++)
+		{
+			const FIndex3i& Tri = Triangles[TriIdx];
+			for (int32 LastIdx = 2, Idx = 0; Idx < 3; LastIdx = Idx++)
+			{
+				EdgeToTri.Add(FIndex2i(Tri[LastIdx], Tri[Idx]), TriIdx);
+			}
+		}
+		TriNeighbors.SetNum(Triangles.Num());
+		for (int32 TriIdx = 0; TriIdx < Triangles.Num(); TriIdx++)
+		{
+			const FIndex3i& Tri = Triangles[TriIdx];
+			for (int32 LastIdx = 2, Idx = 0; Idx < 3; LastIdx = Idx++)
+			{
+				FIndex2i BackEdge(Tri[Idx], Tri[LastIdx]);
+				int32* Val = EdgeToTri.Find(BackEdge);
+				TriNeighbors[TriIdx][LastIdx] = Val ? *Val : -1;
+			}
+		}
+	}
+
+	/**
+	 * @param ValueOut The volume of the tetrahedron formed by TriPts and Pt
+	 * @return true if Pt is on the 'positive' side of the triangle
+	 */
+	bool IsVisible(const FVector3d TriPts[3], const FVector3d& Pt, double& ValueOut)
+	{
+		ValueOut = ExactPredicates::Orient3D(TriPts[0], TriPts[1], TriPts[2], Pt);
+		return ValueOut > 0;
+	}
+
+	/**
+	 * Set triangle positions to prep for IsVisible calls
+	 */
+	void SetTriPts(const FIndex3i& Tri, TFunctionRef<FVector3<RealType>(int32)> GetPointFunc,  FVector3d TriPtsOut[3])
+	{
+		TriPtsOut[0] = (FVector3d)GetPointFunc(Tri[0]);
+		TriPtsOut[1] = (FVector3d)GetPointFunc(Tri[1]);
+		TriPtsOut[2] = (FVector3d)GetPointFunc(Tri[2]);
+	}
+
+	/**
+	 * Bucket all points based on which triangles can see them
+	 */
+	void InitVisibility(const TArray<FIndex3i>& Triangles, int32 NumPoints, TFunctionRef<FVector3<RealType>(int32)> GetPointFunc, TFunctionRef<bool(int32)> FilterFunc)
+	{
+		VisiblePoints.SetNum(Triangles.Num());
+
+		FVector3d TriPts[3];
+		FVector3d Pt;
+		for (int32 PtIdx = 0; PtIdx < NumPoints; PtIdx++)
+		{
+			if (!FilterFunc(PtIdx))
+			{
+				continue;
+			}
+			Pt = (FVector3d)GetPointFunc(PtIdx);
+			for (int32 TriIdx = 0; TriIdx < Triangles.Num(); TriIdx++)
+			{
+				SetTriPts(Triangles[TriIdx], GetPointFunc, TriPts);
+				double Value;
+				if (IsVisible(TriPts, Pt, Value))
+				{
+					if (!VisiblePoints[TriIdx].Num())
+					{
+						TrisWithPoints.Add(TriIdx);
+					}
+					VisiblePoints[TriIdx].AddPt(PtIdx, Value);
+				}
+			}
+		}
+	}
+
+	/**
+	 * Update any neighbor "back pointers" to TriIdx to instead point to NewIdx
+	 * Used during triangle deletion when we swap a triangle to a new index
+	 */
+	void UpdateNeighbors(int32 TriIdx, int32 NewIdx)
+	{
+		FIndex3i Neighbors = TriNeighbors[TriIdx];
+		for (int32 Idx = 0; Idx < 3; Idx++)
+		{
+			int32 NbrIdx = Neighbors[Idx];
+			if (NbrIdx > -1 && NbrIdx < TriNeighbors.Num())
+			{
+				FIndex3i& NbrNbrs = TriNeighbors[NbrIdx];
+				for (int SubIdx = 0; SubIdx < 3; SubIdx++)
+				{
+					if (NbrNbrs[SubIdx] == TriIdx)
+					{
+						NbrNbrs[SubIdx] = NewIdx;
+						break; // should only be one back-pointer per neighbor (or sometimes zero if it was already partially unhooked)
+					}
+				}
+			}
+		}
+	}
+
+	/**
+	 * Rewrite one neighbor connection (and NOT the backwards connection)
+	 *
+	 * @param Triangles All triangles
+	 * @param TriIdx The triangle to update
+	 * @param EdgeFirstVertex The first vertex of the edge that should be connected to the new neighbor
+	 * @param NewNbrIdx The new neighbor to connect to
+	 */
+	void UpdateNeighbor(TArray<FIndex3i>& Triangles, int32 TriIdx, int32 EdgeFirstVertex, int32 NewNbrIdx)
+	{
+		const FIndex3i& Triangle = Triangles[TriIdx];
+		for (int32 Idx = 0; Idx < 3; Idx++)
+		{
+			if (Triangle[Idx] == EdgeFirstVertex)
+			{
+				TriNeighbors[TriIdx][Idx] = NewNbrIdx;
+				return;
+			}
+		}
+		check(false); // EdgeFirstVertex wasn't found in triangle; invalid UpdateNeighbor call
+	}
+
+	/**
+	 * Debugging aid to validate the neighbor connectivity
+	 */
+	bool ValidateConnectivity(TArray<FIndex3i>& Triangles, TSet<int32> Skip = TSet<int32>())
+	{
+		TMap<FIndex2i, int32> EdgeToTri;
+		for (int32 TriIdx = 0; TriIdx < Triangles.Num(); TriIdx++)
+		{
+			if (Skip.Contains(TriIdx))
+			{
+				continue;
+			}
+			const FIndex3i& Tri = Triangles[TriIdx];
+			for (int32 LastIdx = 2, Idx = 0; Idx < 3; LastIdx = Idx++)
+			{
+				EdgeToTri.Add(FIndex2i(Tri[LastIdx], Tri[Idx]), TriIdx);
+			}
+		}
+		for (int32 TriIdx = 0; TriIdx < Triangles.Num(); TriIdx++)
+		{
+			if (Skip.Contains(TriIdx))
+			{
+				continue;
+			}
+			const FIndex3i& Tri = Triangles[TriIdx];
+			const FIndex3i& Nbrs = TriNeighbors[TriIdx];
+			for (int32 LastIdx = 2, Idx = 0; Idx < 3; LastIdx = Idx++)
+			{
+				FIndex2i BackEdge = FIndex2i(Tri[Idx], Tri[LastIdx]);
+				int32* FoundNbr = EdgeToTri.Find(BackEdge);
+				if (!FoundNbr || Nbrs[LastIdx] != *FoundNbr)
+				{
+					ensure(false);
+					return false;
+				}
+			}
+		}
+
+		return true;
+	}
+
+	/** 
+	 * Remove a triangle and update connectivity accordingly.  Removes by swapping, and updates references as needed.
+	 * Note: removes a whole set at a time, because indices change during deletion
+	 *
+	 * @param Triangles All triangles
+	 * @param ToDelete All triangles to be deleted
+	 */
+	void DeleteTriangles(TArray<FIndex3i>& Triangles, const TSet<int32>& ToDelete)
+	{
+		TSet<int32> Moved; // indices of deleted tris that non-delete triangles have been moved into, must not be deleted
+		for (int32 TriIdx : ToDelete)
+		{
+			checkSlow(!Moved.Contains(TriIdx));
+
+			// skip triangles we already greedy-deleted while looking for the LastIdx in a previous iter
+			if (TriIdx >= Triangles.Num())
+			{
+				continue;
+			}
+			checkSlow(Triangles.Num() == TriNeighbors.Num());
+			checkSlow(Triangles.Num() == VisiblePoints.Num());
+
+			// nuke back-refs
+			UpdateNeighbors(TriIdx, -1);
+
+			int32 LastIdx = Triangles.Num() - 1;
+			// greedily remove any triangles from the end rather than swapping them back
+			while (TriIdx < LastIdx && (ToDelete.Contains(LastIdx) && !Moved.Contains(LastIdx)))
+			{
+				TrisWithPoints.Remove(LastIdx);
+				UpdateNeighbors(LastIdx, -1);
+				LastIdx--;
+			}
+
+			// found a non-delete tri at end to swap in
+			if (TriIdx < LastIdx)
+			{
+				Moved.Add(TriIdx); // track so we don't later "greedy delete" the now non-delete tri
+				UpdateNeighbors(LastIdx, TriIdx);
+				TriNeighbors[TriIdx] = TriNeighbors[LastIdx];
+				Triangles[TriIdx] = Triangles[LastIdx];
+				VisiblePoints[TriIdx] = MoveTemp(VisiblePoints[LastIdx]);
+				if (TrisWithPoints.Remove(LastIdx))
+				{
+					TrisWithPoints.Add(TriIdx);
+				}
+			}
+			else // this was the last tri, just let it pop off
+			{
+				TrisWithPoints.Remove(TriIdx);
+			}
+			Triangles.SetNum(LastIdx, false);
+			TriNeighbors.SetNum(LastIdx, false);
+			VisiblePoints.SetNum(LastIdx, false);
+		}
+	}
+
+	/**
+	 * @return A triangle index and visible-from-that-triangle point index that can be added to the hull next
+	 */
+	FIndex2i ChooseVisiblePoint()
+	{
+		FIndex2i FoundTriPointPair(-1, -1);
+		for (int32 TriIdx : TrisWithPoints)
+		{
+			ensure(VisiblePoints[TriIdx].Num() > 0);
 
+			FoundTriPointPair[0] = TriIdx;
+			// choose the "max point" -- the point with the largest volume when it makes a tetrahedron w/ the triangle
+			FoundTriPointPair[1] = VisiblePoints[TriIdx].MaxPt();
+			break;
+		}
+		return FoundTriPointPair;
+	}
+
+	/**
+	 * Helper struct carrying enough information to describe triangles that will be added when we UpdateHullWithNewPoint
+	 */
+	struct FNewTriangle
+	{
+		FNewTriangle()
+		{}
+		FNewTriangle(int32 V0, int32 V1, int32 ConnectedTri) : EdgeVertices(V0, V1), ConnectedTri(ConnectedTri)
+		{}
+
+		// The edge the new point will connect to
+		FIndex2i EdgeVertices;
+		// The neighboring triangle across that edge
+		int32 ConnectedTri;
+	};
+
+	/**
+	 * Recursive traversal strategy that starts from a 'visible' triangle and walks to the connected set of all visible triangles, s.t. it visits the boundary in order
+	 * See explanation here: http://algolist.ru/maths/geom/convhull/qhull3d.php
+	 *
+	 * @param Triangles All triangles
+	 * @param GetPointFunc Point access function
+	 * @param Pt Current 'eye point'; triangles will be deleted if this point is on the 'positive' side of the triangle
+	 * @param NewlyUnclaimed (Output) As we mark triangles for deletion, we move the visible points stored on the triangles to this giant list of points that need to be re-assigned (or culled) wrt the new triangles
+	 * @param ToDelete (Output) Set of all triangles that we will delete
+	 * @param ToAdd (Output) Info required to add and connect up new triangles to the eye point
+	 * @param TriIdx Index of triangle to currently process
+	 * @param CrossedEdgeFirstVertex The first vertex of the edge that was 'crossed over' to traverse to this TriIdx, or -1 for the initial call
+	 * @return false if triangle is beyond horizon (so we don't need to search through it / remove it), true otherwise
+	 */
+	bool HorizonHelper(const TArray<FIndex3i>& Triangles, TFunctionRef<FVector3<RealType>(int32)> GetPointFunc, const FVector3d& Pt, TArray<int32>& NewlyUnclaimed, TSet<int32>& ToDelete, TArray<FNewTriangle>& ToAdd, int32 TriIdx, int32 CrossedEdgeFirstVertex)
+	{
+		// if it's not the first triangle, crossed edge should be set and we should check if the triangle is visible / actually needs to be replaced
+		if (CrossedEdgeFirstVertex != -1)
+		{
+			FVector3d TriPts[3];
+			SetTriPts(Triangles[TriIdx], GetPointFunc, TriPts);
+			double UnusedValue;
+			if (!IsVisible(TriPts, Pt, UnusedValue))
+			{
+				return false;
+			}
+		}
+
+		// track the triangle as needing deletion; for simplicity wait until after traversal to actually delete
+		ToDelete.Add(TriIdx); // TODO: could we delete as we go rather than keep this set?  seems tricky
+
+		// clean out the visible points right away
+		if (VisiblePoints[TriIdx].Num() > 0)
+		{
+			TrisWithPoints.Remove(TriIdx);
+			NewlyUnclaimed.Append(VisiblePoints[TriIdx].Indices);
+			VisiblePoints[TriIdx].Reset();
+		}
+		
+
+		FIndex3i Tri = Triangles[TriIdx];
+		int32 FirstOff = 0, OffMax = 3; // Cross all three edges for first triangle
+		if (CrossedEdgeFirstVertex > -1) // Cross the two edges we didn't 'come from' for all other triangles
+		{
+			FirstOff = Tri.IndexOf(CrossedEdgeFirstVertex) + 1; // index of the crossed edge wrt this triangle
+			OffMax = FirstOff + 3;
+		}
+		for (int32 EdgeIdxOff = FirstOff; EdgeIdxOff < OffMax; EdgeIdxOff++)
+		{
+			int32 AcrossTri = TriNeighbors[TriIdx][EdgeIdxOff % 3];
+			if (ToDelete.Contains(AcrossTri))
+			{
+				continue;
+			}
+			// second vertex of the edge (should be first seen when looking for the edge on opposite side)
+			int32 EdgeSecondVertex = Tri[(EdgeIdxOff + 1) % 3];
+			if (!HorizonHelper(Triangles, GetPointFunc, Pt, NewlyUnclaimed, ToDelete, ToAdd, AcrossTri, EdgeSecondVertex))
+			{
+				int32 EdgeFirstVertex = Tri[EdgeIdxOff % 3];
+				ToAdd.Add(FNewTriangle(EdgeFirstVertex, EdgeSecondVertex, AcrossTri));
+			}
+		}
+
+		return true;
+	}
+
+	void UpdateHullWithNewPoint(TArray<FIndex3i>& Triangles, TFunctionRef<FVector3<RealType>(int32)> GetPointFunc, int32 StartTriIdx, int32 PtIdx)
+	{
+		//ValidateConnectivity(Triangles);
+
+		FVector3d Pt = (FVector3d)GetPointFunc(PtIdx);
+
+		// Call the recursive helper to find all the tris we need to delete + the 'horizon' of edges that will form new triangles
+		TArray<int32> NewlyUnclaimed;
+		TSet<int32> ToDelete;
+		TArray<FNewTriangle> ToAdd;
+		HorizonHelper(Triangles, GetPointFunc, Pt, NewlyUnclaimed, ToDelete, ToAdd, StartTriIdx, -1);
+
+		// Connect up all the horizon triangles
+		int32 NewTriStart = Triangles.Num();
+		int32 NumAdd = ToAdd.Num();
+		FVector3d TriPts[3];
+		for (int32 AddIdx = 0; AddIdx < NumAdd; AddIdx++)
+		{
+			const FNewTriangle& TriData = ToAdd[AddIdx];
+			int32 NewTriIdx = Triangles.Add(FIndex3i(PtIdx, TriData.EdgeVertices.A, TriData.EdgeVertices.B));
+			int32 PrevTriIdx = NewTriStart + ((AddIdx - 1 + NumAdd) % NumAdd);
+			int32 NextTriIdx = NewTriStart + ((AddIdx + 1) % NumAdd);
+
+			int32 AcrossTriIdx = TriData.ConnectedTri;
+			TriNeighbors.Add(FIndex3i(PrevTriIdx, AcrossTriIdx, NextTriIdx)); // TODO: is this the right order? need to validate!
+			UpdateNeighbor(Triangles, AcrossTriIdx, TriData.EdgeVertices.B, NewTriIdx);
+			FVisiblePoints& Visible = VisiblePoints.Emplace_GetRef();
+			SetTriPts(Triangles[NewTriIdx], GetPointFunc, TriPts);
+			// claim any claim-able points
+			for (int32 UnclaimedIdx = 0; UnclaimedIdx < NewlyUnclaimed.Num(); UnclaimedIdx++)
+			{
+				int32 UnPtIdx = NewlyUnclaimed[UnclaimedIdx];
+				double Value;
+				if (IsVisible(TriPts, (FVector3d)GetPointFunc(UnPtIdx), Value))
+				{
+					Visible.AddPt(UnPtIdx, Value);
+					NewlyUnclaimed.RemoveAtSwap(UnclaimedIdx, 1, false);
+					UnclaimedIdx--;
+					continue;
+				}
+			}
+			if (Visible.Num() > 0)
+			{
+				TrisWithPoints.Add(NewTriIdx);
+			}
+		}
+
+		//ValidateConnectivity(Triangles, ToDelete);
+
+		TArray<FIndex3i> OldNbrs = TriNeighbors;
+
+		// do the deletions *last* so we don't invalidate tri indices in ToAdd
+		DeleteTriangles(Triangles, ToDelete);
+
+		//ValidateConnectivity(Triangles);
+	}
+};
 
 
 
@@ -179,6 +603,7 @@ template<class RealType>
 bool TConvexHull3<RealType>::Solve(int32 NumPoints, TFunctionRef<FVector3<RealType>(int32)> GetPointFunc, TFunctionRef<bool(int32)> FilterFunc)
 {
 	Hull.Reset();
+	NumHullPoints = 0;
 
 	TExtremePoints3<RealType> InitialTet(NumPoints, GetPointFunc, FilterFunc);
 	Dimension = InitialTet.Dimension;
@@ -209,72 +634,26 @@ bool TConvexHull3<RealType>::Solve(int32 NumPoints, TFunctionRef<FVector3<RealTy
 	Hull.Add(FIndex3i(InitialTet.Extreme[0], InitialTet.Extreme[1], InitialTet.Extreme[3]));
 	Hull.Add(FIndex3i(InitialTet.Extreme[0], InitialTet.Extreme[2], InitialTet.Extreme[1]));
 
-	// The set of processed points is maintained to eliminate exact duplicates
-	//  (should not be needed for correctness but lets us count unique points and is probably faster)
-	TSet<FVector3<RealType>> Processed;
-	for (int i = 0; i < 4; ++i)
-	{
-		Processed.Add(GetPointFunc(InitialTet.Extreme[i]));
-	}
-	// Incrementally update the hull.  
-	for (int i = 0; i < NumPoints; ++i)
+	NumHullPoints = 4;
+
+	FHullConnectivity<RealType> Connectivity;
+	Connectivity.BuildNeighbors(Hull);
+	Connectivity.InitVisibility(Hull, NumPoints, GetPointFunc, FilterFunc);
+	while (true)
 	{
-		if (FilterFunc(i) && !Processed.Contains(GetPointFunc(i)))
+		FIndex2i Visible = Connectivity.ChooseVisiblePoint();
+		if (Visible.A == -1)
 		{
-			Insert(GetPointFunc, i);
-			Processed.Add(GetPointFunc(i));
+			break;
 		}
+		NumHullPoints++;
+		Connectivity.UpdateHullWithNewPoint(Hull, GetPointFunc, Visible.A, Visible.B);
 	}
-	NumUniquePoints = Processed.Num();
 
 	return true;
 }
 
 
-template<class RealType>
-void TConvexHull3<RealType>::Insert(TFunctionRef<FVector3<RealType>(int32)> GetPointFunc, int32 PtIdx)
-{
-	TArray<int> PtSides; PtSides.SetNum(Hull.Num());
-	FVector3d Pt = (FVector3d)GetPointFunc(PtIdx); // double precision point to use double precision plane test
-	bool bSingleThread = true; // TODO: have yet to find a dataset where the ParallelFor actually helps -- maybe if Hull.Num() gets very large?
-	ParallelFor(Hull.Num(), [this, &GetPointFunc, &PtSides, Pt](int32 TriIdx)
-		{
-			FIndex3i Tri = Hull[TriIdx];
-			double Side = FMath::Sign(ExactPredicates::Orient3D((FVector3d)GetPointFunc(Tri.A), (FVector3d)GetPointFunc(Tri.B), (FVector3d)GetPointFunc(Tri.C), Pt));
-			PtSides[TriIdx] = Side > 0 ? 1 : (Side < 0 ? -1 : 0);
-		}, bSingleThread);
-
-	TSet<FIndex2i> TerminatorEdges; // set of 'terminator' edges -- edges that we see only once among the removed triangles -- to be connected to the new vertex
-	for (int TriIdx = 0; TriIdx < Hull.Num(); TriIdx++)
-	{
-		int PtSide = PtSides[TriIdx];
-		if (PtSide > 0)
-		{
-			const FIndex3i& Tri = Hull[TriIdx];
-			for (int E0 = 2, E1 = 0; E1 < 3; E0 = E1++)
-			{
-				int V0 = Tri[E0], V1 = Tri[E1];
-				FIndex2i Edge(V0, V1);
-				// if we've seen this edge before, it would be backwards -- try to remove it
-				FIndex2i BackEdge(V1, V0);
-				if (TerminatorEdges.Remove(BackEdge) == 0)
-				{
-					// we hadn't seen it before, so add it as a candidate terminator
-					TerminatorEdges.Add(Edge);
-				}
-			}
-
-			Hull.RemoveAtSwap(TriIdx, 1, false);
-			PtSides.RemoveAtSwap(TriIdx, 1, false);
-			TriIdx--;
-		}
-	}
-
-	for (const FIndex2i& Edge : TerminatorEdges)
-	{
-		Hull.Add(FIndex3i(PtIdx, Edge.A, Edge.B)); // add triangle connecting inserted point to terminator edge
-	}
-}
 
 
 template class TConvexHull3<float>;
diff --git a/Engine/Plugins/Experimental/GeometryProcessing/Source/GeometryAlgorithms/Public/ConvexHull3.h b/Engine/Plugins/Experimental/GeometryProcessing/Source/GeometryAlgorithms/Public/ConvexHull3.h
--- a/Engine/Plugins/Experimental/GeometryProcessing/Source/GeometryAlgorithms/Public/ConvexHull3.h
+++ b/Engine/Plugins/Experimental/GeometryProcessing/Source/GeometryAlgorithms/Public/ConvexHull3.h
@@ -96,7 +96,7 @@ public:
 	void Empty()
 	{
 		Dimension = 0;
-		NumUniquePoints = 0;
+		NumHullPoints = 0;
 		Hull.Empty();
 	}
 
@@ -118,21 +118,19 @@ public:
 		return Plane;
 	}
 
-	/** @return Number of unique points considered by convex hull construction (excludes exact duplicate points and filtered-out points) */
-	int GetNumUniquePoints() const
+	/** @return Number of points on the output hull */
+	int GetNumHullPoints() const
 	{
-		return NumUniquePoints;
+		return NumHullPoints;
 	}
 
 protected:
-	// Incrementally insert ith point
-	void Insert(TFunctionRef<FVector3<RealType>(int32)> GetPointFunc, int Idx);
 
 	int32 Dimension;
 	TLine3<RealType> Line;
 	TPlane3<RealType> Plane;
 
-	int NumUniquePoints = 0;
+	int NumHullPoints = 0;
 	TArray<FIndex3i> Hull;
 	
 };
