BoxTypes:
- fix IsEmpty() behavior
- set default-constructed AABBs to be the result of Empty(), i.e. "inside-out" by default
- an empty FDynamicMesh3 returns a default AABB from GetBounds()
- unit tests for this behavior

#rnx

diff --git a/Engine/Plugins/Experimental/GeometryProcessing/Source/DynamicMesh/Private/DynamicMesh3_Queries.cpp b/Engine/Plugins/Experimental/GeometryProcessing/Source/DynamicMesh/Private/DynamicMesh3_Queries.cpp
--- a/Engine/Plugins/Experimental/GeometryProcessing/Source/DynamicMesh/Private/DynamicMesh3_Queries.cpp
+++ b/Engine/Plugins/Experimental/GeometryProcessing/Source/DynamicMesh/Private/DynamicMesh3_Queries.cpp
@@ -591,6 +591,11 @@ int FDynamicMesh3::FindTriangle(int a, int b, int c) const
  */
 FAxisAlignedBox3d FDynamicMesh3::GetBounds() const
 {
+	if (VertexCount() == 0) 
+	{
+		return FAxisAlignedBox3d::Empty();
+	}
+
 	FVector3d MinVec = Vertices[*(VertexIndicesItr().begin())];
 	FVector3d MaxVec = MinVec;
 	for (int vi : VertexIndicesItr())
diff --git a/Engine/Plugins/Experimental/GeometryProcessing/Source/GeometricObjects/Public/BoxTypes.h b/Engine/Plugins/Experimental/GeometryProcessing/Source/GeometricObjects/Public/BoxTypes.h
--- a/Engine/Plugins/Experimental/GeometryProcessing/Source/GeometricObjects/Public/BoxTypes.h
+++ b/Engine/Plugins/Experimental/GeometryProcessing/Source/GeometricObjects/Public/BoxTypes.h
@@ -13,9 +13,11 @@ struct TInterval1
 	RealType Min;
 	RealType Max;
 
-	TInterval1()
+	TInterval1() :
+		TInterval1(Empty())
 	{
 	}
+
 	TInterval1(const RealType& Min, const RealType& Max)
 	{
 		this->Min = Min;
@@ -186,7 +188,7 @@ struct TInterval1
 
 	inline bool IsEmpty() const
 	{
-		return Max > Min;
+		return Max < Min;
 	}
 
 	void Expand(RealType Radius)
@@ -208,9 +210,11 @@ struct TAxisAlignedBox3
 	FVector3<RealType> Min;
 	FVector3<RealType> Max;
 
-	TAxisAlignedBox3()
+	TAxisAlignedBox3() : 
+		TAxisAlignedBox3(TAxisAlignedBox3<RealType>::Empty())
 	{
 	}
+
 	TAxisAlignedBox3(const FVector3<RealType>& Min, const FVector3<RealType>& Max)
 	{
 		this->Min = Min;
@@ -470,7 +474,7 @@ struct TAxisAlignedBox3
 
 	inline bool IsEmpty() const
 	{
-		return Max.X > Min.X || Max.Y > Min.Y || Max.Z > Min.Z;
+		return Max.X < Min.X || Max.Y < Min.Y || Max.Z < Min.Z;
 	}
 
 	void Expand(RealType Radius)
@@ -490,9 +494,11 @@ struct TAxisAlignedBox2
 	FVector2<RealType> Min;
 	FVector2<RealType> Max;
 
-	TAxisAlignedBox2()
+	TAxisAlignedBox2() : 
+		TAxisAlignedBox2(Empty())
 	{
 	}
+
 	TAxisAlignedBox2(const FVector2<RealType>& Min, const FVector2<RealType>& Max)
 		: Min(Min), Max(Max)
 	{
@@ -663,7 +669,7 @@ struct TAxisAlignedBox2
 
 	inline bool IsEmpty() const
 	{
-		return Max.X > Min.X || Max.Y > Min.Y;
+		return Max.X < Min.X || Max.Y < Min.Y;
 	}
 
 	void Expand(RealType Radius)
