FFDLattice:
- add some basic unit tests
- replace a bunch of unrolled 3D vector operations by utilizing component-wise operator overloads.
- cache the initial lattice cell size rather than recomputing it.

diff --git a/Engine/Plugins/Experimental/GeometryProcessing/Source/DynamicMesh/Private/Operations/FFDLattice.cpp b/Engine/Plugins/Experimental/GeometryProcessing/Source/DynamicMesh/Private/Operations/FFDLattice.cpp
--- a/Engine/Plugins/Experimental/GeometryProcessing/Source/DynamicMesh/Private/Operations/FFDLattice.cpp
+++ b/Engine/Plugins/Experimental/GeometryProcessing/Source/DynamicMesh/Private/Operations/FFDLattice.cpp
@@ -7,7 +7,10 @@
 FFFDLattice::FFFDLattice(const FVector3i& InDims, const FDynamicMesh3& Mesh, float Padding) :
 	Dims(InDims)
 {
+	check(InDims.X > 1 && InDims.Y > 1 && InDims.Z > 1);
+
 	InitialBounds = Mesh.GetBounds();
+	check(!InitialBounds.IsEmpty());
 
 	// Expand the initial bounding box to make the computation of which grid cell a mesh vertex is inside of a little 
 	// less susceptible to numerical error issues.
@@ -15,34 +18,35 @@ FFFDLattice::FFFDLattice(const FVector3i& InDims, const FDynamicMesh3& Mesh, flo
 	const float ClampedPadding = FMath::Clamp(Padding, 0.01f, 5.0f);
 	const FVector3d Center = InitialBounds.Center();
 	FVector3d Extents = InitialBounds.Extents();
-	Extents *= (1.0 + 0.5 * ClampedPadding);
+
+	// Compute padding based on maximum component of the diagonal. Avoids problems when one or more component is zero.
+	double MaxDiagonal = InitialBounds.Diagonal().MaxElement();
+	Extents = Extents + (0.5 * ClampedPadding * MaxDiagonal);
 
 	InitialBounds.Min = Center - Extents;
 	InitialBounds.Max = Center + Extents;
 	
+	FVector3d Diag = InitialBounds.Diagonal();
+	CellSize = Diag / FVector3d(Dims - 1);
+
 	ComputeInitialEmbedding(Mesh);
 }
 
 void FFFDLattice::GenerateInitialLatticePositions(TArray<FVector3d>& OutLatticePositions) const
 {
-	double DX = (InitialBounds.Max.X - InitialBounds.Min.X) / (Dims.X - 1);
-	double DY = (InitialBounds.Max.Y - InitialBounds.Min.Y) / (Dims.Y - 1);
-	double DZ = (InitialBounds.Max.Z - InitialBounds.Min.Z) / (Dims.Z - 1);
-
 	int TotalNumLatticePoints = Dims.X * Dims.Y * Dims.Z;
-
 	OutLatticePositions.SetNum(TotalNumLatticePoints);
 
 	for (int i = 0; i < Dims.X; ++i)
 	{
-		double X = DX * i;
+		double X = CellSize.X * i;
 		for (int j = 0; j < Dims.Y; ++j)
 		{
-			double Y = DY * j;
+			double Y = CellSize.Y * j;
 			for (int k = 0; k < Dims.Z; ++k)
 			{
 				int PointID = ControlPointIndexFromCoordinates(i, j, k);
-				double Z = DZ * k;
+				double Z = CellSize.Z * k;
 
 				OutLatticePositions[PointID] = InitialBounds.Min + FVector3d{ X,Y,Z };
 			}
@@ -86,27 +90,9 @@ void FFFDLattice::GenerateLatticeEdges(TArray<FVector2i>& OutLatticeEdges) const
 
 FVector3d FFFDLattice::ComputeTrilinearWeights(const FVector3d& QueryPoint, FVector3i& GridCoordinates) const
 {
-	double DX = (InitialBounds.Max.X - InitialBounds.Min.X) / (Dims.X - 1);
-	double DY = (InitialBounds.Max.Y - InitialBounds.Min.Y) / (Dims.Y - 1);
-	double DZ = (InitialBounds.Max.Z - InitialBounds.Min.Z) / (Dims.Z - 1);
-
-	FVector3d GridPoint{
-		((QueryPoint.X - InitialBounds.Min.X) / DX),
-		((QueryPoint.Y - InitialBounds.Min.Y) / DY),
-		((QueryPoint.Z - InitialBounds.Min.Z) / DZ) };
-
-	// compute integer coordinates
-	int X0 = (int)GridPoint.X;
-	int Y0 = (int)GridPoint.Y;
-	int Z0 = (int)GridPoint.Z;
-
-	GridCoordinates = FVector3i{ X0, Y0, Z0 };
-
-	FVector3d Weights = FVector3d{
-		GridPoint.X - (double)X0,
-		GridPoint.Y - (double)Y0,
-		GridPoint.Z - (double)Z0 };
-
+	FVector3d GridPoint = (QueryPoint - InitialBounds.Min) / CellSize;
+	GridCoordinates = FVector3i(GridPoint);
+	FVector3d Weights = GridPoint - FVector3d(GridCoordinates);
 	return Weights;
 }
 
@@ -190,20 +176,13 @@ static float CubicBSplineKernel(float A)
 }
 
 
-FVector3d FFFDLattice::InterpolatedPositionCubic(const FEmbedding& VertexEmbedding, const TArray<FVector3d>& LatticeControlPoints) const
+FVector3d FFFDLattice::InterpolatedPositionCubic(const FEmbedding& VertexEmbedding, 
+												 const TArray<FVector3d>& LatticeControlPoints) const
 {
-	int X0 = VertexEmbedding.LatticeCell.X;
-	int Y0 = VertexEmbedding.LatticeCell.Y;
-	int Z0 = VertexEmbedding.LatticeCell.Z;
-
 	double T = VertexEmbedding.CellWeighting.X;
 	double U = VertexEmbedding.CellWeighting.Y;
 	double V = VertexEmbedding.CellWeighting.Z;
 
-	double DX = (InitialBounds.Max.X - InitialBounds.Min.X) / (Dims.X - 1);
-	double DY = (InitialBounds.Max.Y - InitialBounds.Min.Y) / (Dims.Y - 1);
-	double DZ = (InitialBounds.Max.Z - InitialBounds.Min.Z) / (Dims.Z - 1);
-
 	FVector3d Sum{ 0.0f, 0.0f, 0.0f };
 
 	// TODO: This can probably be replaced with some relatively simple linear algebra
@@ -221,9 +200,9 @@ FVector3d FFFDLattice::InterpolatedPositionCubic(const FEmbedding& VertexEmbeddi
 				double WeightZ = CubicBSplineKernel(V - DK);
 				double Weight = WeightX * WeightY * WeightZ;
 
-				int i = X0 + DI;
-				int j = Y0 + DJ;
-				int k = Z0 + DK;
+				int i = VertexEmbedding.LatticeCell.X + DI;
+				int j = VertexEmbedding.LatticeCell.Y + DJ;
+				int k = VertexEmbedding.LatticeCell.Z + DK;
 
 				FVector3d LatticePoint;
 				if (i < 0 || i >= Dims.X || j < 0 || j >= Dims.Y || k < 0 || k >= Dims.Z)
@@ -300,10 +279,9 @@ FVector3d FFFDLattice::InterpolatedPosition(const FEmbedding& VertexEmbedding, c
 
 FVector3d FFFDLattice::ClosestLatticePosition(const FVector3i& VirtualControlPointIndex, const TArray<FVector3d>& LatticeControlPoints) const
 {
-	FVector3i NearestControlPointIndex;
-	NearestControlPointIndex.X = FMath::Clamp(VirtualControlPointIndex.X, 0, Dims.X - 1);
-	NearestControlPointIndex.Y = FMath::Clamp(VirtualControlPointIndex.Y, 0, Dims.Y - 1);
-	NearestControlPointIndex.Z = FMath::Clamp(VirtualControlPointIndex.Z, 0, Dims.Z - 1);
+	// Clamp to valid lattice index
+	FVector3i NearestControlPointIndex = Max(Min(VirtualControlPointIndex, Dims - 1), FVector3i(0, 0, 0));
+
 	return LatticeControlPoints[ControlPointIndexFromCoordinates(NearestControlPointIndex)];
 }
 
@@ -312,11 +290,7 @@ FVector3d FFFDLattice::ExtrapolatedLatticePosition(const FVector3i& VirtualContr
 	// Use the location of the nearest control point and the location of a control point in the opposite direction of 
 	// the extrapolation to get the extrapolated location.
 
-	FVector3i NearestControlPointIndex;
-	NearestControlPointIndex.X = FMath::Clamp(VirtualControlPointIndex.X, 0, Dims.X - 1);
-	NearestControlPointIndex.Y = FMath::Clamp(VirtualControlPointIndex.Y, 0, Dims.Y - 1);
-	NearestControlPointIndex.Z = FMath::Clamp(VirtualControlPointIndex.Z, 0, Dims.Z - 1);
-
+	FVector3i NearestControlPointIndex = Max(Min(VirtualControlPointIndex, Dims - 1), FVector3i(0, 0, 0));
 	FVector3i Delta = VirtualControlPointIndex - NearestControlPointIndex;
 	check(Delta != FVector3i::Zero());
 
diff --git a/Engine/Plugins/Experimental/GeometryProcessing/Source/DynamicMesh/Public/Operations/FFDLattice.h b/Engine/Plugins/Experimental/GeometryProcessing/Source/DynamicMesh/Public/Operations/FFDLattice.h
--- a/Engine/Plugins/Experimental/GeometryProcessing/Source/DynamicMesh/Public/Operations/FFDLattice.h
+++ b/Engine/Plugins/Experimental/GeometryProcessing/Source/DynamicMesh/Public/Operations/FFDLattice.h
@@ -46,7 +46,10 @@ protected:
 	// Extents of the lattice before it is deformed
 	FAxisAlignedBox3d InitialBounds;
 
-	/// Interpolation weights and cell indices per vertex of the input FDynamicMesh3
+	// 3D size of a lattice cell
+	FVector3d CellSize;
+
+	// Interpolation weights and cell indices per vertex of the input FDynamicMesh3
 	TArray<FEmbedding> VertexEmbeddings;
 
 
