Parallelize computing mesh deformation in FFDLattice.

#rnx

diff --git a/Engine/Plugins/Experimental/GeometryProcessing/Source/DynamicMesh/Private/Operations/FFDLattice.cpp b/Engine/Plugins/Experimental/GeometryProcessing/Source/DynamicMesh/Private/Operations/FFDLattice.cpp
--- a/Engine/Plugins/Experimental/GeometryProcessing/Source/DynamicMesh/Private/Operations/FFDLattice.cpp
+++ b/Engine/Plugins/Experimental/GeometryProcessing/Source/DynamicMesh/Private/Operations/FFDLattice.cpp
@@ -3,6 +3,7 @@
 #include "Operations/FFDLattice.h"
 #include "DynamicMesh3.h"
 #include "Util/ProgressCancel.h"
+#include "Async/ParallelFor.h"
 
 FFFDLattice::FFFDLattice(const FVector3i& InDims, const FDynamicMesh3& Mesh, float Padding) :
 	Dims(InDims)
@@ -96,30 +97,62 @@ FVector3d FFFDLattice::ComputeTrilinearWeights(const FVector3d& QueryPoint, FVec
 	return Weights;
 }
 
-void FFFDLattice::ComputeInitialEmbedding(const FDynamicMesh3& Mesh)
+
+void FFFDLattice::ComputeInitialEmbedding(const FDynamicMesh3& Mesh, FLatticeExecutionInfo ExecutionInfo)
 {
-	VertexEmbeddings.Reset(Mesh.VertexCount());
-	for (const FVector3d& VertexPosition : Mesh.VerticesItr())
+	int MaxVertexID = Mesh.MaxVertexID();
+	VertexEmbeddings.SetNum(MaxVertexID);
+
+	EParallelForFlags ParallelForFlags = ExecutionInfo.bParallel ? EParallelForFlags::None : EParallelForFlags::ForceSingleThread;
+
+	ParallelFor(MaxVertexID,
+				[this, &Mesh, MaxVertexID](int VertexID)
 	{
-		FEmbedding Embedding;
+		if (!Mesh.IsVertex(VertexID))
+		{
+			return;
+		}
+
+		FVector3d VertexPosition = Mesh.GetVertex(VertexID);
+
+		FEmbedding& Embedding = VertexEmbeddings[VertexID];
 		Embedding.CellWeighting = ComputeTrilinearWeights(VertexPosition, Embedding.LatticeCell);
-		VertexEmbeddings.Add(Embedding);
-	}
+	},
+				ParallelForFlags);
 
-	check(VertexEmbeddings.Num() == Mesh.VertexCount());
 }
 
 
-void FFFDLattice::GetDeformedMeshVertexPositions(const TArray<FVector3d>& LatticeControlPoints, 
-												 TArray<FVector3d>& OutVertexPositions, 
+void FFFDLattice::GetDeformedMeshVertexPositions(const TArray<FVector3d>& LatticeControlPoints,
+												 TArray<FVector3d>& OutVertexPositions,
 												 ELatticeInterpolation Interpolation,
+												 FLatticeExecutionInfo ExecutionInfo,
 												 FProgressCancel* Progress) const
 {
-	int NumVertices = VertexEmbeddings.Num();
-	OutVertexPositions.SetNumZeroed(NumVertices);
+	int MaxVertexID = VertexEmbeddings.Num();
+	OutVertexPositions.SetNumZeroed(MaxVertexID);
+
+	EParallelForFlags ParallelForFlags = ExecutionInfo.bParallel ? EParallelForFlags::None : EParallelForFlags::ForceSingleThread;
+	ParallelForFlags |= EParallelForFlags::BackgroundPriority;
 
-	for (int VertexID = 0; VertexID < NumVertices; ++VertexID)
+	bool bCancelled = false;
+
+	check(VertexEmbeddings.Num() == OutVertexPositions.Num());
+
+	auto InterpolationJob = [this, &OutVertexPositions, &LatticeControlPoints, &bCancelled, ExecutionInfo, MaxVertexID, Interpolation, Progress]
+	(int VertexID)
 	{
+		// Every once in a while, check for cancellation
+		if ((VertexID % ExecutionInfo.CancelCheckSize == 0) && Progress && Progress->Cancelled())
+		{
+			bCancelled = true;
+		}
+
+		if (bCancelled || VertexEmbeddings[VertexID].LatticeCell[0] < 0)
+		{
+			return;
+		}
+
 		if (Interpolation == ELatticeInterpolation::Cubic)
 		{
 			OutVertexPositions[VertexID] = InterpolatedPositionCubic(VertexEmbeddings[VertexID], LatticeControlPoints);
@@ -128,13 +161,10 @@ void FFFDLattice::GetDeformedMeshVertexPositions(const TArray<FVector3d>& Lattic
 		{
 			OutVertexPositions[VertexID] = InterpolatedPosition(VertexEmbeddings[VertexID], LatticeControlPoints);
 		}
+	};
+
+	ParallelFor(MaxVertexID, InterpolationJob, ParallelForFlags);
 
-		// Every once in a while, check for cancellation
-		if (Progress && (VertexID % 1000 == 0) && Progress->Cancelled())
-		{
-			return;
-		}
-	}
 }
 
 
diff --git a/Engine/Plugins/Experimental/GeometryProcessing/Source/DynamicMesh/Public/Operations/FFDLattice.h b/Engine/Plugins/Experimental/GeometryProcessing/Source/DynamicMesh/Public/Operations/FFDLattice.h
--- a/Engine/Plugins/Experimental/GeometryProcessing/Source/DynamicMesh/Public/Operations/FFDLattice.h
+++ b/Engine/Plugins/Experimental/GeometryProcessing/Source/DynamicMesh/Public/Operations/FFDLattice.h
@@ -16,6 +16,11 @@ enum class ELatticeInterpolation : uint8
 	Cubic = 1
 };
 
+struct FLatticeExecutionInfo
+{
+	bool bParallel = true;
+	int CancelCheckSize = 100000;		// Number of vertices to process before checking for cancel
+};
 
 /** 
  * Free-form deformation lattice. Initialize it with a mesh and desired resolution, get the initial lattice points out.
@@ -28,8 +33,8 @@ public:
 	/// Linear interpolation information. There should be one of these per mesh vertex.
 	struct FEmbedding
 	{
-		FVector3i LatticeCell;		// Original point is in the hexahedron whose Min corner is LatticeCell
-		FVector3d CellWeighting;	// Original linear interpolation weights of the point in the cell
+		FVector3i LatticeCell = { -1, -1, -1 };		// Original point is in the hexahedron whose Min corner is LatticeCell
+		FVector3d CellWeighting = { 0,0,0 };	// Original linear interpolation weights of the point in the cell
 	};
 
 	/// Create a lattice that fits the given mesh and has the given resolution along each dimension. Also precompute
@@ -41,6 +46,7 @@ public:
 	void GetDeformedMeshVertexPositions(const TArray<FVector3d>& LatticeControlPoints,
 										TArray<FVector3d>& OutVertexPositions, 
 										ELatticeInterpolation Interpolation = ELatticeInterpolation::Linear,
+										FLatticeExecutionInfo ExecutionInfo = FLatticeExecutionInfo(),
 										FProgressCancel* Progress = nullptr) const;
 
 	void GenerateInitialLatticePositions(TArray<FVector3d>& OutLatticePositions) const;
@@ -76,7 +82,7 @@ protected:
 	}
 
 	/// For each vertex in Mesh, compute the lattice cell it resides in and its weighting
-	void ComputeInitialEmbedding(const FDynamicMesh3& Mesh);
+	void ComputeInitialEmbedding(const FDynamicMesh3& Mesh, FLatticeExecutionInfo ExecutionInfo = FLatticeExecutionInfo());
 
 	/// Compute cell index and linear interpolation weights for a given point
 	FVector3d ComputeTrilinearWeights(const FVector3d& Pt, FVector3i& GridCoordinates) const;
diff --git a/Engine/Plugins/Experimental/MeshModelingToolset/Source/ModelingOperators/Private/DeformationOps/LatticeDeformerOp.cpp b/Engine/Plugins/Experimental/MeshModelingToolset/Source/ModelingOperators/Private/DeformationOps/LatticeDeformerOp.cpp
--- a/Engine/Plugins/Experimental/MeshModelingToolset/Source/ModelingOperators/Private/DeformationOps/LatticeDeformerOp.cpp
+++ b/Engine/Plugins/Experimental/MeshModelingToolset/Source/ModelingOperators/Private/DeformationOps/LatticeDeformerOp.cpp
@@ -18,7 +18,9 @@ void FLatticeDeformerOp::CalculateResult(FProgressCancel* Progress)
 	}
 
 	TArray<FVector3d> DeformedPositions;
-	Lattice->GetDeformedMeshVertexPositions(LatticeControlPoints, DeformedPositions, InterpolationType, Progress);
+	FLatticeExecutionInfo ExecutionInfo = FLatticeExecutionInfo();
+	ExecutionInfo.bParallel = true;
+	Lattice->GetDeformedMeshVertexPositions(LatticeControlPoints, DeformedPositions, InterpolationType, ExecutionInfo, Progress);
 
 	check(ResultMesh->VertexCount() == DeformedPositions.Num());
 
