fix CompactInPlace bug for unreferenced vertices

#rnx

diff --git a/Engine/Plugins/Experimental/GeometryProcessing/Source/GeometricObjects/Private/Util/SmallListSet.cpp b/Engine/Plugins/Experimental/GeometryProcessing/Source/GeometricObjects/Private/Util/SmallListSet.cpp
--- a/Engine/Plugins/Experimental/GeometryProcessing/Source/GeometricObjects/Private/Util/SmallListSet.cpp
+++ b/Engine/Plugins/Experimental/GeometryProcessing/Source/GeometricObjects/Private/Util/SmallListSet.cpp
@@ -144,7 +144,6 @@ void FSmallListSet::Move(int32 FromIndex, int32 ToIndex)
 	check(FromIndex >= 0);
 	check(ToIndex >= 0);
 	check(ListHeads[ToIndex] == NullValue);
-	check(ListHeads[FromIndex] != NullValue);
 	ListHeads[ToIndex] = ListHeads[FromIndex];
 	ListHeads[FromIndex] = NullValue;
 }
