change the way generic attributes are handled so that we can preserve generic attributes across mesheditor AppendMesh operations

+ a couple minor helpful additions to DynamicMeshEditor: add a cleanup function for removing unreferenced vertices, and add a parameter for StitchSparselyCorrespondedVertexLoops to optionally reverse its triangle orientations

#rnx

diff --git a/Engine/Plugins/Experimental/GeometryProcessing/Source/DynamicMesh/Private/DynamicMeshEditor.cpp b/Engine/Plugins/Experimental/GeometryProcessing/Source/DynamicMesh/Private/DynamicMeshEditor.cpp
--- a/Engine/Plugins/Experimental/GeometryProcessing/Source/DynamicMesh/Private/DynamicMeshEditor.cpp
+++ b/Engine/Plugins/Experimental/GeometryProcessing/Source/DynamicMesh/Private/DynamicMeshEditor.cpp
@@ -11,17 +11,6 @@
 #include "Selections/MeshConnectedComponents.h"
 
 
-void FMeshIndexMappings::Initialize(FDynamicMesh3* Mesh)
-{
-	if (Mesh->HasAttributes())
-	{
-		FDynamicMeshAttributeSet* Attribs = Mesh->Attributes();
-		UVMaps.SetNum(Attribs->NumUVLayers());
-		NormalMaps.SetNum(Attribs->NumNormalLayers());
-	}
-}
-
-
 
 
 
@@ -44,7 +33,18 @@ void FDynamicMeshEditResult::GetAllTriangles(TArray<int>& TrianglesOut) const
 
 
 
-
+bool FDynamicMeshEditor::RemoveIsolatedVertices()
+{
+	bool bSuccess = true;
+	for (int VID : Mesh->VertexIndicesItr())
+	{
+		if (Mesh->GetVtxEdgeCount(VID) == 0)
+		{
+			bSuccess = Mesh->RemoveVertex(VID, false, false) == EMeshResult::Ok && bSuccess;
+		}
+	}
+	return bSuccess;
+}
 
 
 bool FDynamicMeshEditor::StitchVertexLoopsMinimal(const TArray<int>& Loop1, const TArray<int>& Loop2, FDynamicMeshEditResult& ResultOut)
@@ -180,7 +180,7 @@ bool FDynamicMeshEditor::WeldVertexLoops(const TArray<int32>& Loop1, const TArra
 
 
 
-bool FDynamicMeshEditor::StitchSparselyCorrespondedVertexLoops(const TArray<int>& VertexIDs1, const TArray<int>& MatchedIndices1, const TArray<int>& VertexIDs2, const TArray<int>& MatchedIndices2, FDynamicMeshEditResult& ResultOut)
+bool FDynamicMeshEditor::StitchSparselyCorrespondedVertexLoops(const TArray<int>& VertexIDs1, const TArray<int>& MatchedIndices1, const TArray<int>& VertexIDs2, const TArray<int>& MatchedIndices2, FDynamicMeshEditResult& ResultOut, bool bReverseOrientation)
 {
 	int CorrespondN = MatchedIndices1.Num();
 	if (!ensureMsgf(CorrespondN == MatchedIndices2.Num(), TEXT("FDynamicMeshEditor::StitchSparselyCorrespondedVertices: correspondence arrays are not the same length!")))
@@ -248,6 +248,10 @@ bool FDynamicMeshEditor::StitchSparselyCorrespondedVertexLoops(const TArray<int>
 				LenAlong[1] += NextV.Distance(Vertex[1]);
 				Vertex[1] = NextV;
 			}
+			if (bReverseOrientation)
+			{
+				Swap(Tri.B, Tri.C);
+			}
 			int Tid = Mesh->AppendTriangle(Tri, NewGroupID);
 			if (Tid < 0)
 			{
@@ -1484,6 +1488,15 @@ void FDynamicMeshEditor::AppendMesh(const FDynamicMesh3* AppendMesh,
 				ToMaterialIDs->SetValue(TriangleMap.GetTo(TriID), FromMaterialIDs->GetValue(TriID));
 			}
 		}
+
+		for (const TPair<FName, TUniquePtr<FDynamicMeshAttributeBase>>& AttribPair : AppendMesh->Attributes()->GetAttachedAttributes())
+		{
+			if (Mesh->Attributes()->HasAttachedAttribute(AttribPair.Key))
+			{
+				FDynamicMeshAttributeBase* ToAttrib = Mesh->Attributes()->GetAttachedAttribute(AttribPair.Key);
+				ToAttrib->CopyThroughMapping(AttribPair.Value.Get(), IndexMapsOut);
+			}
+		}
 	}
 }
 
@@ -1653,6 +1666,15 @@ static void AppendAttributes(const FDynamicMesh3* FromMesh, int FromTriangleID,
 		FDynamicMeshMaterialAttribute* ToMaterialIDs = ToMesh->Attributes()->GetMaterialID();
 		ToMaterialIDs->SetValue(ToTriangleID, FromMaterialIDs->GetValue(FromTriangleID));
 	}
+
+	for (const TPair<FName, TUniquePtr<FDynamicMeshAttributeBase>>& AttribPair : FromMesh->Attributes()->GetAttachedAttributes())
+	{
+		if (ToMesh->Attributes()->HasAttachedAttribute(AttribPair.Key))
+		{
+			FDynamicMeshAttributeBase* ToAttrib = ToMesh->Attributes()->GetAttachedAttribute(AttribPair.Key);
+			ToAttrib->CopyThroughMapping(AttribPair.Value.Get(), IndexMaps);
+		}
+	}
 }
 
 
diff --git a/Engine/Plugins/Experimental/GeometryProcessing/Source/DynamicMesh/Private/MeshIndexMappings.cpp b/Engine/Plugins/Experimental/GeometryProcessing/Source/DynamicMesh/Private/MeshIndexMappings.cpp
new file mode 100644
--- /dev/null
+++ b/Engine/Plugins/Experimental/GeometryProcessing/Source/DynamicMesh/Private/MeshIndexMappings.cpp
@@ -0,0 +1,16 @@
+// Copyright Epic Games, Inc. All Rights Reserved.
+
+
+#include "MeshIndexMappings.h"
+#include "DynamicMeshAttributeSet.h"
+
+
+void FMeshIndexMappings::Initialize(FDynamicMesh3* Mesh)
+{
+	if (Mesh->HasAttributes())
+	{
+		FDynamicMeshAttributeSet* Attribs = Mesh->Attributes();
+		UVMaps.SetNum(Attribs->NumUVLayers());
+		NormalMaps.SetNum(Attribs->NumNormalLayers());
+	}
+}
diff --git a/Engine/Plugins/Experimental/GeometryProcessing/Source/DynamicMesh/Public/DynamicAttribute.h b/Engine/Plugins/Experimental/GeometryProcessing/Source/DynamicMesh/Public/DynamicAttribute.h
--- a/Engine/Plugins/Experimental/GeometryProcessing/Source/DynamicMesh/Public/DynamicAttribute.h
+++ b/Engine/Plugins/Experimental/GeometryProcessing/Source/DynamicMesh/Public/DynamicAttribute.h
@@ -3,6 +3,7 @@
 #pragma once
 
 #include "DynamicMesh3.h"
+#include "MeshIndexMappings.h"
 
 template<typename ParentType>
 class TDynamicAttributeBase;
@@ -80,6 +81,20 @@ public:
 	/** Update any held pointer to the parent */
 	virtual void Reparent(ParentType* NewParent) = 0;
 
+	/**
+	  * Copy data from a different attribute to this one, using the mesh index mapping to determine the correspondence 
+	  * @param Source copy attribute data from this source
+	  * @param Mapping the correspondence from Source's parent to this attribute's parent 
+	  * @return true if the copy succeeded, false otherwise (e.g. false if the data from the source attribute was not compatible and the CopyOut failed to copy across)
+	  */
+	virtual bool CopyThroughMapping(const TDynamicAttributeBase* Source, const FMeshIndexMappings& Mapping) = 0;
+
+	/** Generic function to copy data out of an attribute; it's up to the derived class to map RawID to chunks of attribute data */
+	virtual bool CopyOut(int RawID, void* Buffer, int BufferSize) const = 0;
+
+	/** Generic function to copy data in to an attribute; it's up to the derived class to map RawID to chunks of attribute data */
+	virtual bool CopyIn(int RawID, void* Buffer, int BufferSize) = 0;
+
 	virtual void OnNewVertex(int VertexID, bool bInserted)
 	{
 	}
@@ -171,6 +186,11 @@ protected:
 		RegisteredAttributes.Add(Attribute);
 	}
 
+	void UnregisterExternalAttribute(TDynamicAttributeBase<ParentType>* Attribute)
+	{
+		RegisteredAttributes.Remove(Attribute);
+	}
+
 	void ResetRegisteredAttributes()
 	{
 		RegisteredAttributes.Reset();
diff --git a/Engine/Plugins/Experimental/GeometryProcessing/Source/DynamicMesh/Public/DynamicMeshAttributeSet.h b/Engine/Plugins/Experimental/GeometryProcessing/Source/DynamicMesh/Public/DynamicMeshAttributeSet.h
--- a/Engine/Plugins/Experimental/GeometryProcessing/Source/DynamicMesh/Public/DynamicMeshAttributeSet.h
+++ b/Engine/Plugins/Experimental/GeometryProcessing/Source/DynamicMesh/Public/DynamicMeshAttributeSet.h
@@ -61,10 +61,9 @@ public:
 
 		GenericAttributes.Reset();
 		ResetRegisteredAttributes();
-		for (int Idx = 0; Idx < Copy.GenericAttributes.Num(); Idx++)
+		for (const TPair<FName, TUniquePtr<FDynamicMeshAttributeBase>>& AttribPair : Copy.GenericAttributes)
 		{
-			const FDynamicMeshAttributeBase* SourceAttrib = Copy.GenericAttributes[Idx].Get();
-			AttachAttribute(SourceAttrib->MakeCopy(ParentMesh));
+			AttachAttribute(AttribPair.Key, AttribPair.Value->MakeCopy(ParentMesh));
 		}
 
 		// parent mesh is *not* copied!
@@ -116,10 +115,9 @@ public:
 
 		GenericAttributes.Reset();
 		ResetRegisteredAttributes();
-		for (int Idx = 0; Idx < Copy.GenericAttributes.Num(); Idx++)
+		for (const TPair<FName, TUniquePtr<FDynamicMeshAttributeBase>>& AttribPair : Copy.GenericAttributes)
 		{
-			const FDynamicMeshAttributeBase* SourceAttrib = Copy.GenericAttributes[Idx].Get();
-			AttachAttribute(SourceAttrib->MakeCompactCopy(CompactMaps, ParentMesh));
+			AttachAttribute(AttribPair.Key, AttribPair.Value->MakeCompactCopy(CompactMaps, ParentMesh));
 		}
 
 		// parent mesh is *not* copied!
@@ -144,9 +142,9 @@ public:
 			MaterialIDAttrib->CompactInPlace(CompactMaps);
 		}
 
-		for (int Idx = 0; Idx < GenericAttributes.Num(); Idx++)
+		for (const TPair<FName, TUniquePtr<FDynamicMeshAttributeBase>>& AttribPair : GenericAttributes)
 		{
-			GenericAttributes[Idx]->CompactInPlace(CompactMaps);
+			AttribPair.Value->CompactInPlace(CompactMaps);
 		}
 	}
 
@@ -174,10 +172,9 @@ public:
 
 		GenericAttributes.Reset();
 		ResetRegisteredAttributes();
-		for (int Idx = 0; Idx < ToMatch.GenericAttributes.Num(); Idx++)
+		for (const TPair<FName, TUniquePtr<FDynamicMeshAttributeBase>>& AttribPair : ToMatch.GenericAttributes)
 		{
-			const FDynamicMeshAttributeBase* SourceAttrib = ToMatch.GenericAttributes[Idx].Get();
-			AttachAttribute(SourceAttrib->MakeNew(ParentMesh));
+			AttachAttribute(AttribPair.Key, AttribPair.Value->MakeNew(ParentMesh));
 		}
 	}
 
@@ -203,9 +200,9 @@ private:
 			MaterialIDAttrib->Reparent( NewParent );
 		}
 
-		for (int Idx = 0; Idx < GenericAttributes.Num(); Idx++)
+		for (const TPair<FName, TUniquePtr<FDynamicMeshAttributeBase>>& AttribPair : GenericAttributes)
 		{
-			GenericAttributes[Idx]->Reparent( NewParent );
+			AttribPair.Value->Reparent( NewParent );
 		}
 	}
 
@@ -340,23 +337,36 @@ public:
 	}
 
 	// Attach a new attribute (and transfer ownership of it to the attribute set)
-	int AttachAttribute(FDynamicMeshAttributeBase* Attribute)
+	void AttachAttribute(FName AttribName, FDynamicMeshAttributeBase* Attribute)
 	{
-		int AttributeID = GenericAttributes.Add(TUniquePtr<FDynamicMeshAttributeBase>(Attribute));
+		if (GenericAttributes.Contains(AttribName))
+		{
+			UnregisterExternalAttribute(GenericAttributes[AttribName].Get());
+		}
+		GenericAttributes.Add(AttribName, TUniquePtr<FDynamicMeshAttributeBase>(Attribute));
 		RegisterExternalAttribute(Attribute);
-		return AttributeID;
 	}
 
-	FDynamicMeshAttributeBase* GetAttachedAttribute(int AttributeID)
+	FDynamicMeshAttributeBase* GetAttachedAttribute(FName AttribName) const
 	{
-		return GenericAttributes[AttributeID].Get();
+		return GenericAttributes[AttribName].Get();
 	}
 
-	int NumAttachedAttributes()
+	int NumAttachedAttributes() const
 	{
 		return GenericAttributes.Num();
 	}
 
+	bool HasAttachedAttribute(FName AttribName) const
+	{
+		return GenericAttributes.Contains(AttribName);
+	}
+
+	const TMap<FName, TUniquePtr<FDynamicMeshAttributeBase>>& GetAttachedAttributes() const
+	{
+		return GenericAttributes;
+	}
+
 protected:
 	/** Parent mesh of this attribute set */
 	FDynamicMesh3* ParentMesh;
@@ -369,7 +379,7 @@ protected:
 
 	TUniquePtr<FDynamicMeshMaterialAttribute> MaterialIDAttrib;
 
-	TArray<TUniquePtr<FDynamicMeshAttributeBase>> GenericAttributes;
+	TMap<FName, TUniquePtr<FDynamicMeshAttributeBase>> GenericAttributes;
 
 protected:
 	friend class FDynamicMesh3;
diff --git a/Engine/Plugins/Experimental/GeometryProcessing/Source/DynamicMesh/Public/DynamicMeshEditor.h b/Engine/Plugins/Experimental/GeometryProcessing/Source/DynamicMesh/Public/DynamicMeshEditor.h
--- a/Engine/Plugins/Experimental/GeometryProcessing/Source/DynamicMesh/Public/DynamicMeshEditor.h
+++ b/Engine/Plugins/Experimental/GeometryProcessing/Source/DynamicMesh/Public/DynamicMeshEditor.h
@@ -8,81 +8,10 @@
 #include "DynamicMeshAttributeSet.h"
 #include "EdgeLoop.h"
 #include "Util/SparseIndexCollectionTypes.h"
+#include "MeshIndexMappings.h"
 
 struct FDynamicSubmesh3;
 
-/**
- * FMeshIndexMappings stores a set of integer IndexMaps for a mesh
- * This is a convenient object to have, to avoid passing around large numbers of separate maps.
- * The individual maps are not necessarily all filled by every operation.
- */
-struct DYNAMICMESH_API FMeshIndexMappings
-{
-protected:
-	FIndexMapi VertexMap;
-	FIndexMapi TriangleMap;
-	FIndexMapi GroupMap;
-
-	TArray<FIndexMapi> UVMaps;
-	TArray<FIndexMapi> NormalMaps;
-
-public:
-	/** Size internal arrays-of-maps to be suitable for this Mesh */
-	void Initialize(FDynamicMesh3* Mesh);
-
-	/** @return the value used to indicate "invalid" in the mapping */
-	constexpr int InvalidID() const { return VertexMap.UnmappedID(); }
-
-	void Reset()
-	{
-		VertexMap.Reset();
-		TriangleMap.Reset();
-		GroupMap.Reset();
-		for (FIndexMapi& UVMap : UVMaps)
-		{
-			UVMap.Reset();
-		}
-		for (FIndexMapi& NormalMap : NormalMaps)
-		{
-			NormalMap.Reset();
-		}
-	}
-
-	void ResetTriangleMap()
-	{
-		TriangleMap.Reset();
-	}
-
-	FIndexMapi& GetVertexMap() { return VertexMap; }
-	const FIndexMapi& GetVertexMap() const { return VertexMap; }
-	inline void SetVertex(int FromID, int ToID) { VertexMap.Add(FromID, ToID); }
-	inline int GetNewVertex(int FromID) const { return VertexMap.GetTo(FromID); }
-	inline bool ContainsVertex(int FromID) const { return VertexMap.ContainsFrom(FromID); }
-
-	FIndexMapi& GetTriangleMap() { return TriangleMap; }
-	const FIndexMapi& GetTriangleMap() const { return TriangleMap; }
-	void SetTriangle(int FromID, int ToID) { TriangleMap.Add(FromID, ToID); }
-	int GetNewTriangle(int FromID) const { return TriangleMap.GetTo(FromID); }
-	inline bool ContainsTriangle(int FromID) const { return TriangleMap.ContainsFrom(FromID); }
-
-	FIndexMapi& GetGroupMap() { return GroupMap; }
-	void SetGroup(int FromID, int ToID) { GroupMap.Add(FromID, ToID); }
-	int GetNewGroup(int FromID) const { return GroupMap.GetTo(FromID); }
-	inline bool ContainsGroup(int FromID) const { return GroupMap.ContainsFrom(FromID); }
-
-	FIndexMapi& GetUVMap(int UVLayer) { return UVMaps[UVLayer]; }
-	void SetUV(int UVLayer, int FromID, int ToID) { UVMaps[UVLayer].Add(FromID, ToID); }
-	int GetNewUV(int UVLayer, int FromID) const { return UVMaps[UVLayer].GetTo(FromID); }
-	inline bool ContainsUV(int UVLayer, int FromID) const { return UVMaps[UVLayer].ContainsFrom(FromID); }
-
-	FIndexMapi& GetNormalMap(int NormalLayer) { return NormalMaps[NormalLayer]; }
-	void SetNormal(int NormalLayer, int FromID, int ToID) { NormalMaps[NormalLayer].Add(FromID, ToID); }
-	int GetNewNormal(int NormalLayer, int FromID) const { return NormalMaps[NormalLayer].GetTo(FromID); }
-	inline bool ContainsNormal(int NormalLayer, int FromID) const { return NormalMaps[NormalLayer].ContainsFrom(FromID); }
-
-};
-
-
 /**
  * FDynamicMeshEditResult is used to return information about new mesh elements
  * created by mesh changes, primarily in FDynamicMeshEditor
@@ -139,6 +68,10 @@ public:
 		Mesh = MeshIn;
 	}
 
+	/**
+	 * Remove any vertices that are not used by any triangles
+	 */
+	bool RemoveIsolatedVertices();
 
 	//////////////////////////////////////////////////////////////////////////
 	// Create and Remove Triangle Functions
@@ -165,7 +98,7 @@ public:
 	 * @param ResultOut lists of newly created triangles/vertices/etc
 	 * @return true if operation succeeded.  If a failure occurs, any added triangles are removed via RemoveTriangles
 	 */
-	bool StitchSparselyCorrespondedVertexLoops(const TArray<int>& VertexIDs1, const TArray<int>& MatchedIndices1, const TArray<int>& VertexIDs2, const TArray<int>& MatchedIndices2, FDynamicMeshEditResult& ResultOut);
+	bool StitchSparselyCorrespondedVertexLoops(const TArray<int>& VertexIDs1, const TArray<int>& MatchedIndices1, const TArray<int>& VertexIDs2, const TArray<int>& MatchedIndices2, FDynamicMeshEditResult& ResultOut, bool bReverseOrientation = false);
 
 
 	/**
diff --git a/Engine/Plugins/Experimental/GeometryProcessing/Source/DynamicMesh/Public/DynamicMeshTriangleAttribute.h b/Engine/Plugins/Experimental/GeometryProcessing/Source/DynamicMesh/Public/DynamicMeshTriangleAttribute.h
--- a/Engine/Plugins/Experimental/GeometryProcessing/Source/DynamicMesh/Public/DynamicMeshTriangleAttribute.h
+++ b/Engine/Plugins/Experimental/GeometryProcessing/Source/DynamicMesh/Public/DynamicMeshTriangleAttribute.h
@@ -168,6 +168,49 @@ public:
 	// Accessors/Queries
 	//  
 
+	virtual bool CopyThroughMapping(const TDynamicAttributeBase<FDynamicMesh3>* Source, const FMeshIndexMappings& Mapping) override
+	{
+		AttribValueType BufferData[AttribDimension];
+		int BufferSize = sizeof(BufferData);
+		for (const TPair<int32, int32>& MapTID : Mapping.GetTriangleMap().GetForwardMap())
+		{
+			if (!ensure(Source->CopyOut(MapTID.Key, BufferData, BufferSize)))
+			{
+				return false;
+			}
+			SetValue(MapTID.Value, BufferData);
+		}
+		return true;
+	}
+	virtual bool CopyOut(int RawID, void* Buffer, int BufferSize) const override
+	{
+		if (sizeof(AttribValueType) * AttribDimension != BufferSize)
+		{
+			return false;
+		}
+		AttribValueType* BufferData = static_cast<AttribValueType*>(Buffer);
+		int k = RawID * AttribDimension;
+		for (int i = 0; i < AttribDimension; ++i)
+		{
+			BufferData[i] = AttribValues[k + i];
+		}
+		return true;
+	}
+	virtual bool CopyIn(int RawID, void* Buffer, int BufferSize) override
+	{
+		if (sizeof(AttribValueType) * AttribDimension != BufferSize)
+		{
+			return false;
+		}
+		AttribValueType* BufferData = static_cast<AttribValueType*>(Buffer);
+		int k = RawID * AttribDimension;
+		for (int i = 0; i < AttribDimension; ++i)
+		{
+			AttribValues[k + i] = BufferData[i];
+		}
+		return true;
+	}
+
 	/** Get the element at a given index */
 	inline void GetValue(int TriangleID, AttribValueType* Data) const
 	{
diff --git a/Engine/Plugins/Experimental/GeometryProcessing/Source/DynamicMesh/Public/DynamicVertexAttribute.h b/Engine/Plugins/Experimental/GeometryProcessing/Source/DynamicMesh/Public/DynamicVertexAttribute.h
--- a/Engine/Plugins/Experimental/GeometryProcessing/Source/DynamicMesh/Public/DynamicVertexAttribute.h
+++ b/Engine/Plugins/Experimental/GeometryProcessing/Source/DynamicMesh/Public/DynamicVertexAttribute.h
@@ -156,7 +156,50 @@ public:
 
 	//
 	// Accessors/Queries
-	//  
+	//
+
+	virtual bool CopyThroughMapping(const TDynamicAttributeBase<ParentType>* Source, const FMeshIndexMappings& Mapping) override
+	{
+		AttribValueType BufferData[AttribDimension];
+		int BufferSize = sizeof(BufferData);
+		for (const TPair<int32, int32>& MapVID : Mapping.GetVertexMap().GetForwardMap())
+		{
+			if (!ensure(Source->CopyOut(MapVID.Key, BufferData, BufferSize)))
+			{
+				return false;
+			}
+			SetValue(MapVID.Value, BufferData);
+		}
+		return true;
+	}
+	virtual bool CopyOut(int RawID, void* Buffer, int BufferSize) const override
+	{
+		if (sizeof(AttribValueType)*AttribDimension != BufferSize)
+		{
+			return false;
+		}
+		AttribValueType* BufferData = static_cast<AttribValueType*>(Buffer);
+		int k = RawID * AttribDimension;
+		for (int i = 0; i < AttribDimension; ++i)
+		{
+			BufferData[i] = AttribValues[k + i];
+		}
+		return true;
+	}
+	virtual bool CopyIn(int RawID, void* Buffer, int BufferSize) override
+	{
+		if (sizeof(AttribValueType) * AttribDimension != BufferSize)
+		{
+			return false;
+		}
+		AttribValueType* BufferData = static_cast<AttribValueType*>(Buffer);
+		int k = RawID * AttribDimension;
+		for (int i = 0; i < AttribDimension; ++i)
+		{
+			AttribValues[k + i] = BufferData[i];
+		}
+		return true;
+	}
 
 	/** Get the element at a given index */
 	inline void GetValue(int VertexID, AttribValueType* Data) const
diff --git a/Engine/Plugins/Experimental/GeometryProcessing/Source/DynamicMesh/Public/MeshIndexMappings.h b/Engine/Plugins/Experimental/GeometryProcessing/Source/DynamicMesh/Public/MeshIndexMappings.h
new file mode 100644
--- /dev/null
+++ b/Engine/Plugins/Experimental/GeometryProcessing/Source/DynamicMesh/Public/MeshIndexMappings.h
@@ -0,0 +1,76 @@
+// Copyright Epic Games, Inc. All Rights Reserved.
+
+#pragma once
+
+#include "DynamicMesh3.h"
+
+/**
+ * FMeshIndexMappings stores a set of integer IndexMaps for a mesh
+ * This is a convenient object to have, to avoid passing around large numbers of separate maps.
+ * The individual maps are not necessarily all filled by every operation.
+ */
+struct DYNAMICMESH_API FMeshIndexMappings
+{
+protected:
+	FIndexMapi VertexMap;
+	FIndexMapi TriangleMap;
+	FIndexMapi GroupMap;
+
+	TArray<FIndexMapi> UVMaps;
+	TArray<FIndexMapi> NormalMaps;
+
+public:
+	/** Size internal arrays-of-maps to be suitable for this Mesh */
+	void Initialize(FDynamicMesh3* Mesh);
+
+	/** @return the value used to indicate "invalid" in the mapping */
+	constexpr int InvalidID() const { return VertexMap.UnmappedID(); }
+
+	void Reset()
+	{
+		VertexMap.Reset();
+		TriangleMap.Reset();
+		GroupMap.Reset();
+		for (FIndexMapi& UVMap : UVMaps)
+		{
+			UVMap.Reset();
+		}
+		for (FIndexMapi& NormalMap : NormalMaps)
+		{
+			NormalMap.Reset();
+		}
+	}
+
+	void ResetTriangleMap()
+	{
+		TriangleMap.Reset();
+	}
+
+	FIndexMapi& GetVertexMap() { return VertexMap; }
+	const FIndexMapi& GetVertexMap() const { return VertexMap; }
+	inline void SetVertex(int FromID, int ToID) { VertexMap.Add(FromID, ToID); }
+	inline int GetNewVertex(int FromID) const { return VertexMap.GetTo(FromID); }
+	inline bool ContainsVertex(int FromID) const { return VertexMap.ContainsFrom(FromID); }
+
+	FIndexMapi& GetTriangleMap() { return TriangleMap; }
+	const FIndexMapi& GetTriangleMap() const { return TriangleMap; }
+	void SetTriangle(int FromID, int ToID) { TriangleMap.Add(FromID, ToID); }
+	int GetNewTriangle(int FromID) const { return TriangleMap.GetTo(FromID); }
+	inline bool ContainsTriangle(int FromID) const { return TriangleMap.ContainsFrom(FromID); }
+
+	FIndexMapi& GetGroupMap() { return GroupMap; }
+	void SetGroup(int FromID, int ToID) { GroupMap.Add(FromID, ToID); }
+	int GetNewGroup(int FromID) const { return GroupMap.GetTo(FromID); }
+	inline bool ContainsGroup(int FromID) const { return GroupMap.ContainsFrom(FromID); }
+
+	FIndexMapi& GetUVMap(int UVLayer) { return UVMaps[UVLayer]; }
+	void SetUV(int UVLayer, int FromID, int ToID) { UVMaps[UVLayer].Add(FromID, ToID); }
+	int GetNewUV(int UVLayer, int FromID) const { return UVMaps[UVLayer].GetTo(FromID); }
+	inline bool ContainsUV(int UVLayer, int FromID) const { return UVMaps[UVLayer].ContainsFrom(FromID); }
+
+	FIndexMapi& GetNormalMap(int NormalLayer) { return NormalMaps[NormalLayer]; }
+	void SetNormal(int NormalLayer, int FromID, int ToID) { NormalMaps[NormalLayer].Add(FromID, ToID); }
+	int GetNewNormal(int NormalLayer, int FromID) const { return NormalMaps[NormalLayer].GetTo(FromID); }
+	inline bool ContainsNormal(int NormalLayer, int FromID) const { return NormalMaps[NormalLayer].ContainsFrom(FromID); }
+
+};
diff --git a/Engine/Plugins/Experimental/MeshModelingToolset/Source/MeshModelingTools/Private/PlaneCutTool.cpp b/Engine/Plugins/Experimental/MeshModelingToolset/Source/MeshModelingTools/Private/PlaneCutTool.cpp
--- a/Engine/Plugins/Experimental/MeshModelingToolset/Source/MeshModelingTools/Private/PlaneCutTool.cpp
+++ b/Engine/Plugins/Experimental/MeshModelingToolset/Source/MeshModelingTools/Private/PlaneCutTool.cpp
@@ -140,12 +140,10 @@ void UPlaneCutTool::Setup()
 		OriginalDynamicMesh->EnableAttributes();
 		TDynamicMeshScalarTriangleAttribute<int>* SubObjectIDs = new TDynamicMeshScalarTriangleAttribute<int>(OriginalDynamicMesh);
 		SubObjectIDs->Initialize(0);
-		int AttribIndex = OriginalDynamicMesh->Attributes()->AttachAttribute(SubObjectIDs);
+		OriginalDynamicMesh->Attributes()->AttachAttribute(FPlaneCutOp::ObjectIndexAttribute, SubObjectIDs);
 
 		/// fill in the MeshesToCut array
 		UDynamicMeshReplacementChangeTarget* Target = MeshesToCut.Add_GetRef(NewObject<UDynamicMeshReplacementChangeTarget>());
-		MeshSubObjectAttribIndices.Add(AttribIndex);
-		check(MeshSubObjectAttribIndices.Num() == MeshesToCut.Num());
 		// store a UV scale based on the original mesh bounds (we don't want to recompute this between cuts b/c we want consistent UV scale)
 		MeshUVScaleFactor.Add(1.0 / OriginalDynamicMesh->GetBounds().MaxDim());
 
@@ -369,7 +367,6 @@ TUniquePtr<FDynamicMeshOperator> UPlaneCutOperatorFactory::MakeNewOperator()
 	CutOp->bKeepBothHalves = CutTool->BasicProperties->bKeepBothHalves;
 	CutOp->CutPlaneLocalThickness = CutTool->BasicProperties->SpacingBetweenHalves * NormalScaleFactor;
 	CutOp->UVScaleFactor = CutTool->MeshUVScaleFactor[ComponentIndex];
-	CutOp->SubObjectsAttribIndex = CutTool->MeshSubObjectAttribIndices[ComponentIndex];
 	
 
 	return CutOp;
@@ -555,7 +552,7 @@ void UPlaneCutTool::GenerateAsset(const TArray<FDynamicMeshOpResult>& Results)
 		{
 			TDynamicMeshScalarTriangleAttribute<int>* SubMeshIDs =
 				static_cast<TDynamicMeshScalarTriangleAttribute<int>*>(UseMesh->Attributes()->GetAttachedAttribute(
-					MeshSubObjectAttribIndices[OrigMeshIdx]));
+					FPlaneCutOp::ObjectIndexAttribute));
 			TArray<FDynamicMesh3>& SplitMeshes = AllSplitMeshes[OrigMeshIdx];
 			bool bWasSplit = FDynamicMeshEditor::SplitMesh(UseMesh, SplitMeshes, [SubMeshIDs](int TID)
 			{
diff --git a/Engine/Plugins/Experimental/MeshModelingToolset/Source/MeshModelingTools/Public/PlaneCutTool.h b/Engine/Plugins/Experimental/MeshModelingToolset/Source/MeshModelingTools/Public/PlaneCutTool.h
--- a/Engine/Plugins/Experimental/MeshModelingToolset/Source/MeshModelingTools/Public/PlaneCutTool.h
+++ b/Engine/Plugins/Experimental/MeshModelingToolset/Source/MeshModelingTools/Public/PlaneCutTool.h
@@ -178,8 +178,6 @@ protected:
 	UPROPERTY()
 	TArray<UDynamicMeshReplacementChangeTarget*> MeshesToCut;
 
-	// for each mesh in MeshesToCut, the index of the attached generic triangle attribute tracking the object index
-	TArray<int> MeshSubObjectAttribIndices;
 	// UV Scale factor is cached based on the bounding box of the mesh before any cuts are performed, so you don't get inconsistent UVs if you multi-cut the object to smaller sizes
 	TArray<float> MeshUVScaleFactor;
 
diff --git a/Engine/Plugins/Experimental/MeshModelingToolset/Source/ModelingOperators/Private/CuttingOps/PlaneCutOp.cpp b/Engine/Plugins/Experimental/MeshModelingToolset/Source/ModelingOperators/Private/CuttingOps/PlaneCutOp.cpp
--- a/Engine/Plugins/Experimental/MeshModelingToolset/Source/ModelingOperators/Private/CuttingOps/PlaneCutOp.cpp
+++ b/Engine/Plugins/Experimental/MeshModelingToolset/Source/ModelingOperators/Private/CuttingOps/PlaneCutOp.cpp
@@ -10,6 +10,7 @@
 #include "Operations/MeshPlaneCut.h"
 #include "ConstrainedDelaunay2.h"
 
+const FName FPlaneCutOp::ObjectIndexAttribute = "ObjectIndexAttribute";
 
 void FPlaneCutOp::SetTransform(const FTransform& Transform) {
 	ResultTransform = (FTransform3d)Transform;
@@ -41,7 +42,7 @@ void FPlaneCutOp::CalculateResult(FProgressCancel* Progress)
 	if (bKeepBothHalves)
 	{
 		int MaxSubObjectID = -1;
-		TDynamicMeshScalarTriangleAttribute<int>* SubObjectAttrib = static_cast<TDynamicMeshScalarTriangleAttribute<int>*>(ResultMesh->Attributes()->GetAttachedAttribute(SubObjectsAttribIndex));
+		TDynamicMeshScalarTriangleAttribute<int>* SubObjectAttrib = static_cast<TDynamicMeshScalarTriangleAttribute<int>*>(ResultMesh->Attributes()->GetAttachedAttribute(ObjectIndexAttribute));
 		for (int TID : ResultMesh->TriangleIndicesItr())
 		{
 			MaxSubObjectID = FMath::Max(MaxSubObjectID, SubObjectAttrib->GetValue(TID));
@@ -96,7 +97,7 @@ void FPlaneCutOp::CalculateResult(FProgressCancel* Progress)
 	}
 	if (bFillCutHole && bKeepBothHalves)
 	{
-		TDynamicMeshScalarTriangleAttribute<int>* SubObjectAttrib = static_cast<TDynamicMeshScalarTriangleAttribute<int>*>(ResultMesh->Attributes()->GetAttachedAttribute(SubObjectsAttribIndex));
+		TDynamicMeshScalarTriangleAttribute<int>* SubObjectAttrib = static_cast<TDynamicMeshScalarTriangleAttribute<int>*>(ResultMesh->Attributes()->GetAttachedAttribute(ObjectIndexAttribute));
 		Cut.TransferTriangleLabelsToHoleFillTriangles(SubObjectAttrib);
 	}
 }
\ No newline at end of file
diff --git a/Engine/Plugins/Experimental/MeshModelingToolset/Source/ModelingOperators/Public/CuttingOps/PlaneCutOp.h b/Engine/Plugins/Experimental/MeshModelingToolset/Source/ModelingOperators/Public/CuttingOps/PlaneCutOp.h
--- a/Engine/Plugins/Experimental/MeshModelingToolset/Source/ModelingOperators/Public/CuttingOps/PlaneCutOp.h
+++ b/Engine/Plugins/Experimental/MeshModelingToolset/Source/ModelingOperators/Public/CuttingOps/PlaneCutOp.h
@@ -21,7 +21,7 @@ public:
 	bool bKeepBothHalves;
 	float CutPlaneLocalThickness = 0; // plane thickness in the local space of the mesh
 	float UVScaleFactor;
-	int SubObjectsAttribIndex;
+	static const FName ObjectIndexAttribute;
 	TSharedPtr<const FDynamicMesh3> OriginalMesh;
 
 	void SetTransform(const FTransform& Transform);
