MeshModelingTools: Modifying converter to keep polygroup ID's the same when going back and forth between DynamicMesh and MeshDescription.

#rnx

diff --git a/Engine/Plugins/Experimental/GeometryProcessing/Source/MeshConversion/Private/MeshDescriptionToDynamicMesh.cpp b/Engine/Plugins/Experimental/GeometryProcessing/Source/MeshConversion/Private/MeshDescriptionToDynamicMesh.cpp
--- a/Engine/Plugins/Experimental/GeometryProcessing/Source/MeshConversion/Private/MeshDescriptionToDynamicMesh.cpp
+++ b/Engine/Plugins/Experimental/GeometryProcessing/Source/MeshConversion/Private/MeshDescriptionToDynamicMesh.cpp
@@ -195,22 +195,23 @@ void FMeshDescriptionToDynamicMesh::Convert(const FMeshDescription* MeshIn, FDyn
 			TriData.TriInstances[2] = InstanceTri[2];
 
 			int GroupID = 0;
-			if (GroupMode == EPrimaryGroupMode::SetToPolyGroup)
+			switch (GroupMode)
 			{
+			case EPrimaryGroupMode::SetToPolyGroup:
 				if (PolyGroups.IsValid())
 				{
 					GroupID = PolyGroups.Get(PolygonID, 0);
 				}
+				break;
+			case EPrimaryGroupMode::SetToPolygonID:
+				GroupID = PolygonID.GetValue() + 1; // Shift IDs up by 1 to leave ID 0 as a default/unassigned group
+				break;
+			case EPrimaryGroupMode::SetToPolygonGroupID:
+				GroupID = PolygonGroupID + 1; // Shift IDs up by 1 to leave ID 0 as a default/unassigned group
+				break;
+			case EPrimaryGroupMode::SetToZero:
+				break; // keep at 0
 			}
-			else if (GroupMode == EPrimaryGroupMode::SetToPolygonID)
-			{
-				GroupID = PolygonID.GetValue();
-			}
-			else if (GroupMode == EPrimaryGroupMode::SetToPolygonGroupID)
-			{
-				GroupID = PolygonGroupID;
-			}
-			GroupID++; // Shift IDs up by 1 to leave ID 0 as a default/unassigned group
 
 			// append triangle
 			int32 VertexID0 = MeshIn->GetVertexInstanceVertex(InstanceTri[0]).GetValue();
