add a GetCorners() to the 2D AABB (similar to the one for 3D AABBs)

#rnx

diff --git a/Engine/Plugins/Experimental/GeometryProcessing/Source/GeometricObjects/Public/BoxTypes.h b/Engine/Plugins/Experimental/GeometryProcessing/Source/GeometricObjects/Public/BoxTypes.h
--- a/Engine/Plugins/Experimental/GeometryProcessing/Source/GeometricObjects/Public/BoxTypes.h
+++ b/Engine/Plugins/Experimental/GeometryProcessing/Source/GeometricObjects/Public/BoxTypes.h
@@ -565,6 +565,19 @@ struct TAxisAlignedBox2
 		return (Max - Min) * (RealType).5;
 	}
 
+	/**
+	 * Corners are ordered to follow the perimeter of the bounding rectangle, starting from the (Min.X, Min.Y) corner and ending at (Min.X, Max.Y)
+	 * @param Index which corner to return, must be in range [0,3]
+	 * @return Corner of the bounding rectangle
+	 */
+	FVector2<RealType> GetCorner(int Index) const
+	{
+		check(Index >= 0 && Index <= 3);
+		RealType X = ((Index % 3) == 0) ? (Min.X) : (Max.X);
+		RealType Y = ((Index & 2) == 0) ? (Min.Y) : (Max.Y);
+		return FVector2<RealType>(X, Y);
+	}
+
 	inline void Contain(const FVector2<RealType>& V)
 	{
 		if (V.X < Min.X)
