fix dynamic mesh's compact in place handling of parent vertices for overlays -- all elements' parents should be remapped, not just the ones that are moved
 + change the corresponding test to cover this
 + also fix the generic attribute test to use the updated api

#rnx

diff --git a/Engine/Plugins/Experimental/GeometryProcessing/Source/DynamicMesh/Public/DynamicMeshOverlay.h b/Engine/Plugins/Experimental/GeometryProcessing/Source/DynamicMesh/Public/DynamicMeshOverlay.h
--- a/Engine/Plugins/Experimental/GeometryProcessing/Source/DynamicMesh/Public/DynamicMeshOverlay.h
+++ b/Engine/Plugins/Experimental/GeometryProcessing/Source/DynamicMesh/Public/DynamicMeshOverlay.h
@@ -141,6 +141,11 @@ public:
 		{
 			// mapping is 1:1 by default; sparsely re-mapped below
 			MapE[ID] = ID;
+			// remap all parents
+			if (ParentVertices[ID] >= 0)
+			{
+				ParentVertices[ID] = CompactMaps.GetVertex(ParentVertices[ID]);
+			}
 		}
 
 		TDynamicVector<short>& ERef = ElementsRefCounts.GetRawRefCountsUnsafe();
@@ -150,15 +155,7 @@ public:
 			// remap the element data
 			GetElement(iLastE, Data);
 			SetElement(iCurE, Data);
-			int OrigParent = ParentVertices[iLastE];
-			if (OrigParent == IndexConstants::InvalidID)
-			{
-				ParentVertices[iCurE] = IndexConstants::InvalidID;
-			}
-			else 
-			{
-				ParentVertices[iCurE] = CompactMaps.GetVertex(OrigParent);
-			}
+			ParentVertices[iCurE] = ParentVertices[iLastE];
 			ERef[iCurE] = ERef[iLastE];
 			ERef[iLastE] = FRefCountVector::INVALID_REF_COUNT;
 			MapE[iLastE] = iCurE;
