optimizations and an option to skip edge welding for mesh booleans
+ fix include and link issues with Sparse Dynamic Octree

#rnx

diff --git a/Engine/Plugins/Experimental/GeometryProcessing/Source/DynamicMesh/Private/Operations/MeshBoolean.cpp b/Engine/Plugins/Experimental/GeometryProcessing/Source/DynamicMesh/Private/Operations/MeshBoolean.cpp
--- a/Engine/Plugins/Experimental/GeometryProcessing/Source/DynamicMesh/Private/Operations/MeshBoolean.cpp
+++ b/Engine/Plugins/Experimental/GeometryProcessing/Source/DynamicMesh/Private/Operations/MeshBoolean.cpp
@@ -7,6 +7,7 @@
 
 #include "Async/ParallelFor.h"
 #include "MeshTransforms.h"
+#include "Spatial/SparseDynamicOctree3.h"
 
 #include "Algo/RemoveIf.h"
 
@@ -205,7 +206,10 @@ bool FMeshBoolean::Compute()
 {
 	// copy meshes
 	FDynamicMesh3 CutMeshB(*Meshes[1]);
-	*Result = *Meshes[0];
+	if (Result != Meshes[0])
+	{
+		*Result = *Meshes[0];
+	}
 	FDynamicMesh3* CutMesh[2]{ Result, &CutMeshB }; // just an alias to keep things organized
 
 	// transform the copies to a shared space (centered at the origin)
@@ -477,6 +481,44 @@ bool FMeshBoolean::Compute()
 			int OtherMeshIdx = 1 - MeshIdx;
 			FDynamicMesh3& OtherMesh = *CutMesh[OtherMeshIdx];
 
+
+			FSparseDynamicOctree3 EdgeOctree;
+			EdgeOctree.RootDimension = .25;
+			EdgeOctree.SetMaxTreeDepth(7);
+			auto EdgeBounds = [&OtherMesh](int EID)
+			{
+				FDynamicMesh3::FEdge Edge = OtherMesh.GetEdge(EID);
+				FVector3d A = OtherMesh.GetVertex(Edge.Vert.A);
+				FVector3d B = OtherMesh.GetVertex(Edge.Vert.B);
+				if (A.X > B.X)
+				{
+					Swap(A.X, B.X);
+				}
+				if (A.Y > B.Y)
+				{
+					Swap(A.Y, B.Y);
+				}
+				if (A.Z > B.Z)
+				{
+					Swap(A.Z, B.Z);
+				}
+				return FAxisAlignedBox3d(A, B);
+			};
+			auto AddEdge = [&EdgeOctree, &OtherMesh, EdgeBounds](int EID)
+			{
+				EdgeOctree.InsertObject(EID, EdgeBounds(EID));
+			};
+			auto UpdateEdge = [&EdgeOctree, &OtherMesh, EdgeBounds](int EID)
+			{
+				EdgeOctree.ReinsertObject(EID, EdgeBounds(EID));
+			};
+			for (int EID : CutBoundaryEdges[OtherMeshIdx])
+			{
+				AddEdge(EID);
+			}
+			TArray<int> EdgesInRange;
+
+
 			// mapping from OtherMesh VIDs to ProcessMesh VIDs
 			// used to ensure we only keep the best match, in cases where multiple boundary vertices map to a given vertex on the other mesh boundary
 			TMap<int, int> FoundMatches;
@@ -521,7 +563,11 @@ bool FMeshBoolean::Compute()
 				if (NearestVID == FDynamicMesh3::InvalidID)
 				{
 					// vertex had no match -- try to split edge to match it
-					int OtherEID = FindNearestEdge(OtherMesh, CutBoundaryEdges[OtherMeshIdx], Pos);
+					FAxisAlignedBox3d QueryBox(Pos, SnapTolerance);
+					EdgesInRange.Reset();
+					EdgeOctree.RangeQuery(QueryBox, EdgesInRange);
+
+					int OtherEID = FindNearestEdge(OtherMesh, EdgesInRange, Pos);
 					if (OtherEID != FDynamicMesh3::InvalidID)
 					{
 						FVector3d EdgePts[2];
@@ -537,6 +583,8 @@ bool FMeshBoolean::Compute()
 								FoundMatches.Add(SplitInfo.NewVertex, BoundaryVID);
 								OtherMesh.SetVertex(SplitInfo.NewVertex, Pos);
 								CutBoundaryEdges[OtherMeshIdx].Add(SplitInfo.NewEdges.A);
+								UpdateEdge(OtherEID);
+								AddEdge(SplitInfo.NewEdges.A);
 								// Note: Do not update PossUnmatchedBdryVerts with the new vertex, because it is already matched by construction
 								// Likewise do not update the pointhash -- we don't want it to find vertices that were already perfectly matched
 							}
@@ -582,8 +630,11 @@ bool FMeshBoolean::Compute()
 		FMeshIndexMappings IndexMaps;
 		Editor.AppendMesh(CutMesh[1], IndexMaps);
 
-		bool bWeldSuccess = MergeEdges(IndexMaps, CutMesh, CutBoundaryEdges, AllVIDMatches);
-		bSuccess = bSuccess && bWeldSuccess;
+		if (bWeldSharedEdges)
+		{
+			bool bWeldSuccess = MergeEdges(IndexMaps, CutMesh, CutBoundaryEdges, AllVIDMatches);
+			bSuccess = bSuccess && bWeldSuccess;
+		}
 	}
 	else
 	{
diff --git a/Engine/Plugins/Experimental/GeometryProcessing/Source/DynamicMesh/Public/Operations/MeshBoolean.h b/Engine/Plugins/Experimental/GeometryProcessing/Source/DynamicMesh/Public/Operations/MeshBoolean.h
--- a/Engine/Plugins/Experimental/GeometryProcessing/Source/DynamicMesh/Public/Operations/MeshBoolean.h
+++ b/Engine/Plugins/Experimental/GeometryProcessing/Source/DynamicMesh/Public/Operations/MeshBoolean.h
@@ -58,6 +58,9 @@ public:
 	/** Put the Result mesh in the same space as the input.  If true, ResultTransform will be the identity transform. */
 	bool bPutResultInInputSpace = true;
 
+	/** Weld newly-created cut edges where the input meshes meet.  If false, the input meshes will remain topologically disconnected. */
+	bool bWeldSharedEdges = true;
+
 	/** Set this to be able to cancel running operation */
 	FProgressCancel* Progress = nullptr;
 
diff --git a/Engine/Plugins/Experimental/GeometryProcessing/Source/GeometricObjects/Private/Spatial/SparseDynamicOctree3.cpp b/Engine/Plugins/Experimental/GeometryProcessing/Source/GeometricObjects/Private/Spatial/SparseDynamicOctree3.cpp
new file mode 100644
--- /dev/null
+++ b/Engine/Plugins/Experimental/GeometryProcessing/Source/GeometricObjects/Private/Spatial/SparseDynamicOctree3.cpp
@@ -0,0 +1,577 @@
+// Copyright Epic Games, Inc. All Rights Reserved.
+
+
+#include "Spatial/SparseDynamicOctree3.h"
+
+
+// NB: These have to be here until C++17 allows inline variables
+constexpr int32 FDynamicFlagArray::GrowChunkSize;
+constexpr uint32 FSparseOctreeCell::InvalidID;
+constexpr uint8 FSparseOctreeCell::InvalidLevel;
+constexpr uint32 FSparseDynamicOctree3::MaxSupportedTreeDepth;
+constexpr uint32 FSparseDynamicOctree3::InvalidCellID;
+constexpr uint32 FSparseDynamicOctree3::SpillCellID;
+
+
+bool FSparseDynamicOctree3::ContainsObject(int32 ObjectID) const
+{
+	return ValidObjectIDs.Get(ObjectID);
+}
+
+void FSparseDynamicOctree3::InsertObject(int32 ObjectID, const FAxisAlignedBox3d& Bounds)
+{
+	checkSlow(ContainsObject(ObjectID) == false);
+
+	FSparseOctreeCell CurrentCell = FindCurrentContainingCell(Bounds);
+
+	// if we could not find a containing root cell, we spill
+	if (CurrentCell.Level == FSparseOctreeCell::InvalidLevel)
+	{
+		Insert_Spill(ObjectID, Bounds);
+		return;
+	}
+	checkSlow(CanFit(CurrentCell, Bounds));
+
+	// if we found a containing root cell but it doesn't exist, create it and insert
+	if (CurrentCell.Level == 0 && CurrentCell.IsExistingCell() == false)
+	{
+		Insert_NewRoot(ObjectID, Bounds, CurrentCell);
+		return;
+	}
+
+	// YIKES this currently does max-depth insertion...
+	//   desired behavior is that parent cell accumulates and then splits later!
+
+	int PotentialChildIdx = ToChildCellIndex(CurrentCell, Bounds.Center());
+	// if current cell does not have this child we might fit there so try it
+	if (CurrentCell.HasChild(PotentialChildIdx) == false)
+	{
+		// todo can we do a fast check based on level and dimensions??
+		FSparseOctreeCell NewChildCell = CurrentCell.MakeChildCell(PotentialChildIdx);
+		if (NewChildCell.Level <= MaxTreeDepth && CanFit(NewChildCell, Bounds))
+		{
+			Insert_NewChildCell(ObjectID, Bounds, CurrentCell.CellID, NewChildCell, PotentialChildIdx);
+			return;
+		}
+	}
+
+	// insert into current cell if
+	//   1) child cell exists (in which case we didn't fit or FindCurrentContainingCell would have returned)
+	//   2) we tried to fit in child cell and failed
+	Insert_ToCell(ObjectID, Bounds, CurrentCell);
+}
+
+void FSparseDynamicOctree3::Insert_NewChildCell(int32 ObjectID, const FAxisAlignedBox3d& Bounds,
+	int ParentCellID, FSparseOctreeCell NewChildCell, int ChildIdx)
+{
+	FSparseOctreeCell& OrigParentCell = Cells[ParentCellID];
+	checkSlow(OrigParentCell.HasChild(ChildIdx) == false);
+
+	NewChildCell.CellID = CellRefCounts.Allocate();
+	Cells.InsertAt(NewChildCell, NewChildCell.CellID);
+
+	ObjectIDToCellMap.InsertAt(NewChildCell.CellID, ObjectID);
+	ValidObjectIDs.Set(ObjectID, true);
+
+	CellObjectLists.AllocateAt(NewChildCell.CellID);
+	CellObjectLists.Insert(NewChildCell.CellID, ObjectID);
+
+	OrigParentCell.SetChild(ChildIdx, NewChildCell);
+
+	checkSlow(CanFit(NewChildCell, Bounds));
+	// this check is unstable if the center point is within machine-epsilon of the cell border
+	//check(PointToIndex(NewChildCell.Level, Bounds.Center()) == NewChildCell.Index);
+}
+
+
+void FSparseDynamicOctree3::Insert_ToCell(int32 ObjectID, const FAxisAlignedBox3d& Bounds, const FSparseOctreeCell& ExistingCell)
+{
+	checkSlow(CellRefCounts.IsValid(ExistingCell.CellID));
+
+	ObjectIDToCellMap.InsertAt(ExistingCell.CellID, ObjectID);
+	ValidObjectIDs.Set(ObjectID, true);
+
+	CellObjectLists.Insert(ExistingCell.CellID, ObjectID);
+
+	checkSlow(CanFit(ExistingCell, Bounds));
+	// this check is unstable if the center point is within machine-epsilon of the cell border
+	//check(PointToIndex(ExistingCell.Level, Bounds.Center()) == ExistingCell.Index);
+}
+
+
+void FSparseDynamicOctree3::Insert_NewRoot(int32 ObjectID, const FAxisAlignedBox3d& Bounds, FSparseOctreeCell NewRootCell)
+{
+	checkSlow(RootCells.Has(NewRootCell.Index) == false);
+
+	NewRootCell.CellID = CellRefCounts.Allocate();
+	Cells.InsertAt(NewRootCell, NewRootCell.CellID);
+
+	ObjectIDToCellMap.InsertAt(NewRootCell.CellID, ObjectID);
+	ValidObjectIDs.Set(ObjectID, true);
+
+	uint32* RootCellElem = RootCells.Get(NewRootCell.Index, true);
+	*RootCellElem = NewRootCell.CellID;
+
+	CellObjectLists.AllocateAt(NewRootCell.CellID);
+	CellObjectLists.Insert(NewRootCell.CellID, ObjectID);
+
+	checkSlow(CanFit(NewRootCell, Bounds));
+}
+
+void FSparseDynamicOctree3::Insert_Spill(int32 ObjectID, const FAxisAlignedBox3d& Bounds)
+{
+	SpillObjectSet.Add(ObjectID);
+	ObjectIDToCellMap.InsertAt(SpillCellID, ObjectID);
+	ValidObjectIDs.Set(ObjectID, true);
+}
+
+
+
+
+bool FSparseDynamicOctree3::RemoveObject(int32 ObjectID)
+{
+	if (ContainsObject(ObjectID) == false)
+	{
+		return false;
+	}
+
+	uint32 CellID = GetCellForObject(ObjectID);
+	if (CellID == SpillCellID)
+	{
+		int32 RemovedCount = SpillObjectSet.Remove(ObjectID);
+		checkSlow(RemovedCount > 0);
+		ValidObjectIDs.Set(ObjectID, false);
+		return (RemovedCount > 0);
+	}
+	if (CellID == InvalidCellID)
+	{
+		return false;
+	}
+
+	ObjectIDToCellMap[ObjectID] = InvalidCellID;
+	ValidObjectIDs.Set(ObjectID, false);
+
+	bool bInList = CellObjectLists.Remove(CellID, ObjectID);
+	checkSlow(bInList);
+	return true;
+}
+
+
+
+
+
+void FSparseDynamicOctree3::ReinsertObject(int32 ObjectID, const FAxisAlignedBox3d& NewBounds)
+{
+	if (ContainsObject(ObjectID))
+	{
+		uint32 CellID = GetCellForObject(ObjectID);
+		if (CellID != SpillCellID && CellID != InvalidCellID)
+		{
+			FSparseOctreeCell& CurrentCell = Cells[CellID];
+			if (CanFit(CurrentCell, NewBounds))
+			{
+				return;		// everything is fine
+			}
+
+		}
+	}
+
+	RemoveObject(ObjectID);
+	InsertObject(ObjectID, NewBounds);
+}
+
+
+
+
+
+double FSparseDynamicOctree3::FindNearestRayCellIntersection(const FSparseOctreeCell& Cell, const FRay3d& Ray) const
+{
+	FAxisAlignedBox3d Box = GetCellBox(Cell, MaxExpandFactor);
+	double ray_t = TNumericLimits<double>::Max();
+	if (FIntrRay3AxisAlignedBox3d::FindIntersection(Ray, Box, ray_t))
+	{
+		return ray_t;
+	}
+	else
+	{
+		return TNumericLimits<double>::Max();
+	}
+}
+
+
+
+int32 FSparseDynamicOctree3::FindNearestHitObject(const FRay3d& Ray,
+	TFunctionRef<FAxisAlignedBox3d(int)> GetObjectBoundsFunc,
+	TFunctionRef<double(int, const FRay3d&)> HitObjectDistFunc,
+	double MaxDistance) const
+{
+	// this should take advantage of raster!
+
+	// always test against all spill objects
+	int32 HitObjectID = -1;
+	for (int ObjectID : SpillObjectSet)
+	{
+		double HitDist = HitObjectDistFunc(ObjectID, Ray);
+		if (HitDist < MaxDistance)
+		{
+			MaxDistance = HitDist;
+			HitObjectID = ObjectID;
+		}
+	}
+
+	// we use queue instead of recursion
+	TArray<const FSparseOctreeCell*> Queue;
+	Queue.Reserve(64);
+
+	// push all root cells onto queue if they are hit by ray
+	RootCells.AllocatedIteration([&](const uint32* RootCellID)
+	{
+		const FSparseOctreeCell* RootCell = &Cells[*RootCellID];
+		double RayHitParam = FindNearestRayCellIntersection(*RootCell, Ray);
+		if (RayHitParam < MaxDistance)
+		{
+			Queue.Add(&Cells[*RootCellID]);
+		}
+	});
+
+	// test cells until the queue is empty
+	while (Queue.Num() > 0)
+	{
+		const FSparseOctreeCell* CurCell = Queue.Pop(false);
+		
+		// process elements
+		for (int ObjectID : CellObjectLists.Values(CurCell->CellID))
+		{
+			double HitDist = HitObjectDistFunc(ObjectID, Ray);
+			if (HitDist < MaxDistance)
+			{
+				MaxDistance = HitDist;
+				HitObjectID = ObjectID;
+			}
+		}
+
+		// descend to child cells
+		// sort by distance? use DDA?
+		for (int k = 0; k < 8; ++k) 
+		{
+			if (CurCell->HasChild(k))
+			{
+				const FSparseOctreeCell* ChildCell = &Cells[CurCell->GetChildCellID(k)];
+				double RayHitParam = FindNearestRayCellIntersection(*ChildCell, Ray);
+				if (RayHitParam < MaxDistance)
+				{
+					Queue.Add(ChildCell);
+				}
+			}
+		}
+	}
+
+	return HitObjectID;
+}
+
+
+
+
+
+void FSparseDynamicOctree3::RangeQuery(
+	const FAxisAlignedBox3d& Bounds,
+	TFunctionRef<void(int)> ObjectIDFunc) const
+{
+	// todo: this should take advantage of raster!
+
+	// always process spill objects
+	for (int ObjectID : SpillObjectSet)
+	{
+		ObjectIDFunc(ObjectID);
+	}
+
+	TArray<const FSparseOctreeCell*> Queue;
+
+	// start at root cells
+	RootCells.AllocatedIteration([&](const uint32* RootCellID)
+	{
+		const FSparseOctreeCell* RootCell = &Cells[*RootCellID];
+		if ( GetCellBox(*RootCell, MaxExpandFactor).Intersects(Bounds) )
+		{
+			Queue.Add(&Cells[*RootCellID]);
+		}
+	});
+
+
+	while (Queue.Num() > 0)
+	{
+		const FSparseOctreeCell* CurCell = Queue.Pop(false);
+
+		// process elements
+		for (int ObjectID : CellObjectLists.Values(CurCell->CellID))
+		{
+			ObjectIDFunc(ObjectID);
+		}
+
+		for (int k = 0; k < 8; ++k)
+		{
+			if (CurCell->HasChild(k))
+			{
+				const FSparseOctreeCell* ChildCell = &Cells[CurCell->GetChildCellID(k)];
+				if (GetCellBox(*ChildCell, MaxExpandFactor).Intersects(Bounds))
+				{
+					Queue.Add(ChildCell);
+				}
+			}
+		}
+	}
+}
+
+
+
+
+
+
+
+void FSparseDynamicOctree3::RangeQuery(
+	const FAxisAlignedBox3d& Bounds,
+	TArray<int>& ObjectIDs) const
+{
+	// todo: this should take advantage of raster!
+
+	// always collect spill objects
+	for (int ObjectID : SpillObjectSet)
+	{
+		ObjectIDs.Add(ObjectID);
+	}
+
+	TArray<const FSparseOctreeCell*> Queue;
+
+	// start at root cells
+	RootCells.AllocatedIteration([&](const uint32* RootCellID)
+	{
+		const FSparseOctreeCell* RootCell = &Cells[*RootCellID];
+		if (GetCellBox(*RootCell, MaxExpandFactor).Intersects(Bounds))
+		{
+			Queue.Add(&Cells[*RootCellID]);
+		}
+	});
+
+
+	while (Queue.Num() > 0)
+	{
+		const FSparseOctreeCell* CurCell = Queue.Pop(false);
+
+		// process elements
+		for (int ObjectID : CellObjectLists.Values(CurCell->CellID))
+		{
+			ObjectIDs.Add(ObjectID);
+		}
+
+		for (int k = 0; k < 8; ++k)
+		{
+			if (CurCell->HasChild(k))
+			{
+				const FSparseOctreeCell* ChildCell = &Cells[CurCell->GetChildCellID(k)];
+				if (GetCellBox(*ChildCell, MaxExpandFactor).Intersects(Bounds))
+				{
+					Queue.Add(ChildCell);
+				}
+			}
+		}
+	}
+}
+
+
+
+
+
+
+FSparseOctreeCell FSparseDynamicOctree3::FindCurrentContainingCell(const FAxisAlignedBox3d& Bounds) const
+{
+	double BoxWidth = Bounds.MaxDim();
+	FVector3d BoxCenter = Bounds.Center();
+
+	// look up root cell, which may not exist
+	FVector3i RootIndex = PointToIndex(0, BoxCenter);
+	const uint32* RootCellID = RootCells.Get(RootIndex);
+	if (RootCellID == nullptr)
+	{
+		FSparseOctreeCell RootCell(0, RootIndex);
+		// have to make sure we can fit in this root cell
+		if (CanFit(RootCell, Bounds))
+		{
+			return RootCell;
+		}
+		else
+		{
+			return FSparseOctreeCell(FSparseOctreeCell::InvalidLevel, FVector3i::Zero());
+		}
+
+	}
+	checkSlow(CellRefCounts.IsValid(*RootCellID));
+
+	// check if box is contained in root cell, if not we have to spill
+	// (should we do this before checking for existence? we can...)
+	const FSparseOctreeCell* RootCell = &Cells[*RootCellID];
+	if (CanFit(*RootCell, Bounds) == false)
+	{
+		return FSparseOctreeCell(FSparseOctreeCell::InvalidLevel, FVector3i::Zero());
+	}
+
+	const FSparseOctreeCell* CurrentCell = RootCell;
+	do
+	{
+		int ChildIdx = ToChildCellIndex(*CurrentCell, BoxCenter);
+		if (CurrentCell->HasChild(ChildIdx))
+		{
+			int32 ChildCellID = CurrentCell->GetChildCellID(ChildIdx);
+			checkSlow(CellRefCounts.IsValid(ChildCellID));
+			const FSparseOctreeCell* ChildCell = &Cells[ChildCellID];
+			if (CanFit(*ChildCell, Bounds))
+			{
+				CurrentCell = ChildCell;
+				continue;
+			}
+		}
+
+		return *CurrentCell;
+
+	} while (true);		// loop will always terminate
+
+	return FSparseOctreeCell(FSparseOctreeCell::InvalidLevel, FVector3i::Zero());
+}
+
+
+void FSparseDynamicOctree3::CheckValidity(
+	TFunctionRef<bool(int)> IsValidObjectIDFunc,
+	TFunctionRef<FAxisAlignedBox3d(int)> GetObjectBoundsFunc,
+	EValidityCheckFailMode FailMode,
+	bool bVerbose,
+	bool bFailOnMissingObjects) const
+{
+	bool is_ok = true;
+	TFunction<void(bool)> CheckOrFailF = [&](bool b)
+	{
+		is_ok = is_ok && b;
+	};
+	if (FailMode == EValidityCheckFailMode::Check)
+	{
+		CheckOrFailF = [&](bool b)
+		{
+			checkf(b, TEXT("FSparseDynamicOctree3::CheckValidity failed!"));
+			is_ok = is_ok && b;
+		};
+	}
+	else if (FailMode == EValidityCheckFailMode::Ensure)
+	{
+		CheckOrFailF = [&](bool b)
+		{
+			ensureMsgf(b, TEXT("FSparseDynamicOctree3::CheckValidity failed!"));
+			is_ok = is_ok && b;
+		};
+	}
+
+	TArray<int> CellsAtLevels, ObjectsAtLevel;
+	CellsAtLevels.Init(0, 32);
+	ObjectsAtLevel.Init(0, 32);
+	uint32 SpillObjectCount = 0;
+	uint32 MissingObjectCount = 0;
+	uint8 MaxLevel = 0;
+
+	// check that all object IDs in per-cell object lists is valid
+	for (int32 CellID : CellRefCounts.Indices())
+	{
+		for (int32 ObjectID : CellObjectLists.Values(CellID))
+		{
+			CheckOrFailF(IsValidObjectIDFunc(ObjectID));
+		}
+	}
+
+	uint32 NumObjectIDs = ObjectIDToCellMap.Num();
+	for (uint32 ObjectID = 0; ObjectID < NumObjectIDs; ObjectID++)
+	{
+		bool IsValidObjectID = IsValidObjectIDFunc(ObjectID);
+		if (IsValidObjectID)
+		{
+			FAxisAlignedBox3d ObjectBounds = GetObjectBoundsFunc(ObjectID);
+
+			CheckOrFailF(ObjectID < ObjectIDToCellMap.Num());
+			uint32 CellID = ObjectIDToCellMap[ObjectID];
+
+			if (bFailOnMissingObjects)
+			{
+				CheckOrFailF(CellID != InvalidCellID);
+			}
+
+			if (CellID == SpillCellID)
+			{
+				// this is a spill node...
+				SpillObjectCount++;
+				CheckOrFailF(SpillObjectSet.Contains(ObjectID));
+			}
+			else if (CellID == InvalidCellID)
+			{
+				MissingObjectCount++;
+				CheckOrFailF(SpillObjectSet.Contains(ObjectID) == false);
+			}
+			else
+			{
+				CheckOrFailF(CellRefCounts.IsValid(CellID));
+				FSparseOctreeCell Cell = Cells[CellID];
+				FAxisAlignedBox3d CellBounds = GetCellBox(Cell, MaxExpandFactor);
+				CheckOrFailF(CellBounds.Contains(ObjectBounds));
+				CheckOrFailF(CellObjectLists.Contains(CellID, ObjectID));
+
+				ObjectsAtLevel[Cell.Level]++;
+			}
+		}
+	}
+
+
+	for (int32 CellID : CellRefCounts.Indices())
+	{
+		const FSparseOctreeCell& Cell = Cells[CellID];
+		CellsAtLevels[Cell.Level]++;
+		MaxLevel = FMath::Max(MaxLevel, Cell.Level);
+	}
+
+	if (bVerbose)
+	{
+		UE_LOG(LogTemp, Warning, TEXT("FSparseDynamicOctree3::CheckValidity: MaxLevel %d  SpillCount %d  MissingCount %d"), MaxLevel, SpillObjectCount, MissingObjectCount);
+		for (uint32 k = 0; k <= MaxLevel; ++k)
+		{
+			UE_LOG(LogTemp, Warning, TEXT("    Level %4d  Cells %4d  Objects %4d"), k, CellsAtLevels[k], ObjectsAtLevel[k]);
+		}
+	}
+}
+
+
+
+void FSparseDynamicOctree3::ComputeStatistics(FStatistics& StatsOut) const
+{
+	StatsOut.SpillObjCount = SpillObjectSet.Num();
+
+	StatsOut.Levels = 0;
+	for (int32 CellID : CellRefCounts.Indices())
+	{
+		const FSparseOctreeCell& Cell = Cells[CellID];
+		StatsOut.Levels = FMath::Max(StatsOut.Levels, (int32)Cell.Level);
+	}
+	StatsOut.Levels++;
+	StatsOut.LevelObjCounts.Init(0, StatsOut.Levels);
+	StatsOut.LevelBoxCounts.Init(0, StatsOut.Levels);
+	for (int32 CellID : CellRefCounts.Indices())
+	{
+		const FSparseOctreeCell& Cell = Cells[CellID];
+		StatsOut.LevelBoxCounts[Cell.Level]++;
+		StatsOut.LevelObjCounts[Cell.Level] += CellObjectLists.GetCount(CellID);
+	}
+}
+
+FString FSparseDynamicOctree3::FStatistics::ToString() const
+{
+	FString Result = FString::Printf(
+		TEXT("Levels %2d   SpillCount %5d \r\n"), Levels, SpillObjCount);
+	for (int k = 0; k < Levels; ++k)
+	{
+		Result += FString::Printf(TEXT("  Level %2d:  Cells %8d  Tris %8d  Avg %5.3f\r\n"), 
+			k, LevelBoxCounts[k], LevelObjCounts[k], ((float)LevelObjCounts[k] / (float)LevelBoxCounts[k])  );
+	}
+	return Result;
+}
diff --git a/Engine/Plugins/Experimental/GeometryProcessing/Source/GeometricObjects/Public/Spatial/SparseDynamicOctree3.h b/Engine/Plugins/Experimental/GeometryProcessing/Source/GeometricObjects/Public/Spatial/SparseDynamicOctree3.h
--- a/Engine/Plugins/Experimental/GeometryProcessing/Source/GeometricObjects/Public/Spatial/SparseDynamicOctree3.h
+++ b/Engine/Plugins/Experimental/GeometryProcessing/Source/GeometricObjects/Public/Spatial/SparseDynamicOctree3.h
@@ -6,11 +6,13 @@
 #include "CoreMinimal.h"
 #include "Containers/StaticArray.h"
 #include "BoxTypes.h"
+#include "GeometryTypes.h"
 #include "Util/DynamicVector.h"
 #include "Util/SmallListSet.h"
+#include "Util/RefCountVector.h"
 #include "Spatial/SparseGrid3.h"
 #include "Intersection/IntrRay3AxisAlignedBox3.h"
-#include "ASync/ParallelFor.h"
+#include "Async/ParallelFor.h"
 
 
 /**
@@ -22,7 +24,7 @@ class FDynamicFlagArray
 public:
 	TBitArray<> BitArray;
 	uint32 MaxIndex = 0;
-	static constexpr int32 GrowChunkSize = 0xFFF;
+	GEOMETRICOBJECTS_API static constexpr int32 GrowChunkSize = 0xFFF;
 
 	FDynamicFlagArray()
 	{
@@ -55,8 +57,8 @@ public:
  */
 struct FSparseOctreeCell
 {
-	static constexpr uint32 InvalidID = TNumericLimits<uint32>::Max();;
-	static constexpr uint8 InvalidLevel = TNumericLimits<uint8>::Max();
+	GEOMETRICOBJECTS_API static constexpr uint32 InvalidID = TNumericLimits<uint32>::Max();;
+	GEOMETRICOBJECTS_API static constexpr uint8 InvalidLevel = TNumericLimits<uint8>::Max();
 
 	/** ID of cell (index into cell list) */
 	uint32 CellID;
@@ -135,7 +137,7 @@ struct FSparseOctreeCell
  * The octree is dynamic. Objects can be removed and re-inserted.
  * 
  */
-class FSparseDynamicOctree3
+class GEOMETRICOBJECTS_API FSparseDynamicOctree3
 {
 	// potential optimizations/improvements
 	//    - Cell sizes are known at each level...can keep a lookup table? will this improve performance?
@@ -197,7 +199,7 @@ public:
 	 * @param ObjectID ID of the object
 	 * @return true if ObjectID is stored in this octree
 	 */
-	inline bool ContainsObject(int32 ObjectID) const;
+	bool ContainsObject(int32 ObjectID) const;
 
 	/**
 	 * Insert ObjectID into the Octree
@@ -205,21 +207,21 @@ public:
 	 * @param Bounds bounding box of the object
 	 *
 	 */
-	inline void InsertObject(int32 ObjectID, const FAxisAlignedBox3d& Bounds);
+	void InsertObject(int32 ObjectID, const FAxisAlignedBox3d& Bounds);
 
 	/**
 	 * Remove an object from the octree
 	 * @param ObjectID ID of the object
 	 * @return true if the object was in the tree and removed
 	 */
-	inline bool RemoveObject(int32 ObjectID);
+	bool RemoveObject(int32 ObjectID);
 
 	/**
 	 * Update the position of an object in the octree. This is more efficient than doing a remove+insert
 	 * @param ObjectID ID of the object
 	 * @param NewBounds enw bounding box of the object
 	 */
-	inline void ReinsertObject(int32 ObjectID, const FAxisAlignedBox3d& NewBounds);
+	void ReinsertObject(int32 ObjectID, const FAxisAlignedBox3d& NewBounds);
 
 	/**
 	 * Find nearest ray-hit point with objects in tree
@@ -239,7 +241,7 @@ public:
 	 * @param Bounds query box
 	 * @param ObjectIDFunc this function is called 
 	 */
-	inline void RangeQuery(const FAxisAlignedBox3d& Bounds,
+	void RangeQuery(const FAxisAlignedBox3d& Bounds,
 		TFunctionRef<void(int)> ObjectIDFunc) const;
 
 	/**
@@ -247,7 +249,7 @@ public:
 	 * @param Bounds query box
 	 * @param ObjectIDsOut collected ObjectIDs are stored here 
 	 */
-	inline void RangeQuery(const FAxisAlignedBox3d& Bounds,
+	void RangeQuery(const FAxisAlignedBox3d& Bounds,
 		TArray<int>& ObjectIDsOut ) const;
 
 
@@ -259,7 +261,7 @@ public:
 	 * @param bVerbose if true, print some debug info via UE_LOG
 	 * @param bFailOnMissingObjects if true, assume ObjectIDs are dense and that all ObjectIDs must be in the tree
 	 */
-	inline void CheckValidity(
+	void CheckValidity(
 		TFunctionRef<bool(int)> IsValidObjectIDFunc,
 		TFunctionRef<FAxisAlignedBox3d(int)> GetObjectBoundsFunc,
 		EValidityCheckFailMode FailMode = EValidityCheckFailMode::Check,
@@ -275,7 +277,7 @@ public:
 		TArray<int32> LevelBoxCounts;
 		TArray<int32> LevelObjCounts;
 		int32 SpillObjCount;
-		inline FString ToString() const;
+		FString ToString() const;
 	};
 
 	/**
@@ -382,580 +384,16 @@ protected:
 		return InvalidCellID;
 	}
 
-	inline FSparseOctreeCell FindCurrentContainingCell(const FAxisAlignedBox3d& Bounds) const;
+	FSparseOctreeCell FindCurrentContainingCell(const FAxisAlignedBox3d& Bounds) const;
 
 
-	inline void Insert_Spill(int32 ObjectID, const FAxisAlignedBox3d& Bounds);
-	inline void Insert_NewRoot(int32 ObjectID, const FAxisAlignedBox3d& Bounds, FSparseOctreeCell NewRootCell);
-	inline void Insert_ToCell(int32 ObjectID, const FAxisAlignedBox3d& Bounds, const FSparseOctreeCell& ExistingCell);
-	inline void Insert_NewChildCell(int32 ObjectID, const FAxisAlignedBox3d& Bounds, int ParentCellID, FSparseOctreeCell NewChildCell, int ChildIdx);
+	void Insert_Spill(int32 ObjectID, const FAxisAlignedBox3d& Bounds);
+	void Insert_NewRoot(int32 ObjectID, const FAxisAlignedBox3d& Bounds, FSparseOctreeCell NewRootCell);
+	void Insert_ToCell(int32 ObjectID, const FAxisAlignedBox3d& Bounds, const FSparseOctreeCell& ExistingCell);
+	void Insert_NewChildCell(int32 ObjectID, const FAxisAlignedBox3d& Bounds, int ParentCellID, FSparseOctreeCell NewChildCell, int ChildIdx);
 
 
-	inline double FindNearestRayCellIntersection(const FSparseOctreeCell& Cell, const FRay3d& Ray) const;
+	double FindNearestRayCellIntersection(const FSparseOctreeCell& Cell, const FRay3d& Ray) const;
 
 };
 
-
-bool FSparseDynamicOctree3::ContainsObject(int32 ObjectID) const
-{
-	return ValidObjectIDs.Get(ObjectID);
-}
-
-
-void FSparseDynamicOctree3::InsertObject(int32 ObjectID, const FAxisAlignedBox3d& Bounds)
-{
-	checkSlow(ContainsObject(ObjectID) == false);
-
-	FSparseOctreeCell CurrentCell = FindCurrentContainingCell(Bounds);
-
-	// if we could not find a containing root cell, we spill
-	if (CurrentCell.Level == FSparseOctreeCell::InvalidLevel)
-	{
-		Insert_Spill(ObjectID, Bounds);
-		return;
-	}
-	checkSlow(CanFit(CurrentCell, Bounds));
-
-	// if we found a containing root cell but it doesn't exist, create it and insert
-	if (CurrentCell.Level == 0 && CurrentCell.IsExistingCell() == false)
-	{
-		Insert_NewRoot(ObjectID, Bounds, CurrentCell);
-		return;
-	}
-
-	// YIKES this currently does max-depth insertion...
-	//   desired behavior is that parent cell accumulates and then splits later!
-
-	int PotentialChildIdx = ToChildCellIndex(CurrentCell, Bounds.Center());
-	// if current cell does not have this child we might fit there so try it
-	if (CurrentCell.HasChild(PotentialChildIdx) == false)
-	{
-		// todo can we do a fast check based on level and dimensions??
-		FSparseOctreeCell NewChildCell = CurrentCell.MakeChildCell(PotentialChildIdx);
-		if (NewChildCell.Level <= MaxTreeDepth && CanFit(NewChildCell, Bounds))
-		{
-			Insert_NewChildCell(ObjectID, Bounds, CurrentCell.CellID, NewChildCell, PotentialChildIdx);
-			return;
-		}
-	}
-
-	// insert into current cell if
-	//   1) child cell exists (in which case we didn't fit or FindCurrentContainingCell would have returned)
-	//   2) we tried to fit in child cell and failed
-	Insert_ToCell(ObjectID, Bounds, CurrentCell);
-}
-
-void FSparseDynamicOctree3::Insert_NewChildCell(int32 ObjectID, const FAxisAlignedBox3d& Bounds,
-	int ParentCellID, FSparseOctreeCell NewChildCell, int ChildIdx)
-{
-	FSparseOctreeCell& OrigParentCell = Cells[ParentCellID];
-	checkSlow(OrigParentCell.HasChild(ChildIdx) == false);
-
-	NewChildCell.CellID = CellRefCounts.Allocate();
-	Cells.InsertAt(NewChildCell, NewChildCell.CellID);
-
-	ObjectIDToCellMap.InsertAt(NewChildCell.CellID, ObjectID);
-	ValidObjectIDs.Set(ObjectID, true);
-
-	CellObjectLists.AllocateAt(NewChildCell.CellID);
-	CellObjectLists.Insert(NewChildCell.CellID, ObjectID);
-
-	OrigParentCell.SetChild(ChildIdx, NewChildCell);
-
-	checkSlow(CanFit(NewChildCell, Bounds));
-	// this check is unstable if the center point is within machine-epsilon of the cell border
-	//check(PointToIndex(NewChildCell.Level, Bounds.Center()) == NewChildCell.Index);
-}
-
-
-void FSparseDynamicOctree3::Insert_ToCell(int32 ObjectID, const FAxisAlignedBox3d& Bounds, const FSparseOctreeCell& ExistingCell)
-{
-	checkSlow(CellRefCounts.IsValid(ExistingCell.CellID));
-
-	ObjectIDToCellMap.InsertAt(ExistingCell.CellID, ObjectID);
-	ValidObjectIDs.Set(ObjectID, true);
-
-	CellObjectLists.Insert(ExistingCell.CellID, ObjectID);
-
-	checkSlow(CanFit(ExistingCell, Bounds));
-	// this check is unstable if the center point is within machine-epsilon of the cell border
-	//check(PointToIndex(ExistingCell.Level, Bounds.Center()) == ExistingCell.Index);
-}
-
-
-void FSparseDynamicOctree3::Insert_NewRoot(int32 ObjectID, const FAxisAlignedBox3d& Bounds, FSparseOctreeCell NewRootCell)
-{
-	checkSlow(RootCells.Has(NewRootCell.Index) == false);
-
-	NewRootCell.CellID = CellRefCounts.Allocate();
-	Cells.InsertAt(NewRootCell, NewRootCell.CellID);
-
-	ObjectIDToCellMap.InsertAt(NewRootCell.CellID, ObjectID);
-	ValidObjectIDs.Set(ObjectID, true);
-
-	uint32* RootCellElem = RootCells.Get(NewRootCell.Index, true);
-	*RootCellElem = NewRootCell.CellID;
-
-	CellObjectLists.AllocateAt(NewRootCell.CellID);
-	CellObjectLists.Insert(NewRootCell.CellID, ObjectID);
-
-	checkSlow(CanFit(NewRootCell, Bounds));
-}
-
-void FSparseDynamicOctree3::Insert_Spill(int32 ObjectID, const FAxisAlignedBox3d& Bounds)
-{
-	SpillObjectSet.Add(ObjectID);
-	ObjectIDToCellMap.InsertAt(SpillCellID, ObjectID);
-	ValidObjectIDs.Set(ObjectID, true);
-}
-
-
-
-
-bool FSparseDynamicOctree3::RemoveObject(int32 ObjectID)
-{
-	if (ContainsObject(ObjectID) == false)
-	{
-		return false;
-	}
-
-	uint32 CellID = GetCellForObject(ObjectID);
-	if (CellID == SpillCellID)
-	{
-		int32 RemovedCount = SpillObjectSet.Remove(ObjectID);
-		checkSlow(RemovedCount > 0);
-		ValidObjectIDs.Set(ObjectID, false);
-		return (RemovedCount > 0);
-	}
-	if (CellID == InvalidCellID)
-	{
-		return false;
-	}
-
-	ObjectIDToCellMap[ObjectID] = InvalidCellID;
-	ValidObjectIDs.Set(ObjectID, false);
-
-	bool bInList = CellObjectLists.Remove(CellID, ObjectID);
-	checkSlow(bInList);
-	return true;
-}
-
-
-
-
-
-void FSparseDynamicOctree3::ReinsertObject(int32 ObjectID, const FAxisAlignedBox3d& NewBounds)
-{
-	if (ContainsObject(ObjectID))
-	{
-		uint32 CellID = GetCellForObject(ObjectID);
-		if (CellID != SpillCellID && CellID != InvalidCellID)
-		{
-			FSparseOctreeCell& CurrentCell = Cells[CellID];
-			if (CanFit(CurrentCell, NewBounds))
-			{
-				return;		// everything is fine
-			}
-
-		}
-	}
-
-	RemoveObject(ObjectID);
-	InsertObject(ObjectID, NewBounds);
-}
-
-
-
-
-
-double FSparseDynamicOctree3::FindNearestRayCellIntersection(const FSparseOctreeCell& Cell, const FRay3d& Ray) const
-{
-	FAxisAlignedBox3d Box = GetCellBox(Cell, MaxExpandFactor);
-	double ray_t = TNumericLimits<double>::Max();
-	if (FIntrRay3AxisAlignedBox3d::FindIntersection(Ray, Box, ray_t))
-	{
-		return ray_t;
-	}
-	else
-	{
-		return TNumericLimits<double>::Max();
-	}
-}
-
-
-
-int32 FSparseDynamicOctree3::FindNearestHitObject(const FRay3d& Ray,
-	TFunctionRef<FAxisAlignedBox3d(int)> GetObjectBoundsFunc,
-	TFunctionRef<double(int, const FRay3d&)> HitObjectDistFunc,
-	double MaxDistance) const
-{
-	// this should take advantage of raster!
-
-	// always test against all spill objects
-	int32 HitObjectID = -1;
-	for (int ObjectID : SpillObjectSet)
-	{
-		double HitDist = HitObjectDistFunc(ObjectID, Ray);
-		if (HitDist < MaxDistance)
-		{
-			MaxDistance = HitDist;
-			HitObjectID = ObjectID;
-		}
-	}
-
-	// we use queue instead of recursion
-	TArray<const FSparseOctreeCell*> Queue;
-	Queue.Reserve(64);
-
-	// push all root cells onto queue if they are hit by ray
-	RootCells.AllocatedIteration([&](const uint32* RootCellID)
-	{
-		const FSparseOctreeCell* RootCell = &Cells[*RootCellID];
-		double RayHitParam = FindNearestRayCellIntersection(*RootCell, Ray);
-		if (RayHitParam < MaxDistance)
-		{
-			Queue.Add(&Cells[*RootCellID]);
-		}
-	});
-
-	// test cells until the queue is empty
-	while (Queue.Num() > 0)
-	{
-		const FSparseOctreeCell* CurCell = Queue.Pop(false);
-		
-		// process elements
-		for (int ObjectID : CellObjectLists.Values(CurCell->CellID))
-		{
-			double HitDist = HitObjectDistFunc(ObjectID, Ray);
-			if (HitDist < MaxDistance)
-			{
-				MaxDistance = HitDist;
-				HitObjectID = ObjectID;
-			}
-		}
-
-		// descend to child cells
-		// sort by distance? use DDA?
-		for (int k = 0; k < 8; ++k) 
-		{
-			if (CurCell->HasChild(k))
-			{
-				const FSparseOctreeCell* ChildCell = &Cells[CurCell->GetChildCellID(k)];
-				double RayHitParam = FindNearestRayCellIntersection(*ChildCell, Ray);
-				if (RayHitParam < MaxDistance)
-				{
-					Queue.Add(ChildCell);
-				}
-			}
-		}
-	}
-
-	return HitObjectID;
-}
-
-
-
-
-
-void FSparseDynamicOctree3::RangeQuery(
-	const FAxisAlignedBox3d& Bounds,
-	TFunctionRef<void(int)> ObjectIDFunc) const
-{
-	// todo: this should take advantage of raster!
-
-	// always process spill objects
-	for (int ObjectID : SpillObjectSet)
-	{
-		ObjectIDFunc(ObjectID);
-	}
-
-	TArray<const FSparseOctreeCell*> Queue;
-
-	// start at root cells
-	RootCells.AllocatedIteration([&](const uint32* RootCellID)
-	{
-		const FSparseOctreeCell* RootCell = &Cells[*RootCellID];
-		if ( GetCellBox(*RootCell, MaxExpandFactor).Intersects(Bounds) )
-		{
-			Queue.Add(&Cells[*RootCellID]);
-		}
-	});
-
-
-	while (Queue.Num() > 0)
-	{
-		const FSparseOctreeCell* CurCell = Queue.Pop(false);
-
-		// process elements
-		for (int ObjectID : CellObjectLists.Values(CurCell->CellID))
-		{
-			ObjectIDFunc(ObjectID);
-		}
-
-		for (int k = 0; k < 8; ++k)
-		{
-			if (CurCell->HasChild(k))
-			{
-				const FSparseOctreeCell* ChildCell = &Cells[CurCell->GetChildCellID(k)];
-				if (GetCellBox(*ChildCell, MaxExpandFactor).Intersects(Bounds))
-				{
-					Queue.Add(ChildCell);
-				}
-			}
-		}
-	}
-}
-
-
-
-
-
-
-
-void FSparseDynamicOctree3::RangeQuery(
-	const FAxisAlignedBox3d& Bounds,
-	TArray<int>& ObjectIDs) const
-{
-	// todo: this should take advantage of raster!
-
-	// always collect spill objects
-	for (int ObjectID : SpillObjectSet)
-	{
-		ObjectIDs.Add(ObjectID);
-	}
-
-	TArray<const FSparseOctreeCell*> Queue;
-
-	// start at root cells
-	RootCells.AllocatedIteration([&](const uint32* RootCellID)
-	{
-		const FSparseOctreeCell* RootCell = &Cells[*RootCellID];
-		if (GetCellBox(*RootCell, MaxExpandFactor).Intersects(Bounds))
-		{
-			Queue.Add(&Cells[*RootCellID]);
-		}
-	});
-
-
-	while (Queue.Num() > 0)
-	{
-		const FSparseOctreeCell* CurCell = Queue.Pop(false);
-
-		// process elements
-		for (int ObjectID : CellObjectLists.Values(CurCell->CellID))
-		{
-			ObjectIDs.Add(ObjectID);
-		}
-
-		for (int k = 0; k < 8; ++k)
-		{
-			if (CurCell->HasChild(k))
-			{
-				const FSparseOctreeCell* ChildCell = &Cells[CurCell->GetChildCellID(k)];
-				if (GetCellBox(*ChildCell, MaxExpandFactor).Intersects(Bounds))
-				{
-					Queue.Add(ChildCell);
-				}
-			}
-		}
-	}
-}
-
-
-
-
-
-
-FSparseOctreeCell FSparseDynamicOctree3::FindCurrentContainingCell(const FAxisAlignedBox3d& Bounds) const
-{
-	double BoxWidth = Bounds.MaxDim();
-	FVector3d BoxCenter = Bounds.Center();
-
-	// look up root cell, which may not exist
-	FVector3i RootIndex = PointToIndex(0, BoxCenter);
-	const uint32* RootCellID = RootCells.Get(RootIndex);
-	if (RootCellID == nullptr)
-	{
-		FSparseOctreeCell RootCell(0, RootIndex);
-		// have to make sure we can fit in this root cell
-		if (CanFit(RootCell, Bounds))
-		{
-			return RootCell;
-		}
-		else
-		{
-			return FSparseOctreeCell(FSparseOctreeCell::InvalidLevel, FVector3i::Zero());
-		}
-
-	}
-	checkSlow(CellRefCounts.IsValid(*RootCellID));
-
-	// check if box is contained in root cell, if not we have to spill
-	// (should we do this before checking for existence? we can...)
-	const FSparseOctreeCell* RootCell = &Cells[*RootCellID];
-	if (CanFit(*RootCell, Bounds) == false)
-	{
-		return FSparseOctreeCell(FSparseOctreeCell::InvalidLevel, FVector3i::Zero());
-	}
-
-	const FSparseOctreeCell* CurrentCell = RootCell;
-	do
-	{
-		int ChildIdx = ToChildCellIndex(*CurrentCell, BoxCenter);
-		if (CurrentCell->HasChild(ChildIdx))
-		{
-			int32 ChildCellID = CurrentCell->GetChildCellID(ChildIdx);
-			checkSlow(CellRefCounts.IsValid(ChildCellID));
-			const FSparseOctreeCell* ChildCell = &Cells[ChildCellID];
-			if (CanFit(*ChildCell, Bounds))
-			{
-				CurrentCell = ChildCell;
-				continue;
-			}
-		}
-
-		return *CurrentCell;
-
-	} while (true);		// loop will always terminate
-
-	return FSparseOctreeCell(FSparseOctreeCell::InvalidLevel, FVector3i::Zero());
-}
-
-
-void FSparseDynamicOctree3::CheckValidity(
-	TFunctionRef<bool(int)> IsValidObjectIDFunc,
-	TFunctionRef<FAxisAlignedBox3d(int)> GetObjectBoundsFunc,
-	EValidityCheckFailMode FailMode,
-	bool bVerbose,
-	bool bFailOnMissingObjects) const
-{
-	bool is_ok = true;
-	TFunction<void(bool)> CheckOrFailF = [&](bool b)
-	{
-		is_ok = is_ok && b;
-	};
-	if (FailMode == EValidityCheckFailMode::Check)
-	{
-		CheckOrFailF = [&](bool b)
-		{
-			checkf(b, TEXT("FSparseDynamicOctree3::CheckValidity failed!"));
-			is_ok = is_ok && b;
-		};
-	}
-	else if (FailMode == EValidityCheckFailMode::Ensure)
-	{
-		CheckOrFailF = [&](bool b)
-		{
-			ensureMsgf(b, TEXT("FSparseDynamicOctree3::CheckValidity failed!"));
-			is_ok = is_ok && b;
-		};
-	}
-
-	TArray<int> CellsAtLevels, ObjectsAtLevel;
-	CellsAtLevels.Init(0, 32);
-	ObjectsAtLevel.Init(0, 32);
-	uint32 SpillObjectCount = 0;
-	uint32 MissingObjectCount = 0;
-	uint8 MaxLevel = 0;
-
-	// check that all object IDs in per-cell object lists is valid
-	for (int32 CellID : CellRefCounts.Indices())
-	{
-		for (int32 ObjectID : CellObjectLists.Values(CellID))
-		{
-			CheckOrFailF(IsValidObjectIDFunc(ObjectID));
-		}
-	}
-
-	uint32 NumObjectIDs = ObjectIDToCellMap.Num();
-	for (uint32 ObjectID = 0; ObjectID < NumObjectIDs; ObjectID++)
-	{
-		bool IsValidObjectID = IsValidObjectIDFunc(ObjectID);
-		if (IsValidObjectID)
-		{
-			FAxisAlignedBox3d ObjectBounds = GetObjectBoundsFunc(ObjectID);
-
-			CheckOrFailF(ObjectID < ObjectIDToCellMap.Num());
-			uint32 CellID = ObjectIDToCellMap[ObjectID];
-
-			if (bFailOnMissingObjects)
-			{
-				CheckOrFailF(CellID != InvalidCellID);
-			}
-
-			if (CellID == SpillCellID)
-			{
-				// this is a spill node...
-				SpillObjectCount++;
-				CheckOrFailF(SpillObjectSet.Contains(ObjectID));
-			}
-			else if (CellID == InvalidCellID)
-			{
-				MissingObjectCount++;
-				CheckOrFailF(SpillObjectSet.Contains(ObjectID) == false);
-			}
-			else
-			{
-				CheckOrFailF(CellRefCounts.IsValid(CellID));
-				FSparseOctreeCell Cell = Cells[CellID];
-				FAxisAlignedBox3d CellBounds = GetCellBox(Cell, MaxExpandFactor);
-				CheckOrFailF(CellBounds.Contains(ObjectBounds));
-				CheckOrFailF(CellObjectLists.Contains(CellID, ObjectID));
-
-				ObjectsAtLevel[Cell.Level]++;
-			}
-		}
-	}
-
-
-	for (int32 CellID : CellRefCounts.Indices())
-	{
-		const FSparseOctreeCell& Cell = Cells[CellID];
-		CellsAtLevels[Cell.Level]++;
-		MaxLevel = FMath::Max(MaxLevel, Cell.Level);
-	}
-
-	if (bVerbose)
-	{
-		UE_LOG(LogTemp, Warning, TEXT("FSparseDynamicOctree3::CheckValidity: MaxLevel %d  SpillCount %d  MissingCount %d"), MaxLevel, SpillObjectCount, MissingObjectCount);
-		for (uint32 k = 0; k <= MaxLevel; ++k)
-		{
-			UE_LOG(LogTemp, Warning, TEXT("    Level %4d  Cells %4d  Objects %4d"), k, CellsAtLevels[k], ObjectsAtLevel[k]);
-		}
-	}
-}
-
-
-
-void FSparseDynamicOctree3::ComputeStatistics(FStatistics& StatsOut) const
-{
-	StatsOut.SpillObjCount = SpillObjectSet.Num();
-
-	StatsOut.Levels = 0;
-	for (int32 CellID : CellRefCounts.Indices())
-	{
-		const FSparseOctreeCell& Cell = Cells[CellID];
-		StatsOut.Levels = FMath::Max(StatsOut.Levels, (int32)Cell.Level);
-	}
-	StatsOut.Levels++;
-	StatsOut.LevelObjCounts.Init(0, StatsOut.Levels);
-	StatsOut.LevelBoxCounts.Init(0, StatsOut.Levels);
-	for (int32 CellID : CellRefCounts.Indices())
-	{
-		const FSparseOctreeCell& Cell = Cells[CellID];
-		StatsOut.LevelBoxCounts[Cell.Level]++;
-		StatsOut.LevelObjCounts[Cell.Level] += CellObjectLists.GetCount(CellID);
-	}
-}
-
-FString FSparseDynamicOctree3::FStatistics::ToString() const
-{
-	FString Result = FString::Printf(
-		TEXT("Levels %2d   SpillCount %5d \r\n"), Levels, SpillObjCount);
-	for (int k = 0; k < Levels; ++k)
-	{
-		Result += FString::Printf(TEXT("  Level %2d:  Cells %8d  Tris %8d  Avg %5.3f\r\n"), 
-			k, LevelBoxCounts[k], LevelObjCounts[k], ((float)LevelObjCounts[k] / (float)LevelBoxCounts[k])  );
-	}
-	return Result;
-}
\ No newline at end of file
