add RemoveAttribute() function to allow generic attributes to be removed from a mesh attribute set

diff --git a/Engine/Plugins/Experimental/GeometryProcessing/Source/DynamicMesh/Public/DynamicMeshAttributeSet.h b/Engine/Plugins/Experimental/GeometryProcessing/Source/DynamicMesh/Public/DynamicMeshAttributeSet.h
--- a/Engine/Plugins/Experimental/GeometryProcessing/Source/DynamicMesh/Public/DynamicMeshAttributeSet.h
+++ b/Engine/Plugins/Experimental/GeometryProcessing/Source/DynamicMesh/Public/DynamicMeshAttributeSet.h
@@ -347,6 +347,15 @@ public:
 		RegisterExternalAttribute(Attribute);
 	}
 
+	void RemoveAttribute(FName AttribName)
+	{
+		if (GenericAttributes.Contains(AttribName))
+		{
+			UnregisterExternalAttribute(GenericAttributes[AttribName].Get());
+			GenericAttributes.Remove(AttribName);
+		}
+	}
+
 	FDynamicMeshAttributeBase* GetAttachedAttribute(FName AttribName) const
 	{
 		return GenericAttributes[AttribName].Get();
