remove inline from functions that were moved to cpp file for SparseDynamicOctree
#rnx

diff --git a/Engine/Plugins/Experimental/GeometryProcessing/Source/GeometricObjects/Public/Spatial/SparseDynamicOctree3.h b/Engine/Plugins/Experimental/GeometryProcessing/Source/GeometricObjects/Public/Spatial/SparseDynamicOctree3.h
--- a/Engine/Plugins/Experimental/GeometryProcessing/Source/GeometricObjects/Public/Spatial/SparseDynamicOctree3.h
+++ b/Engine/Plugins/Experimental/GeometryProcessing/Source/GeometricObjects/Public/Spatial/SparseDynamicOctree3.h
@@ -231,7 +231,7 @@ public:
 	 * @param MaxDistance maximum hit distance
 	 * @return ObjectID of hit object, or -1 on miss
 	 */
-	inline int32 FindNearestHitObject(const FRay3d& Ray,
+	int32 FindNearestHitObject(const FRay3d& Ray,
 		TFunctionRef<FAxisAlignedBox3d(int)> GetObjectBoundsFunc,
 		TFunctionRef<double(int, const FRay3d&)> HitObjectDistFunc,
 		double MaxDistance = TNumericLimits<double>::Max()) const;
@@ -283,7 +283,7 @@ public:
 	/**
 	 * Populate given FStatistics with info about the octree
 	 */
-	inline void ComputeStatistics(FStatistics& StatsOut) const;
+	void ComputeStatistics(FStatistics& StatsOut) const;
 
 protected:
 	// this identifier is used for unknown cells
