add min area and min tri count filters for jacketing tool, to help remove isolated floating tris from output

diff --git a/Engine/Plugins/Experimental/GeometryProcessing/Source/DynamicMesh/Private/DynamicMeshEditor.cpp b/Engine/Plugins/Experimental/GeometryProcessing/Source/DynamicMesh/Private/DynamicMeshEditor.cpp
--- a/Engine/Plugins/Experimental/GeometryProcessing/Source/DynamicMesh/Private/DynamicMeshEditor.cpp
+++ b/Engine/Plugins/Experimental/GeometryProcessing/Source/DynamicMesh/Private/DynamicMeshEditor.cpp
@@ -342,7 +342,7 @@ bool FDynamicMeshEditor::RemoveTriangles(const TArray<int>& Triangles, bool bRem
 
 
 
-int FDynamicMeshEditor::RemoveSmallComponents(double MinVolume, double MinArea)
+int FDynamicMeshEditor::RemoveSmallComponents(double MinVolume, double MinArea, int MinTriangleCount)
 {
 	FMeshConnectedComponents C(Mesh);
 	C.FindConnectedTriangles();
@@ -352,8 +352,13 @@ int FDynamicMeshEditor::RemoveSmallComponents(double MinVolume, double MinArea)
 	}
 	int Removed = 0;
 	for (FMeshConnectedComponents::FComponent& Comp : C) {
-		FVector2d VolArea = TMeshQueries<FDynamicMesh3>::GetVolumeArea(*Mesh, Comp.Indices);
-		if (VolArea.X < MinVolume || VolArea.Y < MinArea) {
+		bool bRemove = Comp.Indices.Num() < MinTriangleCount;
+		if (!bRemove)
+		{
+			FVector2d VolArea = TMeshQueries<FDynamicMesh3>::GetVolumeArea(*Mesh, Comp.Indices);
+			bRemove = VolArea.X < MinVolume || VolArea.Y < MinArea;
+		}
+		if (bRemove) {
 			RemoveTriangles(Comp.Indices, true);
 			Removed++;
 		}
diff --git a/Engine/Plugins/Experimental/GeometryProcessing/Source/DynamicMesh/Public/DynamicMeshEditor.h b/Engine/Plugins/Experimental/GeometryProcessing/Source/DynamicMesh/Public/DynamicMeshEditor.h
--- a/Engine/Plugins/Experimental/GeometryProcessing/Source/DynamicMesh/Public/DynamicMeshEditor.h
+++ b/Engine/Plugins/Experimental/GeometryProcessing/Source/DynamicMesh/Public/DynamicMeshEditor.h
@@ -229,7 +229,7 @@ public:
 	 * @param MinArea Remove components with less area than this
 	 * @return number of components removed
 	 */
-	int RemoveSmallComponents(double MinVolume, double MinArea = 0.0);
+	int RemoveSmallComponents(double MinVolume, double MinArea = 0.0, int MinTriangleCount = 0);
 
 
 	/**
diff --git a/Engine/Plugins/Experimental/MeshModelingToolset/Source/MeshModelingTools/Private/RemoveOccludedTrianglesTool.cpp b/Engine/Plugins/Experimental/MeshModelingToolset/Source/MeshModelingTools/Private/RemoveOccludedTrianglesTool.cpp
--- a/Engine/Plugins/Experimental/MeshModelingToolset/Source/MeshModelingTools/Private/RemoveOccludedTrianglesTool.cpp
+++ b/Engine/Plugins/Experimental/MeshModelingToolset/Source/MeshModelingTools/Private/RemoveOccludedTrianglesTool.cpp
@@ -318,6 +318,10 @@ TUniquePtr<FDynamicMeshOperator> URemoveOccludedTrianglesOperatorFactory::MakeNe
 	Op->AddRandomRays = Tool->BasicProperties->AddRandomRays;
 
 	Op->AddTriangleSamples = Tool->BasicProperties->AddTriangleSamples;
+
+	Op->MinAreaConnectedComponent = Tool->BasicProperties->MinAreaIsland;
+
+	Op->MinTriCountConnectedComponent = Tool->BasicProperties->MinTriCountIsland;
 	
 	Op->SetTransform(LocalToWorld);
 
diff --git a/Engine/Plugins/Experimental/MeshModelingToolset/Source/MeshModelingTools/Public/RemoveOccludedTrianglesTool.h b/Engine/Plugins/Experimental/MeshModelingToolset/Source/MeshModelingTools/Public/RemoveOccludedTrianglesTool.h
--- a/Engine/Plugins/Experimental/MeshModelingToolset/Source/MeshModelingTools/Public/RemoveOccludedTrianglesTool.h
+++ b/Engine/Plugins/Experimental/MeshModelingToolset/Source/MeshModelingTools/Public/RemoveOccludedTrianglesTool.h
@@ -84,6 +84,12 @@ public:
 	/** If false, when multiple meshes are selected the meshes can occlude each other.  When true, we process each selected mesh independently and only consider self-occlusions. */
 	UPROPERTY(EditAnywhere, Category = OcclusionCalculation)
 	bool bOnlySelfOcclude = false;
+
+	UPROPERTY(EditAnywhere, Category = RemoveIslands, meta = (UIMin = "0", ClampMin = "0"))
+	double MinAreaIsland = 0;
+
+	UPROPERTY(EditAnywhere, Category = RemoveIslands, meta = (UIMin = "0", ClampMin = "0"))
+	int MinTriCountIsland = 0;
 };
 
 
diff --git a/Engine/Plugins/Experimental/MeshModelingToolset/Source/ModelingOperators/Private/CleaningOps/RemoveOccludedTrianglesOp.cpp b/Engine/Plugins/Experimental/MeshModelingToolset/Source/ModelingOperators/Private/CleaningOps/RemoveOccludedTrianglesOp.cpp
--- a/Engine/Plugins/Experimental/MeshModelingToolset/Source/ModelingOperators/Private/CleaningOps/RemoveOccludedTrianglesOp.cpp
+++ b/Engine/Plugins/Experimental/MeshModelingToolset/Source/ModelingOperators/Private/CleaningOps/RemoveOccludedTrianglesOp.cpp
@@ -8,6 +8,7 @@
 #include "DynamicMeshAABBTree3.h"
 
 #include "Operations/RemoveOccludedTriangles.h"
+#include "Selections/MeshConnectedComponents.h"
 
 
 #include "Async/ParallelFor.h"
@@ -89,4 +90,10 @@ void FRemoveOccludedTrianglesOp::CalculateResult(FProgressCancel* Progress)
 		Jacket.AddTriangleSamples = AddTriangleSamples;
 		Jacket.Apply(MeshTransforms, &CombinedMeshTrees->AABB, &CombinedMeshTrees->FastWinding);
 	}
+
+	if (MinTriCountConnectedComponent > 0 || MinAreaConnectedComponent > 0)
+	{
+		FDynamicMeshEditor Editor(ResultMesh.Get());
+		Editor.RemoveSmallComponents(0, MinAreaConnectedComponent, MinTriCountConnectedComponent);
+	}
 }
\ No newline at end of file
diff --git a/Engine/Plugins/Experimental/MeshModelingToolset/Source/ModelingOperators/Public/CleaningOps/RemoveOccludedTrianglesOp.h b/Engine/Plugins/Experimental/MeshModelingToolset/Source/ModelingOperators/Public/CleaningOps/RemoveOccludedTrianglesOp.h
--- a/Engine/Plugins/Experimental/MeshModelingToolset/Source/ModelingOperators/Public/CleaningOps/RemoveOccludedTrianglesOp.h
+++ b/Engine/Plugins/Experimental/MeshModelingToolset/Source/ModelingOperators/Public/CleaningOps/RemoveOccludedTrianglesOp.h
@@ -62,6 +62,10 @@ public:
 	// if true, *ignore* the combined selected occluder data in CombinedMeshTrees, and only consider self-occlusion of this mesh
 	bool bOnlySelfOcclude = false;
 
+	double MinAreaConnectedComponent = 0;
+
+	int MinTriCountConnectedComponent = 0;
+
 	EOcclusionCalculationMode InsideMode =
 		EOcclusionCalculationMode::FastWindingNumber;
 
