fixes for handling of merge edges when one of the vertices of the removed edge is not removed

#rnx

diff --git a/Engine/Plugins/Experimental/GeometryProcessing/Source/DynamicMesh/Private/DynamicMeshOverlay.cpp b/Engine/Plugins/Experimental/GeometryProcessing/Source/DynamicMesh/Private/DynamicMeshOverlay.cpp
--- a/Engine/Plugins/Experimental/GeometryProcessing/Source/DynamicMesh/Private/DynamicMeshOverlay.cpp
+++ b/Engine/Plugins/Experimental/GeometryProcessing/Source/DynamicMesh/Private/DynamicMeshOverlay.cpp
@@ -1136,6 +1136,10 @@ void TDynamicMeshOverlay<RealType, ElementSize>::OnMergeEdges(const FDynamicMesh
 	{
 		int KeptVID = MergeInfo.KeptVerts[i];
 		int RemovedVID = MergeInfo.RemovedVerts[i];
+		if (RemovedVID == FDynamicMesh3::InvalidID)
+		{
+			continue;
+		}
 		// this for loop is very similar to GetVertexElements() but accounts for the base mesh already being updated
 		for (int TID : ParentMesh->VtxTrianglesItr(KeptVID)) // only care about triangles connected to the *new* vertex; these are updated
 		{
diff --git a/Engine/Plugins/Experimental/GeometryProcessing/Source/DynamicMesh/Public/DynamicVertexAttribute.h b/Engine/Plugins/Experimental/GeometryProcessing/Source/DynamicMesh/Public/DynamicVertexAttribute.h
--- a/Engine/Plugins/Experimental/GeometryProcessing/Source/DynamicMesh/Public/DynamicVertexAttribute.h
+++ b/Engine/Plugins/Experimental/GeometryProcessing/Source/DynamicMesh/Public/DynamicVertexAttribute.h
@@ -314,8 +314,14 @@ public:
 	void OnMergeEdges(const FDynamicMesh3::FMergeEdgesInfo& MergeInfo) override
 	{
 		// just blend the attributes?
-		SetAttributeFromLerp(MergeInfo.KeptVerts.A, MergeInfo.KeptVerts.A, MergeInfo.RemovedVerts.A, .5);
-		SetAttributeFromLerp(MergeInfo.KeptVerts.B, MergeInfo.KeptVerts.B, MergeInfo.RemovedVerts.B, .5);
+		if (MergeInfo.RemovedVerts.A != FDynamicMesh3::InvalidID)
+		{
+			SetAttributeFromLerp(MergeInfo.KeptVerts.A, MergeInfo.KeptVerts.A, MergeInfo.RemovedVerts.A, .5);
+		}
+		if (MergeInfo.RemovedVerts.B != FDynamicMesh3::InvalidID)
+		{
+			SetAttributeFromLerp(MergeInfo.KeptVerts.B, MergeInfo.KeptVerts.B, MergeInfo.RemovedVerts.B, .5);
+		}
 	}
 
 	/** Update the overlay to reflect a vertex split in the parent */
