GeometryProcessing improvements/additions

- new correspondence-finding strategy in FMeshImageBakingCache. Added a global Thickness parameter, and now preferentially uses the ray-hit found by casting inwards from Point+Thickness*Normal. This handles the cases where a greeble/etc is just layered on top of base mesh. These kind of mesh elements are only intended to show up in the normal map, so they aren't stitched in.
  - expose Thickness and World-Space options in BakeMeshAttributeMapsTool
  - initialize image in Normal and ResampleImage Bakers, otherwise initial value is garbage
- add TSampledScalarField2::BilinearSampleGradientClamped and fix typo in BilinearSampleClamped() that resulted in incorrect interpolation
- add Texture2DUtil, contains functions for reading a TImageBuilder from a UTexture2D. Can read from source data if available.
- ColorConstants now returns white for group 0
- add double version of FMeshDescriptionToDynamicMesh::CopyTangents()
- add support for computing on subset of vertiecs in FMeshConvexHull
- fix up template export in TMeshTangents
- make mesh const in TImplicitMorphology, TImplicitSolidify
- accessor/etc additions in TDenseGrid2, TImageBuilder, TSampledScalarField2, TIndexedWeightMap
- add source UStaticMesh reference in FPhysicsDataCollection, function to initialize from a UStaticMesh
  - update Tools affected by above changes
- add TSampleSetStatisticBuilder::ComputeMultiPass() helper function to compute stats for iterable collections in one line
- remove dead code in UVLayoutOp.cpp

#rnx

diff --git a/Engine/Plugins/Experimental/GeometryProcessing/Source/DynamicMesh/Private/MeshNormals.cpp b/Engine/Plugins/Experimental/GeometryProcessing/Source/DynamicMesh/Private/MeshNormals.cpp
--- a/Engine/Plugins/Experimental/GeometryProcessing/Source/DynamicMesh/Private/MeshNormals.cpp
+++ b/Engine/Plugins/Experimental/GeometryProcessing/Source/DynamicMesh/Private/MeshNormals.cpp
@@ -208,13 +208,13 @@ void FMeshNormals::QuickComputeVertexNormalsForTriangles(FDynamicMesh3& Mesh, co
 
 
 
-bool FMeshNormals::QuickRecomputeOverlayNormals(FDynamicMesh3& Mesh, bool bInvert)
+bool FMeshNormals::QuickRecomputeOverlayNormals(FDynamicMesh3& Mesh, bool bInvert, bool bWeightByArea, bool bWeightByAngle)
 {
 	if (Mesh.HasAttributes() && Mesh.Attributes()->GetNormalLayer(0) != nullptr)
 	{
 		FDynamicMeshNormalOverlay* NormalOverlay = Mesh.Attributes()->GetNormalLayer(0);
 		FMeshNormals Normals(&Mesh);
-		Normals.RecomputeOverlayNormals(NormalOverlay);
+		Normals.RecomputeOverlayNormals(NormalOverlay, bWeightByArea, bWeightByAngle);
 		Normals.CopyToOverlay(NormalOverlay, bInvert);
 		return true;
 	}
@@ -330,6 +330,28 @@ void FMeshNormals::InitializeOverlayToPerTriangleNormals(FDynamicMeshNormalOverl
 }
 
 
+void FMeshNormals::InitializeOverlayTopologyFromOpeningAngle(const FDynamicMesh3* Mesh, FDynamicMeshNormalOverlay* NormalOverlay, 
+	double AngleThresholdDeg)
+{
+	double NormalDotProdThreshold = FMathd::Cos(AngleThresholdDeg * FMathd::DegToRad);
+
+	FMeshNormals FaceNormals(Mesh);
+	FaceNormals.ComputeTriangleNormals();
+	const TArray<FVector3d>& Normals = FaceNormals.GetNormals();
+	NormalOverlay->CreateFromPredicate([&Normals, &NormalDotProdThreshold](int VID, int TA, int TB) {
+		return Normals[TA].Dot(Normals[TB]) > NormalDotProdThreshold;
+	}, 0);
+}
+
+void FMeshNormals::InitializeOverlayTopologyFromFaceGroups(const FDynamicMesh3* Mesh, FDynamicMeshNormalOverlay* NormalOverlay)
+{
+	ensure(Mesh->HasTriangleGroups());
+	NormalOverlay->CreateFromPredicate([Mesh](int VID, int TA, int TB) {
+		return Mesh->GetTriangleGroup(TA) == Mesh->GetTriangleGroup(TB);
+	}, 0);
+}
+
+
 void FMeshNormals::InitializeMeshToPerTriangleNormals(FDynamicMesh3* Mesh)
 {
 	if (Mesh->HasAttributes() == false)
diff --git a/Engine/Plugins/Experimental/GeometryProcessing/Source/DynamicMesh/Private/Operations/MeshConvexHull.cpp b/Engine/Plugins/Experimental/GeometryProcessing/Source/DynamicMesh/Private/Operations/MeshConvexHull.cpp
--- a/Engine/Plugins/Experimental/GeometryProcessing/Source/DynamicMesh/Private/Operations/MeshConvexHull.cpp
+++ b/Engine/Plugins/Experimental/GeometryProcessing/Source/DynamicMesh/Private/Operations/MeshConvexHull.cpp
@@ -7,8 +7,56 @@
 #include "ConvexHull3.h"
 
 
-
 bool FMeshConvexHull::Compute()
+{
+	bool bOK = false;
+	if (VertexSet.Num() > 0)
+	{
+		bOK = Compute_VertexSubset();
+	}
+	else
+	{
+		bOK = Compute_FullMesh();
+	}
+	if (!bOK)
+	{
+		return false;
+	}
+
+
+	if (bPostSimplify)
+	{
+		bool bSimplified = false;
+		if (ConvexHull.TriangleCount() > MaxTargetFaceCount)
+		{
+			FVolPresMeshSimplification Simplifier(&ConvexHull);
+			Simplifier.CollapseMode = FVolPresMeshSimplification::ESimplificationCollapseModes::MinimalExistingVertexError;
+			Simplifier.SimplifyToTriangleCount(MaxTargetFaceCount);
+			bSimplified = true;
+		}
+
+
+		if (bSimplified)
+		{
+			// recalculate convex hull
+			// TODO: test if simplified mesh is convex first, can just re-use in that case!!
+			FMeshConvexHull SimplifiedHull(&ConvexHull);
+			if (SimplifiedHull.Compute())
+			{
+				ConvexHull = MoveTemp(SimplifiedHull.ConvexHull);
+			}
+		}
+
+	}
+
+	return bOK;
+}
+
+
+
+
+
+bool FMeshConvexHull::Compute_FullMesh()
 {
 	FConvexHull3d HullCompute;
 	bool bOK = HullCompute.Solve(Mesh->MaxVertexID(),
@@ -43,32 +91,44 @@ bool FMeshConvexHull::Compute()
 		ConvexHull.AppendTriangle(Triangle);
 	});
 
-	if (bPostSimplify)
+	return true;
+}
+
+
+
+bool FMeshConvexHull::Compute_VertexSubset()
+{
+	FConvexHull3d HullCompute;
+	bool bOK = HullCompute.Solve(VertexSet.Num(),
+		[this](int32 Index) { return Mesh->GetVertex(VertexSet[Index]); });
+	if (!bOK)
 	{
-		bool bSimplified = false;
-		if (ConvexHull.TriangleCount() > MaxTargetFaceCount)
-		{
-			FVolPresMeshSimplification Simplifier(&ConvexHull);
-			Simplifier.CollapseMode = FVolPresMeshSimplification::ESimplificationCollapseModes::MinimalExistingVertexError;
-			Simplifier.SimplifyToTriangleCount(MaxTargetFaceCount);
-			bSimplified = true;
-		}
+		return false;
+	}
 
+	TMap<int32, int32> HullVertMap;
 
-		if (bSimplified)
+	ConvexHull = FDynamicMesh3(EMeshComponents::None);
+	HullCompute.GetTriangles([&](FIndex3i Triangle)
+	{
+		for (int32 j = 0; j < 3; ++j)
 		{
-			// recalculate convex hull
-			// TODO: test if simplified mesh is convex first, can just re-use in that case!!
-			FMeshConvexHull SimplifiedHull(&ConvexHull);
-			if (SimplifiedHull.Compute())
+			int32 Index = Triangle[j];
+			if (HullVertMap.Contains(Index) == false)
 			{
-				ConvexHull = MoveTemp(SimplifiedHull.ConvexHull);
+				FVector3d OrigPos = Mesh->GetVertex(VertexSet[Index]);
+				int32 NewVID = ConvexHull.AppendVertex(OrigPos);
+				HullVertMap.Add(Index, NewVID);
+				Triangle[j] = NewVID;
+			}
+			else
+			{
+				Triangle[j] = HullVertMap[Index];
 			}
 		}
 
-	}
-
+		ConvexHull.AppendTriangle(Triangle);
+	});
 
 	return true;
-
 }
\ No newline at end of file
diff --git a/Engine/Plugins/Experimental/GeometryProcessing/Source/DynamicMesh/Private/Sampling/MeshImageBakingCache.cpp b/Engine/Plugins/Experimental/GeometryProcessing/Source/DynamicMesh/Private/Sampling/MeshImageBakingCache.cpp
--- a/Engine/Plugins/Experimental/GeometryProcessing/Source/DynamicMesh/Private/Sampling/MeshImageBakingCache.cpp
+++ b/Engine/Plugins/Experimental/GeometryProcessing/Source/DynamicMesh/Private/Sampling/MeshImageBakingCache.cpp
@@ -39,6 +39,14 @@ void FMeshImageBakingCache::SetUVLayer(int32 UVLayerIn)
 }
 
 
+void FMeshImageBakingCache::SetThickness(double ThicknessIn)
+{
+	Thickness = ThicknessIn;
+	InvalidateSamples();
+	InvalidateOccupancy();   // do we need to do this?
+}
+
+
 const FDynamicMeshNormalOverlay* FMeshImageBakingCache::GetDetailNormals() const
 {
 	check(DetailMesh && DetailMesh->HasAttributes());
@@ -82,10 +90,15 @@ void FMeshImageBakingCache::InvalidateOccupancy()
 
 /**
  * Find point on Detail mesh that corresponds to point on Base mesh.
- * If nearest point on Detail mesh is within DistanceThreshold, uses that point (cleanly handles coplanar/etc).
- * Otherwise casts a ray in Normal direction.
- * If Normal-direction ray misses, use reverse direction.
- * If both miss, we return false, no correspondence found
+ * Strategy is:
+ *    1) cast a ray inwards along -Normal from BasePoint + Thickness*Normal
+ *    2) cast a ray outwards along Normal from BasePoint 
+ *    3) cast a ray inwards along -Normal from BasePoint
+ * We take (1) preferentially, and then (2), and then (3)
+ * 
+ * If all of those fail, we try to find a nearest-point within distance Thickess
+ * 
+ * If all the above fail, return false
  */
 static bool GetDetailTrianglePoint(
 	const FDynamicMesh3& DetailMesh,
@@ -94,43 +107,67 @@ static bool GetDetailTrianglePoint(
 	const FVector3d& BaseNormal,
 	int32& DetailTriangleOut,
 	FVector3d& DetailTriBaryCoords,
-	double DistanceThreshold = FMathf::ZeroTolerance * 100.0f)
+	double Thickness)
 {
-	// check if we are within on-surface tolerance, if so we use nearest point
-	IMeshSpatial::FQueryOptions OnSurfQueryOptions;
-	OnSurfQueryOptions.MaxDistance = DistanceThreshold;
-	double NearDistSqr = 0;
-	int32 NearestTriID = DetailSpatial.FindNearestTriangle(BasePoint, NearDistSqr, OnSurfQueryOptions);
-	if (DetailMesh.IsTriangle(NearestTriID))
-	{
-		DetailTriangleOut = NearestTriID;
-		FDistPoint3Triangle3d DistQuery = TMeshQueries<FDynamicMesh3>::TriangleDistance(DetailMesh, NearestTriID, BasePoint);
-		DetailTriBaryCoords = DistQuery.TriangleBaryCoords;
-		return true;
-	}
-
 	// TODO: should we check normals here? inverse normal should probably not be considered valid
 
 	// shoot rays forwards and backwards
-	FRay3d Ray(BasePoint, BaseNormal), BackwardsRay(BasePoint, -BaseNormal);
-	int32 HitTID = IndexConstants::InvalidID, BackwardHitTID = IndexConstants::InvalidID;
-	double HitDist, BackwardHitDist;
-	bool bHitForward = DetailSpatial.FindNearestHitTriangle(Ray, HitDist, HitTID);
-	bool bHitBackward = DetailSpatial.FindNearestHitTriangle(BackwardsRay, BackwardHitDist, BackwardHitTID);
-
-	// use the backwards hit if it is closer than the forwards hit
-	if ((bHitBackward && bHitForward == false) || (bHitForward && bHitBackward && BackwardHitDist < HitDist))
+	FRay3d InwardRay = FRay3d(BasePoint + Thickness * BaseNormal, -BaseNormal);
+	FRay3d ForwardRay(BasePoint, BaseNormal);
+	FRay3d BackwardRay(BasePoint, -BaseNormal);
+	int32 ForwardHitTID = IndexConstants::InvalidID, InwardHitTID = IndexConstants::InvalidID, BackwardHitTID = IndexConstants::InvalidID;
+	double ForwardHitDist, InwardHitDist, BackwardHitDist;
+
+	IMeshSpatial::FQueryOptions Options;
+	Options.MaxDistance = Thickness;
+	bool bHitInward = DetailSpatial.FindNearestHitTriangle(InwardRay, InwardHitDist, InwardHitTID, Options);
+	bool bHitForward = DetailSpatial.FindNearestHitTriangle(ForwardRay, ForwardHitDist, ForwardHitTID, Options);
+	bool bHitBackward = DetailSpatial.FindNearestHitTriangle(BackwardRay, BackwardHitDist, BackwardHitTID, Options);
+
+	FRay3d HitRay;
+	int32 HitTID = IndexConstants::InvalidID;
+	double HitDist = TNumericLimits<double>::Max();
+
+	if (bHitInward)
 	{
-		Ray = BackwardsRay;
+		HitRay = InwardRay;
+		HitTID = InwardHitTID;
+		HitDist = InwardHitDist;
+	}
+	else if (bHitForward)
+	{
+		HitRay = ForwardRay;
+		HitTID = ForwardHitTID;
+		HitDist = ForwardHitDist;
+	}
+	else if (bHitBackward)
+	{
+		HitRay = BackwardRay;
 		HitTID = BackwardHitTID;
 		HitDist = BackwardHitDist;
 	}
 
+	// if we did not find any hits, try nearest-point
+	if (DetailMesh.IsTriangle(HitTID) == false)
+	{
+		IMeshSpatial::FQueryOptions OnSurfQueryOptions;
+		OnSurfQueryOptions.MaxDistance = Thickness;
+		double NearDistSqr = 0;
+		int32 NearestTriID = DetailSpatial.FindNearestTriangle(BasePoint, NearDistSqr, OnSurfQueryOptions);
+		if (DetailMesh.IsTriangle(NearestTriID))
+		{
+			DetailTriangleOut = NearestTriID;
+			FDistPoint3Triangle3d DistQuery = TMeshQueries<FDynamicMesh3>::TriangleDistance(DetailMesh, NearestTriID, BasePoint);
+			DetailTriBaryCoords = DistQuery.TriangleBaryCoords;
+			return true;
+		}
+	}
+
 	// if we got a valid ray hit, use it
 	if (DetailMesh.IsTriangle(HitTID))
 	{
 		DetailTriangleOut = HitTID;
-		FIntrRay3Triangle3d IntrQuery = TMeshQueries<FDynamicMesh3>::TriangleIntersection(DetailMesh, HitTID, Ray);
+		FIntrRay3Triangle3d IntrQuery = TMeshQueries<FDynamicMesh3>::TriangleIntersection(DetailMesh, HitTID, HitRay);
 		DetailTriBaryCoords = IntrQuery.TriangleBaryCoords;
 		return true;
 	}
@@ -202,9 +239,11 @@ bool FMeshImageBakingCache::ValidateCache()
 
 			ValueOut.BaseSample = SampleInfo;
 
+			double SampleThickness = this->GetThickness();		// could modulate w/ a map here...
+
 			// find detail mesh triangle point
 			bool bFoundTri = GetDetailTrianglePoint(*DetailMesh, *DetailSpatial, SampleInfo.SurfacePoint, RayDir,
-				ValueOut.DetailTriID, ValueOut.DetailBaryCoords);
+				ValueOut.DetailTriID, ValueOut.DetailBaryCoords, SampleThickness);
 			if (!bFoundTri)
 			{
 				ValueOut.DetailTriID = FDynamicMesh3::InvalidID;
diff --git a/Engine/Plugins/Experimental/GeometryProcessing/Source/DynamicMesh/Private/Sampling/MeshNormalMapBaker.cpp b/Engine/Plugins/Experimental/GeometryProcessing/Source/DynamicMesh/Private/Sampling/MeshNormalMapBaker.cpp
--- a/Engine/Plugins/Experimental/GeometryProcessing/Source/DynamicMesh/Private/Sampling/MeshNormalMapBaker.cpp
+++ b/Engine/Plugins/Experimental/GeometryProcessing/Source/DynamicMesh/Private/Sampling/MeshNormalMapBaker.cpp
@@ -11,6 +11,10 @@ void FMeshNormalMapBaker::Bake()
 	const FDynamicMeshNormalOverlay* DetailNormalOverlay = BakeCache->GetDetailNormals();
 	check(DetailNormalOverlay);
 
+	// make sure we have per-triangle tangents computed so that GetInterpolatedTriangleTangent() below works
+	int32 MaxTriID = BakeCache->GetBakeTargetMesh()->MaxTriangleID();
+	check( BaseMeshTangents->GetTangents().Num() >= 3*BakeCache->GetBakeTargetMesh()->MaxTriangleID());
+
 	auto NormalSampleFunction = [&](const FMeshImageBakingCache::FCorrespondenceSample& SampleData)
 	{
 		int32 DetailTriID = SampleData.DetailTriID;
@@ -41,6 +45,8 @@ void FMeshNormalMapBaker::Bake()
 
 	NormalsBuilder = MakeUnique<TImageBuilder<FVector3f>>();
 	NormalsBuilder->SetDimensions(BakeCache->GetDimensions());
+	FVector3f DefaultMapNormal = (DefaultNormal + FVector3f::One()) * 0.5;
+	NormalsBuilder->Clear(DefaultMapNormal);
 
 
 	BakeCache->EvaluateSamples([&](const FVector2i& Coords, const FMeshImageBakingCache::FCorrespondenceSample& Sample)
diff --git a/Engine/Plugins/Experimental/GeometryProcessing/Source/DynamicMesh/Private/Sampling/MeshResampleImageBaker.cpp b/Engine/Plugins/Experimental/GeometryProcessing/Source/DynamicMesh/Private/Sampling/MeshResampleImageBaker.cpp
--- a/Engine/Plugins/Experimental/GeometryProcessing/Source/DynamicMesh/Private/Sampling/MeshResampleImageBaker.cpp
+++ b/Engine/Plugins/Experimental/GeometryProcessing/Source/DynamicMesh/Private/Sampling/MeshResampleImageBaker.cpp
@@ -28,6 +28,7 @@ void FMeshResampleImageBaker::Bake()
 
 	ResultBuilder = MakeUnique<TImageBuilder<FVector4f>>();
 	ResultBuilder->SetDimensions(BakeCache->GetDimensions());
+	ResultBuilder->Clear(DefaultColor);
 
 	BakeCache->EvaluateSamples([&](const FVector2i& Coords, const FMeshImageBakingCache::FCorrespondenceSample& Sample)
 	{
diff --git a/Engine/Plugins/Experimental/GeometryProcessing/Source/DynamicMesh/Public/MeshNormals.h b/Engine/Plugins/Experimental/GeometryProcessing/Source/DynamicMesh/Public/MeshNormals.h
--- a/Engine/Plugins/Experimental/GeometryProcessing/Source/DynamicMesh/Public/MeshNormals.h
+++ b/Engine/Plugins/Experimental/GeometryProcessing/Source/DynamicMesh/Public/MeshNormals.h
@@ -138,6 +138,13 @@ public:
 	 */
 	static void InitializeOverlayToPerTriangleNormals(FDynamicMeshNormalOverlay* NormalOverlay);
 
+
+	static void InitializeOverlayTopologyFromOpeningAngle(const FDynamicMesh3* Mesh, FDynamicMeshNormalOverlay* NormalOverlay, 
+		double AngleThresholdDeg);
+
+	static void InitializeOverlayTopologyFromFaceGroups(const FDynamicMesh3* Mesh, FDynamicMeshNormalOverlay* NormalOverlay);
+
+
 	/**
 	 * Initialize the given Mesh with per-face normals, ie separate overlay element for each vertex of each triangle.
 	 */
@@ -154,7 +161,8 @@ public:
 	 * Compute overlay normals for the given mesh
 	 * @param bInvert if true, normals are flipped
 	 */
-	static bool QuickRecomputeOverlayNormals(FDynamicMesh3& Mesh, bool bInvert = false);
+	static bool QuickRecomputeOverlayNormals(FDynamicMesh3& Mesh, bool bInvert = false, bool bWeightByArea = true, bool bWeightByAngle = true);
+
 
 protected:
 	/** Compute per-vertex normals using area-weighted averaging of one-ring triangle normals */
diff --git a/Engine/Plugins/Experimental/GeometryProcessing/Source/DynamicMesh/Public/MeshTangents.h b/Engine/Plugins/Experimental/GeometryProcessing/Source/DynamicMesh/Public/MeshTangents.h
--- a/Engine/Plugins/Experimental/GeometryProcessing/Source/DynamicMesh/Public/MeshTangents.h
+++ b/Engine/Plugins/Experimental/GeometryProcessing/Source/DynamicMesh/Public/MeshTangents.h
@@ -23,7 +23,7 @@ struct DYNAMICMESH_API FComputeTangentsOptions
  * tangent vectors for a FDynamicMesh.
  */
 template<typename RealType>
-class DYNAMICMESH_API TMeshTangents
+class TMeshTangents
 {
 protected:
 	/** Target Mesh */
@@ -126,10 +126,18 @@ public:
 	void GetInterpolatedTriangleTangent(int32 TriangleID, const FVector3<RealType>& BaryCoords, FVector3<RealType>& TangentOut, FVector3<RealType>& BitangentOut) const
 	{
 		int32 k = TriangleID * 3;
-		TangentOut = FVector3<RealType>::Blend3(Tangents[k], Tangents[k+1], Tangents[k+2], BaryCoords.X, BaryCoords.Y, BaryCoords.Z);
-		TangentOut.Normalize();
-		BitangentOut = FVector3<RealType>::Blend3(Bitangents[k], Bitangents[k+1], Bitangents[k+2], BaryCoords.X, BaryCoords.Y, BaryCoords.Z);
-		BitangentOut.Normalize();
+		if (k >= 0 && (k+2) < Tangents.Num())
+		{
+			TangentOut = FVector3<RealType>::Blend3(Tangents[k], Tangents[k+1], Tangents[k+2], BaryCoords.X, BaryCoords.Y, BaryCoords.Z);
+			TangentOut.Normalize();
+			BitangentOut = FVector3<RealType>::Blend3(Bitangents[k], Bitangents[k+1], Bitangents[k+2], BaryCoords.X, BaryCoords.Y, BaryCoords.Z);
+			BitangentOut.Normalize();
+		}
+		else
+		{
+			TangentOut = FVector3<RealType>::UnitX();
+			BitangentOut = FVector3<RealType>::UnitY();
+		}
 	}
 
 
diff --git a/Engine/Plugins/Experimental/GeometryProcessing/Source/DynamicMesh/Public/Operations/MeshConvexHull.h b/Engine/Plugins/Experimental/GeometryProcessing/Source/DynamicMesh/Public/Operations/MeshConvexHull.h
--- a/Engine/Plugins/Experimental/GeometryProcessing/Source/DynamicMesh/Public/Operations/MeshConvexHull.h
+++ b/Engine/Plugins/Experimental/GeometryProcessing/Source/DynamicMesh/Public/Operations/MeshConvexHull.h
@@ -16,6 +16,9 @@ public:
 	/** Input Mesh */
 	const FDynamicMesh3* Mesh;
 
+	/** If set, hull will be computed on subset of vertices */
+	TArray<int32> VertexSet;
+
 	/** If true, output convex hull is simplified down to MaxTargetFaceCount */
 	bool bPostSimplify = false;
 	/** Target triangle count of the output Convex Hull */
@@ -37,4 +40,9 @@ public:
 	 */
 	bool Compute();
 
+
+
+protected:
+	bool Compute_FullMesh();
+	bool Compute_VertexSubset();
 };
\ No newline at end of file
diff --git a/Engine/Plugins/Experimental/GeometryProcessing/Source/DynamicMesh/Public/Sampling/MeshImageBakingCache.h b/Engine/Plugins/Experimental/GeometryProcessing/Source/DynamicMesh/Public/Sampling/MeshImageBakingCache.h
--- a/Engine/Plugins/Experimental/GeometryProcessing/Source/DynamicMesh/Public/Sampling/MeshImageBakingCache.h
+++ b/Engine/Plugins/Experimental/GeometryProcessing/Source/DynamicMesh/Public/Sampling/MeshImageBakingCache.h
@@ -19,10 +19,11 @@ public:
 
 	void SetDimensions(FImageDimensions Dimensions);
 	void SetUVLayer(int32 UVLayer);
-
+	void SetThickness(double Thickness);
 
 	FImageDimensions GetDimensions() const { return Dimensions; }
 	int32 GetUVLayer() const { return UVLayer; }
+	double GetThickness() const { return Thickness; }
 
 	const FDynamicMesh3* GetBakeTargetMesh() const { return TargetMesh; }
 	const FDynamicMeshUVOverlay* GetBakeTargetUVs() const;
@@ -61,7 +62,7 @@ protected:
 
 	FImageDimensions Dimensions;
 	int32 UVLayer;
-
+	double Thickness;
 
 	TDenseGrid2<FCorrespondenceSample> SampleMap;
 	bool bSamplesValid = false;
diff --git a/Engine/Plugins/Experimental/GeometryProcessing/Source/DynamicMesh/Public/Sampling/MeshNormalMapBaker.h b/Engine/Plugins/Experimental/GeometryProcessing/Source/DynamicMesh/Public/Sampling/MeshNormalMapBaker.h
--- a/Engine/Plugins/Experimental/GeometryProcessing/Source/DynamicMesh/Public/Sampling/MeshNormalMapBaker.h
+++ b/Engine/Plugins/Experimental/GeometryProcessing/Source/DynamicMesh/Public/Sampling/MeshNormalMapBaker.h
@@ -20,6 +20,7 @@ public:
 
 	const TMeshTangents<double>* BaseMeshTangents;
 
+	FVector3f DefaultNormal = FVector3f(0, 0, 1);
 
 	//
 	// Compute functions
diff --git a/Engine/Plugins/Experimental/GeometryProcessing/Source/DynamicMesh/Public/Sampling/MeshResampleImageBaker.h b/Engine/Plugins/Experimental/GeometryProcessing/Source/DynamicMesh/Public/Sampling/MeshResampleImageBaker.h
--- a/Engine/Plugins/Experimental/GeometryProcessing/Source/DynamicMesh/Public/Sampling/MeshResampleImageBaker.h
+++ b/Engine/Plugins/Experimental/GeometryProcessing/Source/DynamicMesh/Public/Sampling/MeshResampleImageBaker.h
@@ -20,6 +20,8 @@ public:
 	TFunction<FVector4f(FVector2d)> SampleFunction = [](FVector2d Position) { return FVector4f::Zero(); };
 	const FDynamicMeshUVOverlay* DetailUVOverlay = nullptr;
 
+	FVector4f DefaultColor = FVector4f(0, 0, 0, 1);
+
 	//
 	// Compute functions
 	//
diff --git a/Engine/Plugins/Experimental/GeometryProcessing/Source/GeometricObjects/Private/Util/ColorConstants.cpp b/Engine/Plugins/Experimental/GeometryProcessing/Source/GeometricObjects/Private/Util/ColorConstants.cpp
--- a/Engine/Plugins/Experimental/GeometryProcessing/Source/GeometricObjects/Private/Util/ColorConstants.cpp
+++ b/Engine/Plugins/Experimental/GeometryProcessing/Source/GeometricObjects/Private/Util/ColorConstants.cpp
@@ -14,5 +14,5 @@ FColor LinearColors::SelectFColor(int32 Index)
 		MediumAquamarine3b(), Thistle3b(), Wheat3b(),
 		LightSkyBlue3b(), LightPink3b(), MediumSpringGreen3b()
 	};
-	return (Index < 0) ? FColor::White : ColorMap[Index % (7 * 3)];
+	return (Index <= 0) ? White3b() : ColorMap[(Index-1) % (7*3)];
 }
\ No newline at end of file
diff --git a/Engine/Plugins/Experimental/GeometryProcessing/Source/GeometricObjects/Public/Image/ImageBuilder.h b/Engine/Plugins/Experimental/GeometryProcessing/Source/GeometricObjects/Public/Image/ImageBuilder.h
--- a/Engine/Plugins/Experimental/GeometryProcessing/Source/GeometricObjects/Public/Image/ImageBuilder.h
+++ b/Engine/Plugins/Experimental/GeometryProcessing/Source/GeometricObjects/Public/Image/ImageBuilder.h
@@ -86,6 +86,22 @@ public:
 	}
 
 
+	/**
+	 * Convert to a different data type of same Dimensions using ConvertFunc
+	 */
+	template<typename OtherType>
+	void Convert(TFunctionRef<OtherType(const PixelType&)> ConvertFunc,
+				 TImageBuilder<OtherType>& ConvertedImageOut) const
+	{
+		ConvertedImageOut.SetDimensions(Dimensions);
+		int64 NumPixels = Dimensions.Num();
+		for (int64 k = 0; k < NumPixels; ++k)
+		{
+			ConvertedImageOut.Image[k] = ConvertFunc(Image[k]);
+		}
+	}
+
+
 	/**
 	 * Sample the image value at floating-point pixel coords with Bilinear interpolation
 	 */
diff --git a/Engine/Plugins/Experimental/GeometryProcessing/Source/GeometricObjects/Public/Implicit/Morphology.h b/Engine/Plugins/Experimental/GeometryProcessing/Source/GeometricObjects/Public/Implicit/Morphology.h
--- a/Engine/Plugins/Experimental/GeometryProcessing/Source/GeometricObjects/Public/Implicit/Morphology.h
+++ b/Engine/Plugins/Experimental/GeometryProcessing/Source/GeometricObjects/Public/Implicit/Morphology.h
@@ -45,7 +45,7 @@ public:
 	///
 	/// Inputs
 	///
-	TriangleMeshType* Source = nullptr;
+	const TriangleMeshType* Source = nullptr;
 	TMeshAABBTree3<TriangleMeshType>* SourceSpatial = nullptr;
 	EMorphologyOp MorphologyOp = EMorphologyOp::Dilate;
 
diff --git a/Engine/Plugins/Experimental/GeometryProcessing/Source/GeometricObjects/Public/Implicit/Solidify.h b/Engine/Plugins/Experimental/GeometryProcessing/Source/GeometricObjects/Public/Implicit/Solidify.h
--- a/Engine/Plugins/Experimental/GeometryProcessing/Source/GeometricObjects/Public/Implicit/Solidify.h
+++ b/Engine/Plugins/Experimental/GeometryProcessing/Source/GeometricObjects/Public/Implicit/Solidify.h
@@ -16,7 +16,7 @@ class TImplicitSolidify
 {
 public:
 
-	TImplicitSolidify(TriangleMeshType* Source = nullptr, TMeshAABBTree3<TriangleMeshType>* SourceSpatial = nullptr, TFastWindingTree<TriangleMeshType>* SourceWinding = nullptr)
+	TImplicitSolidify(const TriangleMeshType* Source = nullptr, TMeshAABBTree3<TriangleMeshType>* SourceSpatial = nullptr, TFastWindingTree<TriangleMeshType>* SourceWinding = nullptr)
 		: Source(Source), SourceSpatial(SourceSpatial), SourceWinding(SourceWinding)
 	{
 	}
@@ -29,7 +29,7 @@ public:
 	/// Inputs
 	///
 
-	TriangleMeshType* Source = nullptr;
+	const TriangleMeshType* Source = nullptr;
 	TMeshAABBTree3<TriangleMeshType>* SourceSpatial = nullptr;
 	TFastWindingTree<TriangleMeshType>* SourceWinding = nullptr;
 
diff --git a/Engine/Plugins/Experimental/GeometryProcessing/Source/GeometricObjects/Public/Sampling/SampleSetStatistics.h b/Engine/Plugins/Experimental/GeometryProcessing/Source/GeometricObjects/Public/Sampling/SampleSetStatistics.h
--- a/Engine/Plugins/Experimental/GeometryProcessing/Source/GeometricObjects/Public/Sampling/SampleSetStatistics.h
+++ b/Engine/Plugins/Experimental/GeometryProcessing/Source/GeometricObjects/Public/Sampling/SampleSetStatistics.h
@@ -95,23 +95,23 @@ typedef TSampleSetStatistics<double> FSampleSetStatisticsd;
 template<typename RealType>
 struct TSampleSetStatisticBuilder
 {
-	int32 Num;
+	int32 NumStatistics;
 	TArray<TSampleSetStatistics<RealType>> Statistics;
 	TArray<typename TSampleSetStatistics<RealType>::FFixedCountData> FixedCountBuildData;
 
 
-	TSampleSetStatisticBuilder(int32 Count)
+	TSampleSetStatisticBuilder(int32 NumStatisticsIn = 1)
 	{
-		Num = Count;
-		Statistics.Init(TSampleSetStatistics<RealType>(), Num);
-		FixedCountBuildData.SetNum(Num);
+		NumStatistics = NumStatisticsIn;
+		Statistics.Init(TSampleSetStatistics<RealType>(), NumStatistics);
+		FixedCountBuildData.SetNum(NumStatistics);
 	}
 
 	const TSampleSetStatistics<RealType>& operator[](int32 Index) const { return Statistics[Index]; }
 
 	void Begin_FixedCount(int64 Count)
 	{
-		for (int32 j = 0; j < Num; ++j)
+		for (int32 j = 0; j < NumStatistics; ++j)
 		{
 			FixedCountBuildData[j] = Statistics[j].Begin_FixedCount(Count);
 		}
@@ -119,7 +119,7 @@ struct TSampleSetStatisticBuilder
 
 	void StartSecondPass_FixedCount()
 	{
-		for (int32 j = 0; j < Num; ++j)
+		for (int32 j = 0; j < NumStatistics; ++j)
 		{
 			Statistics[j].StartSecondPass_FixedCount(FixedCountBuildData[j]);
 		}
@@ -127,15 +127,34 @@ struct TSampleSetStatisticBuilder
 
 	void CompleteSecondPass_FixedCount()
 	{
-		for (int32 j = 0; j < Num; ++j)
+		for (int32 j = 0; j < NumStatistics; ++j)
 		{
 			Statistics[j].CompleteSecondPass_FixedCount(FixedCountBuildData[j]);
 		}
 	}
 
-	void AccumulateValue_FixedCount(int32 Index, const RealType Value)
+	void AccumulateValue_FixedCount(int32 StatisticIndex, const RealType Value)
 	{
-		Statistics[Index].AccumulateValue_FixedCount(Value, FixedCountBuildData[Index]);
+		Statistics[StatisticIndex].AccumulateValue_FixedCount(Value, FixedCountBuildData[StatisticIndex]);
 	}
+
+
+	template<typename EnumerableType>
+	void ComputeMultiPass(int64 Count, const EnumerableType& MultiEnumerableType)
+	{
+		check(NumStatistics == 1);
+		Begin_FixedCount(Count);
+		for (RealType Value : MultiEnumerableType)
+		{
+			Statistics[0].AccumulateValue_FixedCount(Value, FixedCountBuildData[0]);
+		}
+		StartSecondPass_FixedCount();
+		for (RealType Value : MultiEnumerableType)
+		{
+			Statistics[0].AccumulateValue_FixedCount(Value, FixedCountBuildData[0]);
+		}
+		CompleteSecondPass_FixedCount();
+	}
+
 };
 
diff --git a/Engine/Plugins/Experimental/GeometryProcessing/Source/GeometricObjects/Public/Spatial/DenseGrid2.h b/Engine/Plugins/Experimental/GeometryProcessing/Source/GeometricObjects/Public/Spatial/DenseGrid2.h
--- a/Engine/Plugins/Experimental/GeometryProcessing/Source/GeometricObjects/Public/Spatial/DenseGrid2.h
+++ b/Engine/Plugins/Experimental/GeometryProcessing/Source/GeometricObjects/Public/Spatial/DenseGrid2.h
@@ -42,6 +42,16 @@ public:
 		return DimensionX * DimensionY;
 	}
 
+	int64 Width() const
+	{
+		return DimensionX;
+	}
+
+	int64 Height() const
+	{
+		return DimensionY;
+	}
+
 	FVector2i GetDimensions() const
 	{
 		return FVector2i((int32)DimensionX, (int32)DimensionY);
@@ -114,6 +124,11 @@ public:
 		return Buffer[X + DimensionX * Y];
 	}
 
+	const TArray64<ElemType>& GridValues() const
+	{
+		return Buffer;
+	}
+
 
 	void GetXPair(int32 X0, int32 Y, ElemType& AOut, ElemType& BOut) const
 	{
diff --git a/Engine/Plugins/Experimental/GeometryProcessing/Source/GeometricObjects/Public/Spatial/SampledScalarField2.h b/Engine/Plugins/Experimental/GeometryProcessing/Source/GeometricObjects/Public/Spatial/SampledScalarField2.h
--- a/Engine/Plugins/Experimental/GeometryProcessing/Source/GeometricObjects/Public/Spatial/SampledScalarField2.h
+++ b/Engine/Plugins/Experimental/GeometryProcessing/Source/GeometricObjects/Public/Spatial/SampledScalarField2.h
@@ -7,6 +7,7 @@
 #include "CoreMinimal.h"
 #include "MathUtil.h"
 #include "VectorTypes.h"
+#include "DenseGrid2.h"
 
 
 
@@ -22,8 +23,7 @@ template<typename RealType, typename ValueType>
 class TSampledScalarField2
 {
 public:
-	FVector2i GridDimensions;
-	TArray<ValueType> GridValues;
+	TDenseGrid2<ValueType> GridValues;
 
 	FVector2<RealType> GridOrigin;
 	FVector2<RealType> CellDimensions;
@@ -33,19 +33,40 @@ public:
 	 */
 	TSampledScalarField2()
 	{
-		GridDimensions = FVector2i(2, 2);
-		GridValues.SetNum(4);
+		GridValues.Resize(2, 2);
 		GridOrigin = FVector2<RealType>::Zero();
 		CellDimensions = FVector2<RealType>::One();
 	}
 
+	void CopyConfiguration(const TSampledScalarField2<RealType,ValueType>& OtherField)
+	{
+		GridValues.Resize(OtherField.GridValues.Width(), OtherField.GridValues.Height());
+		GridOrigin = OtherField.GridOrigin;
+		CellDimensions = OtherField.CellDimensions;
+	}
+
 	/**
 	 * Resize the grid to given Width/Height, and initialize w/ given InitValue
 	 */
-	void Resize(int Width, int Height, const ValueType& InitValue)
+	void Resize(int64 Width, int64 Height, const ValueType& InitValue)
+	{
+		GridValues.Resize(Width, Height);
+		GridValues.AssignAll(InitValue);
+	}
+
+	int32 Width() const
 	{
-		GridDimensions = FVector2i(Width, Height);
-		GridValues.Init(InitValue, Width*Height);
+		return (int32)GridValues.Width();
+	}
+
+	int32 Height() const
+	{
+		return (int32)GridValues.Height();
+	}
+
+	int64 Num() const
+	{
+		return GridValues.Size();
 	}
 
 	/**
@@ -77,16 +98,17 @@ public:
 			((Position.Y - GridOrigin.Y) / CellDimensions.Y));
 
 		// compute integer grid coordinates
-		int x0 = (int)GridPoint.X;
-		int x1 = x0 + 1;
-		int y0 = (int)GridPoint.Y;
-		int y1 = y0 + 1;
+		int64 x0 = (int64)GridPoint.X;
+		int64 x1 = x0 + 1;
+		int64 y0 = (int64)GridPoint.Y;
+		int64 y1 = y0 + 1;
 
 		// clamp to valid range
-		x0 = FMath::Clamp(x0, 0, GridDimensions.X - 1);
-		x1 = FMath::Clamp(x1, 0, GridDimensions.X - 1);
-		y0 = FMath::Clamp(y0, 0, GridDimensions.Y - 1);
-		y1 = FMath::Clamp(y1, 0, GridDimensions.Y - 1);
+		int64 Width = GridValues.Width(), Height = GridValues.Height();
+		x0 = FMath::Clamp(x0, (int64)0, Width - 1);
+		x1 = FMath::Clamp(x1, (int64)0, Width - 1);
+		y0 = FMath::Clamp(y0, (int64)0, Height - 1);
+		y1 = FMath::Clamp(y1, (int64)0, Height - 1);
 
 		// convert real-valued grid coords to [0,1] range
 		RealType fAx = FMath::Clamp(GridPoint.X - (RealType)x0, (RealType)0, (RealType)1);
@@ -95,10 +117,10 @@ public:
 		RealType OneMinusfAy = (RealType)1 - fAy;
 
 		// fV## is grid cell corner index
-		const ValueType& fV00 = GridValues[y0*GridDimensions.X + x0];
-		const ValueType& fV01 = GridValues[y0*GridDimensions.X + x1];
-		const ValueType& fV10 = GridValues[y1*GridDimensions.X + x0];
-		const ValueType& fV11 = GridValues[y1*GridDimensions.X + x1];
+		const ValueType& fV00 = GridValues[y0*Width + x0];
+		const ValueType& fV10 = GridValues[y0*Width + x1];
+		const ValueType& fV01 = GridValues[y1*Width + x0];
+		const ValueType& fV11 = GridValues[y1*Width + x1];
 
 		return
 			(OneMinusfAx * OneMinusfAy) * fV00 +
@@ -107,6 +129,59 @@ public:
 			(fAx         * fAy)         * fV11;
 	}
 
+
+
+
+	/**
+	 * Sample scalar field gradient with bilinear interpolation at given Position
+	 * @param Position sample point relative to grid origin/dimensions
+	 * @return interpolated value at this position
+	 */
+	void BilinearSampleGradientClamped(const FVector2<RealType>& Position, ValueType& GradXOut, ValueType& GradYOut) const
+	{
+		// transform Position into grid coordinates
+		FVector2<RealType> GridPoint(
+			((Position.X - GridOrigin.X) / CellDimensions.X),
+			((Position.Y - GridOrigin.Y) / CellDimensions.Y));
+
+		// compute integer grid coordinates
+		int64 x0 = (int64)GridPoint.X;
+		int64 x1 = x0 + 1;
+		int64 y0 = (int64)GridPoint.Y;
+		int64 y1 = y0 + 1;
+
+		// clamp to valid range
+		int64 Width = GridValues.Width(), Height = GridValues.Height();
+		x0 = FMath::Clamp(x0, (int64)0, Width - 1);
+		x1 = FMath::Clamp(x1, (int64)0, Width - 1);
+		y0 = FMath::Clamp(y0, (int64)0, Height - 1);
+		y1 = FMath::Clamp(y1, (int64)0, Height - 1);
+
+		// convert real-valued grid coords to [0,1] range
+		RealType fAx = FMath::Clamp(GridPoint.X - (RealType)x0, (RealType)0, (RealType)1);
+		RealType fAy = FMath::Clamp(GridPoint.Y - (RealType)y0, (RealType)0, (RealType)1);
+		RealType OneMinusfAx = (RealType)1 - fAx;
+		RealType OneMinusfAy = (RealType)1 - fAy;
+
+		// fV## is grid cell corner index
+		const ValueType& fV00 = GridValues[y0 * Width + x0];
+		const ValueType& fV10 = GridValues[y0 * Width + x1];
+		const ValueType& fV01 = GridValues[y1 * Width + x0];
+		const ValueType& fV11 = GridValues[y1 * Width + x1];
+
+		GradXOut =
+			-fV00 * (OneMinusfAy) +
+			-fV01 * (fAy) +
+			fV10 * (OneMinusfAy) +
+			fV11 * (fAy);
+
+		GradYOut =
+			-fV00 * (OneMinusfAx) +
+			fV01 * (OneMinusfAx)  +
+			-fV10 * (fAx) +
+			fV11 * (fAx);
+	}
+
 };
 
 typedef TSampledScalarField2<double, double> FSampledScalarField2d;
diff --git a/Engine/Plugins/Experimental/GeometryProcessing/Source/GeometricObjects/Public/Util/ColorConstants.h b/Engine/Plugins/Experimental/GeometryProcessing/Source/GeometricObjects/Public/Util/ColorConstants.h
--- a/Engine/Plugins/Experimental/GeometryProcessing/Source/GeometricObjects/Public/Util/ColorConstants.h
+++ b/Engine/Plugins/Experimental/GeometryProcessing/Source/GeometricObjects/Public/Util/ColorConstants.h
@@ -212,7 +212,7 @@ namespace LinearColors
 			MediumAquamarine3f(), Thistle3f(), Wheat3f(),
 			LightSkyBlue3f(), LightPink3f(), MediumSpringGreen3f()
 		};
-		return (Index < 0) ? VectorType(White3f()) : VectorType(ColorMap[Index % (7 * 3)]);
+		return VectorType( (Index <= 0) ? White3f() : ColorMap[(Index-1) % (7*3)] );
 	}
 
 
diff --git a/Engine/Plugins/Experimental/GeometryProcessing/Source/GeometricObjects/Public/WeightMapTypes.h b/Engine/Plugins/Experimental/GeometryProcessing/Source/GeometricObjects/Public/WeightMapTypes.h
--- a/Engine/Plugins/Experimental/GeometryProcessing/Source/GeometricObjects/Public/WeightMapTypes.h
+++ b/Engine/Plugins/Experimental/GeometryProcessing/Source/GeometricObjects/Public/WeightMapTypes.h
@@ -18,18 +18,39 @@ public:
 	RealType DefaultValue;
 	TArray<RealType> Values;
 
-	RealType GetValue(int32 Index) const
+	int32 Num() const { return Values.Num(); }
+
+	void SetNum(int32 NewNum)
+	{
+		Values.SetNum(NewNum);
+	}
+
+	RealType GetValueUnsafe(int32 Index) const
 	{
 		return Values[Index];
 	}
 
-	RealType GetInterpValue(const FIndex3i& Indices, const FVector3d& BaryCoords) const
+	RealType GetValue(int32 Index) const
+	{
+		return (Index >= 0 && Index < Values.Num()) ? Values[Index] : DefaultValue;
+	}
+
+	RealType GetInterpValueUnsafe(const FIndex3i& Indices, const FVector3d& BaryCoords) const
 	{
 		return (RealType)((double)Values[Indices.A] * BaryCoords.X
 			+ (double)Values[Indices.B] * BaryCoords.Y
 			+ (double)Values[Indices.C] * BaryCoords.Z);
 	}
 
+	RealType GetInterpValue(const FIndex3i& Indices, const FVector3d& BaryCoords) const
+	{
+		RealType A = GetValue(Indices.A);
+		RealType B = GetValue(Indices.B);
+		RealType C = GetValue(Indices.C);
+
+		return (RealType)((double)A*BaryCoords.X + (double)B*BaryCoords.Y + (double)C*BaryCoords.Z);
+	}
+
 	void InvertWeightMap(TInterval1<RealType> Range = TInterval1<RealType>((RealType)0, (RealType)1.0))
 	{
 		int32 Num = Values.Num();
diff --git a/Engine/Plugins/Experimental/GeometryProcessing/Source/MeshConversion/Private/MeshDescriptionToDynamicMesh.cpp b/Engine/Plugins/Experimental/GeometryProcessing/Source/MeshConversion/Private/MeshDescriptionToDynamicMesh.cpp
--- a/Engine/Plugins/Experimental/GeometryProcessing/Source/MeshConversion/Private/MeshDescriptionToDynamicMesh.cpp
+++ b/Engine/Plugins/Experimental/GeometryProcessing/Source/MeshConversion/Private/MeshDescriptionToDynamicMesh.cpp
@@ -409,12 +409,9 @@ void FMeshDescriptionToDynamicMesh::Convert(const FMeshDescription* MeshIn, FDyn
 
 
 
-
-
-void FMeshDescriptionToDynamicMesh::CopyTangents(const FMeshDescription* SourceMesh, const FDynamicMesh3* TargetMesh, TMeshTangents<float>* TangentsOut)
+template<typename RealType>
+static void CopyTangents_Internal(const FMeshDescription* SourceMesh, const FDynamicMesh3* TargetMesh, TMeshTangents<RealType>* TangentsOut, const TArray<FIndex2i>& TriToPolyTriMap)
 {
-	if (!ensureMsgf(bCalculateMaps, TEXT("Cannot CopyTangents unless Maps were calculated"))) return;
-	if (!ensureMsgf(TriToPolyTriMap.Num() == TargetMesh->TriangleCount(), TEXT("Tried to CopyTangents to mesh with different triangle count"))) return;
 
 	FStaticMeshConstAttributes Attributes(*SourceMesh);
 
@@ -428,11 +425,11 @@ void FMeshDescriptionToDynamicMesh::CopyTangents(const FMeshDescription* SourceM
 
 	TangentsOut->SetMesh(TargetMesh);
 	TangentsOut->InitializeTriVertexTangents(false);
-	
+
 	for (int32 TriID : TargetMesh->TriangleIndicesItr())
 	{
 		FIndex2i PolyTriIdx = TriToPolyTriMap[TriID];
-		TArrayView<const FTriangleID> TriangleIDs = SourceMesh->GetPolygonTriangles( FPolygonID(PolyTriIdx.A) );
+		TArrayView<const FTriangleID> TriangleIDs = SourceMesh->GetPolygonTriangles(FPolygonID(PolyTriIdx.A));
 		const FTriangleID TriangleID = TriangleIDs[PolyTriIdx.B];
 		TArrayView<const FVertexInstanceID> InstanceTri = SourceMesh->GetTriangleVertexInstances(TriangleID);
 		for (int32 j = 0; j < 3; ++j)
@@ -440,14 +437,29 @@ void FMeshDescriptionToDynamicMesh::CopyTangents(const FMeshDescription* SourceM
 			FVector Normal = InstanceNormals[InstanceTri[j]];
 			FVector Tangent = InstanceTangents[InstanceTri[j]];
 			float BitangentSign = InstanceSigns[InstanceTri[j]];
-			FVector3f Bitangent = VectorUtil::Bitangent((FVector3f)Normal, (FVector3f)Tangent, (float)BitangentSign);
+			FVector3<RealType> Bitangent = VectorUtil::Bitangent((FVector3<RealType>)Normal, (FVector3<RealType>)Tangent, (RealType)BitangentSign);
 			Tangent.Normalize(); Bitangent.Normalize();
 			TangentsOut->SetPerTriangleTangent(TriID, j, Tangent, Bitangent);
 		}
 	}
+}
+
+
 
+void FMeshDescriptionToDynamicMesh::CopyTangents(const FMeshDescription* SourceMesh, const FDynamicMesh3* TargetMesh, TMeshTangents<float>* TangentsOut)
+{
+	if (!ensureMsgf(bCalculateMaps, TEXT("Cannot CopyTangents unless Maps were calculated"))) return;
+	if (!ensureMsgf(TriToPolyTriMap.Num() == TargetMesh->TriangleCount(), TEXT("Tried to CopyTangents to mesh with different triangle count"))) return;
+	CopyTangents_Internal<float>(SourceMesh, TargetMesh, TangentsOut, TriToPolyTriMap);
 }
 
 
 
+void FMeshDescriptionToDynamicMesh::CopyTangents(const FMeshDescription* SourceMesh, const FDynamicMesh3* TargetMesh, TMeshTangents<double>* TangentsOut)
+{
+	if (!ensureMsgf(bCalculateMaps, TEXT("Cannot CopyTangents unless Maps were calculated"))) return;
+	if (!ensureMsgf(TriToPolyTriMap.Num() == TargetMesh->TriangleCount(), TEXT("Tried to CopyTangents to mesh with different triangle count"))) return;
+	CopyTangents_Internal<double>(SourceMesh, TargetMesh, TangentsOut, TriToPolyTriMap);
+}
+
 
diff --git a/Engine/Plugins/Experimental/GeometryProcessing/Source/MeshConversion/Public/MeshDescriptionToDynamicMesh.h b/Engine/Plugins/Experimental/GeometryProcessing/Source/MeshConversion/Public/MeshDescriptionToDynamicMesh.h
--- a/Engine/Plugins/Experimental/GeometryProcessing/Source/MeshConversion/Public/MeshDescriptionToDynamicMesh.h
+++ b/Engine/Plugins/Experimental/GeometryProcessing/Source/MeshConversion/Public/MeshDescriptionToDynamicMesh.h
@@ -62,4 +62,11 @@ public:
 	 * @warning Convert() must have been used to create the TargetMesh before calling this function
 	 */
 	void CopyTangents(const FMeshDescription* SourceMesh, const FDynamicMesh3* TargetMesh, TMeshTangents<float>* TangentsOut);
+
+	/**
+	 * Copy tangents from MeshDescription to a FMeshTangents instance.
+	 * @warning Convert() must have been used to create the TargetMesh before calling this function
+	 */
+	void CopyTangents(const FMeshDescription* SourceMesh, const FDynamicMesh3* TargetMesh, TMeshTangents<double>* TangentsOut);
+
 };
\ No newline at end of file
diff --git a/Engine/Plugins/Experimental/MeshModelingToolset/Source/MeshModelingTools/Private/BakeMeshAttributeMapsTool.cpp b/Engine/Plugins/Experimental/MeshModelingToolset/Source/MeshModelingTools/Private/BakeMeshAttributeMapsTool.cpp
--- a/Engine/Plugins/Experimental/MeshModelingToolset/Source/MeshModelingTools/Private/BakeMeshAttributeMapsTool.cpp
+++ b/Engine/Plugins/Experimental/MeshModelingToolset/Source/MeshModelingTools/Private/BakeMeshAttributeMapsTool.cpp
@@ -153,7 +153,7 @@ void UBakeMeshAttributeMapsTool::Setup()
 	Settings->WatchProperty(Settings->Resolution, [this](EBakeTextureResolution) { bResultValid = false; });
 	Settings->WatchProperty(Settings->UVLayer, [this](FString) { bResultValid = false; });
 	Settings->WatchProperty(Settings->bUseWorldSpace, [this](bool) { bDetailMeshValid = false; bResultValid = false; });
-
+	Settings->WatchProperty(Settings->Thickness, [this](float) { bResultValid = false; });
 
 	NormalMapProps = NewObject<UBakedNormalMapToolProperties>(this);
 	NormalMapProps->RestoreProperties(this);
@@ -353,6 +353,7 @@ void UBakeMeshAttributeMapsTool::UpdateResult()
 	{
 		UpdateDetailMesh();
 		bDetailMeshValid = true;
+		CachedBakeCacheSettings = FBakeCacheSettings();
 	}
 
 	if (bResultValid)
@@ -370,6 +371,7 @@ void UBakeMeshAttributeMapsTool::UpdateResult()
 	BakeCacheSettings.Dimensions = Dimensions;
 	BakeCacheSettings.UVLayer = FCString::Atoi(*Settings->UVLayer);
 	BakeCacheSettings.DetailTimestamp = this->DetailMeshTimestamp;
+	BakeCacheSettings.Thickness = Settings->Thickness;
 
 	// rebuild bake cache if we need to
 	if (!(CachedBakeCacheSettings == BakeCacheSettings))
@@ -379,7 +381,15 @@ void UBakeMeshAttributeMapsTool::UpdateResult()
 		BakeCache->SetBakeTargetMesh(&BaseMesh);
 		BakeCache->SetDimensions(Dimensions);
 		BakeCache->SetUVLayer(BakeCacheSettings.UVLayer);
+		BakeCache->SetThickness(BakeCacheSettings.Thickness);
 		BakeCache->ValidateCache();
+		CachedBakeCacheSettings = BakeCacheSettings;
+
+		CachedNormalMapSettings = FNormalMapSettings();
+		CachedOcclusionMapSettings = FOcclusionMapSettings();
+		CachedCurvatureMapSettings = FCurvatureMapSettings();
+		CachedMeshPropertyMapSettings = FMeshPropertyMapSettings();
+		CachedTexture2DImageSettings = FTexture2DImageSettings();
 	}
 
 	// rebuild select map type
diff --git a/Engine/Plugins/Experimental/MeshModelingToolset/Source/MeshModelingTools/Private/DisplaceMeshTool.cpp b/Engine/Plugins/Experimental/MeshModelingToolset/Source/MeshModelingTools/Private/DisplaceMeshTool.cpp
--- a/Engine/Plugins/Experimental/MeshModelingToolset/Source/MeshModelingTools/Private/DisplaceMeshTool.cpp
+++ b/Engine/Plugins/Experimental/MeshModelingToolset/Source/MeshModelingTools/Private/DisplaceMeshTool.cpp
@@ -585,15 +585,15 @@ namespace {
 		}
 		else
 		{
-			int TextureWidth = Parameters.DisplacementMap->GetSizeX();
-			int TextureHeight = Parameters.DisplacementMap->GetSizeY();
+			int64 TextureWidth = Parameters.DisplacementMap->GetSizeX();
+			int64 TextureHeight = Parameters.DisplacementMap->GetSizeY();
 			Parameters.DisplaceField.Resize(TextureWidth, TextureHeight, 0.0f);
 			Parameters.DisplaceField.SetCellSize(1.0f / (float)TextureWidth);
 
 			const FColor* FormattedData = TextureAccess.GetData();
-			for (int y = 0; y < TextureHeight; ++y)
+			for (int64 y = 0; y < TextureHeight; ++y)
 			{
-				for (int x = 0; x < TextureWidth; ++x)
+				for (int64 x = 0; x < TextureWidth; ++x)
 				{
 					FColor PixelColor = FormattedData[y * TextureWidth + x];
 					float Value = PixelColor.R / 255.0;
diff --git a/Engine/Plugins/Experimental/MeshModelingToolset/Source/MeshModelingTools/Public/BakeMeshAttributeMapsTool.h b/Engine/Plugins/Experimental/MeshModelingToolset/Source/MeshModelingTools/Public/BakeMeshAttributeMapsTool.h
--- a/Engine/Plugins/Experimental/MeshModelingToolset/Source/MeshModelingTools/Public/BakeMeshAttributeMapsTool.h
+++ b/Engine/Plugins/Experimental/MeshModelingToolset/Source/MeshModelingTools/Public/BakeMeshAttributeMapsTool.h
@@ -81,6 +81,12 @@ public:
 	UPROPERTY(EditAnywhere, Category = MapSettings, meta = (TransientToolProperty))
 	EBakeTextureResolution Resolution = EBakeTextureResolution::Resolution256;
 
+	UPROPERTY(EditAnywhere, Category = MapSettings)
+	bool bUseWorldSpace = false;
+
+	UPROPERTY(EditAnywhere, Category = MapSettings)
+	float Thickness = 3.0;
+
 	/** Which UV layer to use to create the map */
 	UPROPERTY(EditAnywhere, Category = MapSettings, meta = (GetOptions = GetUVLayerNamesFunc))
 	FString UVLayer;
@@ -90,9 +96,6 @@ public:
 	UPROPERTY(meta = (TransientToolProperty))
 	TArray<FString> UVLayerNamesList;
 
-	UPROPERTY(EditAnywhere, Category = MapSettings)
-	bool bUseWorldSpace = false;
-
 	UPROPERTY(VisibleAnywhere, Category = MapSettings, meta = (TransientToolProperty))
 	UTexture2D* Result;
 
@@ -322,10 +325,11 @@ protected:
 		FImageDimensions Dimensions;
 		int32 UVLayer;
 		int32 DetailTimestamp;
+		float Thickness;
 
 		bool operator==(const FBakeCacheSettings& Other) const
 		{
-			return Dimensions == Other.Dimensions && UVLayer == Other.UVLayer && DetailTimestamp == Other.DetailTimestamp;
+			return Dimensions == Other.Dimensions && UVLayer == Other.UVLayer && DetailTimestamp == Other.DetailTimestamp && Thickness == Other.Thickness;
 		}
 	};
 	FBakeCacheSettings CachedBakeCacheSettings;
diff --git a/Engine/Plugins/Experimental/MeshModelingToolset/Source/ModelingComponents/Private/AssetUtils/Texture2DUtil.cpp b/Engine/Plugins/Experimental/MeshModelingToolset/Source/ModelingComponents/Private/AssetUtils/Texture2DUtil.cpp
new file mode 100644
--- /dev/null
+++ b/Engine/Plugins/Experimental/MeshModelingToolset/Source/ModelingComponents/Private/AssetUtils/Texture2DUtil.cpp
@@ -0,0 +1,150 @@
+// Copyright Epic Games, Inc. All Rights Reserved. 
+
+#include "AssetUtils/Texture2DUtil.h"
+
+
+
+
+
+
+static bool ReadTexture_PlatformData(
+	UTexture2D* TextureMap,
+	FImageDimensions& Dimensions,
+	TImageBuilder<FVector4f>& DestImage)
+{
+	// Read from PlatformData
+	// UBlueprintMaterialTextureNodesBPLibrary::Texture2D_SampleUV_EditorOnly() shows how to read from PlatformData
+	// without converting formats if it's already uncompressed. And can read PF_FloatRGBA. Would make sense to do
+	// that when possible, and perhaps is possible to convert to PF_FloatRGBA instead of using TC_VectorDisplacementmap?
+	// 
+	// Note that the current code cannot run on a background thread, UpdateResource() will call FlushRenderingCommands()
+	// which will check() if it's on the Game Thread
+
+	check(TextureMap->PlatformData);
+	int32 Width = TextureMap->PlatformData->Mips[0].SizeX;
+	int32 Height = TextureMap->PlatformData->Mips[0].SizeY;
+	Dimensions = FImageDimensions(Width, Height);
+	DestImage.SetDimensions(Dimensions);
+	int64 Num = Dimensions.Num();
+
+	// convert built platform texture data to uncompressed RGBA8 format
+	TextureCompressionSettings InitialCompressionSettings = TextureMap->CompressionSettings;
+	bool bWasSRGB = TextureMap->SRGB;
+#if WITH_EDITOR
+	TextureMipGenSettings InitialMipGenSettings = TextureMap->MipGenSettings;
+#endif
+	TextureMap->CompressionSettings = TextureCompressionSettings::TC_VectorDisplacementmap;
+	TextureMap->SRGB = false;
+#if WITH_EDITOR
+	TextureMap->MipGenSettings = TextureMipGenSettings::TMGS_NoMipmaps;
+#endif
+	TextureMap->UpdateResource();
+
+	// lock texture and represet as FColor
+	const FColor* FormattedImageData = reinterpret_cast<const FColor*>(TextureMap->PlatformData->Mips[0].BulkData.LockReadOnly());
+
+	// maybe could be done more quickly by row?
+	for (int32 i = 0; i < Num; ++i)
+	{
+		FColor ByteColor = FormattedImageData[i];
+		FLinearColor FloatColor(ByteColor);
+		DestImage.SetPixel(i, FVector4f(FloatColor));
+	}
+
+	// restore built platform texture data to initial state
+	TextureMap->PlatformData->Mips[0].BulkData.Unlock();
+	TextureMap->CompressionSettings = InitialCompressionSettings;
+	TextureMap->SRGB = bWasSRGB;
+#if WITH_EDITOR
+	TextureMap->MipGenSettings = InitialMipGenSettings;
+#endif
+	TextureMap->UpdateResource();
+
+	return true;
+}
+
+
+#if WITH_EDITOR
+static bool ReadTexture_SourceData(
+	UTexture2D* TextureMap,
+	FImageDimensions& Dimensions,
+	TImageBuilder<FVector4f>& DestImage)
+{
+	FTextureSource& TextureSource = TextureMap->Source;
+
+	int32 Width = TextureSource.GetSizeX();
+	int32 Height = TextureSource.GetSizeY();
+	Dimensions = FImageDimensions(Width, Height);
+	DestImage.SetDimensions(Dimensions);
+	int64 Num = Dimensions.Num();
+
+	TArray64<uint8> SourceData;
+	TextureMap->Source.GetMipData(SourceData, 0, 0, 0);
+	ETextureSourceFormat SourceFormat = TextureSource.GetFormat();
+	int32 BytesPerPixel = TextureSource.GetBytesPerPixel();
+	const uint8* SourceDataPtr = SourceData.GetData();
+
+	// code below is derived from UBlueprintMaterialTextureNodesBPLibrary::Texture2D_SampleUV_EditorOnly()
+	if ((SourceFormat == TSF_BGRA8 || SourceFormat == TSF_BGRE8))
+	{
+		check(BytesPerPixel == sizeof(FColor));
+		for (int32 i = 0; i < Num; ++i)
+		{
+			const uint8* PixelPtr = SourceDataPtr + (i * BytesPerPixel);
+			FColor PixelColor = *((FColor*)PixelPtr);
+
+			FLinearColor FloatColor = (TextureMap->SRGB) ?
+				FLinearColor::FromSRGBColor(PixelColor) :
+				FLinearColor(float(PixelColor.R), float(PixelColor.G), float(PixelColor.B), float(PixelColor.A)) / 255.0f;
+
+			DestImage.SetPixel(i, FVector4f(FloatColor));
+		}
+	}
+	else if ((SourceFormat == TSF_RGBA16 || SourceFormat == TSF_RGBA16F))
+	{
+		check(BytesPerPixel == sizeof(FFloat16Color));
+		for (int32 i = 0; i < Num; ++i)
+		{
+			const uint8* PixelPtr = SourceDataPtr + (i * BytesPerPixel);
+			FFloat16Color PixelColor = *((FFloat16Color*)PixelPtr);
+			FLinearColor FloatColor(float(PixelColor.R), float(PixelColor.G), float(PixelColor.B), float(PixelColor.A));
+			DestImage.SetPixel(i, FVector4f(FloatColor));
+		}
+	}
+	else if (SourceFormat == TSF_G8)
+	{
+		check(BytesPerPixel == 1);
+		for (int32 i = 0; i < Num; ++i)
+		{
+			const uint8* PixelPtr = SourceDataPtr + (i * BytesPerPixel);
+			uint8 PixelColor = *PixelPtr;
+			float PixelColorf = float(PixelColor) / 255.0f;
+			FLinearColor FloatColor = (TextureMap->SRGB) ?
+				FLinearColor::FromSRGBColor(FColor(PixelColor, PixelColor, PixelColor, 255)) :
+				FLinearColor(PixelColorf, PixelColorf, PixelColorf, 1.0);
+			DestImage.SetPixel(i, FVector4f(FloatColor));
+		}
+	}
+
+	return true;
+}
+#endif 
+
+
+bool UE::AssetUtils::ReadTexture(
+	UTexture2D* TextureMap,
+	FImageDimensions& Dimensions,
+	TImageBuilder<FVector4f>& DestImage,
+	bool bPreferPlatformData)
+{
+	if (ensure(TextureMap) == false) return false;
+
+#if WITH_EDITOR
+	if (TextureMap->Source.IsValid() && bPreferPlatformData == false)
+	{
+		return ReadTexture_SourceData(TextureMap, Dimensions, DestImage);
+	}
+#endif
+
+	return ReadTexture_PlatformData(TextureMap, Dimensions, DestImage);
+}
\ No newline at end of file
diff --git a/Engine/Plugins/Experimental/MeshModelingToolset/Source/ModelingComponents/Private/Physics/PhysicsDataCollection.cpp b/Engine/Plugins/Experimental/MeshModelingToolset/Source/ModelingComponents/Private/Physics/PhysicsDataCollection.cpp
--- a/Engine/Plugins/Experimental/MeshModelingToolset/Source/ModelingComponents/Private/Physics/PhysicsDataCollection.cpp
+++ b/Engine/Plugins/Experimental/MeshModelingToolset/Source/ModelingComponents/Private/Physics/PhysicsDataCollection.cpp
@@ -15,6 +15,22 @@ void FPhysicsDataCollection::InitializeFromComponent(const UActorComponent* Comp
 	const UStaticMesh* StaticMesh = StaticMeshComponent->GetStaticMesh();
 	
 	SourceComponent = StaticMeshComponent;
+	SourceStaticMesh = StaticMesh;
+	BodySetup = StaticMesh->GetBodySetup();
+
+	ExternalScale3D = FVector(1.f, 1.f, 1.f);
+
+	if (bInitializeAggGeom)
+	{
+		AggGeom = BodySetup->AggGeom;
+		// transfer AggGeom to FSimpleShapeSet3d...
+	}
+}
+
+
+void FPhysicsDataCollection::InitializeFromStaticMesh(const UStaticMesh* StaticMesh, bool bInitializeAggGeom)
+{
+	SourceStaticMesh = StaticMesh;
 	BodySetup = StaticMesh->GetBodySetup();
 
 	ExternalScale3D = FVector(1.f, 1.f, 1.f);
@@ -30,6 +46,7 @@ void FPhysicsDataCollection::InitializeFromComponent(const UActorComponent* Comp
 void FPhysicsDataCollection::InitializeFromExisting(const FPhysicsDataCollection& Other)
 {
 	SourceComponent = Other.SourceComponent;
+	SourceStaticMesh = Other.SourceStaticMesh;
 	BodySetup = Other.BodySetup;
 
 	ExternalScale3D = Other.ExternalScale3D;
diff --git a/Engine/Plugins/Experimental/MeshModelingToolset/Source/ModelingComponents/Public/AssetUtils/Texture2DUtil.h b/Engine/Plugins/Experimental/MeshModelingToolset/Source/ModelingComponents/Public/AssetUtils/Texture2DUtil.h
new file mode 100644
--- /dev/null
+++ b/Engine/Plugins/Experimental/MeshModelingToolset/Source/ModelingComponents/Public/AssetUtils/Texture2DUtil.h
@@ -0,0 +1,26 @@
+// Copyright Epic Games, Inc. All Rights Reserved.
+
+#pragma once
+
+#include "MathUtil.h"
+#include "VectorTypes.h"
+#include "Image/ImageDimensions.h"
+#include "Image/ImageBuilder.h"
+#include "Engine/Classes/Engine/Texture2D.h"
+
+
+
+namespace UE
+{
+namespace AssetUtils
+{
+
+	MODELINGCOMPONENTS_API bool ReadTexture(
+		UTexture2D* TextureMap,
+		FImageDimensions& DimensionsOut,
+		TImageBuilder<FVector4f>& DestImageOut,
+		bool bPreferPlatformData = false);
+
+}
+}
+
diff --git a/Engine/Plugins/Experimental/MeshModelingToolset/Source/ModelingComponents/Public/Physics/PhysicsDataCollection.h b/Engine/Plugins/Experimental/MeshModelingToolset/Source/ModelingComponents/Public/Physics/PhysicsDataCollection.h
--- a/Engine/Plugins/Experimental/MeshModelingToolset/Source/ModelingComponents/Public/Physics/PhysicsDataCollection.h
+++ b/Engine/Plugins/Experimental/MeshModelingToolset/Source/ModelingComponents/Public/Physics/PhysicsDataCollection.h
@@ -28,10 +28,12 @@ public:
 
 	/** The Component this physics data came from / is for */
 	TWeakObjectPtr<const UActorComponent> SourceComponent;
+	/** The StaticMesh this physics data came from / is for */
+	TWeakObjectPtr<const UStaticMesh> SourceStaticMesh;
 	/** The BodySetup in use by the SourceComponent */
 	TWeakObjectPtr<const UBodySetup> BodySetup;
 	/** Scaling factor applied to the SourceComponent, which should be transferred to Collision Geometry in some cases */
-	FVector ExternalScale3D;
+	FVector ExternalScale3D = FVector::OneVector;
 
 
 	//
@@ -54,6 +56,11 @@ public:
 	 */
 	void InitializeFromComponent(const UActorComponent* Component, bool bInitializeAggGeom);
 
+	/**
+	 * Initialize from the given StaticMesh, and optionally initialize internal geometry members
+	 */
+	void InitializeFromStaticMesh(const UStaticMesh* StaticMesh, bool bInitializeAggGeom);
+
 	/**
 	 * Initialize 
 	 */
diff --git a/Engine/Plugins/Experimental/MeshModelingToolset/Source/ModelingOperatorsEditorOnly/Private/CleaningOps/SimplifyMeshOp.cpp b/Engine/Plugins/Experimental/MeshModelingToolset/Source/ModelingOperatorsEditorOnly/Private/CleaningOps/SimplifyMeshOp.cpp
--- a/Engine/Plugins/Experimental/MeshModelingToolset/Source/ModelingOperatorsEditorOnly/Private/CleaningOps/SimplifyMeshOp.cpp
+++ b/Engine/Plugins/Experimental/MeshModelingToolset/Source/ModelingOperatorsEditorOnly/Private/CleaningOps/SimplifyMeshOp.cpp
@@ -65,8 +65,11 @@ void ComputeSimplify(FDynamicMesh3* TargetMesh, const bool bReproject,
 														 true, !bPreserveSharpEdges, bAllowSeamCollapse);
 	Reducer.SetExternalConstraints(MoveTemp(constraints));
 
-	FMeshProjectionTarget ProjTarget(&OriginalMesh, &OriginalMeshSpatial);
-	Reducer.SetProjectionTarget(&ProjTarget);
+	if (bReproject)
+	{
+		FMeshProjectionTarget ProjTarget(&OriginalMesh, &OriginalMeshSpatial);
+		Reducer.SetProjectionTarget(&ProjTarget);
+	}
 
 	if (TargetMode == ESimplifyTargetType::Percentage)
 	{
diff --git a/Engine/Plugins/Experimental/MeshModelingToolset/Source/ModelingOperatorsEditorOnly/Private/ParameterizationOps/UVLayoutOp.cpp b/Engine/Plugins/Experimental/MeshModelingToolset/Source/ModelingOperatorsEditorOnly/Private/ParameterizationOps/UVLayoutOp.cpp
--- a/Engine/Plugins/Experimental/MeshModelingToolset/Source/ModelingOperatorsEditorOnly/Private/ParameterizationOps/UVLayoutOp.cpp
+++ b/Engine/Plugins/Experimental/MeshModelingToolset/Source/ModelingOperatorsEditorOnly/Private/ParameterizationOps/UVLayoutOp.cpp
@@ -2,9 +2,6 @@
 
 #include "ParameterizationOps/UVLayoutOp.h"
 
-#include "OverlappingCorners.h"
-#include "LayoutUV.h"
-
 #include "DynamicMesh3.h"
 #include "DynamicMeshAttributeSet.h"
 
@@ -17,103 +14,6 @@ void FUVLayoutOp::SetTransform(const FTransform& Transform) {
 	ResultTransform = (FTransform3d)Transform;
 }
 
-// very task-specific lightweight view of a FDynamicMesh3; ONLY for compact meshes with attributes
-// not intended for use outside of this narrow context
-class FCompactDynamicMeshWithAttributesLayoutView : public FLayoutUV::IMeshView
-{
-	FDynamicMesh3* Mesh;
-	int UVLayerInput, UVLayerOutput;
-
-public:
-	FCompactDynamicMeshWithAttributesLayoutView(FDynamicMesh3* Mesh, int UVLayerIn=0, int UVLayerOut=0) : Mesh(Mesh), UVLayerInput(UVLayerIn), UVLayerOutput(UVLayerOut)
-	{
-		check(Mesh->HasAttributes());
-		check(Mesh->IsCompact());
-		check(Mesh->Attributes()->NumUVLayers() > UVLayerIn);
-		check(Mesh->Attributes()->NumUVLayers() > UVLayerOut);
-	}
-
-	uint32 GetNumIndices() const override
-	{
-		return Mesh->TriangleCount() * 3;
-	}
-	FVector GetPosition(uint32 Index) const override
-	{
-		return FVector(Mesh->GetVertex(Mesh->GetTriangle(Index / 3)[Index % 3]));
-	}
-	FVector GetNormal(uint32 Index) const override
-	{
-		const FDynamicMeshNormalOverlay* NormalOverlay = Mesh->Attributes()->PrimaryNormals();
-		FIndex3i NormalEl = NormalOverlay->GetTriangle(Index / 3);
-		FVector3f Normal;
-		NormalOverlay->GetElement(NormalEl[Index % 3], Normal);
-		return FVector(Normal);
-	}
-	FVector2D GetInputTexcoord(uint32 Index) const override
-	{
-		const FDynamicMeshUVOverlay* UVOverlay = Mesh->Attributes()->GetUVLayer(UVLayerInput);
-		FIndex3i UVEl = UVOverlay->GetTriangle(Index / 3);
-		FVector2f UV;
-		UVOverlay->GetElement(UVEl[Index % 3], UV);
-		return FVector2D(UV);
-	}
-
-	/**
-	 * This function is kind of nonsense for our use case as we cannot initialize a UV overlay from a single number
-	 * (unless we make every triangle disconnected in a fully raw wedge thing, which we never ever want to do)
-	 * So we assume the calling code will only call this w/ NumIndices matching the input layer & the intent of making the UV layers the same.
-	 */
-	void InitOutputTexcoords(uint32 Num) override
-	{
-		if (UVLayerInput != UVLayerOutput)
-		{
-			Mesh->Attributes()->GetUVLayer(UVLayerOutput)->Copy(*Mesh->Attributes()->GetUVLayer(UVLayerInput));
-		}
-		check(Num == Mesh->TriangleCount() * 3);
-	}
-
-	void SetOutputTexcoord(uint32 Index, const FVector2D& Value) override
-	{
-		FDynamicMeshUVOverlay* UVOverlay = Mesh->Attributes()->GetUVLayer(UVLayerOutput);
-		FIndex3i UVEl = UVOverlay->GetTriangle(Index / 3);
-		UVOverlay->SetElement(UVEl[Index % 3], FVector2f(Value));
-	}
-};
-
-/**
- * Create an overlapping corner map to identify wedge indices that share the same UV element index for a given UV layer
- */
-FOverlappingCorners OverlappingCornersFromUVs(const FDynamicMesh3* Mesh, int UVLayerIndex)
-{
-	// track all wedge indices that map to the same UV element
-	const FDynamicMeshUVOverlay* UVOverlay = Mesh->Attributes()->GetUVLayer(UVLayerIndex);
-	FOverlappingCorners Overlaps;
-	Overlaps.Init(Mesh->TriangleCount()*3);
-	for (int ElementID : UVOverlay->ElementIndicesItr())
-	{
-		int VertID = UVOverlay->GetParentVertex(ElementID);
-		int LastWedgeIdx = -1;
-		for (int TriID : Mesh->VtxTrianglesItr(VertID))
-		{
-			FIndex3i ElTri = UVOverlay->GetTriangle(TriID);
-			for (int SubIdx = 0; SubIdx < 3; SubIdx++)
-			{
-				if (ElTri[SubIdx] == ElementID)
-				{
-					int WedgeIdx = TriID * 3 + SubIdx;
-					if (LastWedgeIdx != -1)
-					{
-						Overlaps.Add(LastWedgeIdx, WedgeIdx);
-					}
-					LastWedgeIdx = WedgeIdx;
-				}
-			}
-		}
-	}
-	Overlaps.FinishAdding();
-	return Overlaps;
-}
-
 
 void FUVLayoutOp::CalculateResult(FProgressCancel* Progress)
 {
