Update DynamicMesh to,from MeshDescription Conversions.  Now supports the shared UVs on MeshDecription and works directly with triangles (not polygons).

#rnx

diff --git a/Engine/Plugins/Experimental/GeometryProcessing/Source/MeshConversion/Private/DynamicMeshToMeshDescription.cpp b/Engine/Plugins/Experimental/GeometryProcessing/Source/MeshConversion/Private/DynamicMeshToMeshDescription.cpp
--- a/Engine/Plugins/Experimental/GeometryProcessing/Source/MeshConversion/Private/DynamicMeshToMeshDescription.cpp
+++ b/Engine/Plugins/Experimental/GeometryProcessing/Source/MeshConversion/Private/DynamicMeshToMeshDescription.cpp
@@ -12,31 +12,23 @@
 
 namespace DynamicMeshToMeshDescriptionConversionHelper
 {
-	// NOTE: assumes the order of triangles in the MeshIn correspond to the ordering you'd get by iterating over polygons, then tris-in-polygons, on MeshOut
+	// NOTE: assumes the order of triangles in the MeshIn correspond to the ordering over tris on MeshOut
 	// This matches conversion currently used in MeshDescriptionToDynamicMesh.cpp, but if that changes we will need to change this function to match!
 	template <typename OutAttributeType, int VecLen, typename InAttributeType>
 	void SetAttributesFromOverlay(
 		const FDynamicMesh3* MeshInArg, const FMeshDescription& MeshOutArg,
 		TVertexInstanceAttributesRef<OutAttributeType>& InstanceAttrib, const TDynamicMeshVectorOverlay<float, VecLen, InAttributeType>* Overlay, int AttribIndex=0)
 	{
-		const FPolygonArray& Polygons = MeshOutArg.Polygons();
-		int MeshInTriIdx = 0;
-		for (const FPolygonID PolygonID : Polygons.GetElementIDs())
+		for (const FTriangleID TriangleID : MeshOutArg.Triangles().GetElementIDs())
 		{
-			FPolygonGroupID PolygonGroupID = MeshOutArg.GetPolygonPolygonGroup(PolygonID);
+			TArrayView<const FVertexInstanceID> InstanceTri = MeshOutArg.GetTriangleVertexInstances(TriangleID);
 
-			TArrayView<const FTriangleID> TriangleIDs = MeshOutArg.GetPolygonTriangles(PolygonID);
-			int NumTriangles = TriangleIDs.Num();
-			for (int TriIdx = 0; TriIdx < NumTriangles; ++TriIdx, ++MeshInTriIdx)
-			{
-				const FTriangleID TriangleID = TriangleIDs[TriIdx];
-				TArrayView<const FVertexInstanceID> InstanceTri = MeshOutArg.GetTriangleVertexInstances(TriangleID);
+			int32 MeshInTriIdx = TriangleID.GetValue();   
 
-				FIndex3i OverlayVertIndices = Overlay->GetTriangle(MeshInTriIdx);
-				InstanceAttrib.Set(InstanceTri[0], AttribIndex, OutAttributeType(Overlay->GetElement(OverlayVertIndices.A)));
-				InstanceAttrib.Set(InstanceTri[1], AttribIndex, OutAttributeType(Overlay->GetElement(OverlayVertIndices.B)));
-				InstanceAttrib.Set(InstanceTri[2], AttribIndex, OutAttributeType(Overlay->GetElement(OverlayVertIndices.C)));
-			}
+			FIndex3i OverlayVertIndices = Overlay->GetTriangle(MeshInTriIdx);
+			InstanceAttrib.Set(InstanceTri[0], AttribIndex, OutAttributeType(Overlay->GetElement(OverlayVertIndices.A)));
+			InstanceAttrib.Set(InstanceTri[1], AttribIndex, OutAttributeType(Overlay->GetElement(OverlayVertIndices.B)));
+			InstanceAttrib.Set(InstanceTri[2], AttribIndex, OutAttributeType(Overlay->GetElement(OverlayVertIndices.C)));
 		}
 	}
 }
@@ -49,6 +41,11 @@ void FDynamicMeshToMeshDescription::Update(const FDynamicMesh3* MeshIn, FMeshDes
 
 	check(MeshIn->IsCompactV());
 
+	// can't currently update the shared UV connectivity data structure on the MeshDescription :(
+	//  -- see FDynamicMeshToMeshDescription::UpdateAttributes()
+	// after this has been fixed, please update USimpleDynamicMeshComponent::Bake() to use the Update() path accordingly  
+	check(bUpdateUVs == false);
+
 	// update positions
 	int32 NumVertices = MeshOut.Vertices().Num();
 	check(NumVertices <= MeshIn->VertexCount());
@@ -66,6 +63,14 @@ void FDynamicMeshToMeshDescription::UpdateAttributes(const FDynamicMesh3* MeshIn
 {
 	check(MeshIn->IsCompactV());
 
+	if (bUpdateUVs)
+	{
+		// currently we can't update the shared UVs on a MeshDescription
+		// Updating only per-instance UVs will result in the two representations for UVs to be out of sync. 
+		// @todo UpdateUV path below when MeshDescription supports updating the shared UVs 
+		check(0);
+	}
+
 	FStaticMeshAttributes Attributes(MeshOut);
 
 	if (bUpdateNormals)
@@ -185,6 +190,11 @@ void FDynamicMeshToMeshDescription::Convert_NoAttributes(const FDynamicMesh3* Me
 	FMeshDescriptionBuilder Builder;
 	Builder.SetMeshDescription(&MeshOut);
 
+	Builder.SuspendMeshDescriptionIndexing();
+	const int32 UVLayerIndex = 0;
+	Builder.SetNumUVLayers(1);
+	Builder.ReserveNewUVs(MeshIn->VertexCount(), UVLayerIndex);
+
 	bool bCopyGroupToPolyGroup = false;
 	if (ConversionOptions.bSetPolyGroups && MeshIn->HasTriangleGroups())
 	{
@@ -195,6 +205,7 @@ void FDynamicMeshToMeshDescription::Convert_NoAttributes(const FDynamicMesh3* Me
 	// create vertices
 	TArray<FVertexID> MapV; 
 	MapV.SetNum(MeshIn->MaxVertexID());
+	Builder.ReserveNewVertices(MeshIn->VertexCount());
 	for (int VertID : MeshIn->VertexIndicesItr())
 	{
 		MapV[VertID] = Builder.AppendVertex((FVector)MeshIn->GetVertex(VertID));
@@ -204,12 +215,14 @@ void FDynamicMeshToMeshDescription::Convert_NoAttributes(const FDynamicMesh3* Me
 
 	// create new instances when seen
 	TMap<FIndex3i, FVertexInstanceID> InstanceList;
+	TMap<int32, FUVID> InstanceUVIDMap;
 	for (int TriID : MeshIn->TriangleIndicesItr())
 	{
 		FIndex3i Triangle = MeshIn->GetTriangle(TriID);
 		FIndex3i UVTriangle(-1, -1, -1);
 		FIndex3i NormalTriangle = Triangle;
 		FVertexInstanceID InstanceTri[3];
+		FUVID UVIDs[3];
 		for (int j = 0; j < 3; ++j)
 		{
 			FIndex3i InstanceElem(Triangle[j], UVTriangle[j], NormalTriangle[j]);
@@ -217,19 +230,33 @@ void FDynamicMeshToMeshDescription::Convert_NoAttributes(const FDynamicMesh3* Me
 			{
 				FVertexInstanceID NewInstanceID = Builder.AppendInstance(MapV[Triangle[j]]);
 				InstanceList.Add(InstanceElem, NewInstanceID);
-				FVector2D UV = MeshIn->HasVertexUVs() ? FVector2D(MeshIn->GetVertexUV(Triangle[j])) : FVector2D::ZeroVector;
+				
 				FVector Normal = MeshIn->HasVertexNormals() ? FVector(MeshIn->GetVertexNormal(Triangle[j])) : FVector::UpVector;
-				Builder.SetInstance(NewInstanceID, UV, Normal);
+				Builder.SetInstanceNormal(NewInstanceID, Normal);
+
+				// Add UV to MeshDecription UVvertex buffer
+				FVector2D UV = MeshIn->HasVertexUVs() ? FVector2D(MeshIn->GetVertexUV(Triangle[j])) : FVector2D::ZeroVector;
+				FUVID UVID = Builder.AppendUV(UV, UVLayerIndex);
+				
+				// associate UVID with this instance
+				InstanceUVIDMap.Add(NewInstanceID.GetValue(), UVID);				
 			}
 			InstanceTri[j] = InstanceList[InstanceElem];
+			UVIDs[j] = InstanceUVIDMap[InstanceTri[j].GetValue()];
 		}
 
-		FPolygonID NewPolygonID = Builder.AppendTriangle(InstanceTri[0], InstanceTri[1], InstanceTri[2], AllGroupID);
+		FTriangleID NewTriangleID = Builder.AppendTriangle(InstanceTri[0], InstanceTri[1], InstanceTri[2], AllGroupID);
+		
+		// append the UV triangle - builder takes care of the rest
+		Builder.AppendUVTriangle(NewTriangleID, UVIDs[0], UVIDs[1], UVIDs[2], UVLayerIndex);
+		
 		if (bCopyGroupToPolyGroup)
 		{
-			Builder.SetPolyGroupID(NewPolygonID, MeshIn->GetTriangleGroup(TriID));
+			Builder.SetPolyGroupID(NewTriangleID, MeshIn->GetTriangleGroup(TriID));
 		}
 	}
+
+	Builder.ResumeMeshDescriptionIndexing();
 }
 
 
@@ -386,11 +413,11 @@ void FDynamicMeshToMeshDescription::Convert_SharedInstances(const FDynamicMesh3*
 			UsePolygonGroupID = FPolygonGroupID(MaterialID);
 		}
 
-		FPolygonID NewPolygonID = Builder.AppendTriangle(TriVertInstances[TriID].V[0], TriVertInstances[TriID].V[1], TriVertInstances[TriID].V[2], UsePolygonGroupID);
+		FTriangleID NewTriangleID = Builder.AppendTriangle(TriVertInstances[TriID].V[0], TriVertInstances[TriID].V[1], TriVertInstances[TriID].V[2], UsePolygonGroupID);
 
 		if (bCopyGroupToPolyGroup)
 		{
-			Builder.SetPolyGroupID(NewPolygonID, MeshIn->GetTriangleGroup(TriID));
+			Builder.SetPolyGroupID(NewTriangleID, MeshIn->GetTriangleGroup(TriID));
 		}
 	}
 }
@@ -401,10 +428,24 @@ void FDynamicMeshToMeshDescription::Convert_SharedInstances(const FDynamicMesh3*
 
 void FDynamicMeshToMeshDescription::Convert_NoSharedInstances(const FDynamicMesh3* MeshIn, FMeshDescription& MeshOut)
 {
-	// TODO: update this path to support multiple UV layers
+	
+	const bool bHasAttributes = MeshIn->HasAttributes();
+
+	// check if we have per-triangle material ID
+	const FDynamicMeshMaterialAttribute* MaterialIDAttrib = (bHasAttributes && MeshIn->Attributes()->HasMaterialID()) ? MeshIn->Attributes()->GetMaterialID() : nullptr;
+
+	// cache normal and UV overlay info
+	const FDynamicMeshNormalOverlay* NormalOverlay = bHasAttributes ? MeshIn->Attributes()->PrimaryNormals() : nullptr;
+
+	const int32 NumUVLayers = bHasAttributes ? MeshIn->Attributes()->NumUVLayers() : 0;
+	
+	// cache the UV layers
+	TArray<const FDynamicMeshUVOverlay*> UVLayers;
+	for (int32 k = 0; k < NumUVLayers; ++k)
+	{
+		UVLayers.Add(MeshIn->Attributes()->GetUVLayer(k));
+	}
 
-	const FDynamicMeshUVOverlay* UVOverlay = MeshIn->HasAttributes() ? MeshIn->Attributes()->PrimaryUV() : nullptr;
-	const FDynamicMeshNormalOverlay* NormalOverlay = MeshIn->HasAttributes() ? MeshIn->Attributes()->PrimaryNormals() : nullptr;
 
 	MeshOut.Empty();
 
@@ -418,18 +459,41 @@ void FDynamicMeshToMeshDescription::Convert_NoSharedInstances(const FDynamicMesh
 		bCopyGroupToPolyGroup = true;
 	}
 
+	// disable indexing during the full build of the mesh
+	Builder.SuspendMeshDescriptionIndexing();
+
+	// allocate
+	Builder.ReserveNewVertices(MeshIn->VertexCount());
+
+	// create "vertex buffer" in MeshDescription
 	TArray<FVertexID> MapV; MapV.SetNum(MeshIn->MaxVertexID());
 	for (int VertID : MeshIn->VertexIndicesItr())
 	{
 		MapV[VertID] = Builder.AppendVertex((FVector)MeshIn->GetVertex(VertID));
 	}
 
-	FPolygonGroupID ZeroPolygonGroupID = Builder.AppendPolygonGroup();
+	// create UV vertex buffer in MeshDecription
+	Builder.SetNumUVLayers(NumUVLayers);
 
-	// check if we have per-triangle material ID
-	const FDynamicMeshMaterialAttribute* MaterialIDAttrib =
-		(MeshIn->HasAttributes() && MeshIn->Attributes()->HasMaterialID()) ?
-		MeshIn->Attributes()->GetMaterialID() : nullptr;
+	TArray<TArray<FUVID>> MapUVArray; MapUVArray.SetNum(NumUVLayers);
+	for (int32 k = 0; k < NumUVLayers; ++k)
+	{
+		const FDynamicMeshUVOverlay* UVOverlay = UVLayers[k];
+		TArray<FUVID>& MapUV = MapUVArray[k];
+
+		MapUV.SetNum(UVOverlay->MaxElementID());
+		Builder.ReserveNewUVs(UVOverlay->ElementCount(), k); 
+
+		for (int32 ElementID : UVOverlay->ElementIndicesItr())
+		{
+			const FVector2D UVvalue = (FVector2D)UVOverlay->GetElement(ElementID); 
+			MapUV[ElementID] = Builder.AppendUV(UVvalue, k);
+		}
+	}
+	
+
+
+	FPolygonGroupID ZeroPolygonGroupID = Builder.AppendPolygonGroup();
 
 	// need to know max material index value so we can reserve groups in MeshDescription
 	int32 MaxPolygonGroupID = 0;
@@ -454,15 +518,7 @@ void FDynamicMeshToMeshDescription::Convert_NoSharedInstances(const FDynamicMesh
 		}
 	}
 
-	int32 NumUVLayers = MeshIn->HasAttributes() ? MeshIn->Attributes()->NumUVLayers() : 0;
-	Builder.SetNumUVLayers(NumUVLayers);
 	
-	// cache the UV layers
-	TArray<const FDynamicMeshUVOverlay*> UVLayers;
-	for (int32 k = 0; k < NumUVLayers; ++k)
-	{
-		UVLayers.Add( MeshIn->Attributes()->GetUVLayer(k) );
-	}
 
 	TArray<FIndex3i> UVTris;
 	UVTris.SetNum(NumUVLayers);
@@ -471,35 +527,12 @@ void FDynamicMeshToMeshDescription::Convert_NoSharedInstances(const FDynamicMesh
 	{
 		FIndex3i Triangle = MeshIn->GetTriangle(TriID);
 
-		FVertexID TriVertices[3] = { MapV[Triangle[0]] , MapV[Triangle[1]] , MapV[Triangle[2]] };
+		// create new vtx instances for each triangle
 		FVertexInstanceID TriVertInstances[3];
-
-		// look up normal and UV triangles
-		FIndex3i NormalTri = (NormalOverlay) ? NormalOverlay->GetTriangle(TriID) : FIndex3i::Invalid();
-		for (int32 k = 0; k < NumUVLayers; ++k)
-		{
-			UVTris[k] = UVLayers[k]->GetTriangle(TriID);
-		}
-
-		// create new vtx instances and set attributes for each triangle
 		for (int32 j = 0; j < 3; ++j)
 		{
-			FVertexInstanceID NewInstanceID = Builder.AppendInstance(TriVertices[j]);
-			TriVertInstances[j] = NewInstanceID;
-
-			FVector TriVertNormal = FVector::UpVector;
-			if (NormalOverlay && NormalOverlay->IsElement(NormalTri[j]))
-			{
-				TriVertNormal = (FVector)NormalOverlay->GetElement(NormalTri[j]);
-			}
-			Builder.SetInstanceNormal(NewInstanceID, TriVertNormal);
-				
-			for (int32 k = 0; k < NumUVLayers; ++k)
-			{
-				Builder.SetInstanceUV(NewInstanceID,
-					UVLayers[k]->IsElement(UVTris[k][j]) ? (FVector2D)UVLayers[k]->GetElement(UVTris[k][j]) : FVector2D::ZeroVector,
-					k);
-			}
+			const FVertexID TriVertex = MapV[Triangle[j]];
+			TriVertInstances[j] = Builder.AppendInstance(TriVertex);
 		}
 
 		// transfer material index to MeshDescription polygon group (by convention)
@@ -511,11 +544,64 @@ void FDynamicMeshToMeshDescription::Convert_NoSharedInstances(const FDynamicMesh
 			UsePolygonGroupID = FPolygonGroupID(MaterialID);
 		}
 
-		FPolygonID NewPolygonID = Builder.AppendTriangle(TriVertInstances[0], TriVertInstances[1], TriVertInstances[2], UsePolygonGroupID);
+		// add the triangle to MeshDescription
+		FTriangleID NewTriangleID = Builder.AppendTriangle(TriVertInstances[0], TriVertInstances[1], TriVertInstances[2], UsePolygonGroupID); 
+
+		// transfer UVs.  Note the Builder sets both the shared and per-instance UVs from this
+		for (int32 k = 0; k < NumUVLayers; ++k)
+		{ 
+			FUVID UVIDs[3] = {FUVID(-1), FUVID(-1), FUVID(-1)};
+
+			// add zero uvs for unset triangles.  this mimics the old version of this code.
+			if (!UVLayers[k]->IsSetTriangle(TriID))
+			{
+				for (int32 j = 0; j < 3; ++j)
+				{ 
+					UVIDs[j] = Builder.AppendUV(FVector2D::ZeroVector, k);
+				}
+			}
+			else
+			{ 
+				const TArray<FUVID>& MapUV = MapUVArray[k];
+		
+				// triangle of UV element ids from dynamic mesh.  references values already stored in MeshDescription.
+				FIndex3i UVTri = UVLayers[k]->GetTriangle(TriID);
+
+				// translate to MeshDescription Ids
+				for (int32 j = 0; j < 3; ++j)
+				{ 
+					UVIDs[j] = MapUV[ UVTri[j] ];
+				}
+			}
+			
+			// append the UV triangle - builder takes care of the rest
+			Builder.AppendUVTriangle(NewTriangleID, UVIDs[0], UVIDs[1], UVIDs[2], k); 
+
+		}
+		
+		// look up normal triangles
+		FIndex3i NormalTri = (NormalOverlay) ? NormalOverlay->GetTriangle(TriID) : FIndex3i::Invalid();
+
+		// transfer Normals.   
+		// NB: only per-instance normals are supported in MeshDescription at this time
+		// NB: will need to be updated to follow the pattern used in UVs above when MeshDecription supports shared normals 
+		for (int32 j = 0; j < 3; ++j)
+		{
+			FVertexInstanceID CornerInstanceID = TriVertInstances[j];
+			FVector TriVertNormal = FVector::UpVector;
+			if (NormalOverlay && NormalOverlay->IsElement(NormalTri[j]))
+			{
+				TriVertNormal = (FVector)NormalOverlay->GetElement(NormalTri[j]);
+			}
+			Builder.SetInstanceNormal(CornerInstanceID, TriVertNormal);
+		}
+
 
 		if (bCopyGroupToPolyGroup)
 		{
-			Builder.SetPolyGroupID(NewPolygonID, MeshIn->GetTriangleGroup(TriID));
+			Builder.SetPolyGroupID(NewTriangleID, MeshIn->GetTriangleGroup(TriID));
 		}
 	}
+
+	Builder.ResumeMeshDescriptionIndexing();
 }
diff --git a/Engine/Plugins/Experimental/GeometryProcessing/Source/MeshConversion/Private/MeshDescriptionBuilder.cpp b/Engine/Plugins/Experimental/GeometryProcessing/Source/MeshConversion/Private/MeshDescriptionBuilder.cpp
--- a/Engine/Plugins/Experimental/GeometryProcessing/Source/MeshConversion/Private/MeshDescriptionBuilder.cpp
+++ b/Engine/Plugins/Experimental/GeometryProcessing/Source/MeshConversion/Private/MeshDescriptionBuilder.cpp
@@ -44,6 +44,16 @@ void FMeshDescriptionBuilder::EnablePolyGroups()
 }
 
 
+
+
+void FMeshDescriptionBuilder::ReserveNewVertices(int32 Count)
+{
+	MeshDescription->ReserveNewVertices(Count);
+}
+
+
+
+
 FVertexID FMeshDescriptionBuilder::AppendVertex(const FVector& Position)
 {
 	FVertexID VertexID = MeshDescription->CreateVertex();
@@ -59,13 +69,16 @@ FPolygonGroupID FMeshDescriptionBuilder::AppendPolygonGroup()
 
 
 
-FPolygonID FMeshDescriptionBuilder::AppendTriangle(const FVertexID& Vertex0, const FVertexID& Vertex1, const FVertexID& Vertex2, const FPolygonGroupID& PolygonGroup)
+FTriangleID FMeshDescriptionBuilder::AppendTriangle(const FVertexID& Vertex0, const FVertexID& Vertex1, const FVertexID& Vertex2, const FPolygonGroupID& PolygonGroup)
 {
-	TempBuffer.SetNum(3, false);
-	TempBuffer[0] = Vertex0;
-	TempBuffer[1] = Vertex1;
-	TempBuffer[2] = Vertex2;
-	return AppendPolygon(TempBuffer, PolygonGroup);
+	
+	FVertexInstanceID TriVertexInstances[3];
+	TriVertexInstances[0] = MeshDescription->CreateVertexInstance(Vertex0);
+	TriVertexInstances[1] = MeshDescription->CreateVertexInstance(Vertex1);
+	TriVertexInstances[2] = MeshDescription->CreateVertexInstance(Vertex2);
+
+
+	return AppendTriangle(TriVertexInstances[0], TriVertexInstances[1], TriVertexInstances[2], PolygonGroup);
 }
 
 
@@ -74,7 +87,45 @@ FVertexInstanceID FMeshDescriptionBuilder::AppendInstance(const FVertexID& Verte
 	return MeshDescription->CreateVertexInstance(VertexID);
 }
 
+void FMeshDescriptionBuilder::ReserveNewUVs(int32 Count, int UVLayerIndex)
+{
+	// the reserve new UVs can only be called after the UV channels have been created.
+	
+	check(MeshDescription->GetNumUVElementChannels() > UVLayerIndex )
+
+	MeshDescription->ReserveNewUVs(Count, UVLayerIndex);
+}
+
+FUVID FMeshDescriptionBuilder::AppendUV(const FVector2D& UVvalue, int32 UVLayerIndex)
+{
+	TUVAttributesRef<FVector2D> UVCoordinates = UVCoordinateLayers[UVLayerIndex];
+	FUVID UVID = MeshDescription->CreateUV(UVLayerIndex);
+	UVCoordinates[UVID] = UVvalue;
+	return UVID;
+}
 
+void FMeshDescriptionBuilder::AppendUVTriangle(const FTriangleID& TriangleID, const FUVID UVvertexID0, const FUVID UVvertexID1, const FUVID UVvertexID2, int32 UVLayerIndex)
+{
+	// set the shared UVs
+	TempUVBuffer.SetNum(3, false);
+	TempUVBuffer[0] = UVvertexID0;
+	TempUVBuffer[1] = UVvertexID1;
+	TempUVBuffer[2] = UVvertexID2;
+	MeshDescription->SetTriangleUVIndices(TriangleID, TempUVBuffer, UVLayerIndex);
+
+	TUVAttributesRef<FVector2D> UVCoordinates = UVCoordinateLayers[UVLayerIndex];
+
+	// set per-instance UVs.  
+	// NB: per-instance UVs should go away on MeshDescription.
+	TArrayView<const FVertexInstanceID> TriVertInstances = MeshDescription->GetTriangleVertexInstances(TriangleID);
+	for (int32 j = 0; j < 3; ++j)
+	{
+		const FVertexInstanceID CornerInstanceID = TriVertInstances[j];
+		const FVector2D& UVvalue = UVCoordinates[TempUVBuffer[j]];
+		
+		SetInstanceUV(CornerInstanceID, UVvalue, UVLayerIndex);
+	}
+}
 
 void FMeshDescriptionBuilder::SetPosition(const FVertexID& VertexID, const FVector& NewPosition)
 {
@@ -92,18 +143,6 @@ FVector FMeshDescriptionBuilder::GetPosition(const FVertexInstanceID& InstanceID
 }
 
 
-void FMeshDescriptionBuilder::SetInstance(const FVertexInstanceID& InstanceID, const FVector2D& InstanceUV, const FVector& InstanceNormal)
-{
-	if (InstanceUVs.IsValid())
-	{
-		InstanceUVs.Set(InstanceID, InstanceUV); 
-	}
-	if (InstanceNormals.IsValid())
-	{
-		InstanceNormals.Set(InstanceID, InstanceNormal);
-	}
-}
-
 
 void FMeshDescriptionBuilder::SetInstanceNormal(const FVertexInstanceID& InstanceID, const FVector& Normal)
 {
@@ -118,7 +157,7 @@ void FMeshDescriptionBuilder::SetInstanceUV(const FVertexInstanceID& InstanceID,
 {
 	if (InstanceUVs.IsValid() && ensure(UVLayerIndex < InstanceUVs.GetNumChannels()))
 	{
-		InstanceUVs.Set(InstanceID, UVLayerIndex, InstanceUV); 
+		InstanceUVs.Set(InstanceID, UVLayerIndex, InstanceUV);
 	}
 }
 
@@ -127,8 +166,19 @@ void FMeshDescriptionBuilder::SetNumUVLayers(int32 NumUVLayers)
 {
 	if (ensure(InstanceUVs.IsValid()))
 	{
+		// initialize the instanced UV channels
 		InstanceUVs.SetNumChannels(NumUVLayers);
 	}
+
+	// initialize the shared UV channels
+	MeshDescription->SetNumUVChannels(NumUVLayers);
+
+	// cache reference to uv vertex buffers
+	UVCoordinateLayers.SetNum(NumUVLayers);
+	for (int i = 0; i < NumUVLayers; ++i)
+	{
+		UVCoordinateLayers[i] = MeshDescription->UVAttributes(i).GetAttributesRef<FVector2D>(MeshAttribute::UV::UVCoordinate);
+	}
 }
 
 
@@ -141,40 +191,18 @@ void FMeshDescriptionBuilder::SetInstanceColor(const FVertexInstanceID& Instance
 }
 
 
-FPolygonID FMeshDescriptionBuilder::AppendTriangle(const FVertexID* Triangle, const FPolygonGroupID& PolygonGroup, 
-	const FVector2D * VertexUVs, const FVector* VertexNormals)
+FTriangleID FMeshDescriptionBuilder::AppendTriangle(const FVertexID* Triangle, const FPolygonGroupID& PolygonGroup)
 {
-	TempBuffer.SetNum(3, false);
-	TempBuffer[0] = Triangle[0];
-	TempBuffer[1] = Triangle[1];
-	TempBuffer[2] = Triangle[2];
+	FVertexInstanceID TriVertexInstances[3];
+	TriVertexInstances[0] = MeshDescription->CreateVertexInstance(Triangle[0]);
+	TriVertexInstances[1] = MeshDescription->CreateVertexInstance(Triangle[1]);
+	TriVertexInstances[2] = MeshDescription->CreateVertexInstance(Triangle[2]);
 
-	const TArray<FVector2D> * UseUVs = nullptr;
-	if (VertexUVs != nullptr) 
-	{
-		UVBuffer.SetNum(3, false);
-		UVBuffer[0] = VertexUVs[0];
-		UVBuffer[1] = VertexUVs[1];
-		UVBuffer[2] = VertexUVs[2];
-		UseUVs = &UVBuffer;
-	}
-
-	const TArray<FVector> * UseNormals = nullptr;
-	if (VertexNormals != nullptr)
-	{
-		NormalBuffer.SetNum(3, false);
-		NormalBuffer[0] = VertexNormals[0];
-		NormalBuffer[1] = VertexNormals[1];
-		NormalBuffer[2] = VertexNormals[2];
-		UseNormals = &NormalBuffer;
-	}
-
-	return AppendPolygon(TempBuffer, PolygonGroup, UseUVs, UseNormals);
+	return AppendTriangle(TriVertexInstances[0], TriVertexInstances[1], TriVertexInstances[2], PolygonGroup);
 }
 
 
-FPolygonID FMeshDescriptionBuilder::AppendPolygon(const TArray<FVertexID>& Vertices, const FPolygonGroupID& PolygonGroup, 
-	const TArray<FVector2D> * VertexUVs, const TArray<FVector>* VertexNormals)
+FPolygonID FMeshDescriptionBuilder::AppendPolygon(const TArray<FVertexID>& Vertices, const FPolygonGroupID& PolygonGroup)
 {
 	int NumVertices = Vertices.Num();
 	TArray<FVertexInstanceID> Polygon; 
@@ -184,15 +212,6 @@ FPolygonID FMeshDescriptionBuilder::AppendPolygon(const TArray<FVertexID>& Verti
 		FVertexInstanceID VertexInstance = MeshDescription->CreateVertexInstance(Vertices[j]);
 		Polygon.Add(VertexInstance);
 
-		if (VertexUVs != nullptr)
-		{
-			InstanceUVs.Set(VertexInstance, (*VertexUVs)[j]);
-		}
-		if (VertexNormals != nullptr)
-		{
-			InstanceNormals.Set(VertexInstance, (*VertexNormals)[j]);
-		}
-
 	}
 
 	const FPolygonID NewPolygonID = MeshDescription->CreatePolygon(PolygonGroup, Polygon);
@@ -203,24 +222,26 @@ FPolygonID FMeshDescriptionBuilder::AppendPolygon(const TArray<FVertexID>& Verti
 
 
 
-FPolygonID FMeshDescriptionBuilder::AppendTriangle(const FVertexInstanceID& Instance0, const FVertexInstanceID& Instance1, const FVertexInstanceID& Instance2, const FPolygonGroupID& PolygonGroup)
+FTriangleID FMeshDescriptionBuilder::AppendTriangle(const FVertexInstanceID& Instance0, const FVertexInstanceID& Instance1, const FVertexInstanceID& Instance2, const FPolygonGroupID& PolygonGroup)
 {
-	TArray<FVertexInstanceID> Polygon;
-	Polygon.Add(Instance0);
-	Polygon.Add(Instance1);
-	Polygon.Add(Instance2);
+	TArray<FVertexInstanceID> CornerInstanceIDs;
+	CornerInstanceIDs.Add(Instance0);
+	CornerInstanceIDs.Add(Instance1);
+	CornerInstanceIDs.Add(Instance2);
 
-	const FPolygonID NewPolygonID = MeshDescription->CreatePolygon(PolygonGroup, Polygon);
+	TArray<FEdgeID> NewEdgeIDs; 
+	const FTriangleID NewTriangleID = MeshDescription->CreateTriangle(PolygonGroup, CornerInstanceIDs, &NewEdgeIDs);
 
-	return NewPolygonID;
+	return NewTriangleID;
 }
 
 
 
 
 
-void FMeshDescriptionBuilder::SetPolyGroupID(const FPolygonID& PolygonID, int GroupID)
+void FMeshDescriptionBuilder::SetPolyGroupID(const FTriangleID& TriangleID, int GroupID)
 {
+	FPolygonID PolygonID = MeshDescription->GetTrianglePolygon(TriangleID);
 	PolyGroups.Set(PolygonID, 0, GroupID);
 }
 
@@ -266,4 +287,27 @@ FBox FMeshDescriptionBuilder::ComputeBoundingBox() const
 		bounds.Contain(VertexPositions.Get(VertexID));
 	}
 	return (FBox)bounds;
+}
+
+
+
+
+void FMeshDescriptionBuilder::SuspendMeshDescriptionIndexing()
+{
+	MeshDescription->SuspendVertexInstanceIndexing();
+	MeshDescription->SuspendEdgeIndexing();
+	MeshDescription->SuspendPolygonIndexing();
+	MeshDescription->SuspendPolygonGroupIndexing();
+	MeshDescription->SuspendUVIndexing();
+}
+
+
+
+void FMeshDescriptionBuilder::ResumeMeshDescriptionIndexing()
+{
+	MeshDescription->ResumeVertexInstanceIndexing();
+	MeshDescription->ResumeEdgeIndexing();
+	MeshDescription->ResumePolygonIndexing();
+	MeshDescription->ResumePolygonGroupIndexing();
+	MeshDescription->ResumeUVIndexing();
 }
\ No newline at end of file
diff --git a/Engine/Plugins/Experimental/GeometryProcessing/Source/MeshConversion/Private/MeshDescriptionToDynamicMesh.cpp b/Engine/Plugins/Experimental/GeometryProcessing/Source/MeshConversion/Private/MeshDescriptionToDynamicMesh.cpp
--- a/Engine/Plugins/Experimental/GeometryProcessing/Source/MeshConversion/Private/MeshDescriptionToDynamicMesh.cpp
+++ b/Engine/Plugins/Experimental/GeometryProcessing/Source/MeshConversion/Private/MeshDescriptionToDynamicMesh.cpp
@@ -125,16 +125,21 @@ public:
 
 void FMeshDescriptionToDynamicMesh::Convert(const FMeshDescription* MeshIn, FDynamicMesh3& MeshOut)
 {
+	TriIDMap.Reset();
+	VertIDMap.Reset();
+
+	// allocate the VertIDMap.  Unfortunately the array will need to grow more if MeshIn has non-manifold edges that need to be split
+	VertIDMap.SetNumUninitialized(MeshIn->Vertices().Num());
+
 	// look up vertex positions
-	const FVertexArray& VertexIDs = MeshIn->Vertices();
 	TVertexAttributesConstRef<FVector> VertexPositions = MeshIn->GetVertexPositions();
 
-	// copy vertex positions
-	for (const FVertexID VertexID : VertexIDs.GetElementIDs())
+	// copy vertex positions. Later we may have to append duplicate vertices to resolve non-manifold structures.
+	for (const FVertexID VertexID : MeshIn->Vertices().GetElementIDs())
 	{
 		const FVector Position = VertexPositions.Get(VertexID);
 		int NewVertIdx = MeshOut.AppendVertex(Position);
-		check(NewVertIdx == VertexID.GetValue());
+		VertIDMap[NewVertIdx] = VertexID;
 	}
 
 	// look up vertex-instance UVs and normals
@@ -167,36 +172,29 @@ void FMeshDescriptionToDynamicMesh::Convert(const FMeshDescription* MeshIn, FDyn
 	{
 		FPolygonID PolygonID;
 		int32 PolygonGroupID;
-		int32 TriIndex;
 		FVertexInstanceID TriInstances[3];
 	};
 	TArray<FTriData> AddedTriangles;
 	AddedTriangles.SetNum(MeshIn->Triangles().Num());
 
+	// allocate the TriIDMap
+	TriIDMap.SetNumUninitialized(MeshIn->Triangles().Num());
+
+
+	// Iterate over triangles in the Mesh Description
 	// NOTE: If you change the iteration order here, please update the corresponding iteration in FDynamicMeshToMeshDescription::UpdateAttributes, 
-	//	which assumes the iteration order here is polygons -> triangles, to correspond the triangles when writing updated attributes back!
-	const FPolygonArray& Polygons = MeshIn->Polygons();
-	for (const FPolygonID PolygonID : Polygons.GetElementIDs())
+	//	which assumes the iteration order here is the same as here to correspond the triangles when writing updated attributes back!
+	for (const FTriangleID TriangleID : MeshIn->Triangles().GetElementIDs())
 	{
-		int32 PolygonGroupID = MeshIn->GetPolygonPolygonGroup(PolygonID).GetValue();
 
-		TArrayView<const FTriangleID> TriangleIDs = MeshIn->GetPolygonTriangles(PolygonID);
-		int NumTriangles = TriangleIDs.Num();
-		for (int TriIdx = 0; TriIdx < NumTriangles; ++TriIdx)
+		// Get the PolygonID, PolygonGroupID and general GroupID
+		//---
+		FPolygonID PolygonID = MeshIn->GetTrianglePolygon(TriangleID);
+		FPolygonGroupID PolygonGroupID = MeshIn->GetTrianglePolygonGroup(TriangleID);
+
+		int GroupID = 0;
+		switch (GroupMode)
 		{
-			FTriData TriData;
-			TriData.PolygonID = PolygonID;
-			TriData.PolygonGroupID = PolygonGroupID;
-			TriData.TriIndex = TriIdx;
-
-			TArrayView<const FVertexInstanceID> InstanceTri = MeshIn->GetTriangleVertexInstances(TriangleIDs[TriIdx]);
-			TriData.TriInstances[0] = InstanceTri[0];
-			TriData.TriInstances[1] = InstanceTri[1];
-			TriData.TriInstances[2] = InstanceTri[2];
-
-			int GroupID = 0;
-			switch (GroupMode)
-			{
 			case EPrimaryGroupMode::SetToPolyGroup:
 				if (PolyGroups.IsValid())
 				{
@@ -207,85 +205,132 @@ void FMeshDescriptionToDynamicMesh::Convert(const FMeshDescription* MeshIn, FDyn
 				GroupID = PolygonID.GetValue() + 1; // Shift IDs up by 1 to leave ID 0 as a default/unassigned group
 				break;
 			case EPrimaryGroupMode::SetToPolygonGroupID:
-				GroupID = PolygonGroupID + 1; // Shift IDs up by 1 to leave ID 0 as a default/unassigned group
+				GroupID = PolygonGroupID.GetValue() + 1; // Shift IDs up by 1 to leave ID 0 as a default/unassigned group
 				break;
 			case EPrimaryGroupMode::SetToZero:
 				break; // keep at 0
-			}
+		}
+		
+		FTriData TriData;
+		TriData.PolygonID = PolygonID;
+		TriData.PolygonGroupID = PolygonGroupID;
+		
 
-			// append triangle
-			int32 VertexID0 = MeshIn->GetVertexInstanceVertex(InstanceTri[0]).GetValue();
-			int32 VertexID1 = MeshIn->GetVertexInstanceVertex(InstanceTri[1]).GetValue();
-			int32 VertexID2 = MeshIn->GetVertexInstanceVertex(InstanceTri[2]).GetValue();
-			int NewTriangleID = MeshOut.AppendTriangle(VertexID0, VertexID1, VertexID2, GroupID);
+		// stash the vertex instance IDs for this triangle.  potentially needed for per-instance attribute welding
+		TArrayView<const FVertexInstanceID> InstanceTri = MeshIn->GetTriangleVertexInstances(TriangleID);
+		TriData.TriInstances[0] = InstanceTri[0];
+		TriData.TriInstances[1] = InstanceTri[1];
+		TriData.TriInstances[2] = InstanceTri[2];
 
-			if (NewTriangleID == FDynamicMesh3::DuplicateTriangleID)
-			{
-				continue;
-			}
+		//---
 
-			// if append failed due to non-manifold, duplicate verts
-			if (NewTriangleID == FDynamicMesh3::NonManifoldID)
-			{
-				int e0 = MeshOut.FindEdge(VertexID0, VertexID1);
-				int e1 = MeshOut.FindEdge(VertexID1, VertexID2);
-				int e2 = MeshOut.FindEdge(VertexID2, VertexID0);
+		// Get the vertex IDs for this triangle.
+		TArrayView<const FVertexID> TriangleVertexIDs = MeshIn->GetTriangleVertices(TriangleID);
 
-				// determine which verts need to be duplicated
-				bool bDuplicate[3] = { false, false, false };
-				if (e0 != FDynamicMesh3::InvalidID && MeshOut.IsBoundaryEdge(e0) == false)
-				{
-					bDuplicate[0] = true;
-					bDuplicate[1] = true;
-				}
-				if (e1 != FDynamicMesh3::InvalidID && MeshOut.IsBoundaryEdge(e1) == false)
-				{
-					bDuplicate[1] = true;
-					bDuplicate[2] = true;
-				}
-				if (e2 != FDynamicMesh3::InvalidID && MeshOut.IsBoundaryEdge(e2) == false)
-				{
-					bDuplicate[2] = true;
-					bDuplicate[0] = true;
-				}
-				if (bDuplicate[0])
-				{
-					FVertexID VertexID = MeshIn->GetVertexInstanceVertex(InstanceTri[0]);
-					const FVector Position = VertexPositions.Get(VertexID);
-					int NewVertIdx = MeshOut.AppendVertex(Position);
-					VertexID0 = NewVertIdx;
-				}
-				if (bDuplicate[1])
-				{
-					FVertexID VertexID = MeshIn->GetVertexInstanceVertex(InstanceTri[1]);
-					const FVector Position = VertexPositions.Get(VertexID);
-					int NewVertIdx = MeshOut.AppendVertex(Position);
-					VertexID1 = NewVertIdx;
-				}
-				if (bDuplicate[2])
-				{
-					FVertexID VertexID = MeshIn->GetVertexInstanceVertex(InstanceTri[2]);
-					const FVector Position = VertexPositions.Get(VertexID);
-					int NewVertIdx = MeshOut.AppendVertex(Position);
-					VertexID2 = NewVertIdx;
-				}
+		// sanity
+		checkSlow(TriangleVertexIDs.Num() == 3);
+
+		FIndex3i VertexIDs; 
+		VertexIDs[0] = TriangleVertexIDs[0].GetValue();
+		VertexIDs[1] = TriangleVertexIDs[1].GetValue(); 
+		VertexIDs[2] = TriangleVertexIDs[2].GetValue();
+
+		int NewTriangleID = MeshOut.AppendTriangle(VertexIDs, GroupID);
+
+		// Deal with potential failure cases
+
+		//-- already seen this triangle for some reason.. or the MeshDecription had a degenerate tri
+		if (NewTriangleID == FDynamicMesh3::DuplicateTriangleID || NewTriangleID == FDynamicMesh3::InvalidID)
+		{
+			continue;
+		}
+
+		//-- non manifold 
+		// if append failed due to non-manifold, duplicate verts
+		if (NewTriangleID == FDynamicMesh3::NonManifoldID)
+		{
+			int e0 = MeshOut.FindEdge(VertexIDs[0], VertexIDs[1]);
+			int e1 = MeshOut.FindEdge(VertexIDs[1], VertexIDs[2]);
+			int e2 = MeshOut.FindEdge(VertexIDs[2], VertexIDs[0]);
 
-				NewTriangleID = MeshOut.AppendTriangle(VertexID0, VertexID1, VertexID2, GroupID);
-				checkSlow(NewTriangleID != FDynamicMesh3::NonManifoldID);
+			// determine which verts need to be duplicated
+			bool bDuplicate[3] = { false, false, false };
+			if (e0 != FDynamicMesh3::InvalidID && MeshOut.IsBoundaryEdge(e0) == false)
+			{
+				bDuplicate[0] = true;
+				bDuplicate[1] = true;
+			}
+			if (e1 != FDynamicMesh3::InvalidID && MeshOut.IsBoundaryEdge(e1) == false)
+			{
+				bDuplicate[1] = true;
+				bDuplicate[2] = true;
+			}
+			if (e2 != FDynamicMesh3::InvalidID && MeshOut.IsBoundaryEdge(e2) == false)
+			{
+				bDuplicate[2] = true;
+				bDuplicate[0] = true;
+			}
+			if (bDuplicate[0])
+			{
+				const FVector Position = VertexPositions[TriangleVertexIDs[0]];
+				const int32 NewVertIdx = MeshOut.AppendVertex(Position);
+				VertexIDs[0] = NewVertIdx;
+				VertIDMap.Insert(TriangleVertexIDs[0], NewVertIdx); // note 'insert', because the array may need to grow
+			}
+			if (bDuplicate[1])
+			{
+				const FVector Position = VertexPositions[TriangleVertexIDs[1]];
+				const int32 NewVertIdx = MeshOut.AppendVertex(Position);
+				VertexIDs[1] = NewVertIdx;
+				VertIDMap.Insert(TriangleVertexIDs[1], NewVertIdx);
+			}
+			if (bDuplicate[2])
+			{
+				const FVector Position = VertexPositions[TriangleVertexIDs[2]];
+				const int32 NewVertIdx = MeshOut.AppendVertex(Position);
+				VertexIDs[2] = NewVertIdx;
+				VertIDMap.Insert(TriangleVertexIDs[2], NewVertIdx);
 			}
 
-			checkSlow(NewTriangleID >= 0);
-			AddedTriangles[NewTriangleID] = TriData;
+			NewTriangleID = MeshOut.AppendTriangle(VertexIDs, GroupID);
+			checkSlow(NewTriangleID != FDynamicMesh3::NonManifoldID);
 		}
+
+		checkSlow(NewTriangleID >= 0);
+		AddedTriangles[NewTriangleID] = TriData;
+		TriIDMap[NewTriangleID] = TriangleID;
+	
 	}
 
 	FDateTime Time_AfterTriangles = FDateTime::Now();
 
+
 	//
 	// Enable relevant attributes and initialize UV/Normal welders
 	// 
 
-	int NumUVLayers = InstanceUVs.GetNumChannels();
+	// the shared UV representation.  
+	const int32 NumUVElementChannels = MeshIn->GetNumUVElementChannels();
+
+	// the instanced UV representation.
+	const int NumUVLayers = InstanceUVs.GetNumChannels();
+
+	// determine if we really have shared UVS. Legacy geo might not have them 
+	// - at the time of this writing MeshDescription has not been updated
+	// to always populated the shared UVs during load time for legacy geometry that only has per-instance UVs. 
+	// but that is a promised improvement. 
+
+	bool bUseSharedUVs =  (NumUVLayers == NumUVElementChannels);
+	if (bUseSharedUVs)
+	{
+		for (int UVLayerIndex = 0; UVLayerIndex < NumUVLayers; ++UVLayerIndex)
+		{
+			const int32 NumSharedUVs   = MeshIn->UVs(UVLayerIndex).GetArraySize();
+			bUseSharedUVs = bUseSharedUVs && (NumSharedUVs != 0);
+		}
+	}
+
+
 	TArray<FDynamicMeshUVOverlay*> UVOverlays;
 	TArray<FUVWelder> UVWelders;
 	FDynamicMeshNormalOverlay* NormalOverlay = nullptr;
@@ -313,39 +358,123 @@ void FMeshDescriptionToDynamicMesh::Convert(const FMeshDescription* MeshIn, FDyn
 	}
 
 
-	// we will weld/populate all the attributes simultaneously, hold on to futures in this array and then Wait for them at the end
-	TArray<TFuture<void>> Pending;
-
-	if (bCalculateMaps)
-	{
-		auto Future = Async(EAsyncExecution::ThreadPool, [&]()
-		{
-			for (int32 TriangleID : MeshOut.TriangleIndicesItr())
-			{
-				const FTriData& TriData = AddedTriangles[TriangleID];
-				TriToPolyTriMap.Insert(FIndex2i(TriData.PolygonID.GetValue(), TriData.TriIndex), TriangleID);
-			}
-		});
-		Pending.Add(MoveTemp(Future));
-	}
+	
 
 	if (!bDisableAttributes)
 	{
+		// we will weld/populate all the attributes simultaneously, hold on to futures in this array and then Wait for them at the end
+		TArray<TFuture<void>> Pending;
+
 		for (int UVLayerIndex = 0; UVLayerIndex < NumUVLayers; UVLayerIndex++)
 		{
-			auto UVFuture = Async(EAsyncExecution::ThreadPool, [&, UVLayerIndex]() // must copy UVLayerIndex here!
+			auto UVFuture = Async(EAsyncExecution::ThreadPool, [&, UVLayerIndex, bUseSharedUVs]() // must copy UVLayerIndex here!
 			{
-				for (int32 TriangleID : MeshOut.TriangleIndicesItr())
+
+				if (!bUseSharedUVs) // have to rely on welding the per-instance uvs 
 				{
-					FIndex3i Tri = MeshOut.GetTriangle(TriangleID);
-					const FTriData& TriData = AddedTriangles[TriangleID];
-					FIndex3i TriUV;
-					for (int j = 0; j < 3; ++j)
+					for (int32 TriangleID : MeshOut.TriangleIndicesItr())
 					{
-						FVector2D UV = InstanceUVs.Get(TriData.TriInstances[j], UVLayerIndex);
-						TriUV[j] = UVWelders[UVLayerIndex].FindOrAddUnique(UV, Tri[j]);
+						FIndex3i Tri = MeshOut.GetTriangle(TriangleID);
+						const FTriData& TriData = AddedTriangles[TriangleID];
+						FIndex3i TriUV;
+						for (int j = 0; j < 3; ++j)
+						{
+							FVector2D UV = InstanceUVs.Get(TriData.TriInstances[j], UVLayerIndex);
+							TriUV[j] = UVWelders[UVLayerIndex].FindOrAddUnique(UV, Tri[j]);
+						}
+						UVOverlays[UVLayerIndex]->SetTriangle(TriangleID, TriUV);
+					}
+				}
+				else
+				{
+					// the overlay to fill.
+					FDynamicMeshUVOverlay* UVOverlay = UVOverlays[UVLayerIndex];
+
+					// copy uv "vertex buffer"
+					const FUVArray& UVs = MeshIn->UVs(UVLayerIndex);
+					TUVAttributesRef<const FVector2D> UVCoordinates = UVs.GetAttributes().GetAttributesRef<FVector2D>(MeshAttribute::UV::UVCoordinate);
+						
+					// map to translate UVIds from FUVID::int32() to DynamicOverlay Index.
+					TArray<int32> UVIndexMap;
+					UVIndexMap.Reserve(UVs.GetArraySize() + 1);
+
+					for (FUVID UVID : UVs.GetElementIDs())
+					{
+						const FVector2D UVvalue = UVCoordinates[UVID];
+						int32 NewIndex = UVOverlay->AppendElement(FVector2f(UVvalue));
+
+						// forced to use Insert since this array might have to resize (unfortunately we can only guess the max UVID in the mesh description)
+						UVIndexMap.Insert(NewIndex, UVID.GetValue());
+
+					}
+
+					// copy uv "index buffer"
+					for (int32 TriID : MeshOut.TriangleIndicesItr())
+					{
+						FTriangleID TriangleID = TriIDMap[TriID];
+						// NB: the mesh description lacks a const method variant of this function, hence the const_cast. 
+						// Don't change the UVIndices values!
+						TArrayView<FUVID> UVIndices = const_cast<FMeshDescription*>(MeshIn)->GetTriangleUVIndices(TriangleID, UVLayerIndex);
+
+						// translate to Overlay indicies 
+						FIndex3i TriUV( UVIndexMap[UVIndices[0].GetValue()],
+							            UVIndexMap[UVIndices[1].GetValue()],
+							            UVIndexMap[UVIndices[2].GetValue()] );
+							
+						///--  We have to do some clean-up on the shared UVs that come from MeshDecription --///
+						// This clean up should go away if MeshDescription can solve these problems during import from the source fbx files
+						{
+							// MeshDescription can attach multiple vertices to the same UV element.  DynamicMesh does not.
+							// if we have already used this element for a different mesh vertex, split it.
+							const FIndex3i ParentTriangle = MeshOut.GetTriangle(TriID);
+							for (int i = 0; i < 3; ++i)
+							{
+								int32 ParentVID = UVOverlay->GetParentVertex(TriUV[i]);
+								if (ParentVID != FDynamicMesh3::InvalidID && ParentVID != ParentTriangle[i])
+								{
+									const FVector2D UVvalue = UVCoordinates[UVIndices[1]];
+									TriUV[i] = UVIndexMap.Add(UVOverlay->AppendElement(FVector2f(UVvalue)));
+								}
+							}
+
+							// MeshDescription allows for degenerate UV tris.  Dynamic Mesh does not.
+							// if the UV tri is degenerate we split the degenerate UV edge by adding two new UVs
+							// in its place, or if it is totally degenerate we add 3 new UVs
+
+							if (TriUV[0] == TriUV[1] && TriUV[0] == TriUV[2])
+							{
+								const FVector2D UVvalue = UVCoordinates[UVIndices[1]];
+								TriUV[0] = UVIndexMap.Add(UVOverlay->AppendElement(FVector2f(UVvalue)));
+								TriUV[1] = UVIndexMap.Add(UVOverlay->AppendElement(FVector2f(UVvalue)));
+								TriUV[2] = UVIndexMap.Add(UVOverlay->AppendElement(FVector2f(UVvalue)));
+
+							}
+							else
+							{
+								if (TriUV[0] == TriUV[1])
+								{
+									const FVector2D UVvalue = UVCoordinates[UVIndices[0]];
+									TriUV[0] = UVIndexMap.Add(UVOverlay->AppendElement(FVector2f(UVvalue)));
+									TriUV[1] = UVIndexMap.Add(UVOverlay->AppendElement(FVector2f(UVvalue)));
+								}
+								if (TriUV[0] == TriUV[2])
+								{
+									const FVector2D UVvalue = UVCoordinates[UVIndices[0]];
+									TriUV[0] = UVIndexMap.Add(UVOverlay->AppendElement(FVector2f(UVvalue)));
+									TriUV[2] = UVIndexMap.Add(UVOverlay->AppendElement(FVector2f(UVvalue)));
+								}
+								if (TriUV[1] == TriUV[2])
+								{
+									const FVector2D UVvalue = UVCoordinates[UVIndices[1]];
+									TriUV[1] = UVIndexMap.Add(UVOverlay->AppendElement(FVector2f(UVvalue)));
+									TriUV[2] = UVIndexMap.Add(UVOverlay->AppendElement(FVector2f(UVvalue)));
+								}
+							}
+						}
+
+						// set the triangle in the overlay
+						UVOverlay->SetTriangle(TriID, TriUV);
 					}
-					UVOverlays[UVLayerIndex]->SetTriangle(TriangleID, TriUV);
 				}
 			});
 			Pending.Add(MoveTemp(UVFuture));
@@ -385,12 +514,19 @@ void FMeshDescriptionToDynamicMesh::Convert(const FMeshDescription* MeshIn, FDyn
 			});
 			Pending.Add(MoveTemp(MaterialFuture));
 		}
+
+		// wait for all work to be done
+		for (TFuture<void>& Future : Pending)
+		{
+			Future.Wait();
+		}
 	}
 
-	// wait for all work to be done
-	for (TFuture<void>& Future : Pending)
+	// free maps if no longer needed
+	if (!bCalculateMaps)
 	{
-		Future.Wait();
+		TriIDMap.Empty();
+		VertIDMap.Empty();
 	}
 
 	FDateTime Time_AfterAttribs = FDateTime::Now();
@@ -410,7 +546,7 @@ void FMeshDescriptionToDynamicMesh::Convert(const FMeshDescription* MeshIn, FDyn
 
 
 template<typename RealType>
-static void CopyTangents_Internal(const FMeshDescription* SourceMesh, const FDynamicMesh3* TargetMesh, TMeshTangents<RealType>* TangentsOut, const TArray<FIndex2i>& TriToPolyTriMap)
+static void CopyTangents_Internal(const FMeshDescription* SourceMesh, const FDynamicMesh3* TargetMesh, TMeshTangents<RealType>* TangentsOut, const TArray<FTriangleID>& TriIDMap)
 {
 
 	FStaticMeshConstAttributes Attributes(*SourceMesh);
@@ -425,12 +561,11 @@ static void CopyTangents_Internal(const FMeshDescription* SourceMesh, const FDyn
 
 	TangentsOut->SetMesh(TargetMesh);
 	TangentsOut->InitializeTriVertexTangents(false);
-
+	
 	for (int32 TriID : TargetMesh->TriangleIndicesItr())
 	{
-		FIndex2i PolyTriIdx = TriToPolyTriMap[TriID];
-		TArrayView<const FTriangleID> TriangleIDs = SourceMesh->GetPolygonTriangles(FPolygonID(PolyTriIdx.A));
-		const FTriangleID TriangleID = TriangleIDs[PolyTriIdx.B];
+
+		const FTriangleID TriangleID = TriIDMap[TriID];
 		TArrayView<const FVertexInstanceID> InstanceTri = SourceMesh->GetTriangleVertexInstances(TriangleID);
 		for (int32 j = 0; j < 3; ++j)
 		{
@@ -449,8 +584,8 @@ static void CopyTangents_Internal(const FMeshDescription* SourceMesh, const FDyn
 void FMeshDescriptionToDynamicMesh::CopyTangents(const FMeshDescription* SourceMesh, const FDynamicMesh3* TargetMesh, TMeshTangents<float>* TangentsOut)
 {
 	if (!ensureMsgf(bCalculateMaps, TEXT("Cannot CopyTangents unless Maps were calculated"))) return;
-	if (!ensureMsgf(TriToPolyTriMap.Num() == TargetMesh->TriangleCount(), TEXT("Tried to CopyTangents to mesh with different triangle count"))) return;
-	CopyTangents_Internal<float>(SourceMesh, TargetMesh, TangentsOut, TriToPolyTriMap);
+	if (!ensureMsgf(TriIDMap.Num() == TargetMesh->TriangleCount(), TEXT("Tried to CopyTangents to mesh with different triangle count"))) return;
+	CopyTangents_Internal<float>(SourceMesh, TargetMesh, TangentsOut, TriIDMap);
 }
 
 
@@ -458,8 +593,8 @@ void FMeshDescriptionToDynamicMesh::CopyTangents(const FMeshDescription* SourceM
 void FMeshDescriptionToDynamicMesh::CopyTangents(const FMeshDescription* SourceMesh, const FDynamicMesh3* TargetMesh, TMeshTangents<double>* TangentsOut)
 {
 	if (!ensureMsgf(bCalculateMaps, TEXT("Cannot CopyTangents unless Maps were calculated"))) return;
-	if (!ensureMsgf(TriToPolyTriMap.Num() == TargetMesh->TriangleCount(), TEXT("Tried to CopyTangents to mesh with different triangle count"))) return;
-	CopyTangents_Internal<double>(SourceMesh, TargetMesh, TangentsOut, TriToPolyTriMap);
+	if (!ensureMsgf(TriIDMap.Num() == TargetMesh->TriangleCount(), TEXT("Tried to CopyTangents to mesh with different triangle count"))) return;
+	CopyTangents_Internal<double>(SourceMesh, TargetMesh, TangentsOut, TriIDMap);
 }
 
 
diff --git a/Engine/Plugins/Experimental/GeometryProcessing/Source/MeshConversion/Public/DynamicMeshToMeshDescription.h b/Engine/Plugins/Experimental/GeometryProcessing/Source/MeshConversion/Public/DynamicMeshToMeshDescription.h
--- a/Engine/Plugins/Experimental/GeometryProcessing/Source/MeshConversion/Public/DynamicMeshToMeshDescription.h
+++ b/Engine/Plugins/Experimental/GeometryProcessing/Source/MeshConversion/Public/DynamicMeshToMeshDescription.h
@@ -83,14 +83,14 @@ public:
 
 	/**
 	 * Update only attributes, assuming the mesh topology has not changed.  Does not touch positions.
-	 *	NOTE: assumes the order of triangles in the MeshIn correspond to the ordering you'd get by iterating over polygons, then tris-in-polygons, on MeshOut
+	 *	NOTE: assumes the order of triangles in the MeshIn correspond to the ordering you'd get by iterating over triangles, on MeshOut
 	 *		  This matches conversion currently used in MeshDescriptionToDynamicMesh.cpp, but if that changes we will need to change this function to match!
 	 */
 	void UpdateAttributes(const FDynamicMesh3* MeshIn, FMeshDescription& MeshOut, bool bUpdateNormals, bool bUpdateUVs);
 
 	/**
 	 * Update the Tangent and BinormalSign attributes of the MeshDescription, assuming mesh topology has not changed. Does not modify any other attributes.
-	 *	NOTE: assumes the order of triangles in the MeshIn correspond to the ordering you'd get by iterating over polygons, then tris-in-polygons, on MeshOut
+	 *	NOTE: assumes the order of triangles in the MeshIn correspond to the ordering you'd get by iterating over triangles, on MeshOut
 	 *		  This matches conversion currently used in MeshDescriptionToDynamicMesh.cpp, but if that changes we will need to change this function to match!
 	 */
 	void UpdateTangents(const FDynamicMesh3* MeshIn, FMeshDescription& MeshOut, const TMeshTangents<double>* Tangents);
@@ -113,7 +113,8 @@ public:
 
 	/**
 	 * Convert with no shared VertexInstances. A new VertexInstance is created for
-	 * each triangle vertex (ie corner). However vertex positions are shared.
+	 * each triangle vertex (ie corner). However vertex positions are shared and
+     * the shared UVs structures are populated.
 	 */
 	void Convert_NoSharedInstances(const FDynamicMesh3* MeshIn, FMeshDescription& MeshOut);
 };
\ No newline at end of file
diff --git a/Engine/Plugins/Experimental/GeometryProcessing/Source/MeshConversion/Public/MeshDescriptionBuilder.h b/Engine/Plugins/Experimental/GeometryProcessing/Source/MeshConversion/Public/MeshDescriptionBuilder.h
--- a/Engine/Plugins/Experimental/GeometryProcessing/Source/MeshConversion/Public/MeshDescriptionBuilder.h
+++ b/Engine/Plugins/Experimental/GeometryProcessing/Source/MeshConversion/Public/MeshDescriptionBuilder.h
@@ -25,6 +25,9 @@ class MESHCONVERSION_API FMeshDescriptionBuilder
 public:
 	void SetMeshDescription(FMeshDescription* Description);
 
+	/** Pre-allocate space in the mesh description */
+	void ReserveNewVertices(int32 Count);
+
 	/** Append vertex and return new vertex ID */
 	FVertexID AppendVertex(const FVector& Position);
 
@@ -37,23 +40,28 @@ public:
 	/** Set the position of a vertex */
 	void SetPosition(const FVertexID& VertexID, const FVector& NewPosition);
 
+	/** Set the number of UV layers */
+	void SetNumUVLayers(int32 NumUVLayers);
+
+	/** Pre-allocate space in the mesh description for UVs in the indicated UVLayer */
+	void ReserveNewUVs(int32 Count, int UVLayerIndex);
 
+	/** Append a UV 'vertex' and return a new UV ID*/
+	FUVID AppendUV(const FVector2D& UVvalue, int32 UVLayerIndex);
 
 	/** Append new vertex instance and return ID */
 	FVertexInstanceID AppendInstance(const FVertexID& VertexID);
 
-	/** Set the UV and Normal of a vertex instance*/
-	void SetInstance(const FVertexInstanceID& InstanceID, const FVector2D& InstanceUV, const FVector& InstanceNormal);
+	/** 
+	* Set the UV of a vertex instance.
+	* Note: this generally shouldn't be called directly because it alters the instance UVs but not the shared UVs
+	* When setting UVs use AppendUVTriangle	  
+	*/
+	void SetInstanceUV(const FVertexInstanceID& InstanceID, const FVector2D& InstanceUV, int32 UVLayerIndex = 0);
 
 	/** Set the Normal of a vertex instance*/
 	void SetInstanceNormal(const FVertexInstanceID& InstanceID, const FVector& Normal);
 
-	/** Set the UV of a vertex instance */
-	void SetInstanceUV(const FVertexInstanceID& InstanceID, const FVector2D& InstanceUV, int32 UVLayerIndex = 0);
-
-	/** Set the number of UV layers */
-	void SetNumUVLayers(int32 NumUVLayers);
-
 	/** Set the Color of a vertex instance*/
 	void SetInstanceColor(const FVertexInstanceID& InstanceID, const FVector4& Color);
 
@@ -63,31 +71,30 @@ public:
 	/** Create a new polygon group and return it's ID */
 	FPolygonGroupID AppendPolygonGroup();
 
-	/** Set the PolyTriGroups attribute value to a specific GroupID for a Polygon */
-	void SetPolyGroupID(const FPolygonID& PolygonID, int GroupID);
-
+	/** Set the PolyTriGroups attribute value to a specific GroupID for a Triangle */
+	void SetPolyGroupID(const FTriangleID& TriangleID, int GroupID);
 
+	/** Append a UV triangle to the specified UV layer. This will use both shared and per-instance UV storage*/
+	void AppendUVTriangle(const FTriangleID& TriangleID, const FUVID UVverterxID0, const FUVID UVvertexID1, const FUVID UVvertexID2, int32 UVLayerIndex);
 
 	/** Append a triangle to the mesh with the given PolygonGroup ID */
-	FPolygonID AppendTriangle(const FVertexID& Vertex0, const FVertexID& Vertex1, const FVertexID& Vertex2, const FPolygonGroupID& PolygonGroup);
+	FTriangleID AppendTriangle(const FVertexID& Vertex0, const FVertexID& Vertex1, const FVertexID& Vertex2, const FPolygonGroupID& PolygonGroup);
 
-	/** Append a triangle to the mesh with the given PolygonGroup ID, and optionally with triangle-vertex UVs and Normals */
-	FPolygonID AppendTriangle(const FVertexID* Triangle, const FPolygonGroupID& PolygonGroup, 
-		const FVector2D* VertexUVs = nullptr, const FVector* VertexNormals = nullptr);
 
-	/** 
-	 * Append an arbitrary polygon to the mesh with the given PolygonGroup ID, and optionally with polygon-vertex UVs and Normals
-	 * Unique Vertex instances will be created for each polygon-vertex.
-	 */
-	FPolygonID AppendPolygon(const TArray<FVertexID>& Vertices, const FPolygonGroupID& PolygonGroup, 
-		const TArray<FVector2D>* VertexUVs = nullptr, const TArray<FVector>* VertexNormals = nullptr);
+	/** Append a triangle to the mesh with the given PolygonGroup ID */
+	FTriangleID AppendTriangle(const FVertexID* Triangle, const FPolygonGroupID& PolygonGroup);
+
 
 	/**
 	 * Append a triangle to the mesh using the given vertex instances and PolygonGroup ID
 	 */
-	FPolygonID AppendTriangle(const FVertexInstanceID& Instance0, const FVertexInstanceID& Instance1, const FVertexInstanceID& Instance2, const FPolygonGroupID& PolygonGroup);
-
+	FTriangleID AppendTriangle(const FVertexInstanceID& Instance0, const FVertexInstanceID& Instance1, const FVertexInstanceID& Instance2, const FPolygonGroupID& PolygonGroup);
 
+	/**
+	 * Append an arbitrary polygon to the mesh with the given PolygonGroup ID
+	 * Unique Vertex instances will be created for each polygon-vertex.
+	 */
+	FPolygonID AppendPolygon(const TArray<FVertexID>& Vertices, const FPolygonGroupID& PolygonGroup);
 
 	/** Set MeshAttribute::Edge::IsHard to true for all edges */
 	void SetAllEdgesHardness(bool bHard);
@@ -99,16 +106,24 @@ public:
 	/** Return the current bounding box of the mesh */
 	FBox ComputeBoundingBox() const;
 
+	/** Disable the construction of secondary data structures in the mesh description */
+	void SuspendMeshDescriptionIndexing();
+
+	/** Enable the construction of secondary data structures in the mesh description */
+	void ResumeMeshDescriptionIndexing();
+
 protected:
+
 	FMeshDescription* MeshDescription;
 
 	TVertexAttributesRef<FVector> VertexPositions;
 	TVertexInstanceAttributesRef<FVector2D> InstanceUVs;
 	TVertexInstanceAttributesRef<FVector> InstanceNormals;
 	TVertexInstanceAttributesRef<FVector4> InstanceColors;
+
+	TArray<TUVAttributesRef<FVector2D>> UVCoordinateLayers; 
 	TArray<FVertexID> TempBuffer;
-	TArray<FVector2D> UVBuffer;
-	TArray<FVector> NormalBuffer;
+	TArray<FUVID> TempUVBuffer;
 
 	TPolygonAttributesRef<int> PolyGroups;
 };
\ No newline at end of file
diff --git a/Engine/Plugins/Experimental/GeometryProcessing/Source/MeshConversion/Public/MeshDescriptionToDynamicMesh.h b/Engine/Plugins/Experimental/GeometryProcessing/Source/MeshConversion/Public/MeshDescriptionToDynamicMesh.h
--- a/Engine/Plugins/Experimental/GeometryProcessing/Source/MeshConversion/Public/MeshDescriptionToDynamicMesh.h
+++ b/Engine/Plugins/Experimental/GeometryProcessing/Source/MeshConversion/Public/MeshDescriptionToDynamicMesh.h
@@ -32,9 +32,17 @@ public:
 	/** Ignore all mesh attributes (e.g. UV/Normal layers, material groups) */
 	bool bDisableAttributes = false;
 
-	/** map from triangle ID to (polygon,triangle) pair */
-	TArray<FIndex2i> TriToPolyTriMap;
 
+	/** map from DynamicMesh triangle ID to MeshDescription FTriangleID*/
+	TArray<FTriangleID> TriIDMap;
+
+	/**
+	* map from DynamicMesh vertex Id to MeshDecription FVertexID. 
+	* NB: due to vertex splitting, multiple DynamicMesh vertex ids 
+	* may map to the same MeshDescription FVertexID.
+	*  ( a vertex split is a result of reconciling non-manifold MeshDescription vertex ) 	  
+	*/
+	TArray<FVertexID> VertIDMap;
 
 	/**
 	 * Various modes can be used to create output triangle groups
diff --git a/Engine/Plugins/Experimental/MeshModelingToolset/Source/ModelingComponents/Private/SimpleDynamicMeshComponent.cpp b/Engine/Plugins/Experimental/MeshModelingToolset/Source/ModelingComponents/Private/SimpleDynamicMeshComponent.cpp
--- a/Engine/Plugins/Experimental/MeshModelingToolset/Source/ModelingComponents/Private/SimpleDynamicMeshComponent.cpp
+++ b/Engine/Plugins/Experimental/MeshModelingToolset/Source/ModelingComponents/Private/SimpleDynamicMeshComponent.cpp
@@ -168,8 +168,11 @@ void USimpleDynamicMeshComponent::ApplyTransform(const FTransform3d& Transform,
 
 void USimpleDynamicMeshComponent::Bake(FMeshDescription* MeshDescription, bool bHaveModifiedTopology, const FConversionToMeshDescriptionOptions& ConversionOptions)
 {
+	// shared UVs on MeshDescription can't currently be updated. this forces us to rebuild the mesh.
+	const bool bModifiedUVsOrTopology = ( ConversionOptions.bUpdateUVs || bHaveModifiedTopology ); 
+	
 	FDynamicMeshToMeshDescription Converter(ConversionOptions);
-	if (bHaveModifiedTopology == false && Converter.HaveMatchingElementCounts(Mesh.Get(), MeshDescription))
+	if (bModifiedUVsOrTopology == false && Converter.HaveMatchingElementCounts(Mesh.Get(), MeshDescription))
 	{
 		if (ConversionOptions.bUpdatePositions)
 		{
