GeometryProcessing: move MeshIndexUtil functions to UE namespace (keep old functions for now, calling new)
#rnx

diff --git a/Engine/Plugins/Experimental/GeometryProcessing/Source/DynamicMesh/Private/MeshIndexUtil.cpp b/Engine/Plugins/Experimental/GeometryProcessing/Source/DynamicMesh/Private/MeshIndexUtil.cpp
--- a/Engine/Plugins/Experimental/GeometryProcessing/Source/DynamicMesh/Private/MeshIndexUtil.cpp
+++ b/Engine/Plugins/Experimental/GeometryProcessing/Source/DynamicMesh/Private/MeshIndexUtil.cpp
@@ -3,8 +3,11 @@
 
 #include "MeshIndexUtil.h"
 
-
 void MeshIndexUtil::TriangleToVertexIDs(const FDynamicMesh3* Mesh, const TArray<int>& TriangleIDs, TArray<int>& VertexIDsOut)
+{
+	UE::MeshIndexUtil::TriangleToVertexIDs(Mesh, TriangleIDs, VertexIDsOut);
+}
+void UE::MeshIndexUtil::TriangleToVertexIDs(const FDynamicMesh3* Mesh, const TArray<int>& TriangleIDs, TArray<int>& VertexIDsOut)
 {
 	// if we are getting close to full mesh it is probably more efficient to use a bitmap...
 
@@ -50,8 +53,11 @@ void MeshIndexUtil::TriangleToVertexIDs(const FDynamicMesh3* Mesh, const TArray<
 
 
 
-
 void MeshIndexUtil::VertexToTriangleOneRing(const FDynamicMesh3* Mesh, const TArray<int>& VertexIDs, TSet<int>& TriangleIDsOut)
+{
+	UE::MeshIndexUtil::VertexToTriangleOneRing(Mesh, VertexIDs, TriangleIDsOut);
+}
+void UE::MeshIndexUtil::VertexToTriangleOneRing(const FDynamicMesh3* Mesh, const TArray<int>& VertexIDs, TSet<int>& TriangleIDsOut)
 {
 	int NumVerts = VertexIDs.Num();
 	TriangleIDsOut.Reserve( (NumVerts < 5) ? NumVerts*6 : NumVerts*4);
diff --git a/Engine/Plugins/Experimental/GeometryProcessing/Source/DynamicMesh/Public/MeshIndexUtil.h b/Engine/Plugins/Experimental/GeometryProcessing/Source/DynamicMesh/Public/MeshIndexUtil.h
--- a/Engine/Plugins/Experimental/GeometryProcessing/Source/DynamicMesh/Public/MeshIndexUtil.h
+++ b/Engine/Plugins/Experimental/GeometryProcessing/Source/DynamicMesh/Public/MeshIndexUtil.h
@@ -9,6 +9,10 @@
  */
 namespace MeshIndexUtil
 {
+	//
+	// Do not add any more functions to this namespace. Use UE::MeshIndexUtil instead
+	//
+
 	/**
 	 * Find list of unique vertices that are contained in one or more triangles
 	 * @param Mesh input mesh
@@ -32,6 +36,23 @@ namespace UE
 {
 	namespace MeshIndexUtil
 	{
+		/**
+		 * Find list of unique vertices that are contained in one or more triangles
+		 * @param Mesh input mesh
+		 * @param TriangleIDs list of triangle IDs of Mesh
+		 * @param VertexIDsOut list of vertices contained by triangles
+		 */
+		DYNAMICMESH_API void TriangleToVertexIDs(const FDynamicMesh3* Mesh, const TArray<int>& TriangleIDs, TArray<int>& VertexIDsOut);
+
+
+		/**
+		 * Find all the triangles in all the one rings of a set of vertices
+		 * @param Mesh input mesh
+		 * @param VertexIDs list of Vertex IDs of Mesh
+		 * @param TriangleIDsOut list of triangle IDs where any vertex of triangle touches an element of VertexIDs
+		 */
+		DYNAMICMESH_API void VertexToTriangleOneRing(const FDynamicMesh3* Mesh, const TArray<int>& VertexIDs, TSet<int>& TriangleIDsOut);
+
 		/**
 		 * Walk around VertexID from FromTriangleID to next connected triangle if it exists, walking "away" from PrevTriangleID.
 		 * @param TrisConnectedFunc returns true if two triangles should be considered connected, to support breaking at seams/etc that are not in base mesh topology
