Geometry Tools:  MeshDescriptionBuilder.  Add the static mesh attributes by default.

#rnx

diff --git a/Engine/Plugins/Experimental/GeometryProcessing/Source/MeshConversion/Private/MeshDescriptionBuilder.cpp b/Engine/Plugins/Experimental/GeometryProcessing/Source/MeshConversion/Private/MeshDescriptionBuilder.cpp
--- a/Engine/Plugins/Experimental/GeometryProcessing/Source/MeshConversion/Private/MeshDescriptionBuilder.cpp
+++ b/Engine/Plugins/Experimental/GeometryProcessing/Source/MeshConversion/Private/MeshDescriptionBuilder.cpp
@@ -21,6 +21,14 @@ void FMeshDescriptionBuilder::SetMeshDescription(FMeshDescription* Description)
 {
 	FStaticMeshAttributes Attributes(*Description);
 
+	// make sure we have at least one UV channel
+	if (!Attributes.GetVertexInstanceUVs().IsValid())
+	{
+		Description->VertexInstanceAttributes().RegisterAttribute<FVector2D>(MeshAttribute::VertexInstance::TextureCoordinate, 1, FVector2D::ZeroVector, EMeshAttributeFlags::Lerpable | EMeshAttributeFlags::Mandatory);
+	}
+
+
+	// handles to some of the standard attributes.
 	this->MeshDescription = Description;
 	this->VertexPositions = Attributes.GetVertexPositions();
 	this->InstanceUVs = Attributes.GetVertexInstanceUVs();
@@ -164,7 +172,8 @@ void FMeshDescriptionBuilder::SetInstanceUV(const FVertexInstanceID& InstanceID,
 
 void FMeshDescriptionBuilder::SetNumUVLayers(int32 NumUVLayers)
 {
-	if (ensure(InstanceUVs.IsValid()))
+	bool bValidInstanceUVs = InstanceUVs.IsValid();
+	if (bValidInstanceUVs)
 	{
 		// initialize the instanced UV channels
 		InstanceUVs.SetNumChannels(NumUVLayers);
diff --git a/Engine/Plugins/Experimental/GeometryProcessing/Source/MeshConversion/Public/MeshDescriptionBuilder.h b/Engine/Plugins/Experimental/GeometryProcessing/Source/MeshConversion/Public/MeshDescriptionBuilder.h
--- a/Engine/Plugins/Experimental/GeometryProcessing/Source/MeshConversion/Public/MeshDescriptionBuilder.h
+++ b/Engine/Plugins/Experimental/GeometryProcessing/Source/MeshConversion/Public/MeshDescriptionBuilder.h
@@ -18,7 +18,8 @@ namespace ExtendedMeshAttribute
 
 
 /**
- * Utility class to construct MeshDescription instances
+ * Utility class to construct MeshDescription instances.
+ * NB: this will add a vertex-instance UV to the Description, if it has none. 
  */
 class MESHCONVERSION_API FMeshDescriptionBuilder
 {
