GeometryProcessing: add support for additional Polygroup layers in DynamicMeshAttributeSet.

- add Name FName to TDynamicAttributeBase, update Copy functions in various subclasses to copy this name
- add PolygroupLayers array to FDynamicMeshAttributeSet, add code get/set, call update functions where necessary. Also move some existing function implementations to cpp file
- update TDynamicMeshAttributeChange, DynamicMeshEditor to handle these new attributes
- AttributeEditorTool now creates polygroup layers as Triangle attributes
- MeshDescriptionToDynamicMesh now creates PolygroupLayer attributes from any non-reserved Triangle int32 attributes
- DynamicMeshToMeshDescription now syncs PolygroupLayer attributes with MeshDescription in new ConvertPolygonGroupLayers() function
- add TSet variant of FMeshConnectedComponents::GrowToConnectedTriangles

#rnx

diff --git a/Engine/Plugins/Experimental/GeometryProcessing/Source/DynamicMesh/Private/DynamicMeshAttributeSet.cpp b/Engine/Plugins/Experimental/GeometryProcessing/Source/DynamicMesh/Private/DynamicMeshAttributeSet.cpp
--- a/Engine/Plugins/Experimental/GeometryProcessing/Source/DynamicMesh/Private/DynamicMeshAttributeSet.cpp
+++ b/Engine/Plugins/Experimental/GeometryProcessing/Source/DynamicMesh/Private/DynamicMeshAttributeSet.cpp
@@ -6,6 +6,249 @@
 
 
 
+void FDynamicMeshAttributeSet::Copy(const FDynamicMeshAttributeSet& Copy)
+{
+	SetNumUVLayers(Copy.NumUVLayers());
+	for (int UVIdx = 0; UVIdx < NumUVLayers(); UVIdx++)
+	{
+		UVLayers[UVIdx].Copy(Copy.UVLayers[UVIdx]);
+	}
+	Normals0.Copy(Copy.Normals0);
+
+	if (Copy.MaterialIDAttrib)
+	{
+		EnableMaterialID();
+		MaterialIDAttrib->Copy(*(Copy.MaterialIDAttrib));
+	}
+	else
+	{
+		DisableMaterialID();
+	}
+
+	SetNumPolygroupLayers(Copy.NumPolygroupLayers());
+	for (int GroupIdx = 0; GroupIdx < NumPolygroupLayers(); ++GroupIdx)
+	{
+		PolygroupLayers[GroupIdx].Copy(Copy.PolygroupLayers[GroupIdx]);
+	}
+
+	GenericAttributes.Reset();
+	ResetRegisteredAttributes();
+	for (const TPair<FName, TUniquePtr<FDynamicMeshAttributeBase>>& AttribPair : Copy.GenericAttributes)
+	{
+		AttachAttribute(AttribPair.Key, AttribPair.Value->MakeCopy(ParentMesh));
+	}
+
+	// parent mesh is *not* copied!
+}
+
+
+bool FDynamicMeshAttributeSet::IsCompact()
+{
+	for (int UVIdx = 0; UVIdx < NumUVLayers(); UVIdx++)
+	{
+		if (!UVLayers[UVIdx].IsCompact())
+		{
+			return false;
+		}
+	}
+	if (!Normals0.IsCompact())
+	{
+		return false;
+	}
+	// material ID and generic per-element attributes currently cannot be non-compact
+	return true;
+}
+
+
+
+void FDynamicMeshAttributeSet::CompactCopy(const FCompactMaps& CompactMaps, const FDynamicMeshAttributeSet& Copy)
+{
+	SetNumUVLayers(Copy.NumUVLayers());
+	for (int UVIdx = 0; UVIdx < NumUVLayers(); UVIdx++)
+	{
+		UVLayers[UVIdx].CompactCopy(CompactMaps, Copy.UVLayers[UVIdx]);
+	}
+	Normals0.CompactCopy(CompactMaps, Copy.Normals0);
+
+	if (Copy.MaterialIDAttrib)
+	{
+		EnableMaterialID();
+		MaterialIDAttrib->CompactCopy(CompactMaps, *(Copy.MaterialIDAttrib));
+	}
+	else
+	{
+		DisableMaterialID();
+	}
+
+	SetNumPolygroupLayers(Copy.NumPolygroupLayers());
+	for (int GroupIdx = 0; GroupIdx < NumPolygroupLayers(); ++GroupIdx)
+	{
+		PolygroupLayers[GroupIdx].CompactCopy(CompactMaps, Copy.PolygroupLayers[GroupIdx]);
+	}
+
+	GenericAttributes.Reset();
+	ResetRegisteredAttributes();
+	for (const TPair<FName, TUniquePtr<FDynamicMeshAttributeBase>>& AttribPair : Copy.GenericAttributes)
+	{
+		AttachAttribute(AttribPair.Key, AttribPair.Value->MakeCompactCopy(CompactMaps, ParentMesh));
+	}
+
+	// parent mesh is *not* copied!
+}
+
+
+
+
+void FDynamicMeshAttributeSet::CompactInPlace(const FCompactMaps& CompactMaps)
+{
+	for (int UVIdx = 0; UVIdx < NumUVLayers(); UVIdx++)
+	{
+		UVLayers[UVIdx].CompactInPlace(CompactMaps);
+	}
+	Normals0.CompactInPlace(CompactMaps);
+
+	if (MaterialIDAttrib.IsValid())
+	{
+		MaterialIDAttrib->CompactInPlace(CompactMaps);
+	}
+
+	for (int GroupIdx = 0; GroupIdx < NumPolygroupLayers(); ++GroupIdx)
+	{
+		PolygroupLayers[GroupIdx].CompactInPlace(CompactMaps);
+	}
+
+	for (const TPair<FName, TUniquePtr<FDynamicMeshAttributeBase>>& AttribPair : GenericAttributes)
+	{
+		AttribPair.Value->CompactInPlace(CompactMaps);
+	}
+}
+
+
+
+void FDynamicMeshAttributeSet::EnableMatchingAttributes(const FDynamicMeshAttributeSet& ToMatch)
+{
+	SetNumUVLayers(ToMatch.NumUVLayers());
+	for (int UVIdx = 0; UVIdx < NumUVLayers(); UVIdx++)
+	{
+		UVLayers[UVIdx].ClearElements();
+	}
+	Normals0.ClearElements();
+
+	if (ToMatch.MaterialIDAttrib)
+	{
+		EnableMaterialID();
+	}
+	else
+	{
+		DisableMaterialID();
+	}
+
+	SetNumPolygroupLayers(ToMatch.NumPolygroupLayers());
+	for (int GroupIdx = 0; GroupIdx < NumPolygroupLayers(); ++GroupIdx)
+	{
+		PolygroupLayers[GroupIdx].Initialize((int32)0);
+	}
+
+	GenericAttributes.Reset();
+	ResetRegisteredAttributes();
+	for (const TPair<FName, TUniquePtr<FDynamicMeshAttributeBase>>& AttribPair : ToMatch.GenericAttributes)
+	{
+		AttachAttribute(AttribPair.Key, AttribPair.Value->MakeNew(ParentMesh));
+	}
+}
+
+
+
+void FDynamicMeshAttributeSet::Reparent(FDynamicMesh3* NewParent)
+{
+	ParentMesh = NewParent;
+
+	for (int UVIdx = 0; UVIdx < NumUVLayers(); UVIdx++)
+	{
+		UVLayers[UVIdx].Reparent(NewParent);
+	}
+	Normals0.Reparent(NewParent);
+
+	if (MaterialIDAttrib)
+	{
+		MaterialIDAttrib->Reparent(NewParent);
+	}
+
+	for (int GroupIdx = 0; GroupIdx < NumPolygroupLayers(); ++GroupIdx)
+	{
+		PolygroupLayers[GroupIdx].Reparent(NewParent);
+	}
+
+	for (const TPair<FName, TUniquePtr<FDynamicMeshAttributeBase>>& AttribPair : GenericAttributes)
+	{
+		AttribPair.Value->Reparent(NewParent);
+	}
+}
+
+
+
+void FDynamicMeshAttributeSet::SetNumUVLayers(int Num)
+{
+	if (UVLayers.Num() == Num)
+	{
+		return;
+	}
+	if (Num >= UVLayers.Num())
+	{
+		for (int i = (int)UVLayers.Num(); i < Num; ++i)
+		{
+			UVLayers.Add(new FDynamicMeshUVOverlay(ParentMesh));
+		}
+	}
+	else
+	{
+		UVLayers.RemoveAt(Num, UVLayers.Num() - Num);
+	}
+	check(UVLayers.Num() == Num);
+}
+
+
+
+
+
+int32 FDynamicMeshAttributeSet::NumPolygroupLayers() const
+{
+	return PolygroupLayers.Num();
+}
+
+
+void FDynamicMeshAttributeSet::SetNumPolygroupLayers(int32 Num)
+{
+	if (PolygroupLayers.Num() == Num)
+	{
+		return;
+	}
+	if (Num >= PolygroupLayers.Num())
+	{
+		for (int i = (int)PolygroupLayers.Num(); i < Num; ++i)
+		{
+			PolygroupLayers.Add(new FDynamicMeshPolygroupAttribute(ParentMesh));
+		}
+	}
+	else
+	{
+		PolygroupLayers.RemoveAt(Num, PolygroupLayers.Num() - Num);
+	}
+	check(PolygroupLayers.Num() == Num);
+}
+
+FDynamicMeshPolygroupAttribute* FDynamicMeshAttributeSet::GetPolygroupLayer(int Index)
+{
+	return &PolygroupLayers[Index];
+}
+
+const FDynamicMeshPolygroupAttribute* FDynamicMeshAttributeSet::GetPolygroupLayer(int Index) const
+{
+	return &PolygroupLayers[Index];
+}
+
+
+
 void FDynamicMeshAttributeSet::EnableMaterialID()
 {
 	if (HasMaterialID() == false)
@@ -93,6 +336,12 @@ void FDynamicMeshAttributeSet::OnNewTriangle(int TriangleID, bool bInserted)
 		int NewValue = 0;
 		MaterialIDAttrib->SetNewValue(TriangleID, &NewValue);
 	}
+
+	for (FDynamicMeshPolygroupAttribute& PolygroupLayer : PolygroupLayers)
+	{
+		int32 NewGroup = 0;
+		PolygroupLayer.SetNewValue(TriangleID, &NewGroup);
+	}
 }
 
 
@@ -136,6 +385,11 @@ void FDynamicMeshAttributeSet::OnSplitEdge(const FDynamicMesh3::FEdgeSplitInfo&
 	{
 		MaterialIDAttrib->OnSplitEdge(SplitInfo);
 	}
+
+	for (FDynamicMeshPolygroupAttribute& PolygroupLayer : PolygroupLayers)
+	{
+		PolygroupLayer.OnSplitEdge(SplitInfo);
+	}
 }
 
 void FDynamicMeshAttributeSet::OnFlipEdge(const FDynamicMesh3::FEdgeFlipInfo & flipInfo)
@@ -152,6 +406,11 @@ void FDynamicMeshAttributeSet::OnFlipEdge(const FDynamicMesh3::FEdgeFlipInfo & f
 	{
 		MaterialIDAttrib->OnFlipEdge(flipInfo);
 	}
+
+	for (FDynamicMeshPolygroupAttribute& PolygroupLayer : PolygroupLayers)
+	{
+		PolygroupLayer.OnFlipEdge(flipInfo);
+	}
 }
 
 
@@ -169,6 +428,11 @@ void FDynamicMeshAttributeSet::OnCollapseEdge(const FDynamicMesh3::FEdgeCollapse
 	{
 		MaterialIDAttrib->OnCollapseEdge(collapseInfo);
 	}
+
+	for (FDynamicMeshPolygroupAttribute& PolygroupLayer : PolygroupLayers)
+	{
+		PolygroupLayer.OnCollapseEdge(collapseInfo);
+	}
 }
 
 void FDynamicMeshAttributeSet::OnPokeTriangle(const FDynamicMesh3::FPokeTriangleInfo & pokeInfo)
@@ -185,6 +449,11 @@ void FDynamicMeshAttributeSet::OnPokeTriangle(const FDynamicMesh3::FPokeTriangle
 	{
 		MaterialIDAttrib->OnPokeTriangle(pokeInfo);
 	}
+
+	for (FDynamicMeshPolygroupAttribute& PolygroupLayer : PolygroupLayers)
+	{
+		PolygroupLayer.OnPokeTriangle(pokeInfo);
+	}
 }
 
 void FDynamicMeshAttributeSet::OnMergeEdges(const FDynamicMesh3::FMergeEdgesInfo & mergeInfo)
@@ -201,6 +470,11 @@ void FDynamicMeshAttributeSet::OnMergeEdges(const FDynamicMesh3::FMergeEdgesInfo
 	{
 		MaterialIDAttrib->OnMergeEdges(mergeInfo);
 	}
+
+	for (FDynamicMeshPolygroupAttribute& PolygroupLayer : PolygroupLayers)
+	{
+		PolygroupLayer.OnMergeEdges(mergeInfo);
+	}
 }
 
 void FDynamicMeshAttributeSet::OnSplitVertex(const DynamicMeshInfo::FVertexSplitInfo& SplitInfo, const TArrayView<const int>& TrianglesToUpdate)
@@ -217,5 +491,10 @@ void FDynamicMeshAttributeSet::OnSplitVertex(const DynamicMeshInfo::FVertexSplit
 	{
 		MaterialIDAttrib->OnSplitVertex(SplitInfo, TrianglesToUpdate);
 	}
+
+	for (FDynamicMeshPolygroupAttribute& PolygroupLayer : PolygroupLayers)
+	{
+		PolygroupLayer.OnSplitVertex(SplitInfo, TrianglesToUpdate);
+	}
 }
 
diff --git a/Engine/Plugins/Experimental/GeometryProcessing/Source/DynamicMesh/Private/DynamicMeshChangeTracker.cpp b/Engine/Plugins/Experimental/GeometryProcessing/Source/DynamicMesh/Private/DynamicMeshChangeTracker.cpp
--- a/Engine/Plugins/Experimental/GeometryProcessing/Source/DynamicMesh/Private/DynamicMeshChangeTracker.cpp
+++ b/Engine/Plugins/Experimental/GeometryProcessing/Source/DynamicMesh/Private/DynamicMeshChangeTracker.cpp
@@ -399,11 +399,16 @@ void FDynamicMeshAttributeSetChangeTracker::BeginChange()
 		Change->MaterialIDAttribChange.Emplace();
 	}
 
+	int NumPolygroupLayers = Attribs->NumPolygroupLayers();
+	if (NumPolygroupLayers > 0)
+	{
+		Change->PolygroupChanges.SetNum(NumPolygroupLayers);
+	}
+
 	for (int k = 0, NumRegAttrib = Attribs->NumRegisteredAttributes(); k < NumRegAttrib; k++)
 	{
 		Change->RegisteredAttributeChanges.Add(Attribs->GetRegisteredAttribute(k)->NewBlankChange());
 	}
-
 }
 
 
@@ -452,6 +457,12 @@ void FDynamicMeshAttributeSetChangeTracker::SaveInitialTriangle(int TriangleID)
 		Change->MaterialIDAttribChange->SaveInitialTriangle(Attribs->GetMaterialID(), TriangleID);
 	}
 
+	int NumPolygroupLayers = Attribs->NumPolygroupLayers();
+	for (int k = 0; k < NumPolygroupLayers; ++k)
+	{
+		Change->PolygroupChanges[k].SaveInitialTriangle(Attribs->GetPolygroupLayer(k), TriangleID);
+	}
+
 	for (int k = 0, NumRegAttrib = Attribs->NumRegisteredAttributes(); k < NumRegAttrib; k++)
 	{
 		Change->RegisteredAttributeChanges[k]->SaveInitialTriangle(Attribs->GetRegisteredAttribute(k), TriangleID);
@@ -530,6 +541,12 @@ void FDynamicMeshAttributeSetChangeTracker::StoreAllFinalTriangles(const TArray<
 		Change->MaterialIDAttribChange->StoreAllFinalTriangles(Attribs->GetMaterialID(), TriangleIDs);
 	}
 
+	int NumPolygroupLayers = Attribs->NumPolygroupLayers();
+	for (int k = 0; k < NumPolygroupLayers; ++k)
+	{
+		Change->PolygroupChanges[k].StoreAllFinalTriangles(Attribs->GetPolygroupLayer(k), TriangleIDs);
+	}
+
 	for (int k = 0, NumRegAttrib = Attribs->NumRegisteredAttributes(); k < NumRegAttrib; k++)
 	{
 		Change->RegisteredAttributeChanges[k]->StoreAllFinalTriangles(Attribs->GetRegisteredAttribute(k), TriangleIDs);
@@ -574,6 +591,12 @@ bool FDynamicMeshAttributeChangeSet::Apply(FDynamicMeshAttributeSet* Attributes,
 		MaterialIDAttribChange->Apply(Attributes->GetMaterialID(), bRevert);
 	}
 
+	int NumPolygroupLayers = Attributes->NumPolygroupLayers();
+	for (int k = 0; k < NumPolygroupLayers; ++k)
+	{
+		PolygroupChanges[k].Apply(Attributes->GetPolygroupLayer(k), bRevert);
+	}
+
 	return true;
 }
 
diff --git a/Engine/Plugins/Experimental/GeometryProcessing/Source/DynamicMesh/Private/DynamicMeshEditor.cpp b/Engine/Plugins/Experimental/GeometryProcessing/Source/DynamicMesh/Private/DynamicMeshEditor.cpp
--- a/Engine/Plugins/Experimental/GeometryProcessing/Source/DynamicMesh/Private/DynamicMeshEditor.cpp
+++ b/Engine/Plugins/Experimental/GeometryProcessing/Source/DynamicMesh/Private/DynamicMeshEditor.cpp
@@ -1315,6 +1315,12 @@ void FDynamicMeshEditor::CopyAttributes(int FromTriangleID, int ToTriangleID, FM
 		MaterialIDs->SetValue(ToTriangleID, MaterialIDs->GetValue(FromTriangleID));
 	}
 
+	for (int PolygroupLayerIndex = 0; PolygroupLayerIndex < Mesh->Attributes()->NumPolygroupLayers(); PolygroupLayerIndex++)
+	{
+		FDynamicMeshPolygroupAttribute* Polygroup = Mesh->Attributes()->GetPolygroupLayer(PolygroupLayerIndex);
+		Polygroup->SetValue(ToTriangleID, Polygroup->GetValue(FromTriangleID));
+	}
+
 }
 
 
@@ -1500,6 +1506,18 @@ void FDynamicMeshEditor::AppendMesh(const FDynamicMesh3* AppendMesh,
 			}
 		}
 
+		int NumPolygroupLayers = FMath::Min(Mesh->Attributes()->NumPolygroupLayers(), AppendMesh->Attributes()->NumPolygroupLayers());
+		for (int PolygroupLayerIndex = 0; PolygroupLayerIndex < NumPolygroupLayers; PolygroupLayerIndex++)
+		{
+			// TODO: remap groups? this will be somewhat expensive...
+			const FDynamicMeshPolygroupAttribute* FromPolygroups = AppendMesh->Attributes()->GetPolygroupLayer(PolygroupLayerIndex);
+			FDynamicMeshPolygroupAttribute* ToPolygroups = Mesh->Attributes()->GetPolygroupLayer(PolygroupLayerIndex);
+			for (int TriID : AppendMesh->TriangleIndicesItr())
+			{
+				ToPolygroups->SetValue(TriangleMap.GetTo(TriID), FromPolygroups->GetValue(TriID));
+			}
+		}
+
 		for (const TPair<FName, TUniquePtr<FDynamicMeshAttributeBase>>& AttribPair : AppendMesh->Attributes()->GetAttachedAttributes())
 		{
 			if (Mesh->Attributes()->HasAttachedAttribute(AttribPair.Key))
@@ -1678,6 +1696,15 @@ static void AppendAttributes(const FDynamicMesh3* FromMesh, int FromTriangleID,
 		ToMaterialIDs->SetValue(ToTriangleID, FromMaterialIDs->GetValue(FromTriangleID));
 	}
 
+	int NumPolygroupLayers = FMath::Min(FromMesh->Attributes()->NumPolygroupLayers(), ToMesh->Attributes()->NumPolygroupLayers());
+	for (int PolygroupLayerIndex = 0; PolygroupLayerIndex < NumPolygroupLayers; PolygroupLayerIndex++)
+	{
+		// TODO: remap groups? this will be somewhat expensive...
+		const FDynamicMeshPolygroupAttribute* FromPolygroups = FromMesh->Attributes()->GetPolygroupLayer(PolygroupLayerIndex);
+		FDynamicMeshPolygroupAttribute* ToPolygroups = ToMesh->Attributes()->GetPolygroupLayer(PolygroupLayerIndex);
+		ToPolygroups->SetValue(ToTriangleID, FromPolygroups->GetValue(FromTriangleID));
+	}
+
 	for (const TPair<FName, TUniquePtr<FDynamicMeshAttributeBase>>& AttribPair : FromMesh->Attributes()->GetAttachedAttributes())
 	{
 		if (ToMesh->Attributes()->HasAttachedAttribute(AttribPair.Key))
diff --git a/Engine/Plugins/Experimental/GeometryProcessing/Source/DynamicMesh/Private/Selections/MeshConnectedComponents.cpp b/Engine/Plugins/Experimental/GeometryProcessing/Source/DynamicMesh/Private/Selections/MeshConnectedComponents.cpp
--- a/Engine/Plugins/Experimental/GeometryProcessing/Source/DynamicMesh/Private/Selections/MeshConnectedComponents.cpp
+++ b/Engine/Plugins/Experimental/GeometryProcessing/Source/DynamicMesh/Private/Selections/MeshConnectedComponents.cpp
@@ -306,4 +306,38 @@ void FMeshConnectedComponents::GrowToConnectedTriangles(const FDynamicMesh3* Mes
 			}
 		}
 	}
+}
+
+
+
+void FMeshConnectedComponents::GrowToConnectedTriangles(const FDynamicMesh3* Mesh,
+	const TArray<int>& InputROI, TSet<int>& ResultROI,
+	TArray<int32>* QueueBuffer,
+	TFunctionRef<bool(int32, int32)> CanGrowPredicate
+)
+{
+	TArray<int32> LocalQueue;
+	QueueBuffer = (QueueBuffer == nullptr) ? &LocalQueue : QueueBuffer;
+	QueueBuffer->Reset();
+	QueueBuffer->Insert(InputROI, 0);
+
+	ResultROI.Reset();
+	ResultROI.Append(InputROI);
+
+	while (QueueBuffer->Num() > 0)
+	{
+		int32 CurTri = QueueBuffer->Pop(false);
+		ResultROI.Add(CurTri);
+
+		FIndex3i NbrTris = Mesh->GetTriNeighbourTris(CurTri);
+		for (int j = 0; j < 3; ++j)
+		{
+			int32 tid = NbrTris[j];
+			if (tid != FDynamicMesh3::InvalidID && ResultROI.Contains(tid) == false && CanGrowPredicate(CurTri, tid))
+			{
+				QueueBuffer->Add(tid);
+				ResultROI.Add(tid);
+			}
+		}
+	}
 }
\ No newline at end of file
diff --git a/Engine/Plugins/Experimental/GeometryProcessing/Source/DynamicMesh/Public/DynamicAttribute.h b/Engine/Plugins/Experimental/GeometryProcessing/Source/DynamicMesh/Public/DynamicAttribute.h
--- a/Engine/Plugins/Experimental/GeometryProcessing/Source/DynamicMesh/Public/DynamicAttribute.h
+++ b/Engine/Plugins/Experimental/GeometryProcessing/Source/DynamicMesh/Public/DynamicAttribute.h
@@ -58,7 +58,23 @@ public:
 	{
 	}
 
+
 public:
+	/** Get optional identifier for this attribute set. */
+	FName GetName() const { return Name; }
+
+	/** Set optional identifier for this attribute set. */
+	void SetName(FName NameIn) { Name = NameIn; }
+
+protected:
+	/** Optional FName identifier for this attribute set. Not guaranteed to be unique. */
+	FName Name = FName();
+
+
+
+
+public:
+
 
 	/** Allocate a new copy of the attribute layer, optionally with a different parent */
 	virtual TDynamicAttributeBase* MakeCopy(ParentType* ParentIn) const = 0;
@@ -162,6 +178,16 @@ public:
 	{
 	}
 
+protected:
+
+	/** 
+	 * Implementation of parent-class copy. MakeCopy() and MakeCompactCopy() implementations should call
+	 * this to transfer any custom data added by parent attribute set class.
+	 */
+	virtual void CopyParentClassData(const TDynamicAttributeBase<ParentType>& Other)
+	{
+		Name = Other.Name;
+	}
 };
 
 
diff --git a/Engine/Plugins/Experimental/GeometryProcessing/Source/DynamicMesh/Public/DynamicMeshAttributeSet.h b/Engine/Plugins/Experimental/GeometryProcessing/Source/DynamicMesh/Public/DynamicMeshAttributeSet.h
--- a/Engine/Plugins/Experimental/GeometryProcessing/Source/DynamicMesh/Public/DynamicMeshAttributeSet.h
+++ b/Engine/Plugins/Experimental/GeometryProcessing/Source/DynamicMesh/Public/DynamicMeshAttributeSet.h
@@ -17,6 +17,9 @@ typedef TDynamicMeshVectorOverlay<float, 3, FVector3f> FDynamicMeshNormalOverlay
 /** Standard per-triangle integer material ID */
 typedef TDynamicMeshScalarTriangleAttribute<int32> FDynamicMeshMaterialAttribute;
 
+/** Per-triangle integer polygroup ID */
+typedef TDynamicMeshScalarTriangleAttribute<int32> FDynamicMeshPolygroupAttribute;
+
 /**
  * FDynamicMeshAttributeSet manages a set of extended attributes for a FDynamicMesh3.
  * This includes UV and Normal overlays, etc.
@@ -40,52 +43,10 @@ public:
 	{
 	}
 
-	void Copy(const FDynamicMeshAttributeSet& Copy)
-	{
-		SetNumUVLayers(Copy.NumUVLayers());
-		for (int UVIdx = 0; UVIdx < NumUVLayers(); UVIdx++)
-		{
-			UVLayers[UVIdx].Copy(Copy.UVLayers[UVIdx]);
-		}
-		Normals0.Copy(Copy.Normals0);
-
-		if (Copy.MaterialIDAttrib)
-		{
-			EnableMaterialID();
-			MaterialIDAttrib->Copy(*(Copy.MaterialIDAttrib));
-		}
-		else
-		{
-			DisableMaterialID();
-		}
-
-		GenericAttributes.Reset();
-		ResetRegisteredAttributes();
-		for (const TPair<FName, TUniquePtr<FDynamicMeshAttributeBase>>& AttribPair : Copy.GenericAttributes)
-		{
-			AttachAttribute(AttribPair.Key, AttribPair.Value->MakeCopy(ParentMesh));
-		}
-
-		// parent mesh is *not* copied!
-	}
+	void Copy(const FDynamicMeshAttributeSet& Copy);
 
 	/** returns true if the attached overlays/attributes are compact */
-	bool IsCompact()
-	{
-		for (int UVIdx = 0; UVIdx < NumUVLayers(); UVIdx++)
-		{
-			if (!UVLayers[UVIdx].IsCompact())
-			{
-				return false;
-			}
-		}
-		if (!Normals0.IsCompact())
-		{
-			return false;
-		}
-		// material ID and generic per-element attributes currently cannot be non-compact
-		return true;
-	}
+	bool IsCompact();
 
 	/**
 	 * Performs a CompactCopy of the attached overlays/attributes.
@@ -94,34 +55,7 @@ public:
 	 * @param CompactMaps Maps indicating how vertices and triangles were changes in the parent
 	 * @param Copy The attribute set to be copied
 	 */
-	void CompactCopy(const FCompactMaps& CompactMaps, const FDynamicMeshAttributeSet& Copy)
-	{
-		SetNumUVLayers(Copy.NumUVLayers());
-		for (int UVIdx = 0; UVIdx < NumUVLayers(); UVIdx++)
-		{
-			UVLayers[UVIdx].CompactCopy(CompactMaps, Copy.UVLayers[UVIdx]);
-		}
-		Normals0.CompactCopy(CompactMaps, Copy.Normals0);
-
-		if (Copy.MaterialIDAttrib)
-		{
-			EnableMaterialID();
-			MaterialIDAttrib->CompactCopy(CompactMaps, *(Copy.MaterialIDAttrib));
-		}
-		else
-		{
-			DisableMaterialID();
-		}
-
-		GenericAttributes.Reset();
-		ResetRegisteredAttributes();
-		for (const TPair<FName, TUniquePtr<FDynamicMeshAttributeBase>>& AttribPair : Copy.GenericAttributes)
-		{
-			AttachAttribute(AttribPair.Key, AttribPair.Value->MakeCompactCopy(CompactMaps, ParentMesh));
-		}
-
-		// parent mesh is *not* copied!
-	}
+	void CompactCopy(const FCompactMaps& CompactMaps, const FDynamicMeshAttributeSet& Copy);
 
 	/**
 	 * Compacts the attribute set in place
@@ -129,54 +63,13 @@ public:
 	 *
 	 * @param CompactMaps Maps of how the vertices and triangles were compacted in the parent
 	 */
-	void CompactInPlace(const FCompactMaps& CompactMaps)
-	{
-		for (int UVIdx = 0; UVIdx < NumUVLayers(); UVIdx++)
-		{
-			UVLayers[UVIdx].CompactInPlace(CompactMaps);
-		}
-		Normals0.CompactInPlace(CompactMaps);
-
-		if (MaterialIDAttrib.IsValid())
-		{
-			MaterialIDAttrib->CompactInPlace(CompactMaps);
-		}
-
-		for (const TPair<FName, TUniquePtr<FDynamicMeshAttributeBase>>& AttribPair : GenericAttributes)
-		{
-			AttribPair.Value->CompactInPlace(CompactMaps);
-		}
-	}
+	void CompactInPlace(const FCompactMaps& CompactMaps);
 
 
 	/**
 	 * Enable the matching attributes and overlay layers as the reference Copy set, but do not copy any data across
 	 */
-	void EnableMatchingAttributes(const FDynamicMeshAttributeSet& ToMatch)
-	{
-		SetNumUVLayers(ToMatch.NumUVLayers());
-		for (int UVIdx = 0; UVIdx < NumUVLayers(); UVIdx++)
-		{
-			UVLayers[UVIdx].ClearElements();
-		}
-		Normals0.ClearElements();
-
-		if (ToMatch.MaterialIDAttrib)
-		{
-			EnableMaterialID();
-		}
-		else
-		{
-			DisableMaterialID();
-		}
-
-		GenericAttributes.Reset();
-		ResetRegisteredAttributes();
-		for (const TPair<FName, TUniquePtr<FDynamicMeshAttributeBase>>& AttribPair : ToMatch.GenericAttributes)
-		{
-			AttachAttribute(AttribPair.Key, AttribPair.Value->MakeNew(ParentMesh));
-		}
-	}
+	void EnableMatchingAttributes(const FDynamicMeshAttributeSet& ToMatch);
 
 	/** @return the parent mesh for this overlay */
 	const FDynamicMesh3* GetParentMesh() const { return ParentMesh; }
@@ -185,26 +78,7 @@ public:
 
 private:
 	/** @set the parent mesh for this overlay.  Only safe for use during FDynamicMesh move */
-	void Reparent(FDynamicMesh3* NewParent)
-	{
-		ParentMesh = NewParent;
-
-		for (int UVIdx = 0; UVIdx < NumUVLayers(); UVIdx++)
-		{
-			UVLayers[UVIdx].Reparent( NewParent );
-		}
-		Normals0.Reparent( NewParent );
-
-		if (MaterialIDAttrib)
-		{
-			MaterialIDAttrib->Reparent( NewParent );
-		}
-
-		for (const TPair<FName, TUniquePtr<FDynamicMeshAttributeBase>>& AttribPair : GenericAttributes)
-		{
-			AttribPair.Value->Reparent( NewParent );
-		}
-	}
+	void Reparent(FDynamicMesh3* NewParent);
 
 public:
 	/** @return number of UV layers */
@@ -213,25 +87,7 @@ public:
 		return UVLayers.Num();
 	}
 
-	virtual void SetNumUVLayers(int Num)
-	{
-		if (UVLayers.Num() == Num)
-		{
-			return;
-		}
-		if (Num >= UVLayers.Num())
-		{
-			for (int i = (int)UVLayers.Num(); i < Num; ++i)
-			{
-				UVLayers.Add(new FDynamicMeshUVOverlay(ParentMesh));
-			}
-		}
-		else
-		{
-			UVLayers.RemoveAt(Num, UVLayers.Num() - Num);
-		}
-		check(UVLayers.Num() == Num);
-	}
+	virtual void SetNumUVLayers(int Num);
 
 	/** @return number of Normals layers */
 	virtual int NumNormalLayers() const
@@ -312,6 +168,25 @@ public:
 
 
 
+	//
+	// Polygroup layers
+	//
+
+	/** @return number of Polygroup layers */
+	virtual int32 NumPolygroupLayers() const;
+
+	/** Set the number of Polygroup layers */
+	virtual void SetNumPolygroupLayers(int32 Num);
+
+	/** @return the Polygroup layer at the given Index */
+	FDynamicMeshPolygroupAttribute* GetPolygroupLayer(int Index);
+
+	/** @return the Polygroup layer at the given Index */
+	const FDynamicMeshPolygroupAttribute* GetPolygroupLayer(int Index) const;
+
+
+
+
 	//
 	// Per-Triangle Material ID
 	//
@@ -388,6 +263,8 @@ protected:
 
 	TUniquePtr<FDynamicMeshMaterialAttribute> MaterialIDAttrib;
 
+	TIndirectArray<FDynamicMeshPolygroupAttribute> PolygroupLayers;
+
 	TMap<FName, TUniquePtr<FDynamicMeshAttributeBase>> GenericAttributes;
 
 protected:
@@ -434,6 +311,14 @@ protected:
 			bValid = GetUVLayer(UVLayerIndex)->CheckValidity(bAllowNonmanifold, FailMode) && bValid;
 		}
 		bValid = PrimaryNormals()->CheckValidity(bAllowNonmanifold, FailMode) && bValid;
+		if (MaterialIDAttrib)
+		{
+			bValid = MaterialIDAttrib->CheckValidity(bAllowNonmanifold, FailMode) && bValid;
+		}
+		for (int PolygroupLayerIndex = 0; PolygroupLayerIndex < NumPolygroupLayers(); PolygroupLayerIndex++)
+		{
+			bValid = GetPolygroupLayer(PolygroupLayerIndex)->CheckValidity(bAllowNonmanifold, FailMode) && bValid;
+		}
 		return bValid;
 	}
 };
diff --git a/Engine/Plugins/Experimental/GeometryProcessing/Source/DynamicMesh/Public/DynamicMeshChangeTracker.h b/Engine/Plugins/Experimental/GeometryProcessing/Source/DynamicMesh/Public/DynamicMeshChangeTracker.h
--- a/Engine/Plugins/Experimental/GeometryProcessing/Source/DynamicMesh/Public/DynamicMeshChangeTracker.h
+++ b/Engine/Plugins/Experimental/GeometryProcessing/Source/DynamicMesh/Public/DynamicMeshChangeTracker.h
@@ -84,7 +84,8 @@ protected:
 typedef TDynamicMeshAttributeChange<float,2> FDynamicMeshUVChange;
 /** Standard Normal overlay change type - 3-element float */
 typedef TDynamicMeshAttributeChange<float,3> FDynamicMeshNormalChange;
-
+/** Standard per-triangle integer attribute change type */
+typedef FDynamicMeshTriangleAttributeChange<int32,1> FDynamicMeshTriGroupChange;
 
 /**
  * FDynamicMeshAttributeChangeSet stores a set of UV and Normal changes for a FDynamicMesh3
@@ -97,7 +98,8 @@ public:
 
 	TArray<FDynamicMeshUVChange> UVChanges;
 	TArray<FDynamicMeshNormalChange> NormalChanges;
-	TOptional<FDynamicMeshTriangleAttributeChange<int32, 1>> MaterialIDAttribChange;
+	TOptional<FDynamicMeshTriGroupChange> MaterialIDAttribChange;
+	TArray<FDynamicMeshTriGroupChange> PolygroupChanges;
 	TArray<TUniquePtr<FDynamicMeshAttributeChangeBase>> RegisteredAttributeChanges;
 
 	/** call ::Apply() on all the UV and Normal changes */
diff --git a/Engine/Plugins/Experimental/GeometryProcessing/Source/DynamicMesh/Public/DynamicMeshTriangleAttribute.h b/Engine/Plugins/Experimental/GeometryProcessing/Source/DynamicMesh/Public/DynamicMeshTriangleAttribute.h
--- a/Engine/Plugins/Experimental/GeometryProcessing/Source/DynamicMesh/Public/DynamicMeshTriangleAttribute.h
+++ b/Engine/Plugins/Experimental/GeometryProcessing/Source/DynamicMesh/Public/DynamicMeshTriangleAttribute.h
@@ -102,6 +102,7 @@ public:
 	/** Set this overlay to contain the same arrays as the copy overlay */
 	void Copy(const TDynamicMeshTriangleAttribute<AttribValueType, AttribDimension>& Copy)
 	{
+		CopyParentClassData(Copy);
 		AttribValues = Copy.AttribValues;
 	}
 
@@ -130,6 +131,7 @@ public:
 
 	void CompactCopy(const FCompactMaps& CompactMaps, const TDynamicMeshTriangleAttribute<AttribValueType, AttribDimension>& ToCopy)
 	{
+		CopyParentClassData(ToCopy);
 		check(CompactMaps.MapT.Num() <= int(ToCopy.AttribValues.Num() / AttribDimension));
 		AttribValueType Data[AttribDimension];
 		for (int TID = 0; TID < CompactMaps.MapT.Num(); TID++)
diff --git a/Engine/Plugins/Experimental/GeometryProcessing/Source/DynamicMesh/Public/DynamicVertexAttribute.h b/Engine/Plugins/Experimental/GeometryProcessing/Source/DynamicMesh/Public/DynamicVertexAttribute.h
--- a/Engine/Plugins/Experimental/GeometryProcessing/Source/DynamicMesh/Public/DynamicVertexAttribute.h
+++ b/Engine/Plugins/Experimental/GeometryProcessing/Source/DynamicMesh/Public/DynamicVertexAttribute.h
@@ -93,6 +93,7 @@ public:
 	/** Set this overlay to contain the same arrays as the copy overlay */
 	void Copy(const TDynamicVertexAttribute<AttribValueType, AttribDimension, ParentType>& Copy)
 	{
+		TDynamicAttributeBase<ParentType>::CopyParentClassData(Copy);
 		AttribValues = Copy.AttribValues;
 	}
 
@@ -121,6 +122,7 @@ public:
 
 	void CompactCopy(const FCompactMaps& CompactMaps, const TDynamicVertexAttribute<AttribValueType, AttribDimension, ParentType>& ToCopy)
 	{
+		TDynamicAttributeBase<ParentType>::CopyParentClassData(ToCopy);
 		check(CompactMaps.MapV.Num() <= AttribValues.Num() / AttribDimension);
 		AttribValueType Data[AttribDimension];
 		for (int VID = 0; VID < CompactMaps.MapV.Num(); VID++)
diff --git a/Engine/Plugins/Experimental/GeometryProcessing/Source/DynamicMesh/Public/Selections/MeshConnectedComponents.h b/Engine/Plugins/Experimental/GeometryProcessing/Source/DynamicMesh/Public/Selections/MeshConnectedComponents.h
--- a/Engine/Plugins/Experimental/GeometryProcessing/Source/DynamicMesh/Public/Selections/MeshConnectedComponents.h
+++ b/Engine/Plugins/Experimental/GeometryProcessing/Source/DynamicMesh/Public/Selections/MeshConnectedComponents.h
@@ -149,4 +149,22 @@ public:
 	);
 
 
+	/**
+	 * Utility function to expand a triangle selection to all triangles considered "connected".
+	 * More efficient than using full FMeshConnectedComponents instance if ROI is small relative to Mesh size (or if temp buffers can be re-used)
+	 * This version computes an output TSet instead of output TArray, which is preferable in some cases.
+	 * @param Mesh Mesh to calculate on
+	 * @param InputROI input set of triangles
+	 * @param ResultROI output set of triangles connected to InputROI
+	 * @param QueueBuffer optional buffer used as internal Queue. If passed as nullptr, a TArray will be locally allocated
+	 * @param CanGrowPredicate determines whether two connected mesh triangles should be considered connected while growing
+	 */
+	static void GrowToConnectedTriangles(const FDynamicMesh3* Mesh,
+		const TArray<int>& InputROI,
+		TSet<int>& ResultROI,
+		TArray<int32>* QueueBuffer = nullptr,
+		TFunctionRef<bool(int32, int32)> CanGrowPredicate = [](int32, int32) { return true; }
+	);
+
+
 };
diff --git a/Engine/Plugins/Experimental/GeometryProcessing/Source/MeshConversion/Private/DynamicMeshToMeshDescription.cpp b/Engine/Plugins/Experimental/GeometryProcessing/Source/MeshConversion/Private/DynamicMeshToMeshDescription.cpp
--- a/Engine/Plugins/Experimental/GeometryProcessing/Source/MeshConversion/Private/DynamicMeshToMeshDescription.cpp
+++ b/Engine/Plugins/Experimental/GeometryProcessing/Source/MeshConversion/Private/DynamicMeshToMeshDescription.cpp
@@ -523,6 +523,9 @@ void FDynamicMeshToMeshDescription::Convert_NoSharedInstances(const FDynamicMesh
 	TArray<FIndex3i> UVTris;
 	UVTris.SetNum(NumUVLayers);
 
+	TArray<FTriangleID> IndexToTriangleIDMap;
+	IndexToTriangleIDMap.SetNum(MeshIn->MaxTriangleID());
+
 	for (int TriID : MeshIn->TriangleIndicesItr())
 	{
 		FIndex3i Triangle = MeshIn->GetTriangle(TriID);
@@ -546,6 +549,7 @@ void FDynamicMeshToMeshDescription::Convert_NoSharedInstances(const FDynamicMesh
 
 		// add the triangle to MeshDescription
 		FTriangleID NewTriangleID = Builder.AppendTriangle(TriVertInstances[0], TriVertInstances[1], TriVertInstances[2], UsePolygonGroupID); 
+		IndexToTriangleIDMap[TriID] = NewTriangleID;
 
 		// transfer UVs.  Note the Builder sets both the shared and per-instance UVs from this
 		for (int32 k = 0; k < NumUVLayers; ++k)
@@ -603,5 +607,48 @@ void FDynamicMeshToMeshDescription::Convert_NoSharedInstances(const FDynamicMesh
 		}
 	}
 
+	// convert polygroup layers
+	ConvertPolygroupLayers(MeshIn, MeshOut, IndexToTriangleIDMap);
+
 	Builder.ResumeMeshDescriptionIndexing();
 }
+
+
+
+void FDynamicMeshToMeshDescription::ConvertPolygroupLayers(const FDynamicMesh3* MeshIn, FMeshDescription& MeshOut, const TArray<FTriangleID>& IndexToTriangleIDMap)
+{
+	if (MeshIn->Attributes() == nullptr) return;
+
+	TAttributesSet<FTriangleID>& TriAttribsSet = MeshOut.TriangleAttributes();
+
+	for (int32 li = 0; li < MeshIn->Attributes()->NumPolygroupLayers(); ++li)
+	{
+		const FDynamicMeshPolygroupAttribute* Polygroups = MeshIn->Attributes()->GetPolygroupLayer(li);
+		FName LayerName = Polygroups->GetName();
+	
+		// Find existing attribute with the same name. If not found, create a new one.
+		TTriangleAttributesRef<int32> Attribute;
+		if (TriAttribsSet.HasAttribute(LayerName))
+		{
+			Attribute = TriAttribsSet.GetAttributesRef<int32>(LayerName);
+		}
+		else
+		{
+			TriAttribsSet.RegisterAttribute<int32>(LayerName, 1, 0, EMeshAttributeFlags::AutoGenerated);
+			Attribute = TriAttribsSet.GetAttributesRef<int32>(LayerName);
+		}
+		if (ensure(Attribute.IsValid()))
+		{
+			for (int32 tid : MeshIn->TriangleIndicesItr())
+			{
+				FTriangleID TriangleID = IndexToTriangleIDMap[tid];
+				int32 GroupID = Polygroups->GetValue(tid);
+				Attribute.Set(TriangleID, GroupID);
+			}
+		}
+		else
+		{
+			UE_LOG(LogTemp, Warning, TEXT("FDynamicMeshToMeshDescription::ConvertPolygroupLayers - could not create attribute named %s"), *LayerName.ToString());
+		}
+	}
+}
\ No newline at end of file
diff --git a/Engine/Plugins/Experimental/GeometryProcessing/Source/MeshConversion/Private/MeshDescriptionToDynamicMesh.cpp b/Engine/Plugins/Experimental/GeometryProcessing/Source/MeshConversion/Private/MeshDescriptionToDynamicMesh.cpp
--- a/Engine/Plugins/Experimental/GeometryProcessing/Source/MeshConversion/Private/MeshDescriptionToDynamicMesh.cpp
+++ b/Engine/Plugins/Experimental/GeometryProcessing/Source/MeshConversion/Private/MeshDescriptionToDynamicMesh.cpp
@@ -177,7 +177,7 @@ void FMeshDescriptionToDynamicMesh::Convert(const FMeshDescription* MeshIn, FDyn
 	TArray<FTriData> AddedTriangles;
 	AddedTriangles.SetNum(MeshIn->Triangles().Num());
 
-	// allocate the TriIDMap
+	// allocate the TriIDMap (maps from DynamicMesh triangle ID to MeshDescription FTriangleID)
 	TriIDMap.SetNumUninitialized(MeshIn->Triangles().Num());
 
 
@@ -358,6 +358,21 @@ void FMeshDescriptionToDynamicMesh::Convert(const FMeshDescription* MeshIn, FDyn
 	}
 
 
+	// find polygroup attribs
+	const TAttributesSet<FTriangleID>& TriAttribsSet = MeshIn->TriangleAttributes();
+	TArray<TTriangleAttributesConstRef<int32>> PolygroupAttribs;
+	TArray<FName> PolygroupAttribNames;
+	TriAttribsSet.ForEach([&](const FName AttributeName, auto AttributesRef)
+	{
+		if (IsReservedAttributeName(AttributeName)) return;
+
+		if (TriAttribsSet.template HasAttributeOfType<int32>(AttributeName))
+		{
+			PolygroupAttribs.Add(TriAttribsSet.GetAttributesRef<int32>(AttributeName));
+			PolygroupAttribNames.Add(AttributeName);
+		}
+	});
+
 	
 
 	if (!bDisableAttributes)
@@ -365,6 +380,11 @@ void FMeshDescriptionToDynamicMesh::Convert(const FMeshDescription* MeshIn, FDyn
 		// we will weld/populate all the attributes simultaneously, hold on to futures in this array and then Wait for them at the end
 		TArray<TFuture<void>> Pending;
 
+		if (PolygroupAttribs.Num() > 0)
+		{
+			MeshOut.Attributes()->SetNumPolygroupLayers(PolygroupAttribs.Num());
+		}
+
 		for (int UVLayerIndex = 0; UVLayerIndex < NumUVLayers; UVLayerIndex++)
 		{
 			auto UVFuture = Async(EAsyncExecution::ThreadPool, [&, UVLayerIndex, bUseSharedUVs]() // must copy UVLayerIndex here!
@@ -515,6 +535,27 @@ void FMeshDescriptionToDynamicMesh::Convert(const FMeshDescription* MeshIn, FDyn
 			Pending.Add(MoveTemp(MaterialFuture));
 		}
 
+		// initialize polygroup layers
+		for (int32 GroupLayerIdx = 0; GroupLayerIdx < PolygroupAttribs.Num(); GroupLayerIdx++)
+		{
+			auto PolygroupFuture = Async(EAsyncExecution::ThreadPool, [&,Index=GroupLayerIdx]()
+			{
+				TTriangleAttributesConstRef<int32> InputGroupSet = PolygroupAttribs[Index];
+				FDynamicMeshPolygroupAttribute* OutputGroupSet = MeshOut.Attributes()->GetPolygroupLayer(Index);
+				if (ensure(OutputGroupSet))
+				{
+					OutputGroupSet->SetName(PolygroupAttribNames[Index]);
+					for (int32 tid : MeshOut.TriangleIndicesItr())
+					{
+						FTriangleID SourceTriangleID = TriIDMap[tid];
+						int32 GroupID = InputGroupSet.Get(SourceTriangleID);
+						OutputGroupSet->SetValue(tid, &GroupID);
+					}
+				}
+			});
+			Pending.Add(MoveTemp(PolygroupFuture));
+		}
+
 		// wait for all work to be done
 		for (TFuture<void>& Future : Pending)
 		{
@@ -598,3 +639,33 @@ void FMeshDescriptionToDynamicMesh::CopyTangents(const FMeshDescription* SourceM
 }
 
 
+
+
+bool FMeshDescriptionToDynamicMesh::IsReservedAttributeName(FName AttributeName)
+{
+	// @todo: use FAttributesSetBase::DoesAttributeHaveAnyFlags(EMeshAttributeFlags::Mandatory) to determine this
+	// will need a way to look up MeshDescription attribute sets from their category.
+	return (AttributeName == MeshAttribute::Vertex::Position)
+		|| (AttributeName == MeshAttribute::VertexInstance::VertexIndex)
+		|| (AttributeName == MeshAttribute::VertexInstance::TextureCoordinate)
+		|| (AttributeName == MeshAttribute::VertexInstance::Normal)
+		|| (AttributeName == MeshAttribute::VertexInstance::Tangent)
+		|| (AttributeName == MeshAttribute::VertexInstance::BinormalSign)
+		|| (AttributeName == MeshAttribute::VertexInstance::Color)
+		|| (AttributeName == MeshAttribute::Edge::IsHard)
+		|| (AttributeName == MeshAttribute::Edge::VertexIndex)
+		|| (AttributeName == MeshAttribute::Triangle::VertexInstanceIndex)
+		|| (AttributeName == MeshAttribute::Triangle::PolygonIndex)
+		|| (AttributeName == MeshAttribute::Triangle::EdgeIndex)
+		|| (AttributeName == MeshAttribute::Triangle::VertexIndex)
+		|| (AttributeName == MeshAttribute::Triangle::UVIndex)
+		|| (AttributeName == MeshAttribute::Triangle::PolygonGroupIndex)
+		|| (AttributeName == MeshAttribute::Triangle::Normal)
+		|| (AttributeName == MeshAttribute::Triangle::Tangent)
+		|| (AttributeName == MeshAttribute::Triangle::Binormal)
+		|| (AttributeName == MeshAttribute::UV::UVCoordinate)
+		|| (AttributeName == MeshAttribute::Polygon::PolygonGroupIndex)
+		|| (AttributeName == MeshAttribute::PolygonGroup::ImportedMaterialSlotName)
+		|| (AttributeName == ExtendedMeshAttribute::PolyTriGroups);
+}
+
diff --git a/Engine/Plugins/Experimental/GeometryProcessing/Source/MeshConversion/Public/DynamicMeshToMeshDescription.h b/Engine/Plugins/Experimental/GeometryProcessing/Source/MeshConversion/Public/DynamicMeshToMeshDescription.h
--- a/Engine/Plugins/Experimental/GeometryProcessing/Source/MeshConversion/Public/DynamicMeshToMeshDescription.h
+++ b/Engine/Plugins/Experimental/GeometryProcessing/Source/MeshConversion/Public/DynamicMeshToMeshDescription.h
@@ -117,4 +117,16 @@ public:
      * the shared UVs structures are populated.
 	 */
 	void Convert_NoSharedInstances(const FDynamicMesh3* MeshIn, FMeshDescription& MeshOut);
+
+
+
+protected:
+
+	/**
+	 * Transfer PolygroupLayers from DynamicMesh AttributeSet to MeshDescription.
+	 * Will copy to existing MeshDescription TriangleAttribute<int32> if one with the same name exists.
+	 * Otherwise will register a new one.
+	 */
+	void ConvertPolygroupLayers(const FDynamicMesh3* MeshIn, FMeshDescription& MeshOut, const TArray<FTriangleID>& IndexToTriangleIDMap);
+
 };
\ No newline at end of file
diff --git a/Engine/Plugins/Experimental/GeometryProcessing/Source/MeshConversion/Public/MeshDescriptionToDynamicMesh.h b/Engine/Plugins/Experimental/GeometryProcessing/Source/MeshConversion/Public/MeshDescriptionToDynamicMesh.h
--- a/Engine/Plugins/Experimental/GeometryProcessing/Source/MeshConversion/Public/MeshDescriptionToDynamicMesh.h
+++ b/Engine/Plugins/Experimental/GeometryProcessing/Source/MeshConversion/Public/MeshDescriptionToDynamicMesh.h
@@ -77,4 +77,11 @@ public:
 	 */
 	void CopyTangents(const FMeshDescription* SourceMesh, const FDynamicMesh3* TargetMesh, TMeshTangents<double>* TangentsOut);
 
+
+public:
+
+	/**
+	 * @return true if the Name is the name of a reserved FMeshDescription attribute
+	 */
+	static bool IsReservedAttributeName(FName Name);
 };
\ No newline at end of file
diff --git a/Engine/Plugins/Experimental/MeshModelingToolset/Source/MeshModelingToolsEditorOnly/Private/AttributeEditorTool.cpp b/Engine/Plugins/Experimental/MeshModelingToolset/Source/MeshModelingToolsEditorOnly/Private/AttributeEditorTool.cpp
--- a/Engine/Plugins/Experimental/MeshModelingToolset/Source/MeshModelingToolsEditorOnly/Private/AttributeEditorTool.cpp
+++ b/Engine/Plugins/Experimental/MeshModelingToolset/Source/MeshModelingToolsEditorOnly/Private/AttributeEditorTool.cpp
@@ -6,6 +6,7 @@
 
 #include "ToolSetupUtil.h"
 #include "MathUtil.h"
+#include "MeshDescriptionToDynamicMesh.h"   // for FMeshDescriptionToDynamicMesh::IsReservedAttributeName
 
 #include "MeshDescription.h"
 #include "StaticMeshAttributes.h"
@@ -266,20 +267,6 @@ static bool RemoveAttribute(FMeshDescription* Mesh, EAttributeEditorElementType
 }
 
 
-static bool IsReservedName(FName AttributeName)
-{
-	// @todo: use FAttributesSetBase::DoesAttributeHaveAnyFlags(EMeshAttributeFlags::Mandatory) to determine this
-	// will need a way to look up MeshDescription attribute sets from their category.
-	return (AttributeName == MeshAttribute::Vertex::Position)
-		|| (AttributeName == MeshAttribute::VertexInstance::TextureCoordinate)
-		|| (AttributeName == MeshAttribute::VertexInstance::Normal)
-		|| (AttributeName == MeshAttribute::VertexInstance::Tangent)
-		|| (AttributeName == MeshAttribute::VertexInstance::BinormalSign)
-		|| (AttributeName == MeshAttribute::VertexInstance::Color)
-		|| (AttributeName == MeshAttribute::Edge::IsHard)
-		|| (AttributeName == MeshAttribute::PolygonGroup::ImportedMaterialSlotName);
-}
-
 
 void UAttributeEditorTool::InitializeAttributeLists()
 {
@@ -493,7 +480,7 @@ void UAttributeEditorTool::AddNewWeightMap()
 
 void UAttributeEditorTool::AddNewGroupsLayer()
 {
-	AddNewAttribute(EAttributeEditorElementType::Polygon, EAttributeEditorAttribType::Int32, FName(NewAttributeProps->NewName));
+	AddNewAttribute(EAttributeEditorElementType::Triangle, EAttributeEditorAttribType::Int32, FName(NewAttributeProps->NewName));
 }
 
 
@@ -506,8 +493,9 @@ void UAttributeEditorTool::DeleteAttribute()
 	FMeshDescription* CurMesh = ComponentTargets[0]->GetMesh();
 	FName SelectedName(ModifyAttributeProps->Attribute);
 
+	// @todo: use FAttributesSetBase::DoesAttributeHaveAnyFlags(EMeshAttributeFlags::Mandatory) to determine this ?
 	// @todo: introduce 'reserved' attribute flag which we can check here
-	if (IsReservedName(SelectedName))
+	if (FMeshDescriptionToDynamicMesh::IsReservedAttributeName(SelectedName))
 	{
 		GetToolManager()->DisplayMessage(LOCTEXT("CannotDeleteReservedNameError", "Cannot delete required mesh Attributes"), EToolMessageLevel::UserWarning);
 		return;
