fix append of generic attributes to only happen once, not once-per-triangle!

#rnx

diff --git a/Engine/Plugins/Experimental/GeometryProcessing/Source/DynamicMesh/Private/DynamicMeshEditor.cpp b/Engine/Plugins/Experimental/GeometryProcessing/Source/DynamicMesh/Private/DynamicMeshEditor.cpp
--- a/Engine/Plugins/Experimental/GeometryProcessing/Source/DynamicMesh/Private/DynamicMeshEditor.cpp
+++ b/Engine/Plugins/Experimental/GeometryProcessing/Source/DynamicMesh/Private/DynamicMeshEditor.cpp
@@ -1609,6 +1609,13 @@ void FDynamicMeshEditor::AppendUVs(const FDynamicMesh3* AppendMesh,
 
 // can these be replaced w/ template function?
 
+
+namespace UE
+{
+namespace DynamicMeshEditorInternals
+{
+
+
 // Utility function for ::AppendTriangles()
 static int AppendTriangleUVAttribute(const FDynamicMesh3* FromMesh, int FromElementID, FDynamicMesh3* ToMesh, int UVLayerIndex, FMeshIndexMappings& IndexMaps)
 {
@@ -1642,7 +1649,7 @@ static int AppendTriangleNormalAttribute(const FDynamicMesh3* FromMesh, int From
 
 
 // Utility function for ::AppendTriangles()
-static void AppendAttributes(const FDynamicMesh3* FromMesh, int FromTriangleID, FDynamicMesh3* ToMesh, int ToTriangleID, FMeshIndexMappings& IndexMaps, FDynamicMeshEditResult& ResultOut)
+static void AppendTriangleAttributes(const FDynamicMesh3* FromMesh, int FromTriangleID, FDynamicMesh3* ToMesh, int ToTriangleID, FMeshIndexMappings& IndexMaps, FDynamicMeshEditResult& ResultOut)
 {
 	if (FromMesh->HasAttributes() == false || ToMesh->HasAttributes() == false)
 	{
@@ -1704,7 +1711,14 @@ static void AppendAttributes(const FDynamicMesh3* FromMesh, int FromTriangleID,
 		FDynamicMeshPolygroupAttribute* ToPolygroups = ToMesh->Attributes()->GetPolygroupLayer(PolygroupLayerIndex);
 		ToPolygroups->SetValue(ToTriangleID, FromPolygroups->GetValue(FromTriangleID));
 	}
+}
+
+
 
+// Utility function for ::AppendTriangles()
+static void AppendGenericAttributes(const FDynamicMesh3* FromMesh, FDynamicMesh3* ToMesh, FMeshIndexMappings& IndexMaps)
+{
+	// copy generic attributes after full IndexMaps have been created
 	for (const TPair<FName, TUniquePtr<FDynamicMeshAttributeBase>>& AttribPair : FromMesh->Attributes()->GetAttachedAttributes())
 	{
 		if (ToMesh->Attributes()->HasAttachedAttribute(AttribPair.Key))
@@ -1717,11 +1731,12 @@ static void AppendAttributes(const FDynamicMesh3* FromMesh, int FromTriangleID,
 
 
 
-
-
+}} // namespace UE::DynamicMeshEditorInternals
 
 void FDynamicMeshEditor::AppendTriangles(const FDynamicMesh3* SourceMesh, const TArrayView<const int>& SourceTriangles, FMeshIndexMappings& IndexMaps, FDynamicMeshEditResult& ResultOut, bool bComputeTriangleMap)
 {
+	using namespace UE::DynamicMeshEditorInternals;
+
 	ResultOut.Reset();
 	IndexMaps.Initialize(Mesh);
 
@@ -1774,15 +1789,19 @@ void FDynamicMeshEditor::AppendTriangles(const FDynamicMesh3* SourceMesh, const
 		}
 		ResultOut.NewTriangles.Add(NewTriangleID);
 
-		AppendAttributes(SourceMesh, SourceTriangleID, Mesh, NewTriangleID, IndexMaps, ResultOut);
+		AppendTriangleAttributes(SourceMesh, SourceTriangleID, Mesh, NewTriangleID, IndexMaps, ResultOut);
 
 		//Mesh->CheckValidity(true);
 	}
+
+	AppendGenericAttributes(SourceMesh, Mesh, IndexMaps);
 }
 
 
 bool FDynamicMeshEditor::SplitMesh(const FDynamicMesh3* SourceMesh, TArray<FDynamicMesh3>& SplitMeshes, TFunctionRef<int(int)> TriIDToMeshID, int DeleteMeshID)
 {
+	using namespace UE::DynamicMeshEditorInternals;
+
 	TMap<int, int> MeshIDToIndex;
 	int NumMeshes = 0;
 	bool bAlsoDelete = false;
@@ -1864,7 +1883,12 @@ bool FDynamicMeshEditor::SplitMesh(const FDynamicMesh3* SourceMesh, TArray<FDyna
 
 		int NewTID = Mesh.AppendTriangle(NewTri, NewGID);
 		IndexMaps.SetTriangle(SourceTID, NewTID);
-		AppendAttributes(SourceMesh, SourceTID, &Mesh, NewTID, IndexMaps, UnusedInvalidResultAccumulator);
+		AppendTriangleAttributes(SourceMesh, SourceTID, &Mesh, NewTID, IndexMaps, UnusedInvalidResultAccumulator);
+	}
+
+	for (int Idx = 0; Idx < NumMeshes; Idx++)
+	{
+		AppendGenericAttributes(SourceMesh, &SplitMeshes[Idx], Mappings[Idx]);
 	}
 	
 	return true;
