fix UpdatePoint function for PointHashGrid3 (previously didn't compile but error was hidden because template function was unused I think)

#rnx

diff --git a/Engine/Plugins/Experimental/GeometryProcessing/Source/GeometricObjects/Public/Spatial/PointHashGrid3.h b/Engine/Plugins/Experimental/GeometryProcessing/Source/GeometricObjects/Public/Spatial/PointHashGrid3.h
--- a/Engine/Plugins/Experimental/GeometryProcessing/Source/GeometricObjects/Public/Spatial/PointHashGrid3.h
+++ b/Engine/Plugins/Experimental/GeometryProcessing/Source/GeometricObjects/Public/Spatial/PointHashGrid3.h
@@ -117,7 +117,7 @@ public:
 		bool bWasAtOldPos;
 		{
 			FScopeLock Lock(&CriticalSection);
-			bWasAtOldPos = Hash.RemoveSingle(old_idx, Value);
+			bWasAtOldPos = Hash.RemoveSingle(old_idx, Value) == 1;
 		}
 		check(bWasAtOldPos);
 		{
