Fix up MeshSelfUnion with optimizations, options and improvements from MeshBoolean
 - Use octree to speed up closest-edge finding
 - Center and rescale data before computing intersections (and put back to original space after operation is complete)
 - Add option to skip edge welding, matching MeshBoolean
 - Minor comment and variable name cleanup

#rnx

diff --git a/Engine/Plugins/Experimental/GeometryProcessing/Source/DynamicMesh/Private/Operations/MeshBoolean.cpp b/Engine/Plugins/Experimental/GeometryProcessing/Source/DynamicMesh/Private/Operations/MeshBoolean.cpp
--- a/Engine/Plugins/Experimental/GeometryProcessing/Source/DynamicMesh/Private/Operations/MeshBoolean.cpp
+++ b/Engine/Plugins/Experimental/GeometryProcessing/Source/DynamicMesh/Private/Operations/MeshBoolean.cpp
@@ -212,8 +212,7 @@ bool FMeshBoolean::Compute()
 	}
 	FDynamicMesh3* CutMesh[2]{ Result, &CutMeshB }; // just an alias to keep things organized
 
-	// transform the copies to a shared space (centered at the origin)
-	// TODO: also rescale the meshes to a standard size (?)
+	// transform the copies to a shared space (centered at the origin and scaled to a unit cube)
 	FAxisAlignedBox3d CombinedAABB(CutMesh[0]->GetCachedBounds(), Transforms[0]);
 	FAxisAlignedBox3d MeshB_AABB(CutMesh[1]->GetCachedBounds(), Transforms[1]);
 	CombinedAABB.Contain(MeshB_AABB);
diff --git a/Engine/Plugins/Experimental/GeometryProcessing/Source/DynamicMesh/Private/Operations/MeshSelfUnion.cpp b/Engine/Plugins/Experimental/GeometryProcessing/Source/DynamicMesh/Private/Operations/MeshSelfUnion.cpp
--- a/Engine/Plugins/Experimental/GeometryProcessing/Source/DynamicMesh/Private/Operations/MeshSelfUnion.cpp
+++ b/Engine/Plugins/Experimental/GeometryProcessing/Source/DynamicMesh/Private/Operations/MeshSelfUnion.cpp
@@ -11,6 +11,7 @@
 
 #include "Async/ParallelFor.h"
 #include "MeshTransforms.h"
+#include "Spatial/SparseDynamicOctree3.h"
 
 #include "Algo/RemoveIf.h"
 
@@ -18,6 +19,16 @@
 
 bool FMeshSelfUnion::Compute()
 {
+	// transform the mesh to a shared space (centered at the origin, and scaled to a unit cube)
+	FAxisAlignedBox3d AABB = Mesh->GetCachedBounds();
+	double ScaleFactor = 1.0 / FMath::Clamp(AABB.MaxDim(), 0.01, 1000000.0);
+	FTransform3d TransformToCenteredBox = FTransform3d::Identity();
+	TransformToCenteredBox.SetTranslation(ScaleFactor * (TransformToCenteredBox.GetTranslation() - AABB.Center()));
+	TransformToCenteredBox.SetScale(ScaleFactor * FVector3d::One());
+	MeshTransforms::ApplyTransform(*Mesh, TransformToCenteredBox);
+	FTransform3d ResultTransform(AABB.Center());
+	ResultTransform.SetScale((1.0 / ScaleFactor) * FVector3d::One());
+
 	// build spatial data and use it to find intersections
 	FDynamicMeshAABBTree3 Spatial(Mesh);
 	Spatial.SetTolerance(SnapTolerance);
@@ -191,19 +202,19 @@ bool FMeshSelfUnion::Compute()
 					double DotNormals = Normals[OtherTID].Dot(Normals[TID]);
 					//if (FMath::Abs(DotNormals) > .9) // TODO: do we actually want to check for a normal match? coplanar vertex check below is more robust?
 					{
-						// To be extra sure it's a coplanar match, check the vertices are *also* on the other mesh (w/in SnapTolerance)
+						// To be extra sure it's a coplanar match, check the vertices are *also* on the other connected component (w/in SnapTolerance)
 						FTriangle3d Tri;
 						Mesh->GetTriVertices(TID, Tri.V[0], Tri.V[1], Tri.V[2]);
-						bool bAllTrisOnOtherMesh = true;
+						bool bAllTrisOnOtherComponent = true;
 						for (int Idx = 0; Idx < 3; Idx++)
 						{
 							if (Spatial.FindNearestTriangle(Tri.V[Idx], DSq, QueryOptions) == FDynamicMesh3::InvalidID)
 							{
-								bAllTrisOnOtherMesh = false;
+								bAllTrisOnOtherComponent = false;
 								break;
 							}
 						}
-						if (bAllTrisOnOtherMesh)
+						if (bAllTrisOnOtherComponent)
 						{
 							if (DotNormals < 0)
 							{
@@ -267,6 +278,42 @@ bool FMeshSelfUnion::Compute()
 	{
 		PointHash.InsertPointUnsafe(BoundaryVID, Mesh->GetVertex(BoundaryVID));
 	}
+
+	FSparseDynamicOctree3 EdgeOctree;
+	EdgeOctree.RootDimension = .25;
+	EdgeOctree.SetMaxTreeDepth(7);
+	auto EdgeBounds = [this](int EID)
+	{
+		FDynamicMesh3::FEdge Edge = Mesh->GetEdge(EID);
+		FVector3d A = Mesh->GetVertex(Edge.Vert.A);
+		FVector3d B = Mesh->GetVertex(Edge.Vert.B);
+		if (A.X > B.X)
+		{
+			Swap(A.X, B.X);
+		}
+		if (A.Y > B.Y)
+		{
+			Swap(A.Y, B.Y);
+		}
+		if (A.Z > B.Z)
+		{
+			Swap(A.Z, B.Z);
+		}
+		return FAxisAlignedBox3d(A, B);
+	};
+	auto AddEdge = [&EdgeOctree, EdgeBounds](int EID)
+	{
+		EdgeOctree.InsertObject(EID, EdgeBounds(EID));
+	};
+	auto UpdateEdge = [&EdgeOctree, EdgeBounds](int EID)
+	{
+		EdgeOctree.ReinsertObject(EID, EdgeBounds(EID));
+	};
+	for (int EID : CutBoundaryEdges)
+	{
+		AddEdge(EID);
+	}
+	TArray<int> EdgesInRange;
 	
 	// mapping of all accepted correspondences of boundary vertices (both ways -- so if A is connected to B we add both A->B and B->A) 
 	TMap<int, int> FoundMatches;
@@ -346,7 +393,11 @@ bool FMeshSelfUnion::Compute()
 			if (NearestVID == FDynamicMesh3::InvalidID)
 			{
 				// vertex had no match -- try to split edge to match it
-				int OtherEID = FindNearestEdge(CutBoundaryEdges, BoundaryNbrEdges, Pos);
+				FAxisAlignedBox3d QueryBox(Pos, SnapTolerance);
+				EdgesInRange.Reset();
+				EdgeOctree.RangeQuery(QueryBox, EdgesInRange);
+
+				int OtherEID = FindNearestEdge(EdgesInRange, BoundaryNbrEdges, Pos);
 				if (OtherEID != FDynamicMesh3::InvalidID)
 				{
 					FVector3d EdgePts[2];
@@ -363,6 +414,8 @@ bool FMeshSelfUnion::Compute()
 							FoundMatches.Add(BoundaryVID, SplitInfo.NewVertex);
 							Mesh->SetVertex(SplitInfo.NewVertex, Pos);
 							CutBoundaryEdges.Add(SplitInfo.NewEdges.A);
+							UpdateEdge(OtherEID);
+							AddEdge(SplitInfo.NewEdges.A);
 							// Note: Do not update PossUnmatchedBdryVerts with the new vertex, because it is already matched by construction
 							// Likewise do not update the pointhash -- we don't want it to find vertices that were already perfectly matched
 						}
@@ -388,7 +441,13 @@ bool FMeshSelfUnion::Compute()
 		return false;
 	}
 
-	bool bWeldSuccess = MergeEdges(CutBoundaryEdges, FoundMatches);
+	bool bWeldSuccess = true;
+	if (bWeldSharedEdges)
+	{
+		bWeldSuccess = MergeEdges(CutBoundaryEdges, FoundMatches);
+	}
+
+	MeshTransforms::ApplyTransform(*Mesh, ResultTransform);
 
 	return bWeldSuccess;
 }
diff --git a/Engine/Plugins/Experimental/GeometryProcessing/Source/DynamicMesh/Public/Operations/MeshSelfUnion.h b/Engine/Plugins/Experimental/GeometryProcessing/Source/DynamicMesh/Public/Operations/MeshSelfUnion.h
--- a/Engine/Plugins/Experimental/GeometryProcessing/Source/DynamicMesh/Public/Operations/MeshSelfUnion.h
+++ b/Engine/Plugins/Experimental/GeometryProcessing/Source/DynamicMesh/Public/Operations/MeshSelfUnion.h
@@ -36,7 +36,7 @@ public:
 	double DegenerateEdgeTolFactor = 1.5;
 
 	/** Tolerance distance for considering a point to be on a vertex or edge, especially during mesh-mesh cutting */
-	double SnapTolerance = FMathf::ZeroTolerance * 300.0;
+	double SnapTolerance = FMathf::ZeroTolerance * 1.0;
 
 	/** Amount we nudge samples off the surface when evaluating winding number, to avoid numerical issues */
 	double NormalOffset = FMathd::ZeroTolerance;
@@ -47,6 +47,9 @@ public:
 	/** Whether to remove visible "open" geometry */
 	bool bTrimFlaps = false;
 
+	/** Weld newly-created cut edges where the mesh is unioned with itself.  If false, newly joined surfaces remain topologically disconnected. */
+	bool bWeldSharedEdges = true;
+
 
 	/** Set this to be able to cancel running operation */
 	FProgressCancel* Progress = nullptr;
