DynamicMeshEditor - Check that the meshes have attribute layers before trying to copy the extended attributes.

#rnx

diff --git a/Engine/Plugins/Experimental/GeometryProcessing/Source/DynamicMesh/Private/DynamicMeshEditor.cpp b/Engine/Plugins/Experimental/GeometryProcessing/Source/DynamicMesh/Private/DynamicMeshEditor.cpp
--- a/Engine/Plugins/Experimental/GeometryProcessing/Source/DynamicMesh/Private/DynamicMeshEditor.cpp
+++ b/Engine/Plugins/Experimental/GeometryProcessing/Source/DynamicMesh/Private/DynamicMeshEditor.cpp
@@ -1718,6 +1718,12 @@ static void AppendTriangleAttributes(const FDynamicMesh3* FromMesh, int FromTria
 // Utility function for ::AppendTriangles()
 static void AppendGenericAttributes(const FDynamicMesh3* FromMesh, FDynamicMesh3* ToMesh, FMeshIndexMappings& IndexMaps)
 {
+
+	if (FromMesh->HasAttributes() == false || ToMesh->HasAttributes() == false)
+	{
+		return;
+	}
+
 	// copy generic attributes after full IndexMaps have been created
 	for (const TPair<FName, TUniquePtr<FDynamicMeshAttributeBase>>& AttribPair : FromMesh->Attributes()->GetAttachedAttributes())
 	{
@@ -1793,8 +1799,9 @@ void FDynamicMeshEditor::AppendTriangles(const FDynamicMesh3* SourceMesh, const
 
 		//Mesh->CheckValidity(true);
 	}
-
+ 
 	AppendGenericAttributes(SourceMesh, Mesh, IndexMaps);
+	
 }
 
 
