Fix compile error for vstudio 16.8.2.

CL doesn't like the calling stack peeking into the template argument for allocating the default argument callable function object. Instead default initialize to null and check inside the scope for that case.

diff --git a/Engine/Plugins/Experimental/GeometryProcessing/Source/GeometricObjects/Public/Spatial/MeshAABBTree3.h b/Engine/Plugins/Experimental/GeometryProcessing/Source/GeometricObjects/Public/Spatial/MeshAABBTree3.h
--- a/Engine/Plugins/Experimental/GeometryProcessing/Source/GeometricObjects/Public/Spatial/MeshAABBTree3.h
+++ b/Engine/Plugins/Experimental/GeometryProcessing/Source/GeometricObjects/Public/Spatial/MeshAABBTree3.h
@@ -800,12 +800,17 @@ public:
 	virtual MeshIntersection::FIntersectionsQueryResult FindAllIntersections(
 		const TMeshAABBTree3& OtherTree, const TFunction<FVector3d(const FVector3d&)>& TransformF = nullptr,
 		const FQueryOptions& Options = FQueryOptions(), const FQueryOptions& OtherTreeOptions = FQueryOptions(),
-		TFunctionRef<bool(FIntrTriangle3Triangle3d&)> IntersectionFn = [](FIntrTriangle3Triangle3d& Intr)
-		{
-			return TMeshAABBTree3<TriangleMeshType>::TriangleIntersection(Intr);
-		}
+		TFunction<bool(FIntrTriangle3Triangle3d&)> IntersectionFn = nullptr
 	) const
 	{
+		if (!IntersectionFn)
+		{
+			IntersectionFn = [](FIntrTriangle3Triangle3d& Intr)
+			{
+				return TMeshAABBTree3<TriangleMeshType>::TriangleIntersection(Intr);
+			};
+		}
+
 		MeshIntersection::FIntersectionsQueryResult result;
 
 		if (!ensure(MeshTimestamp == Mesh->GetShapeTimestamp()))
@@ -829,12 +834,17 @@ public:
 	virtual MeshIntersection::FIntersectionsQueryResult FindAllSelfIntersections(
 		bool bIgnoreTopoConnected = true,
 		const FQueryOptions& Options = FQueryOptions(),
-		TFunctionRef<bool(FIntrTriangle3Triangle3d&)> IntersectionFn = [](FIntrTriangle3Triangle3d& Intr)
-		{
-			return TMeshAABBTree3<TriangleMeshType>::TriangleIntersection(Intr);
-		}
+		TFunction<bool(FIntrTriangle3Triangle3d&)> IntersectionFn = nullptr
 	) const
 	{
+		if (!IntersectionFn)
+		{
+			IntersectionFn = [](FIntrTriangle3Triangle3d& Intr)
+			{
+				return TMeshAABBTree3<TriangleMeshType>::TriangleIntersection(Intr);
+			};
+		}
+
 		MeshIntersection::FIntersectionsQueryResult Result;
 
 		if (!ensure(MeshTimestamp == Mesh->GetShapeTimestamp()))
