GeometryProcessing: additional support for Polygroup attribute layers. Add PolygroupUtil.h with utility functions for finding polygroup attributes from name/etc. Add FPolygroupSet class that provides standard Polygroup API for either base FDynamicMesh3 polygroups, or a FDynamicMeshPolygroupAttribute layer, with MaxGroupID tracking/allocation.
#rnx

diff --git a/Engine/Plugins/Experimental/GeometryProcessing/Source/DynamicMesh/Private/Polygroups/PolygroupSet.cpp b/Engine/Plugins/Experimental/GeometryProcessing/Source/DynamicMesh/Private/Polygroups/PolygroupSet.cpp
new file mode 100644
--- /dev/null
+++ b/Engine/Plugins/Experimental/GeometryProcessing/Source/DynamicMesh/Private/Polygroups/PolygroupSet.cpp
@@ -0,0 +1,76 @@
+// Copyright Epic Games, Inc. All Rights Reserved.
+
+#include "Polygroups/PolygroupSet.h"
+#include "Polygroups/PolygroupUtil.h"
+
+using namespace UE::Geometry;
+
+
+FPolygroupSet::FPolygroupSet(const FPolygroupSet* CopyIn)
+{
+	Mesh = CopyIn->Mesh;
+	PolygroupAttrib = CopyIn->PolygroupAttrib;
+	GroupLayerIndex = CopyIn->GroupLayerIndex;
+	MaxGroupID = CopyIn->MaxGroupID;
+}
+
+FPolygroupSet::FPolygroupSet(FDynamicMesh3* MeshIn)
+{
+	Mesh = MeshIn;
+	GroupLayerIndex = -1;
+	RecalculateMaxGroupID();
+}
+
+FPolygroupSet::FPolygroupSet(FDynamicMesh3* MeshIn, FDynamicMeshPolygroupAttribute* PolygroupAttribIn)
+{
+	Mesh = MeshIn;
+	PolygroupAttrib = PolygroupAttribIn;
+	GroupLayerIndex = UE::Geometry::FindPolygroupLayerIndex(*MeshIn, PolygroupAttrib);
+	RecalculateMaxGroupID();
+}
+
+FPolygroupSet::FPolygroupSet(FDynamicMesh3* MeshIn, int32 PolygroupLayerIndex)
+{
+	Mesh = MeshIn;
+	if (ensure(Mesh->Attributes()))
+	{
+		if (PolygroupLayerIndex < Mesh->Attributes()->NumPolygroupLayers())
+		{
+			PolygroupAttrib = Mesh->Attributes()->GetPolygroupLayer(PolygroupLayerIndex);
+			GroupLayerIndex = PolygroupLayerIndex;
+			return;
+		}
+	}
+	RecalculateMaxGroupID();
+	ensureMsgf(false, TEXT("FPolygroupSet: Attribute index missing!"));
+}
+
+
+FPolygroupSet::FPolygroupSet(FDynamicMesh3* MeshIn, FName AttribName)
+{
+	Mesh = MeshIn;
+	PolygroupAttrib = UE::Geometry::FindPolygroupLayerByName(*MeshIn, AttribName);
+	GroupLayerIndex = UE::Geometry::FindPolygroupLayerIndex(*MeshIn, PolygroupAttrib);
+	RecalculateMaxGroupID();
+	ensureMsgf(PolygroupAttrib != nullptr, TEXT("FPolygroupSet: Attribute set missing!"));
+}
+
+
+void FPolygroupSet::RecalculateMaxGroupID()
+{
+	MaxGroupID = 0;
+	if (PolygroupAttrib)
+	{
+		for (int32 tid : Mesh->TriangleIndicesItr())
+		{
+			MaxGroupID = FMath::Max(MaxGroupID, PolygroupAttrib->GetValue(tid));
+		}
+	}
+	else
+	{
+		for (int32 tid : Mesh->TriangleIndicesItr())
+		{
+			MaxGroupID = FMath::Max(MaxGroupID, Mesh->GetTriangleGroup(tid));
+		}
+	}
+}
\ No newline at end of file
diff --git a/Engine/Plugins/Experimental/GeometryProcessing/Source/DynamicMesh/Private/Polygroups/PolygroupUtil.cpp b/Engine/Plugins/Experimental/GeometryProcessing/Source/DynamicMesh/Private/Polygroups/PolygroupUtil.cpp
new file mode 100644
--- /dev/null
+++ b/Engine/Plugins/Experimental/GeometryProcessing/Source/DynamicMesh/Private/Polygroups/PolygroupUtil.cpp
@@ -0,0 +1,36 @@
+// Copyright Epic Games, Inc. All Rights Reserved.
+
+#include "Polygroups/PolygroupUtil.h"
+
+
+
+FDynamicMeshPolygroupAttribute* UE::Geometry::FindPolygroupLayerByName(FDynamicMesh3& Mesh, FName Name)
+{
+	FDynamicMeshAttributeSet* AttributeSet = Mesh.Attributes();
+	if (AttributeSet == nullptr) return nullptr;
+	int32 NumPolygroupLayers = AttributeSet->NumPolygroupLayers();
+	for (int32 k = 0; k < NumPolygroupLayers; ++k)
+	{
+		if (AttributeSet->GetPolygroupLayer(k)->GetName() == Name)
+		{
+			return AttributeSet->GetPolygroupLayer(k);
+		}
+	}
+	return nullptr;
+}
+
+
+int32 UE::Geometry::FindPolygroupLayerIndex(FDynamicMesh3& Mesh, FDynamicMeshPolygroupAttribute* Layer)
+{
+	FDynamicMeshAttributeSet* AttributeSet = Mesh.Attributes();
+	if (AttributeSet == nullptr) return -1;
+	int32 NumPolygroupLayers = AttributeSet->NumPolygroupLayers();
+	for (int32 k = 0; k < NumPolygroupLayers; ++k)
+	{
+		if (AttributeSet->GetPolygroupLayer(k) == Layer)
+		{
+			return k;
+		}
+	}
+	return -1;
+}
\ No newline at end of file
diff --git a/Engine/Plugins/Experimental/GeometryProcessing/Source/DynamicMesh/Public/Polygroups/PolygroupSet.h b/Engine/Plugins/Experimental/GeometryProcessing/Source/DynamicMesh/Public/Polygroups/PolygroupSet.h
new file mode 100644
--- /dev/null
+++ b/Engine/Plugins/Experimental/GeometryProcessing/Source/DynamicMesh/Public/Polygroups/PolygroupSet.h
@@ -0,0 +1,106 @@
+// Copyright Epic Games, Inc. All Rights Reserved.
+
+#pragma once
+
+#include "MathUtil.h"
+#include "DynamicMesh3.h"
+#include "DynamicMeshAttributeSet.h"
+
+
+namespace UE
+{
+namespace Geometry
+{
+
+
+/**
+ * Polygroup sets can be stored in multiple places. The default location is in the per-triangle group integer stored
+ * directly on a FDynamicMesh3. Additional layers may be stored in the FDynamicMeshAttributeSet. Future iterations
+ * could store packed polygroups in other places, store them in separate arrays, and so on.
+ * FPolygroupSet can be used to abstract these different cases, by providing a standard Polygroup Get/Set API.
+ * 
+ * To support unique Polygroup ID allocation, FPolygroupSet calculates the maximum GroupID on creation, and 
+ * updates this maximum across SetGroup() calls. AllocateNewGroupID() can be used to provide new unused GroupIDs.
+ *
+ */
+struct DYNAMICMESH_API FPolygroupSet
+{
+	FDynamicMesh3* Mesh = nullptr;
+	FDynamicMeshPolygroupAttribute* PolygroupAttrib = nullptr;
+	int32 GroupLayerIndex = -1;
+	int32 MaxGroupID = 0;
+
+	/** Initialize a PolygroupSet for the given Mesh, and standard triangle group layer */
+	explicit FPolygroupSet(FDynamicMesh3* MeshIn);
+
+	/** Initialize a PolygroupSet for given Mesh and specific Polygroup attribute layer */
+	explicit FPolygroupSet(FDynamicMesh3* MeshIn, FDynamicMeshPolygroupAttribute* PolygroupAttribIn);
+
+	/** Initialize a PolygroupSet for given Mesh and specific Polygroup attribute layer, found by index. If not valid, fall back to standard triangle group layer. */
+	explicit FPolygroupSet(FDynamicMesh3* MeshIn, int32 PolygroupLayerIndex);
+
+	/** Initialize a PolygroupSet for given Mesh and specific Polygroup attribute layer, found by name. If not valid, fall back to standard triangle group layer. */
+	explicit FPolygroupSet(FDynamicMesh3* MeshIn, FName AttribName);
+
+	/** Initialize a PolygroupSet by copying an existing PolygroupSet */
+	explicit FPolygroupSet(const FPolygroupSet* CopyIn);
+
+
+	/** @return Mesh this PolygroupSet references  */
+	FDynamicMesh3* GetMesh() { return Mesh; }
+
+	/** @return PolygroupAttribute this PolygroupSet references, or null if no PolygroupAttribute is in use */
+	FDynamicMeshPolygroupAttribute* GetPolygroup() { return PolygroupAttrib; }
+
+	/** @return index of current PolygroupAttribute into Mesh AttributeSet, or -1 if this information does not exist */
+	int32 GetPolygroupIndex() const { return GroupLayerIndex; }
+
+
+	/**
+	 * @return PolygroupID for a TriangleID
+	 */
+	int32 GetGroup(int32 TriangleID)
+	{
+		return (PolygroupAttrib) ? PolygroupAttrib->GetValue(TriangleID) : Mesh->GetTriangleGroup(TriangleID);
+	}
+
+	/**
+	 * Set the PolygroupID for a TriangleID
+	 */
+	void SetGroup(int32 TriangleID, int32 NewGroupID)
+	{
+		if (Mesh->IsTriangle(TriangleID))
+		{
+			if (PolygroupAttrib)
+			{
+				PolygroupAttrib->SetValue(TriangleID, NewGroupID);
+			}
+			else
+			{
+				Mesh->SetTriangleGroup(TriangleID, NewGroupID);
+			}
+		}
+		MaxGroupID = FMath::Max(MaxGroupID, NewGroupID);
+	}
+
+	/**
+	 * Calculate the current maximum PolygroupID used in the active set and store in MaxGroupID member
+	 */
+	void RecalculateMaxGroupID();
+
+	/**
+	 * Allocate a new unused PolygroupID by incrementing the MaxGroupID member
+	 */
+	int32 AllocateNewGroupID()
+	{
+		return ++MaxGroupID;
+	}
+};
+
+
+
+
+
+
+}	// end namespace Geometry
+}	// end namespace UE
diff --git a/Engine/Plugins/Experimental/GeometryProcessing/Source/DynamicMesh/Public/Polygroups/PolygroupUtil.h b/Engine/Plugins/Experimental/GeometryProcessing/Source/DynamicMesh/Public/Polygroups/PolygroupUtil.h
new file mode 100644
--- /dev/null
+++ b/Engine/Plugins/Experimental/GeometryProcessing/Source/DynamicMesh/Public/Polygroups/PolygroupUtil.h
@@ -0,0 +1,27 @@
+// Copyright Epic Games, Inc. All Rights Reserved.
+
+#pragma once
+
+#include "DynamicMesh3.h"
+#include "DynamicMeshAttributeSet.h"
+
+
+namespace UE
+{
+namespace Geometry
+{
+
+/**
+ * Find the first FDynamicMeshPolygroupAttribute with the given FName in a the AttributeSet of a Mesh.
+ * @return nullptr if no Polygroup layer is found
+ */
+DYNAMICMESH_API FDynamicMeshPolygroupAttribute* FindPolygroupLayerByName(FDynamicMesh3& Mesh, FName Name);
+
+/**
+ * @return index of Layer in Mesh AttributeSet, or -1 if not found
+ */
+DYNAMICMESH_API int32 FindPolygroupLayerIndex(FDynamicMesh3& Mesh, FDynamicMeshPolygroupAttribute* Layer);
+
+
+}	// end namespace Geometry
+}	// end namespace UE
\ No newline at end of file
