GeometryProcessing: Add DynamicMeshAttributeSet::IsSeamEdge variant that reports back UV and Normal seam state. FGroupTopology can now be configured with specific FDynamicMeshPolygroupAttribute. Add FMeshFaceSelection utility functions for set Difference and Intersection.
#rnx

diff --git a/Engine/Plugins/Experimental/GeometryProcessing/Source/DynamicMesh/Private/DynamicMeshAttributeSet.cpp b/Engine/Plugins/Experimental/GeometryProcessing/Source/DynamicMesh/Private/DynamicMeshAttributeSet.cpp
--- a/Engine/Plugins/Experimental/GeometryProcessing/Source/DynamicMesh/Private/DynamicMeshAttributeSet.cpp
+++ b/Engine/Plugins/Experimental/GeometryProcessing/Source/DynamicMesh/Private/DynamicMeshAttributeSet.cpp
@@ -278,6 +278,21 @@ bool FDynamicMeshAttributeSet::IsSeamEdge(int eid) const
 }
 
 
+bool FDynamicMeshAttributeSet::IsSeamEdge(int EdgeID, bool& bIsUVSeamOut, bool& bIsNormalSeamOut) const
+{
+	bIsUVSeamOut = false;
+	for (const FDynamicMeshUVOverlay& UVLayer : UVLayers)
+	{
+		if (UVLayer.IsSeamEdge(EdgeID))
+		{
+			bIsUVSeamOut = true;
+		}
+	}
+	bIsNormalSeamOut = Normals0.IsSeamEdge(EdgeID);
+	return (bIsUVSeamOut || bIsNormalSeamOut);
+}
+
+
 bool FDynamicMeshAttributeSet::IsSeamVertex(int VID, bool bBoundaryIsSeam) const
 {
 	for (const FDynamicMeshUVOverlay& UVLayer : UVLayers)
diff --git a/Engine/Plugins/Experimental/GeometryProcessing/Source/DynamicMesh/Private/GroupTopology.cpp b/Engine/Plugins/Experimental/GeometryProcessing/Source/DynamicMesh/Private/GroupTopology.cpp
--- a/Engine/Plugins/Experimental/GeometryProcessing/Source/DynamicMesh/Private/GroupTopology.cpp
+++ b/Engine/Plugins/Experimental/GeometryProcessing/Source/DynamicMesh/Private/GroupTopology.cpp
@@ -24,6 +24,18 @@ bool FGroupTopology::FGroupEdge::IsConnectedToVertices(const TSet<int>& Vertices
 FGroupTopology::FGroupTopology(const FDynamicMesh3* MeshIn, bool bAutoBuild)
 {
 	this->Mesh = MeshIn;
+	this->GroupLayer = nullptr;
+	if (bAutoBuild)
+	{
+		RebuildTopology();
+	}
+}
+
+
+FGroupTopology::FGroupTopology(const FDynamicMesh3* MeshIn, const FDynamicMeshPolygroupAttribute* GroupLayerIn, bool bAutoBuild)
+{
+	this->Mesh = MeshIn;
+	this->GroupLayer = GroupLayerIn;
 	if (bAutoBuild)
 	{
 		RebuildTopology();
diff --git a/Engine/Plugins/Experimental/GeometryProcessing/Source/DynamicMesh/Public/DynamicMeshAttributeSet.h b/Engine/Plugins/Experimental/GeometryProcessing/Source/DynamicMesh/Public/DynamicMeshAttributeSet.h
--- a/Engine/Plugins/Experimental/GeometryProcessing/Source/DynamicMesh/Public/DynamicMeshAttributeSet.h
+++ b/Engine/Plugins/Experimental/GeometryProcessing/Source/DynamicMesh/Public/DynamicMeshAttributeSet.h
@@ -99,6 +99,9 @@ public:
 	/** @return true if the given edge is a seam edge in any overlay */
 	virtual bool IsSeamEdge(int EdgeID) const;
 
+	/** @return true if the given edge is a seam edge in any overlay */
+	virtual bool IsSeamEdge(int EdgeID, bool& bIsUVSeamOut, bool& bIsNormalSeamOut) const;
+
 	/** @return true if the given vertex is a seam vertex in any overlay */
 	virtual bool IsSeamVertex(int VertexID, bool bBoundaryIsSeam = true) const;
 
diff --git a/Engine/Plugins/Experimental/GeometryProcessing/Source/DynamicMesh/Public/GroupTopology.h b/Engine/Plugins/Experimental/GeometryProcessing/Source/DynamicMesh/Public/GroupTopology.h
--- a/Engine/Plugins/Experimental/GeometryProcessing/Source/DynamicMesh/Public/GroupTopology.h
+++ b/Engine/Plugins/Experimental/GeometryProcessing/Source/DynamicMesh/Public/GroupTopology.h
@@ -3,6 +3,7 @@
 #pragma once
 
 #include "DynamicMesh3.h"
+#include "DynamicMeshAttributeSet.h"
 #include "EdgeSpan.h"
 #include "Util/IndexUtil.h"
 #include "Containers/BitArray.h"
@@ -132,6 +133,8 @@ class DYNAMICMESH_API FGroupTopology
 public:
 	FGroupTopology() {}
 	FGroupTopology(const FDynamicMesh3* Mesh, bool bAutoBuild);
+	FGroupTopology(const FDynamicMesh3* Mesh, const FDynamicMeshPolygroupAttribute* GroupLayer, bool bAutoBuild);
+	
 
 	virtual ~FGroupTopology() {}
 
@@ -160,7 +163,7 @@ public:
 	 */
 	virtual int GetGroupID(int TriangleID) const
 	{
-		return Mesh->GetTriangleGroup(TriangleID);
+		return (GroupLayer != nullptr) ? GroupLayer->GetValue(TriangleID) : Mesh->GetTriangleGroup(TriangleID);
 	}
 
 
@@ -372,6 +375,7 @@ public:
 
 protected:
 	const FDynamicMesh3* Mesh = nullptr;
+	const FDynamicMeshPolygroupAttribute* GroupLayer = nullptr;
 
 	TArray<int> GroupIDToGroupIndexMap;		// allow fast lookup of index in .Groups, given GroupID
 	TBitArray<> CornerVerticesFlags;		// bit array of corners for fast testing in ExtractGroupEdges  (redundant w/ VertexIDToCornerIDMap?)
diff --git a/Engine/Plugins/Experimental/GeometryProcessing/Source/DynamicMesh/Public/Selections/MeshFaceSelection.h b/Engine/Plugins/Experimental/GeometryProcessing/Source/DynamicMesh/Public/Selections/MeshFaceSelection.h
--- a/Engine/Plugins/Experimental/GeometryProcessing/Source/DynamicMesh/Public/Selections/MeshFaceSelection.h
+++ b/Engine/Plugins/Experimental/GeometryProcessing/Source/DynamicMesh/Public/Selections/MeshFaceSelection.h
@@ -77,6 +77,10 @@ public:
 		return Selected.Contains(tid);
 	}
 
+	bool Contains(int tid) const
+	{
+		return Selected.Contains(tid);
+	}
 
 	void Select(int tid)
 	{
@@ -202,6 +206,36 @@ public:
 		}
 	}
 
+	/**
+	 * Find the elements in current selection that are not in SubtractSet, and store in DifferenceStorage
+	 */
+	template<typename EnumerableType, typename StorageType>
+	void SetDifference(const EnumerableType& SubtractSet, StorageType& DifferenceStorage) const
+	{
+		for ( int32 tid : Selected )
+		{
+			if (SubtractSet.Contains(tid) == false)
+			{
+				DifferenceStorage.Add(tid);
+			}
+		}
+	}
+
+	/**
+	 * Find the elements in current selection that are also in IntersectSet, and store in IntersectionStorage
+	 */
+	template<typename EnumerableType, typename StorageType>
+	void SetIntersection(const EnumerableType& IntersectSet, StorageType& IntersectionStorage) const
+	{
+		for (int32 tid : Selected)
+		{
+			if (IntersectSet.Contains(tid))
+			{
+				IntersectionStorage.Add(tid);
+			}
+		}
+	}
+
 	const TSet<int>& AsSet() const
 	{
 		return Selected;
