GeometryProcessing: fix typo in TImageBuilder that resulted in incorrect range check
#rnx

diff --git a/Engine/Plugins/Experimental/GeometryProcessing/Source/GeometricObjects/Public/Image/ImageBuilder.h b/Engine/Plugins/Experimental/GeometryProcessing/Source/GeometricObjects/Public/Image/ImageBuilder.h
--- a/Engine/Plugins/Experimental/GeometryProcessing/Source/GeometricObjects/Public/Image/ImageBuilder.h
+++ b/Engine/Plugins/Experimental/GeometryProcessing/Source/GeometricObjects/Public/Image/ImageBuilder.h
@@ -116,7 +116,7 @@ public:
 
 		// make sure we are in range
 		if (X0 < 0 || X1 >= Dimensions.GetWidth() ||
-			Y0 < 0 || Y0 >= Dimensions.GetHeight())
+			Y0 < 0 || Y1 >= Dimensions.GetHeight())
 		{
 			return InvalidValue;
 		}
