make hole fillers auto-create a new group for the hole if triangle groups are enabled and the default -1 is specified as the group ID.  (Before it would create a triangle with group ID of -1, which I think we decided we don't want to do ...)

#rnx

diff --git a/Engine/Plugins/Experimental/GeometryProcessing/Source/DynamicMesh/Private/Operations/MinimalHoleFiller.cpp b/Engine/Plugins/Experimental/GeometryProcessing/Source/DynamicMesh/Private/Operations/MinimalHoleFiller.cpp
--- a/Engine/Plugins/Experimental/GeometryProcessing/Source/DynamicMesh/Private/Operations/MinimalHoleFiller.cpp
+++ b/Engine/Plugins/Experimental/GeometryProcessing/Source/DynamicMesh/Private/Operations/MinimalHoleFiller.cpp
@@ -591,6 +591,11 @@ void FMinimalHoleFiller::FlipToImproveAspectRatios()
 
 bool FMinimalHoleFiller::Fill(int32 GroupID)
 {
+	if (GroupID < 0 && Mesh->HasTriangleGroups())
+	{
+		GroupID = Mesh->AllocateTriangleGroup();
+	}
+
 	// do a simple fill
 	FSimpleHoleFiller Simplefill(Mesh, FillLoop);
 
diff --git a/Engine/Plugins/Experimental/GeometryProcessing/Source/DynamicMesh/Private/Operations/PlanarHoleFiller.cpp b/Engine/Plugins/Experimental/GeometryProcessing/Source/DynamicMesh/Private/Operations/PlanarHoleFiller.cpp
--- a/Engine/Plugins/Experimental/GeometryProcessing/Source/DynamicMesh/Private/Operations/PlanarHoleFiller.cpp
+++ b/Engine/Plugins/Experimental/GeometryProcessing/Source/DynamicMesh/Private/Operations/PlanarHoleFiller.cpp
@@ -7,6 +7,11 @@
 
 bool FPlanarHoleFiller::Fill(int GroupID)
 {
+	if (GroupID < 0 && Mesh->HasTriangleGroups())
+	{
+		GroupID = Mesh->AllocateTriangleGroup();
+	}
+
 	FVector3d PlaneX, PlaneY;
 	VectorUtil::MakePerpVectors(PlaneNormal, PlaneX, PlaneY);
 
diff --git a/Engine/Plugins/Experimental/GeometryProcessing/Source/DynamicMesh/Private/Operations/SimpleHoleFiller.cpp b/Engine/Plugins/Experimental/GeometryProcessing/Source/DynamicMesh/Private/Operations/SimpleHoleFiller.cpp
--- a/Engine/Plugins/Experimental/GeometryProcessing/Source/DynamicMesh/Private/Operations/SimpleHoleFiller.cpp
+++ b/Engine/Plugins/Experimental/GeometryProcessing/Source/DynamicMesh/Private/Operations/SimpleHoleFiller.cpp
@@ -9,6 +9,11 @@
 
 bool FSimpleHoleFiller::Fill(int GroupID)
 {
+	if (GroupID < 0 && Mesh->HasTriangleGroups())
+	{
+		GroupID = Mesh->AllocateTriangleGroup();
+	}
+
 	if (Loop.GetVertexCount() < 3)
 	{
 		return false;
