boolean option for polycut
 + also improves texturing options for generalized cylinder (needed b/c that's what we use to make the cutting geometry)
 + also separates the "embedded edges" concept from the "error / problem" edges, and always visualizes the embedded edges (for the boolean, these are different edges; also visualizing the embedded edges makes it a lot easier to understand what the "Insert Polygon" option is doing)

diff --git a/Engine/Plugins/Experimental/GeometryProcessing/Source/DynamicMesh/Private/Operations/MeshBoolean.cpp b/Engine/Plugins/Experimental/GeometryProcessing/Source/DynamicMesh/Private/Operations/MeshBoolean.cpp
--- a/Engine/Plugins/Experimental/GeometryProcessing/Source/DynamicMesh/Private/Operations/MeshBoolean.cpp
+++ b/Engine/Plugins/Experimental/GeometryProcessing/Source/DynamicMesh/Private/Operations/MeshBoolean.cpp
@@ -4,6 +4,7 @@
 
 #include "Operations/MeshBoolean.h"
 #include "Operations/MeshMeshCut.h"
+#include "Selections/MeshConnectedComponents.h"
 
 #include "Async/ParallelFor.h"
 #include "MeshTransforms.h"
@@ -256,12 +257,12 @@ bool FMeshBoolean::Compute()
 		return false;
 	}
 
-	bool bOpIsATrim = Operation == EBooleanOp::TrimInside || Operation == EBooleanOp::TrimOutside;
+	bool bOpOnSingleMesh = Operation == EBooleanOp::TrimInside || Operation == EBooleanOp::TrimOutside || Operation == EBooleanOp::NewGroupInside || Operation == EBooleanOp::NewGroupOutside;
 
 	// cut the meshes
 	FMeshMeshCut Cut(CutMesh[0], CutMesh[1]);
 	Cut.bTrackInsertedVertices = bCollapseDegenerateEdgesOnCut; // to collect candidates to collapse
-	Cut.bMutuallyCut = !bOpIsATrim;
+	Cut.bMutuallyCut = !bOpOnSingleMesh;
 	Cut.SnapTolerance = SnapTolerance;
 	Cut.Cut(Intersections);
 
@@ -270,7 +271,7 @@ bool FMeshBoolean::Compute()
 		return false;
 	}
 
-	int NumMeshesToProcess = bOpIsATrim ? 1 : 2;
+	int NumMeshesToProcess = bOpOnSingleMesh ? 1 : 2;
 
 	// collapse tiny edges along cut boundary
 	if (bCollapseDegenerateEdgesOnCut)
@@ -363,9 +364,9 @@ bool FMeshBoolean::Compute()
 			FDynamicMesh3& ProcessMesh = *CutMesh[MeshIdx];
 			int MaxTriID = ProcessMesh.MaxTriangleID();
 			KeepTri[MeshIdx].SetNumUninitialized(MaxTriID);
-			bool bCoplanarKeepSameDir = (Operation != EBooleanOp::Difference && Operation != EBooleanOp::TrimInside);
+			bool bCoplanarKeepSameDir = (Operation != EBooleanOp::Difference && Operation != EBooleanOp::TrimInside && Operation != EBooleanOp::NewGroupInside);
 			bool bRemoveInside = 1; // whether to remove the inside triangles (e.g. for union) or the outside ones (e.g. for intersection)
-			if (Operation == EBooleanOp::TrimOutside || Operation == EBooleanOp::Intersect || (Operation == EBooleanOp::Difference && MeshIdx == 1))
+			if (Operation == EBooleanOp::NewGroupOutside || Operation == EBooleanOp::TrimOutside || Operation == EBooleanOp::Intersect || (Operation == EBooleanOp::Difference && MeshIdx == 1))
 			{
 				bRemoveInside = 0;
 			}
@@ -440,6 +441,14 @@ bool FMeshBoolean::Compute()
 			}
 		}
 		// now go ahead and delete from both meshes
+		bool bRegroupInsteadOfDelete = Operation == EBooleanOp::NewGroupInside || Operation == EBooleanOp::NewGroupOutside;
+		int NewGroupID = -1;
+		TArray<int> NewGroupTris;
+		if (bRegroupInsteadOfDelete)
+		{
+			ensure(NumMeshesToProcess == 1);
+			NewGroupID = CutMesh[0]->AllocateTriangleGroup();
+		}
 		for (int MeshIdx = 0; MeshIdx < NumMeshesToProcess; MeshIdx++)
 		{
 			FDynamicMesh3& ProcessMesh = *CutMesh[MeshIdx];
@@ -448,7 +457,29 @@ bool FMeshBoolean::Compute()
 			{
 				if (ProcessMesh.IsTriangle(TID) && !KeepTri[MeshIdx][TID])
 				{
-					ProcessMesh.RemoveTriangle(TID, true, false);
+					if (bRegroupInsteadOfDelete)
+					{
+						ProcessMesh.SetTriangleGroup(TID, NewGroupID);
+						NewGroupTris.Add(TID);
+					}
+					else
+					{
+						ProcessMesh.RemoveTriangle(TID, true, false);
+					}
+				}
+			}
+		}
+		if (bRegroupInsteadOfDelete)
+		{
+			// the new triangle group could include disconnected components; best to give them separate triangle groups
+			FMeshConnectedComponents Components(CutMesh[0]);
+			Components.FindConnectedTriangles(NewGroupTris);
+			for (int ComponentIdx = 1; ComponentIdx < Components.Num(); ComponentIdx++)
+			{
+				int SplitGroupID = CutMesh[0]->AllocateTriangleGroup();
+				for (int TID : Components.GetComponent(ComponentIdx).Indices)
+				{
+					CutMesh[0]->SetTriangleGroup(TID, SplitGroupID);
 				}
 			}
 		}
diff --git a/Engine/Plugins/Experimental/GeometryProcessing/Source/DynamicMesh/Public/Operations/MeshBoolean.h b/Engine/Plugins/Experimental/GeometryProcessing/Source/DynamicMesh/Public/Operations/MeshBoolean.h
--- a/Engine/Plugins/Experimental/GeometryProcessing/Source/DynamicMesh/Public/Operations/MeshBoolean.h
+++ b/Engine/Plugins/Experimental/GeometryProcessing/Source/DynamicMesh/Public/Operations/MeshBoolean.h
@@ -39,7 +39,9 @@ public:
 		Difference,
 		Intersect,
 		TrimInside,
-		TrimOutside
+		TrimOutside,
+		NewGroupInside,
+		NewGroupOutside
 	};
 	EBooleanOp Operation;
 	
diff --git a/Engine/Plugins/Experimental/GeometryProcessing/Source/GeometricObjects/Public/Generators/SweepGenerator.h b/Engine/Plugins/Experimental/GeometryProcessing/Source/GeometricObjects/Public/Generators/SweepGenerator.h
--- a/Engine/Plugins/Experimental/GeometryProcessing/Source/GeometricObjects/Public/Generators/SweepGenerator.h
+++ b/Engine/Plugins/Experimental/GeometryProcessing/Source/GeometricObjects/Public/Generators/SweepGenerator.h
@@ -49,16 +49,49 @@ protected:
 							   const TArrayView<const int32>& UVSections,
 							   const TArrayView<const int32>& NormalSections,
 							   const TArrayView<const int32>& SharpNormalsAlongLength,
+							   bool bEvenlySpaceUVs,
+							   const TArrayView<const FVector3d>& Path, // can be empty unless bEvenlySpaceUVs is true
 							   int32 NumCrossSections,
 							   bool bLoop,
 							   const ECapType Caps[2],
-							   FVector2d UVScale, FVector2d UVOffset)
+							   FVector2f SectionsUVScale, FVector2f CapUVScale, FVector2f CapUVOffset)
 	{
 		// per cross section
 		int32 XVerts = CrossSection.VertexCount();
 		int32 XNormals = XVerts + NormalSections.Num();
 		int32 XUVs = XVerts + UVSections.Num() + 1;
 
+		double TotalPerimeter = 0, TotalPathLength = 0;
+		TArray<double> CrossSectionPercentages, PathPercentages;
+		if (bEvenlySpaceUVs)
+		{
+			CrossSectionPercentages.Add(0);
+			PathPercentages.Add(0);
+			for (int XIdx = 0, XNum = CrossSection.VertexCount(); XIdx < XNum; XIdx++)
+			{
+				double SegLen = CrossSection[XIdx].Distance(CrossSection[(XIdx + 1) % XNum]);
+				TotalPerimeter += SegLen;
+				CrossSectionPercentages.Add(TotalPerimeter);
+			}
+			TotalPerimeter = FMath::Max(TotalPerimeter, FMathd::ZeroTolerance);
+			for (int Idx = 1; Idx < CrossSectionPercentages.Num(); Idx++)
+			{
+				CrossSectionPercentages[Idx] /= TotalPerimeter;
+			}
+			int NumPathSegs = bLoop ? Path.Num() : Path.Num() - 1;
+			for (int PIdx = 0; PIdx < NumPathSegs; PIdx++)
+			{
+				double SegLen = Path[PIdx].Distance(Path[(PIdx + 1) % Path.Num()]);
+				TotalPathLength += SegLen;
+				PathPercentages.Add(TotalPathLength);
+			}
+			TotalPathLength = FMath::Max(TotalPathLength, FMathd::ZeroTolerance);
+			for (int Idx = 1; Idx < PathPercentages.Num(); Idx++)
+			{
+				PathPercentages[Idx] /= TotalPathLength;
+			}
+		}
+
 		int32 NumVerts = XVerts * NumCrossSections - (bLoop ? XVerts : 0);
 		int32 NumNormals = NumCrossSections > 1 ? (XNormals * NumCrossSections - (bLoop ? XNormals : 0)) : 0;
 		NumNormals += XNormals * SharpNormalsAlongLength.Num();
@@ -146,7 +179,7 @@ protected:
 					float SideScale = 2 * CapIdx - 1;
 					for (int32 Idx = 0; Idx < XVerts; Idx++)
 					{
-						FVector2d CenteredVert = CrossSection.GetVertices()[Idx] * UVScale + UVOffset;
+						FVector2f CenteredVert = (FVector2f)CrossSection.GetVertices()[Idx] * CapUVScale + CapUVOffset;
 						SetUV(CapUVStart[CapIdx] + Idx, FVector2f(CenteredVert.X * SideScale, CenteredVert.Y), VertOffset + Idx);
 						SetNormal(CapNormalStart[CapIdx] + Idx, FVector3f::Zero(), VertOffset + Idx);
 					}
@@ -172,11 +205,11 @@ protected:
 			int32 NextDupVertIdx = UVSection < NumSections ? UVSections[UVSection] : -1;
 			for (int32 VertSubIdx = 0; VertSubIdx < XVerts; UVSubIdx++)
 			{
-				float UVX = VertSubIdx / float(XVerts);
+				float UVX = bEvenlySpaceUVs ? CrossSectionPercentages[VertSubIdx] : VertSubIdx / float(XVerts);
 				for (int32 XIdx = 0; XIdx < NumCrossSections; XIdx++)
 				{
-					float UVY = XIdx / float(NumCrossSections - 1);
-					SetUV(XIdx * XUVs + UVSubIdx, FVector2f(1-UVX, 1-UVY), (XIdx % CrossSectionsMod) * XVerts + VertSubIdx);
+					float UVY = bEvenlySpaceUVs ? PathPercentages[XIdx] : XIdx / float(NumCrossSections - 1);
+					SetUV(XIdx * XUVs + UVSubIdx, FVector2f(1-UVX, 1-UVY) * SectionsUVScale, (XIdx % CrossSectionsMod) * XVerts + VertSubIdx);
 				}
 
 				if (VertSubIdx == NextDupVertIdx)
@@ -207,8 +240,8 @@ protected:
 				int32 VertSubIdx = 0;
 				for (int32 XIdx = 0; XIdx < NumCrossSections; XIdx++)
 				{
-					float UVY = XIdx / float(NumCrossSections - 1);
-					SetUV(XIdx * XUVs + UVSubIdx, FVector2f(1-UVX, 1-UVY), (XIdx % CrossSectionsMod) * XVerts + VertSubIdx);
+					float UVY = bEvenlySpaceUVs ? PathPercentages[XIdx] : XIdx / float(NumCrossSections - 1);
+					SetUV(XIdx * XUVs + UVSubIdx, FVector2f(1-UVX, 1-UVY) * SectionsUVScale, (XIdx % CrossSectionsMod) * XVerts + VertSubIdx);
 				}
 			}
 			NumSections = NormalSections.Num();
@@ -336,7 +369,7 @@ public:
 		TArray<float> AlongPercents;
 		float LenAlong = ComputeSegLengths(Radii, Heights, AlongPercents);
 
-		ConstructMeshTopology(X, {}, {}, SharpNormalsAlongLength, NumX, false, Caps, FVector2d(.5, .5), FVector2d(.5, .5));
+		ConstructMeshTopology(X, {}, {}, SharpNormalsAlongLength, false, {}, NumX, false, Caps, FVector2f(1, 1), FVector2f(.5f, .5f), FVector2f(.5f, .5f));
 
 		TArray<FVector2d> NormalSides; NormalSides.SetNum(NumX - 1);
 		for (int XIdx = 0; XIdx+1 < NumX; XIdx++)
@@ -545,6 +578,15 @@ public:
 	bool bCapped = false;
 	bool bLoop = false;
 
+	// When true, the generator attempts to scale UV's in a way that preserves scaling across different mesh
+	// results, aiming for 1.0 in UV space to be equal to UnitUVInWorldCoordinates in world space. This in
+	// practice means adjusting the U scale relative to the CrossSection curve length and V scale relative
+	// to the distance between vertices on the Path.
+	bool bUVScaleRelativeWorld = false;
+
+	// Only relevant if bUVScaleRelativeWorld is true (see that description)
+	float UnitUVInWorldCoordinates = 100;
+
 public:
 	/** Generate the mesh */
 	virtual FMeshShapeGenerator& Generate() override
@@ -558,7 +600,19 @@ public:
 			Caps[1] = ECapType::FlatTriangulation;
 		}
 		int PathNum = Path.Num();
-		ConstructMeshTopology(CrossSection, {}, {}, {}, PathNum + (bLoop ? 1 : 0), bLoop, Caps, FVector2d(.5, .5), FVector2d(.5, .5));
+		
+		FAxisAlignedBox2f Bounds = (FAxisAlignedBox2f)CrossSection.Bounds();
+		double BoundsMaxDimInv = 1.0 / FMathd::Max(Bounds.MaxDim(), .001);
+		FVector2f SectionScale(1, 1), CapScale(BoundsMaxDimInv, BoundsMaxDimInv);
+		if (bUVScaleRelativeWorld)
+		{
+			double Perimeter = CrossSection.Perimeter();
+			double PathLen = TCurveUtil<double>::ArcLength(Path, bLoop);
+			SectionScale.X = Perimeter / UnitUVInWorldCoordinates;
+			SectionScale.Y = PathLen / UnitUVInWorldCoordinates;
+			CapScale.X = CapScale.Y = 1.0f / UnitUVInWorldCoordinates;
+		}
+		ConstructMeshTopology(CrossSection, {}, {}, {}, true, Path, PathNum + (bLoop ? 1 : 0), bLoop, Caps, SectionScale, CapScale, Bounds.Center());
 
 		int XNum = CrossSection.VertexCount();
 		TArray<FVector2d> XNormals; XNormals.SetNum(XNum);
diff --git a/Engine/Plugins/Experimental/MeshModelingToolset/Source/MeshModelingToolsEditorOnly/Private/PolygonOnMeshTool.cpp b/Engine/Plugins/Experimental/MeshModelingToolset/Source/MeshModelingToolsEditorOnly/Private/PolygonOnMeshTool.cpp
--- a/Engine/Plugins/Experimental/MeshModelingToolset/Source/MeshModelingToolsEditorOnly/Private/PolygonOnMeshTool.cpp
+++ b/Engine/Plugins/Experimental/MeshModelingToolset/Source/MeshModelingToolsEditorOnly/Private/PolygonOnMeshTool.cpp
@@ -108,8 +108,9 @@ void UPolygonOnMeshTool::Setup()
 		[this](const FDynamicMeshOperator* Op)
 		{
 			const FEmbedPolygonsOp* PolygonsOp = (const FEmbedPolygonsOp*)(Op);
+			EdgesOnFailure = PolygonsOp->EdgesOnFailure;
 			EmbeddedEdges = PolygonsOp->EmbeddedEdges;
-			bEmbedSucceeded = PolygonsOp->bEmbedSucceeded;
+			bOperationSucceeded = PolygonsOp->bOperationSucceeded;
 		}
 	);
 	Preview->OnMeshUpdated.AddLambda(
@@ -150,16 +151,25 @@ void UPolygonOnMeshTool::Setup()
 void UPolygonOnMeshTool::UpdateVisualization()
 {
 	FColor PartialPathEdgeColor(240, 15, 15);
-	float PartialPathEdgeThickness = 2.0;
-	float PartialPathEdgeDepthBias = 2.0f;
+	float PartialPathEdgeThickness = 2.0f;
+	float PartialPathEdgeDepthBias = 3.0f;
+
+	FColor EmbedEdgeColor(100, 240, 100);
+	float EmbedEdgeThickness = 1.0f;
+	float EmbedEdgeDepthBias = 1.0f;
 
 	const FDynamicMesh3* TargetMesh = Preview->PreviewMesh->GetPreviewDynamicMesh();
 	FVector3d A, B;
 
 	DrawnLineSet->Clear();
-	if (!bEmbedSucceeded)
+	for (int EID : EmbeddedEdges)
+	{
+		TargetMesh->GetEdgeV(EID, A, B);
+		DrawnLineSet->AddLine((FVector)A, (FVector)B, EmbedEdgeColor, EmbedEdgeThickness, EmbedEdgeDepthBias);
+	}
+	if (!bOperationSucceeded)
 	{
-		for (int EID : EmbeddedEdges)
+		for (int EID : EdgesOnFailure)
 		{
 			TargetMesh->GetEdgeV(EID, A, B);
 			DrawnLineSet->AddLine((FVector)A, (FVector)B, PartialPathEdgeColor, PartialPathEdgeThickness, PartialPathEdgeDepthBias);
@@ -248,6 +258,8 @@ TUniquePtr<FDynamicMeshOperator> UPolygonOnMeshTool::MakeNewOperator()
 	TUniquePtr<FEmbedPolygonsOp> EmbedOp = MakeUnique<FEmbedPolygonsOp>();
 	EmbedOp->bDiscardAttributes = false;
 	EmbedOp->Operation = BasicProperties->Operation;
+	EmbedOp->bCutWithBoolean = BasicProperties->bCutWithBoolean;
+	EmbedOp->bAttemptFixHolesOnBoolean = BasicProperties->bAttemptFixHoles;
 
 	FFrame3d LocalFrame = DrawPlaneWorld;
 	FTransform3d ToLocal = WorldTransform.Inverse();
diff --git a/Engine/Plugins/Experimental/MeshModelingToolset/Source/MeshModelingToolsEditorOnly/Public/PolygonOnMeshTool.h b/Engine/Plugins/Experimental/MeshModelingToolset/Source/MeshModelingToolsEditorOnly/Public/PolygonOnMeshTool.h
--- a/Engine/Plugins/Experimental/MeshModelingToolset/Source/MeshModelingToolsEditorOnly/Public/PolygonOnMeshTool.h
+++ b/Engine/Plugins/Experimental/MeshModelingToolset/Source/MeshModelingToolsEditorOnly/Public/PolygonOnMeshTool.h
@@ -72,6 +72,14 @@ public:
 	UPROPERTY(EditAnywhere, Category = Shape)
 	EPolygonType Shape = EPolygonType::Circle;
 
+	/** Use a volumetric boolean rather than curve projection; cuts through all layers and across edges */
+	UPROPERTY(EditAnywhere, Category = Operation)
+	bool bCutWithBoolean = true;
+
+	/** Automatically attempt to fill any holes left by CSG (e.g. due to numerical errors) */
+	UPROPERTY(EditAnywhere, Category = Operation, meta = (EditCondition = "bCutWithBoolean"))
+	bool bAttemptFixHoles = true;
+
 	// TODO: re-add if/when extrude is added as a supported operation
 	///** Amount to extrude, if extrude is enabled */
 	//UPROPERTY(EditAnywhere, Category = Options)
@@ -197,8 +205,9 @@ protected:
 	UPROPERTY()
 	ULineSetComponent* DrawnLineSet;
 
+	TArray<int> EdgesOnFailure;
 	TArray<int> EmbeddedEdges;
-	bool bEmbedSucceeded;
+	bool bOperationSucceeded;
 
 protected:
 	UWorld* TargetWorld;
diff --git a/Engine/Plugins/Experimental/MeshModelingToolset/Source/ModelingOperatorsEditorOnly/Private/CuttingOps/EmbedPolygonsOp.cpp b/Engine/Plugins/Experimental/MeshModelingToolset/Source/ModelingOperatorsEditorOnly/Private/CuttingOps/EmbedPolygonsOp.cpp
--- a/Engine/Plugins/Experimental/MeshModelingToolset/Source/ModelingOperatorsEditorOnly/Private/CuttingOps/EmbedPolygonsOp.cpp
+++ b/Engine/Plugins/Experimental/MeshModelingToolset/Source/ModelingOperatorsEditorOnly/Private/CuttingOps/EmbedPolygonsOp.cpp
@@ -6,7 +6,8 @@
 #include "Selections/MeshFaceSelection.h"
 #include "MeshQueries.h"
 #include "Operations/EmbedSurfacePath.h"
-#include "Operations/SimpleHoleFiller.h"
+#include "Operations/MinimalHoleFiller.h"
+#include "Generators/SweepGenerator.h"
 
 #include "Engine/StaticMesh.h"
 
@@ -16,6 +17,7 @@
 #include "MeshSimplification.h"
 #include "MeshConstraintsUtil.h"
 #include "Operations/MeshPlaneCut.h"
+#include "Operations/MeshBoolean.h"
 #include "ConstrainedDelaunay2.h"
 
 
@@ -112,21 +114,145 @@ void CollapseDegenerateEdgesOnVertexPath(FDynamicMesh3& Mesh, TArray<int>& Verte
 	}
 }
 
+void FEmbedPolygonsOp::BooleanPath(FProgressCancel* Progress)
+{
+	FAxisAlignedBox3d Bounds = OriginalMesh->GetBounds();
+	double MeshDiameter = Bounds.MaxDim();
+
+	FFrame3d Frame = PolygonFrame;
+
+	FVector3d Dir = Frame.Z();
+	FInterval1d Range;
+	for (int CornerIdx = 0; CornerIdx < 8; CornerIdx++)
+	{
+		FVector3d Corner = Bounds.GetCorner(CornerIdx);
+		Range.Contain(Dir.Dot(Corner - Frame.Origin));
+	}
+
+
+	Frame.Origin = Frame.Origin + Range.Min * Dir;
+	FGeneralizedCylinderGenerator ExtrudePolyGen;
+	ExtrudePolyGen.CrossSection = EmbedPolygon;
+	ExtrudePolyGen.InitialFrame = Frame;
+	ExtrudePolyGen.Path.Add(Frame.Origin);
+	ExtrudePolyGen.Path.Add(PolygonFrame.Origin + Range.Max * Dir);
+	ExtrudePolyGen.bCapped = true;
+	ExtrudePolyGen.bUVScaleRelativeWorld = true;
+	ExtrudePolyGen.UnitUVInWorldCoordinates = MeshDiameter;
+	FDynamicMesh3 ExtrudePoly;
+	ExtrudePoly.EnableMeshComponents(ResultMesh->GetComponentsFlags());
+	if (ResultMesh->HasAttributes())
+	{
+		ExtrudePoly.EnableAttributes();
+		ExtrudePoly.Attributes()->EnableMatchingAttributes(*ResultMesh->Attributes());
+	}
+	ExtrudePoly.Copy(&ExtrudePolyGen.Generate());
+	FMeshBoolean::EBooleanOp BoolOp = FMeshBoolean::EBooleanOp::Difference;
+	switch (Operation)
+	{
+	case EEmbeddedPolygonOpMethod::InsertPolygon:
+		BoolOp = FMeshBoolean::EBooleanOp::NewGroupInside;
+		break;
+	case EEmbeddedPolygonOpMethod::TrimInside:
+		BoolOp = FMeshBoolean::EBooleanOp::TrimInside;
+		break;
+	case EEmbeddedPolygonOpMethod::TrimOutside:
+		BoolOp = FMeshBoolean::EBooleanOp::TrimOutside;
+		break;
+	case EEmbeddedPolygonOpMethod::CutThrough:
+		BoolOp = FMeshBoolean::EBooleanOp::Difference;
+		break;
+	default:
+		unimplemented();
+	}
+
+	if (Progress->Cancelled())
+	{
+		return;
+	}
+
+	int MaxGroupID = ResultMesh->MaxGroupID();
+	FMeshBoolean Boolean(ResultMesh.Get(), &ExtrudePoly, ResultMesh.Get(), BoolOp);
+	Boolean.Progress = Progress;
+	Boolean.bPutResultInInputSpace = true;
+	bool bBoolSuccess = Boolean.Compute();
+
+	if (Progress->Cancelled())
+	{
+		return;
+	}
+
+	if (Boolean.CreatedBoundaryEdges.Num() > 0 && bAttemptFixHolesOnBoolean)
+	{
+		FMeshBoundaryLoops OpenBoundary(Boolean.Result, false);
+		TSet<int> ConsiderEdges(Boolean.CreatedBoundaryEdges);
+		OpenBoundary.EdgeFilterFunc = [&ConsiderEdges](int EID)
+		{
+			return ConsiderEdges.Contains(EID);
+		};
+		OpenBoundary.Compute();
+
+		if (Progress->Cancelled())
+		{
+			return;
+		}
+
+		for (FEdgeLoop& Loop : OpenBoundary.Loops)
+		{
+			FMinimalHoleFiller Filler(Boolean.Result, Loop);
+			Filler.Fill(ResultMesh->AllocateTriangleGroup());
+		}
+		for (int EID : Boolean.CreatedBoundaryEdges)
+		{
+			if (Boolean.Result->IsEdge(EID) && Boolean.Result->IsBoundaryEdge(EID))
+			{
+				EdgesOnFailure.Add(EID);
+			}
+		}
+
+		bOperationSucceeded = EdgesOnFailure.Num() == 0;
+	}
+	else
+	{
+		EdgesOnFailure = Boolean.CreatedBoundaryEdges;
+		bOperationSucceeded = bBoolSuccess;
+	}
+
+	for (int EID : ResultMesh->EdgeIndicesItr())
+	{
+		FDynamicMesh3::FEdge Edge = ResultMesh->GetEdge(EID);
+		if (ResultMesh->GetTriangleGroup(Edge.Tri[0]) < MaxGroupID !=
+			ResultMesh->GetTriangleGroup(Edge.Tri[1]) < MaxGroupID)
+		{
+			EmbeddedEdges.Add(EID);
+		}
+	}
+
+}
+
 void FEmbedPolygonsOp::CalculateResult(FProgressCancel* Progress)
 {
+	bOperationSucceeded = false;
+
 	if (Progress->Cancelled())
 	{
 		return;
 	}
 	ResultMesh->Copy(*OriginalMesh, true, true, true, !bDiscardAttributes);
 
-	double MeshRadius = OriginalMesh->GetBounds().MaxDim();
-	double UVScaleFactor = 1.0 / MeshRadius;
+	if (bCutWithBoolean)
+	{
+		BooleanPath(Progress);
+		return;
+	}
+
+	double MeshDiameter = OriginalMesh->GetBounds().MaxDim();
+	double UVScaleFactor = 1.0 / MeshDiameter;
 
 	bool bCollapseDegenerateEdges = true; // TODO make this optional?
 	
 	FFrame3d Frame = PolygonFrame;
-	Frame.Origin = Frame.Origin + (2*MeshRadius*Frame.Z());
+	Frame.Origin = Frame.Origin + (2*MeshDiameter*Frame.Z());
 
 	FPolygon2d Polygon = GetPolygon();
 	double Perimeter = Polygon.Perimeter();
@@ -154,6 +280,11 @@ void FEmbedPolygonsOp::CalculateResult(FProgressCancel* Progress)
 		}
 	}
 
+	if (Progress->Cancelled())
+	{
+		return;
+	}
+
 	enum EDeleteMethod
 	{
 		DeleteNone,
@@ -287,6 +418,11 @@ void FEmbedPolygonsOp::CalculateResult(FProgressCancel* Progress)
 		return bResult;
 	};
 
+	if (Progress->Cancelled())
+	{
+		return;
+	}
+
 	EDeleteMethod DeleteMethod = EDeleteMethod::DeleteInside;
 	if (Operation == EEmbeddedPolygonOpMethod::InsertPolygon)
 	{
@@ -304,6 +440,7 @@ void FEmbedPolygonsOp::CalculateResult(FProgressCancel* Progress)
 		RecordEmbeddedEdges(PathVertIDs);
 		if (!bCutSide1 || PathVertIDs.Num() < 2)
 		{
+			EdgesOnFailure = EmbeddedEdges;
 			return;
 		}
 	}
@@ -318,6 +455,7 @@ void FEmbedPolygonsOp::CalculateResult(FProgressCancel* Progress)
 		RecordEmbeddedEdges(AllPathVertIDs[1]);
 		if (!bCutSide2 || AllPathVertIDs[0].Num() < 2 || AllPathVertIDs[1].Num() < 2)
 		{
+			EdgesOnFailure = EmbeddedEdges;
 			return;
 		}
 		FDynamicMeshEditor MeshEditor(ResultMesh.Get());
@@ -370,7 +508,7 @@ void FEmbedPolygonsOp::CalculateResult(FProgressCancel* Progress)
 		//}
 	//}
 
-	bEmbedSucceeded = true;
+	bOperationSucceeded = true;
 }
 
 void FEmbedPolygonsOp::RecordEmbeddedEdges(TArray<int>& PathVertIDs)
diff --git a/Engine/Plugins/Experimental/MeshModelingToolset/Source/ModelingOperatorsEditorOnly/Public/CuttingOps/EmbedPolygonsOp.h b/Engine/Plugins/Experimental/MeshModelingToolset/Source/ModelingOperatorsEditorOnly/Public/CuttingOps/EmbedPolygonsOp.h
--- a/Engine/Plugins/Experimental/MeshModelingToolset/Source/ModelingOperatorsEditorOnly/Public/CuttingOps/EmbedPolygonsOp.h
+++ b/Engine/Plugins/Experimental/MeshModelingToolset/Source/ModelingOperatorsEditorOnly/Public/CuttingOps/EmbedPolygonsOp.h
@@ -31,6 +31,8 @@ public:
 	// inputs
 	FFrame3d PolygonFrame;
 	FPolygon2d EmbedPolygon;
+	bool bCutWithBoolean;
+	bool bAttemptFixHolesOnBoolean;
 
 	// TODO: switch to FGeneralPolygon2d?
 	FPolygon2d GetPolygon()
@@ -54,11 +56,14 @@ public:
 
 
 	// Outputs
+	TArray<int> EdgesOnFailure; // edges to highlight on failure, to help visualize what happened (partial cut edges or hole edges)
 	TArray<int> EmbeddedEdges;
-	bool bEmbedSucceeded = false;
+	bool bOperationSucceeded;
 
 private:
 	void RecordEmbeddedEdges(TArray<int>& PathVertIDs);
+
+	void BooleanPath(FProgressCancel* Progress);
 };
 
 
