GeometryProcessing: convert check to ensure in MeshOcclusionMapBaker
#rnx

diff --git a/Engine/Plugins/Experimental/GeometryProcessing/Source/DynamicMesh/Private/Sampling/MeshOcclusionMapBaker.cpp b/Engine/Plugins/Experimental/GeometryProcessing/Source/DynamicMesh/Private/Sampling/MeshOcclusionMapBaker.cpp
--- a/Engine/Plugins/Experimental/GeometryProcessing/Source/DynamicMesh/Private/Sampling/MeshOcclusionMapBaker.cpp
+++ b/Engine/Plugins/Experimental/GeometryProcessing/Source/DynamicMesh/Private/Sampling/MeshOcclusionMapBaker.cpp
@@ -68,7 +68,7 @@ void FMeshOcclusionMapBaker::Bake()
 			for (FVector3d SphereDir : RayDirections)
 			{
 				FRay3d OcclusionRay(DetailPos, SurfaceFrame.FromFrameVector(SphereDir));
-				check(OcclusionRay.Direction.Dot(DetailTriNormal) > 0);
+				ensure(OcclusionRay.Direction.Dot(DetailTriNormal) > 0);
 
 				// Have weight of point fall off as it becomes more coplanar with face. 
 				// This reduces faceting artifacts that we would otherwise see because geometry does not vary smoothly
