NormalFlowRemesher: do extra iterations of normal flow without remeshing. Don't terminate remeshing iterations if mesh quality is good but normal flow is still moving vertices.

#rnx

diff --git a/Engine/Plugins/Experimental/GeometryProcessing/Source/DynamicMesh/Private/NormalFlowRemesher.cpp b/Engine/Plugins/Experimental/GeometryProcessing/Source/DynamicMesh/Private/NormalFlowRemesher.cpp
--- a/Engine/Plugins/Experimental/GeometryProcessing/Source/DynamicMesh/Private/NormalFlowRemesher.cpp
+++ b/Engine/Plugins/Experimental/GeometryProcessing/Source/DynamicMesh/Private/NormalFlowRemesher.cpp
@@ -39,8 +39,12 @@ void FNormalFlowRemesher::RemeshWithFaceProjection()
 	// to where they will ultimately want to go
 
 	const double OriginalSmoothSpeed = SmoothSpeedT;
+	
+	int Iterations = 0;
+	const double ProjectionDistanceThreshold = 0.1 * MinEdgeLength;
 
-	for (int k = 0; k < MaxRemeshIterations; ++k)
+	bool bContinue = true;
+	while (bContinue)
 	{
 		if (Cancelled())
 		{
@@ -49,17 +53,24 @@ void FNormalFlowRemesher::RemeshWithFaceProjection()
 
 		RemeshIteration();
 
-		if (k > MaxRemeshIterations / 2)
+		if (Iterations > MaxRemeshIterations / 2)
 		{
 			SmoothSpeedT *= 0.9;
 		}
+
+		double MaxProjectionDistance = 0.0;
+		TrackedFaceProjectionPass(MaxProjectionDistance);
+
+		// Stop if we've hit max iterations, or both:
+		// - queue is empty and
+		// - projection isn't moving anything
+		bContinue = (Iterations++ < MaxRemeshIterations) && ((ModifiedEdges->Num() > 0) || (MaxProjectionDistance > ProjectionDistanceThreshold));
 	}
 
 	SmoothSpeedT = OriginalSmoothSpeed;
 
-	// Now face projection iterations.
-
-	if (ProjTarget != nullptr && ProjectionMode == ETargetProjectionMode::AfterRefinement)
+	// Now just face projections and edge flips
+	if (ProjTarget != nullptr)
 	{
 		for (int k = 0; k < NumExtraProjectionIterations; ++k)
 		{
@@ -68,18 +79,23 @@ void FNormalFlowRemesher::RemeshWithFaceProjection()
 				break;
 			}
 
-			TrackedFaceProjectionPass();
-		}
+			double MaxProjectionDistance = 0.0;
+			TrackedFaceProjectionPass(MaxProjectionDistance);
 
-		TrackedFullProjectionPass(true);
+			if (MaxProjectionDistance == 0.0)
+			{
+				break;
+			}
+
+			// See if we can flip edges to improve normal fit
+			TrackedEdgeFlipPass();
+		}
 	}
 
-	// See if we can flip edges to improve normal fit
-	TrackedEdgeFlipPass();
 }
 
 
-void FNormalFlowRemesher::TrackedFaceProjectionPass()
+void FNormalFlowRemesher::TrackedFaceProjectionPass(double& MaxDistanceMoved)
 {
 	ensure(ProjTarget != nullptr);
 
@@ -135,6 +151,8 @@ void FNormalFlowRemesher::TrackedFaceProjectionPass()
 	// did not actually move. We also queue any edges that moved far enough to fall
 	// under min/max edge length thresholds
 
+	MaxDistanceMoved = 0.0;
+
 	for (int VertexID : Mesh->VertexIndicesItr())
 	{
 		TempFlagBuffer[VertexID] = false;
@@ -162,6 +180,8 @@ void FNormalFlowRemesher::TrackedFaceProjectionPass()
 			continue;
 		}
 
+		MaxDistanceMoved = FMath::Max(MaxDistanceMoved, CurrentPosition.Distance(ProjectedPosition));
+
 		TempFlagBuffer[VertexID] = true;
 		TempPosBuffer[VertexID] = ProjectedPosition;
 
@@ -289,7 +309,6 @@ void FNormalFlowRemesher::TrackedEdgeFlipPass()
 				QueueOneRing(OpposingEdgeVertices.A);
 				QueueOneRing(OpposingEdgeVertices.B);
 				OnEdgeFlip(EdgeID, FlipInfo);
-				DoDebugChecks();
 			}
 		}
 	}
diff --git a/Engine/Plugins/Experimental/GeometryProcessing/Source/DynamicMesh/Public/NormalFlowRemesher.h b/Engine/Plugins/Experimental/GeometryProcessing/Source/DynamicMesh/Public/NormalFlowRemesher.h
--- a/Engine/Plugins/Experimental/GeometryProcessing/Source/DynamicMesh/Public/NormalFlowRemesher.h
+++ b/Engine/Plugins/Experimental/GeometryProcessing/Source/DynamicMesh/Public/NormalFlowRemesher.h
@@ -37,14 +37,15 @@ protected:
 	void RemeshWithFaceProjection();
 
 	// Perform face-aligned projection onto the target mesh. Queue edges whose lengths change because of it.
-	void TrackedFaceProjectionPass();
+	void TrackedFaceProjectionPass(double& MaxDistanceMoved);
 
 	// This is called during RemeshIteration 
 	void TrackedFullProjectionPass(bool bParallel) override
 	{
 		for (int i = 0; i < FaceProjectionPassesPerRemeshIteration; ++i)
 		{
-			TrackedFaceProjectionPass();
+			double ProjectionDistance;
+			TrackedFaceProjectionPass(ProjectionDistance);
 		}
 	}
 
diff --git a/Engine/Plugins/Experimental/MeshModelingToolset/Source/MeshModelingTools/Private/RemeshMeshTool.cpp b/Engine/Plugins/Experimental/MeshModelingToolset/Source/MeshModelingTools/Private/RemeshMeshTool.cpp
--- a/Engine/Plugins/Experimental/MeshModelingToolset/Source/MeshModelingTools/Private/RemeshMeshTool.cpp
+++ b/Engine/Plugins/Experimental/MeshModelingToolset/Source/MeshModelingTools/Private/RemeshMeshTool.cpp
@@ -56,6 +56,8 @@ URemeshMeshToolProperties::URemeshMeshToolProperties()
 	TargetTriangleCount = 5000;
 	SmoothingStrength = 0.25;
 	RemeshIterations = 20;
+	MaxRemeshIterations = 20;
+	ExtraProjectionIterations = 5;
 	bDiscardAttributes = false;
 	RemeshType = ERemeshType::Standard;
 	SmoothingType = ERemeshSmoothingType::MeanValue;
@@ -199,6 +201,8 @@ TUniquePtr<FDynamicMeshOperator> URemeshMeshTool::MakeNewOperator()
 	Op->bReproject = BasicProperties->bReproject;
 	Op->bSplits = BasicProperties->bSplits;
 	Op->RemeshIterations = BasicProperties->RemeshIterations;
+	Op->MaxRemeshIterations = BasicProperties->MaxRemeshIterations;
+	Op->ExtraProjectionIterations = BasicProperties->ExtraProjectionIterations;
 	Op->SmoothingStrength = BasicProperties->SmoothingStrength;
 	Op->SmoothingType = BasicProperties->SmoothingType;
 
diff --git a/Engine/Plugins/Experimental/MeshModelingToolset/Source/MeshModelingTools/Public/RemeshMeshTool.h b/Engine/Plugins/Experimental/MeshModelingToolset/Source/MeshModelingTools/Public/RemeshMeshTool.h
--- a/Engine/Plugins/Experimental/MeshModelingToolset/Source/MeshModelingTools/Public/RemeshMeshTool.h
+++ b/Engine/Plugins/Experimental/MeshModelingToolset/Source/MeshModelingTools/Public/RemeshMeshTool.h
@@ -74,6 +74,14 @@ public:
 	UPROPERTY(EditAnywhere, Category = Remeshing, AdvancedDisplay, meta = (EditCondition = "RemeshType == ERemeshType::FullPass", UIMin = "0", UIMax = "50", ClampMin = "0", ClampMax = "1000"))
 	int RemeshIterations;
 
+	/** Maximum number of Remeshing passes, for Remeshers that have convergence criteria */
+	UPROPERTY(EditAnywhere, Category = Remeshing, AdvancedDisplay, meta = (EditCondition = "RemeshType != ERemeshType::FullPass", UIMin = "0", UIMax = "200", ClampMin = "0", ClampMax = "200"))
+	int MaxRemeshIterations;
+
+	/** For NormalFlowRemesher: extra iterations of normal flow with no remeshing */
+	UPROPERTY(EditAnywhere, Category = Remeshing, AdvancedDisplay, meta = (EditCondition = "RemeshType == ERemeshType::NormalFlow", UIMin = "0", UIMax = "200", ClampMin = "0", ClampMax = "200"))
+	int ExtraProjectionIterations;
+
 	/** If true, the target count is ignored and the target edge length is used directly */
 	UPROPERTY(EditAnywhere, Category = Remeshing, AdvancedDisplay)
 	bool bUseTargetEdgeLength;
diff --git a/Engine/Plugins/Experimental/MeshModelingToolset/Source/ModelingOperators/Private/CleaningOps/RemeshMeshOp.cpp b/Engine/Plugins/Experimental/MeshModelingToolset/Source/ModelingOperators/Private/CleaningOps/RemeshMeshOp.cpp
--- a/Engine/Plugins/Experimental/MeshModelingToolset/Source/ModelingOperators/Private/CleaningOps/RemeshMeshOp.cpp
+++ b/Engine/Plugins/Experimental/MeshModelingToolset/Source/ModelingOperators/Private/CleaningOps/RemeshMeshOp.cpp
@@ -12,22 +12,30 @@
 #include "MeshNormals.h"
 #include "NormalFlowRemesher.h"
 
-namespace
+TUniquePtr<FRemesher> FRemeshMeshOp::CreateRemesher(ERemeshType Type, FDynamicMesh3* TargetMesh)
 {
-	TUniquePtr<FRemesher> RemesherFactory(ERemeshType Type, FDynamicMesh3* TargetMesh)
+	switch(Type)
 	{
-		switch(Type)
-		{
-		case ERemeshType::Standard:
-			return MakeUnique<FQueueRemesher>(TargetMesh);
-		case ERemeshType::FullPass:
-			return MakeUnique<FRemesher>(TargetMesh);
-		case ERemeshType::NormalFlow:
-			return MakeUnique<FNormalFlowRemesher>(TargetMesh);
-		default:
-			check(!"Encountered unexpected Remesh Type");
-			return nullptr;
-		}
+	case ERemeshType::Standard:
+	{
+		TUniquePtr<FQueueRemesher> QueueRemesher = MakeUnique<FQueueRemesher>(TargetMesh);
+		QueueRemesher->MaxRemeshIterations = MaxRemeshIterations;
+		return QueueRemesher;
+	}
+	case ERemeshType::FullPass:
+	{
+		return MakeUnique<FRemesher>(TargetMesh);
+	}
+	case ERemeshType::NormalFlow:
+	{
+		TUniquePtr<FNormalFlowRemesher> NormalFlowRemesher = MakeUnique<FNormalFlowRemesher>(TargetMesh);
+		NormalFlowRemesher->MaxRemeshIterations = MaxRemeshIterations;
+		NormalFlowRemesher->NumExtraProjectionIterations = ExtraProjectionIterations;
+		return NormalFlowRemesher;
+	}
+	default:
+		checkf(false, TEXT("Encountered unexpected Remesh Type"));
+		return nullptr;
 	}
 
 }
@@ -54,7 +62,7 @@ void FRemeshMeshOp::CalculateResult(FProgressCancel* Progress)
 
 	FDynamicMesh3* TargetMesh = ResultMesh.Get();
 
-	TUniquePtr<FRemesher> Remesher = RemesherFactory(RemeshType, TargetMesh);
+	TUniquePtr<FRemesher> Remesher = CreateRemesher(RemeshType, TargetMesh);
 
 	Remesher->bEnableSplits = bSplits;
 	Remesher->bEnableFlips = bFlips;
diff --git a/Engine/Plugins/Experimental/MeshModelingToolset/Source/ModelingOperators/Public/CleaningOps/RemeshMeshOp.h b/Engine/Plugins/Experimental/MeshModelingToolset/Source/ModelingOperators/Public/CleaningOps/RemeshMeshOp.h
--- a/Engine/Plugins/Experimental/MeshModelingToolset/Source/ModelingOperators/Public/CleaningOps/RemeshMeshOp.h
+++ b/Engine/Plugins/Experimental/MeshModelingToolset/Source/ModelingOperators/Public/CleaningOps/RemeshMeshOp.h
@@ -10,6 +10,7 @@
 #include "RemeshMeshOp.generated.h"
 
 class FDynamicMesh3;
+class FRemesher;
 
 template <class TriangleMeshType>
 class TMeshAABBTree3;
@@ -56,6 +57,8 @@ public:
 	ERemeshType RemeshType = ERemeshType::Standard;
 
 	int RemeshIterations;
+	int MaxRemeshIterations;
+	int ExtraProjectionIterations = 5;
 	float SmoothingStrength, TargetEdgeLength;
 	ERemeshSmoothingType SmoothingType;
 	bool bDiscardAttributes, bPreserveSharpEdges, bFlips, bSplits, bCollapses, bReproject, bPreventNormalFlips;
@@ -70,6 +73,11 @@ public:
 	// FDynamicMeshOperator implementation
 	// 
 	virtual void CalculateResult(FProgressCancel* Progress) override;
+
+protected:
+
+	TUniquePtr<FRemesher> CreateRemesher(ERemeshType Type, FDynamicMesh3* TargetMesh);
+
 };
 
 
