GeometryProcessing: fix typo in typedef
#rnx

diff --git a/Engine/Plugins/Experimental/GeometryProcessing/Source/GeometricObjects/Public/Intersection/IntrLine2Line2.h b/Engine/Plugins/Experimental/GeometryProcessing/Source/GeometricObjects/Public/Intersection/IntrLine2Line2.h
--- a/Engine/Plugins/Experimental/GeometryProcessing/Source/GeometricObjects/Public/Intersection/IntrLine2Line2.h
+++ b/Engine/Plugins/Experimental/GeometryProcessing/Source/GeometricObjects/Public/Intersection/IntrLine2Line2.h
@@ -183,5 +183,5 @@ public:
 
 };
 
-typedef TIntrLine2Line2<double> FIntrLine2TLine2d;
+typedef TIntrLine2Line2<double> FIntrLine2Line2d;
 typedef TIntrLine2Line2<float> FIntrLine2Line2f;
