for self-union, check if winding would remove the triangle before checking for coplanarity.  Otherwise, we could erroneously keep geometry that is 'inside' the combined surface just because it's coplanar with other geometry.

#rnx

diff --git a/Engine/Plugins/Experimental/GeometryProcessing/Source/DynamicMesh/Private/Operations/MeshSelfUnion.cpp b/Engine/Plugins/Experimental/GeometryProcessing/Source/DynamicMesh/Private/Operations/MeshSelfUnion.cpp
--- a/Engine/Plugins/Experimental/GeometryProcessing/Source/DynamicMesh/Private/Operations/MeshSelfUnion.cpp
+++ b/Engine/Plugins/Experimental/GeometryProcessing/Source/DynamicMesh/Private/Operations/MeshSelfUnion.cpp
@@ -184,9 +184,15 @@ bool FMeshSelfUnion::Compute()
 			}
 			FVector3d Centroid = Mesh->GetTriCentroid(TID);
 
+			double WindingNum = Winding.FastWindingNumber(Centroid + Normals[TID] * NormalOffset);
+			bool bKeep = WindingNum < WindingThreshold; // keep if the outside of the tri is outside the shape
+			if (bTrimFlaps && bKeep) // trimming flaps == also check that the inside of the tri is inside the shape
+			{
+				bKeep = Winding.FastWindingNumber(Centroid - Normals[TID] * NormalOffset) > WindingThreshold;
+			}
 			
-			
-			// first check for the coplanar case
+			// if triangle is a candidate for keeping, check for the coplanar case
+			if (bKeep)
 			{
 				double DSq;
 				int MyComponentID = TriToComponentID[TID];
@@ -222,23 +228,18 @@ bool FMeshSelfUnion::Compute()
 							}
 							else
 							{
+								// for two coplanar components with matching normals,
+								// just keep tris from the component with lower ID
 								int OtherComponentID = TriToComponentID[OtherTID];
-								bool bKeep = MyComponentID < OtherComponentID;
-								KeepTri[TID] = bKeep;
+								bool bHasPriority = MyComponentID < OtherComponentID;
+								KeepTri[TID] = bHasPriority;
 							}
 							return;
 						}
 					}
 				}
 			}
-
-			// didn't already return a coplanar result; use the winding number
-			double WindingNum = Winding.FastWindingNumber(Centroid + Normals[TID] * NormalOffset);
-			bool bKeep = WindingNum < WindingThreshold;
-			if (bTrimFlaps && bKeep)
-			{
-				bKeep = Winding.FastWindingNumber(Centroid - Normals[TID] * NormalOffset) > WindingThreshold;
-			}
+			// didn't already return a coplanar result; use the winding-number-based decision
 			KeepTri[TID] = bKeep;
 		});
 
diff --git a/Engine/Plugins/Experimental/GeometryProcessing/Source/DynamicMesh/Public/Operations/MeshSelfUnion.h b/Engine/Plugins/Experimental/GeometryProcessing/Source/DynamicMesh/Public/Operations/MeshSelfUnion.h
--- a/Engine/Plugins/Experimental/GeometryProcessing/Source/DynamicMesh/Public/Operations/MeshSelfUnion.h
+++ b/Engine/Plugins/Experimental/GeometryProcessing/Source/DynamicMesh/Public/Operations/MeshSelfUnion.h
@@ -39,7 +39,7 @@ public:
 	double SnapTolerance = FMathf::ZeroTolerance * 1.0;
 
 	/** Amount we nudge samples off the surface when evaluating winding number, to avoid numerical issues */
-	double NormalOffset = FMathd::ZeroTolerance;
+	double NormalOffset = FMathf::ZeroTolerance * 1.0;
 
 	/** Threshold to determine whether triangle in one mesh is inside or outside of the other */
 	double WindingThreshold = .5;
