add option to shrink (erode) the set of triangles removed by the jacketing tool

diff --git a/Engine/Plugins/Experimental/GeometryProcessing/Source/DynamicMesh/Public/Operations/RemoveOccludedTriangles.h b/Engine/Plugins/Experimental/GeometryProcessing/Source/DynamicMesh/Public/Operations/RemoveOccludedTriangles.h
--- a/Engine/Plugins/Experimental/GeometryProcessing/Source/DynamicMesh/Public/Operations/RemoveOccludedTriangles.h
+++ b/Engine/Plugins/Experimental/GeometryProcessing/Source/DynamicMesh/Public/Operations/RemoveOccludedTriangles.h
@@ -77,15 +77,16 @@ public:
 	virtual ~TRemoveOccludedTriangles() {}
 
 	/**
-	 * Remove the occluded triangles, considering the given occluder AABB trees (which may represent more geometry than a single mesh)
+	 * Select the occluded triangles, considering the given occluder AABB trees (which may represent more geometry than a single mesh)
 	 * See simpler invocations below for the single instance case or the case where you'd like the spatial data structures built for you
+	 * Selection will be stored in the RemovedT array, but no triangles will be removed.
 	 *
 	 * @param MeshLocalToOccluderSpaces Transforms to take instances of the local mesh into the space of the occluders
 	 * @param Spatials AABB trees for all occluders
 	 * @param FastWindingTrees Precomputed fast winding trees for occluders
 	 * @return true on success
 	 */
-	virtual bool Apply(const TArrayView<const FTransform3d> MeshLocalToOccluderSpaces, 
+	virtual bool Select(const TArrayView<const FTransform3d> MeshLocalToOccluderSpaces,
 		const TArrayView<TMeshAABBTree3<OccluderTriangleMeshType>*> Spatials, const TArrayView<TFastWindingTree<OccluderTriangleMeshType>*> FastWindingTrees)
 	{
 		if (Cancelled())
@@ -238,7 +239,7 @@ public:
 			return false;
 		}
 
-		RemovedT.Empty();
+		RemovedT.Reset();
 		for (int TID = 0; TID < Mesh->MaxTriangleID(); TID++)
 		{
 			if (TriOccluded[TID])
@@ -247,6 +248,14 @@ public:
 			}
 		}
 
+		return true;
+	}
+
+	/**
+	 * Remove triangles that were selected (the triangle IDs in the RemoveT array)
+	 */
+	virtual bool RemoveSelected()
+	{
 		if (RemovedT.Num() > 0)
 		{
 			FDynamicMeshEditor Editor(Mesh);
@@ -261,6 +270,45 @@ public:
 
 		return true;
 	}
+	
+	/**
+	 * Remove the occluded triangles, considering the given occluder AABB trees (which may represent more geometry than a single mesh)
+	 * See simpler invocations below for the single instance case or the case where you'd like the spatial data structures built for you
+	 *
+	 * @param MeshLocalToOccluderSpaces Transforms to take instances of the local mesh into the space of the occluders
+	 * @param Spatials AABB trees for all occluders
+	 * @param FastWindingTrees Precomputed fast winding trees for occluders
+	 * @return true on success
+	 */
+	virtual bool Apply(const TArrayView<const FTransform3d> MeshLocalToOccluderSpaces, 
+		const TArrayView<TMeshAABBTree3<OccluderTriangleMeshType>*> Spatials, const TArrayView<TFastWindingTree<OccluderTriangleMeshType>*> FastWindingTrees)
+	{
+		if (!Select(MeshLocalToOccluderSpaces, Spatials, FastWindingTrees))
+		{
+			return false;
+		}
+
+		return RemoveSelected();
+	}
+
+
+	/**
+	 * Select the occluded triangles, considering the given occluder AABB tree (which may represent more geometry than a single mesh)
+	 * See simpler invocations below for the single instance case or the case where you'd like the spatial data structures built for you
+	 *
+	 * @param MeshLocalToOccluderSpaces Transforms to take instances of the local mesh into the space of the occluders
+	 * @param Spatials AABB trees for all occluders
+	 * @param FastWindingTrees Precomputed fast winding trees for occluders
+	 * @return true on success
+	 */
+	virtual bool Select(const TArrayView<const FTransform3d> MeshLocalToOccluderSpaces,
+		TMeshAABBTree3<OccluderTriangleMeshType>* Spatial, TFastWindingTree<OccluderTriangleMeshType>* FastWindingTree)
+	{
+		TArrayView<TMeshAABBTree3<OccluderTriangleMeshType>*> Spatials(&Spatial, 1);
+		TArrayView<TFastWindingTree<OccluderTriangleMeshType>*> FastWindingTrees(&FastWindingTree, 1);
+		return Select(MeshLocalToOccluderSpaces, Spatials, FastWindingTrees);
+	}
+
 
 	/**
 	 * Remove the occluded triangles, considering the given occluder AABB tree (which may represent more geometry than a single mesh)
@@ -279,20 +327,48 @@ public:
 		return Apply(MeshLocalToOccluderSpaces, Spatials, FastWindingTrees);
 	}
 
+	/**
+	 * Select the occluded triangles -- single instance case
+	 *
+	 * @param LocalToWorld Transform to take the local mesh into the space of the occluder geometry
+	 * @param Occluder AABB tree of occluding geometry
+	 * @return true on success
+	 */
+	virtual bool Select(const FTransform3d& MeshLocalToOccluderSpace, TMeshAABBTree3<OccluderTriangleMeshType>* Spatial, TFastWindingTree<OccluderTriangleMeshType>* FastWindingTree)
+	{
+		TArrayView<const FTransform3d> MeshLocalToOccluderSpaces(&MeshLocalToOccluderSpace, 1); // array view of the single transform
+		return Select(MeshLocalToOccluderSpaces, Spatial, FastWindingTree);
+	}
+
 
 	/**
-	* Remove the occluded triangles -- single instance case
-	*
-	* @param LocalToWorld Transform to take the local mesh into the space of the occluder geometry
-	* @param Occluder AABB tree of occluding geometry
-	* @return true on success
-	*/
+	 * Remove the occluded triangles -- single instance case
+	 *
+	 * @param LocalToWorld Transform to take the local mesh into the space of the occluder geometry
+	 * @param Occluder AABB tree of occluding geometry
+	 * @return true on success
+	 */
 	virtual bool Apply(const FTransform3d& MeshLocalToOccluderSpace, TMeshAABBTree3<OccluderTriangleMeshType>* Spatial, TFastWindingTree<OccluderTriangleMeshType>* FastWindingTree)
 	{
 		TArrayView<const FTransform3d> MeshLocalToOccluderSpaces(&MeshLocalToOccluderSpace, 1); // array view of the single transform
 		return Apply(MeshLocalToOccluderSpaces, Spatial, FastWindingTree);
 	}
 
+
+	/**
+	 * Select the occluded triangles -- single instance case w/out precomputed winding tree
+	 *
+	 * @param LocalToWorld Transform to take the local mesh into the space of the occluder geometry
+	 * @param Occluder AABB tree of occluding geometry
+	 * @return true on success
+	 */
+	virtual bool Select(const FTransform3d& MeshLocalToOccluderSpace, TMeshAABBTree3<OccluderTriangleMeshType>* Occluder)
+	{
+		TFastWindingTree<OccluderTriangleMeshType> FastWindingTree(Occluder, InsideMode == EOcclusionCalculationMode::FastWindingNumber);
+		return Select(MeshLocalToOccluderSpace, Occluder, &FastWindingTree);
+	}
+
+
 	/**
 	 * Remove the occluded triangles -- single instance case w/out precomputed winding tree
 	 *
diff --git a/Engine/Plugins/Experimental/GeometryProcessing/Source/DynamicMesh/Public/Selections/MeshFaceSelection.h b/Engine/Plugins/Experimental/GeometryProcessing/Source/DynamicMesh/Public/Selections/MeshFaceSelection.h
--- a/Engine/Plugins/Experimental/GeometryProcessing/Source/DynamicMesh/Public/Selections/MeshFaceSelection.h
+++ b/Engine/Plugins/Experimental/GeometryProcessing/Source/DynamicMesh/Public/Selections/MeshFaceSelection.h
@@ -433,36 +433,54 @@ public:
 	 *  
 	 *  Return false from FilterF to prevent triangles from being deselected.
 	 */
-	void ContractBorderByOneRingNeighbours()
+	void ContractBorderByOneRingNeighbours(int NumRings = 1, bool bContractFromMeshBoundary = false, const TUniqueFunction<bool(int)>& FilterF = nullptr)
 	{
 		TArray<int> BorderVIDs;   // border vertices
+		// TODO: track the border across removals if NumRings > 1, so we don't have to iterate over the whole selection each time
+		for (int RingIdx = 0; RingIdx < NumRings; ++RingIdx)
+		{
+			BorderVIDs.Reset();
 
-		// [TODO] border vertices are pushed onto the temp list multiple times.
-		// minor inefficiency, but maybe we could improve it?
+			// [TODO] border vertices are pushed onto the temp list multiple times.
+			// minor inefficiency, but maybe we could improve it?
 
-		// find set of vertices on border
-		for (int tid : Selected)
-		{
-			FIndex3i tri_v = Mesh->GetTriangle(tid);
-			for (int j = 0; j < 3; ++j)
+			// find set of vertices on border
+			for (int TID : Selected)
 			{
-				int vid = tri_v[j];
-				for (int nbr_t : Mesh->VtxTrianglesItr(vid))
+				FIndex3i TriV = Mesh->GetTriangle(TID);
+				for (int Idx = 0; Idx < 3; ++Idx)
 				{
-					if (!IsSelected(nbr_t))
+					int VID = TriV[Idx];
+					int UnusedE1, UnusedE2; // Unused edge boundary indices, required to call GetVtxBoundaryEdges to test if vertex is on boundary
+					bool bIsBorder = bContractFromMeshBoundary && Mesh->GetVtxBoundaryEdges(VID, UnusedE1, UnusedE2) > 0;
+					if (!bIsBorder)
+					{
+						for (int NbrT : Mesh->VtxTrianglesItr(VID))
+						{
+							if (!IsSelected(NbrT))
+							{
+								bIsBorder = true;
+								break;
+							}
+						}
+					}
+					if (bIsBorder)
 					{
-						BorderVIDs.Add(vid);
-						break;
+						BorderVIDs.Add(VID);
 					}
 				}
 			}
-		}
 
-		for (int vid : BorderVIDs)
-		{
-			for (int nbr_t : Mesh->VtxTrianglesItr(vid))
+			for (int VID : BorderVIDs)
 			{
-				Deselect(nbr_t);
+				for (int NbrT : Mesh->VtxTrianglesItr(VID))
+				{
+					if (FilterF && !FilterF(NbrT))
+					{
+						continue;
+					}
+					Deselect(NbrT);
+				}
 			}
 		}
 
@@ -471,6 +489,7 @@ public:
 
 
 
+
 	/**
 	 *  Grow selection outwards from seed triangle, until it hits boundaries defined by triangle and edge filters.
 	 *  Edge filter is not effective unless it (possibly combined w/ triangle filter) defines closed loops.
diff --git a/Engine/Plugins/Experimental/MeshModelingToolset/Source/MeshModelingTools/Private/RemoveOccludedTrianglesTool.cpp b/Engine/Plugins/Experimental/MeshModelingToolset/Source/MeshModelingTools/Private/RemoveOccludedTrianglesTool.cpp
--- a/Engine/Plugins/Experimental/MeshModelingToolset/Source/MeshModelingTools/Private/RemoveOccludedTrianglesTool.cpp
+++ b/Engine/Plugins/Experimental/MeshModelingToolset/Source/MeshModelingTools/Private/RemoveOccludedTrianglesTool.cpp
@@ -325,6 +325,8 @@ TUniquePtr<FDynamicMeshOperator> URemoveOccludedTrianglesOperatorFactory::MakeNe
 
 	Op->AddTriangleSamples = Tool->BasicProperties->AddTriangleSamples;
 
+	Op->ShrinkRemoval = Tool->BasicProperties->ShrinkRemoval;
+
 	Op->MinAreaConnectedComponent = Tool->BasicProperties->MinAreaIsland;
 
 	Op->MinTriCountConnectedComponent = Tool->BasicProperties->MinTriCountIsland;
diff --git a/Engine/Plugins/Experimental/MeshModelingToolset/Source/MeshModelingTools/Public/RemoveOccludedTrianglesTool.h b/Engine/Plugins/Experimental/MeshModelingToolset/Source/MeshModelingTools/Public/RemoveOccludedTrianglesTool.h
--- a/Engine/Plugins/Experimental/MeshModelingToolset/Source/MeshModelingTools/Public/RemoveOccludedTrianglesTool.h
+++ b/Engine/Plugins/Experimental/MeshModelingToolset/Source/MeshModelingTools/Public/RemoveOccludedTrianglesTool.h
@@ -85,6 +85,10 @@ public:
 	UPROPERTY(EditAnywhere, Category = OcclusionCalculation)
 	bool bOnlySelfOcclude = false;
 
+	/** Shrink (erode) the boundary of the set of triangles to remove. */
+	UPROPERTY(EditAnywhere, Category = OcclusionCalculation, meta = (UIMin = "0", ClampMin = "0"))
+	int ShrinkRemoval = 0;
+
 	UPROPERTY(EditAnywhere, Category = RemoveIslands, meta = (UIMin = "0", ClampMin = "0"))
 	double MinAreaIsland = 0;
 
diff --git a/Engine/Plugins/Experimental/MeshModelingToolset/Source/ModelingOperators/Private/CleaningOps/RemoveOccludedTrianglesOp.cpp b/Engine/Plugins/Experimental/MeshModelingToolset/Source/ModelingOperators/Private/CleaningOps/RemoveOccludedTrianglesOp.cpp
--- a/Engine/Plugins/Experimental/MeshModelingToolset/Source/ModelingOperators/Private/CleaningOps/RemoveOccludedTrianglesOp.cpp
+++ b/Engine/Plugins/Experimental/MeshModelingToolset/Source/ModelingOperators/Private/CleaningOps/RemoveOccludedTrianglesOp.cpp
@@ -9,6 +9,7 @@
 
 #include "Operations/RemoveOccludedTriangles.h"
 #include "Selections/MeshConnectedComponents.h"
+#include "Selections/MeshFaceSelection.h"
 
 
 #include "Async/ParallelFor.h"
@@ -56,6 +57,14 @@ void FRemoveOccludedTrianglesOp::CalculateResult(FProgressCancel* Progress)
 	bool bDiscardAttributes = false;
 	ResultMesh->Copy(*OriginalMesh, true, true, true, !bDiscardAttributes);
 	
+	auto ShrinkSelection = [](FDynamicMesh3& Mesh, TArray<int>& SelectedTris, int NumShrinks)
+	{
+		FMeshFaceSelection Selection(&Mesh);
+		Selection.Select(SelectedTris);
+		Selection.ContractBorderByOneRingNeighbours(NumShrinks);
+		SelectedTris = Selection.AsArray();
+	};
+
 	if (bOnlySelfOcclude)
 	{
 		TRemoveOccludedTriangles<FDynamicMesh3> SelfJacket(ResultMesh.Get());
@@ -71,7 +80,12 @@ void FRemoveOccludedTrianglesOp::CalculateResult(FProgressCancel* Progress)
 		SelfJacket.NormalOffset = NormalOffset;
 		SelfJacket.AddRandomRays = AddRandomRays;
 		SelfJacket.AddTriangleSamples = AddTriangleSamples;
-		SelfJacket.Apply(FTransform3d::Identity(), &SelfAABB, &SelfFWTree);
+		SelfJacket.Select(FTransform3d::Identity(), &SelfAABB, &SelfFWTree);
+		if (ShrinkRemoval > 0)
+		{
+			ShrinkSelection(*ResultMesh.Get(), SelfJacket.RemovedT, ShrinkRemoval);
+		}
+		SelfJacket.RemoveSelected();
 	}
 	else
 	{
@@ -88,7 +102,12 @@ void FRemoveOccludedTrianglesOp::CalculateResult(FProgressCancel* Progress)
 		Jacket.NormalOffset = NormalOffset;
 		Jacket.AddRandomRays = AddRandomRays;
 		Jacket.AddTriangleSamples = AddTriangleSamples;
-		Jacket.Apply(MeshTransforms, &CombinedMeshTrees->AABB, &CombinedMeshTrees->FastWinding);
+		Jacket.Select(MeshTransforms, &CombinedMeshTrees->AABB, &CombinedMeshTrees->FastWinding);
+		if (ShrinkRemoval > 0)
+		{
+			ShrinkSelection(*ResultMesh.Get(), Jacket.RemovedT, ShrinkRemoval);
+		}
+		Jacket.RemoveSelected();
 	}
 
 	if (MinTriCountConnectedComponent > 0 || MinAreaConnectedComponent > 0)
diff --git a/Engine/Plugins/Experimental/MeshModelingToolset/Source/ModelingOperators/Public/CleaningOps/RemoveOccludedTrianglesOp.h b/Engine/Plugins/Experimental/MeshModelingToolset/Source/ModelingOperators/Public/CleaningOps/RemoveOccludedTrianglesOp.h
--- a/Engine/Plugins/Experimental/MeshModelingToolset/Source/ModelingOperators/Public/CleaningOps/RemoveOccludedTrianglesOp.h
+++ b/Engine/Plugins/Experimental/MeshModelingToolset/Source/ModelingOperators/Public/CleaningOps/RemoveOccludedTrianglesOp.h
@@ -62,6 +62,8 @@ public:
 	// if true, *ignore* the combined selected occluder data in CombinedMeshTrees, and only consider self-occlusion of this mesh
 	bool bOnlySelfOcclude = false;
 
+	int ShrinkRemoval = 0;
+
 	double MinAreaConnectedComponent = 0;
 
 	int MinTriCountConnectedComponent = 0;
