Geometry Tools: MeshDescriptionToDynamicMesh.  Fix bugs in conversion of MeshDescription.
 Handle case where some shared UVs are mising on the mesh description.
Also handle case in tangent transfer where the existence of invalid triangles in the MeshDescription would result in a lower number of triangles in the dynamic mesh.

#rnx

diff --git a/Engine/Plugins/Experimental/GeometryProcessing/Source/MeshConversion/Private/MeshDescriptionToDynamicMesh.cpp b/Engine/Plugins/Experimental/GeometryProcessing/Source/MeshConversion/Private/MeshDescriptionToDynamicMesh.cpp
--- a/Engine/Plugins/Experimental/GeometryProcessing/Source/MeshConversion/Private/MeshDescriptionToDynamicMesh.cpp
+++ b/Engine/Plugins/Experimental/GeometryProcessing/Source/MeshConversion/Private/MeshDescriptionToDynamicMesh.cpp
@@ -178,7 +178,8 @@ void FMeshDescriptionToDynamicMesh::Convert(const FMeshDescription* MeshIn, FDyn
 	AddedTriangles.SetNum(MeshIn->Triangles().Num());
 
 	// allocate the TriIDMap (maps from DynamicMesh triangle ID to MeshDescription FTriangleID)
-	TriIDMap.SetNumUninitialized(MeshIn->Triangles().Num());
+	// reserve space for the data without setting array num()
+	TriIDMap.Reserve(MeshIn->Triangles().Num());
 
 
 	// Iterate over triangles in the Mesh Description
@@ -298,7 +299,8 @@ void FMeshDescriptionToDynamicMesh::Convert(const FMeshDescription* MeshIn, FDyn
 
 		checkSlow(NewTriangleID >= 0);
 		AddedTriangles[NewTriangleID] = TriData;
-		TriIDMap[NewTriangleID] = TriangleID;
+		// .Insert(Value, Index)
+		TriIDMap.Insert(TriangleID, NewTriangleID);
 	
 	}
 
@@ -429,6 +431,8 @@ void FMeshDescriptionToDynamicMesh::Convert(const FMeshDescription* MeshIn, FDyn
 					}
 
 					// copy uv "index buffer"
+					FUVID InvalidUVID(INDEX_NONE);
+					TArray<int32> TrisMissingUVs; 
 					for (int32 TriID : MeshOut.TriangleIndicesItr())
 					{
 						FTriangleID TriangleID = TriIDMap[TriID];
@@ -436,6 +440,13 @@ void FMeshDescriptionToDynamicMesh::Convert(const FMeshDescription* MeshIn, FDyn
 						// Don't change the UVIndices values!
 						TArrayView<FUVID> UVIndices = const_cast<FMeshDescription*>(MeshIn)->GetTriangleUVIndices(TriangleID, UVLayerIndex);
 
+						if (UVIndices[0] == InvalidUVID || UVIndices[1] == InvalidUVID || UVIndices[2] == InvalidUVID)
+						{
+							// keep track of the tris that don't have UVs in the mesh description shared UV channel.
+							TrisMissingUVs.Add(TriID);
+							continue;
+						}
+
 						// translate to Overlay indicies 
 						FIndex3i TriUV( UVIndexMap[UVIndices[0].GetValue()],
 							            UVIndexMap[UVIndices[1].GetValue()],
@@ -452,8 +463,8 @@ void FMeshDescriptionToDynamicMesh::Convert(const FMeshDescription* MeshIn, FDyn
 								int32 ParentVID = UVOverlay->GetParentVertex(TriUV[i]);
 								if (ParentVID != FDynamicMesh3::InvalidID && ParentVID != ParentTriangle[i])
 								{
-									const FVector2D UVvalue = UVCoordinates[UVIndices[1]];
-									TriUV[i] = UVIndexMap.Add(UVOverlay->AppendElement(FVector2f(UVvalue)));
+									const FVector2D UVvalue = UVCoordinates[UVIndices[i]];
+									TriUV[i] = UVOverlay->AppendElement(FVector2f(UVvalue));
 								}
 							}
 
@@ -464,9 +475,9 @@ void FMeshDescriptionToDynamicMesh::Convert(const FMeshDescription* MeshIn, FDyn
 							if (TriUV[0] == TriUV[1] && TriUV[0] == TriUV[2])
 							{
 								const FVector2D UVvalue = UVCoordinates[UVIndices[1]];
-								TriUV[0] = UVIndexMap.Add(UVOverlay->AppendElement(FVector2f(UVvalue)));
-								TriUV[1] = UVIndexMap.Add(UVOverlay->AppendElement(FVector2f(UVvalue)));
-								TriUV[2] = UVIndexMap.Add(UVOverlay->AppendElement(FVector2f(UVvalue)));
+								TriUV[0] = UVOverlay->AppendElement(FVector2f(UVvalue));
+								TriUV[1] = UVOverlay->AppendElement(FVector2f(UVvalue));
+								TriUV[2] = UVOverlay->AppendElement(FVector2f(UVvalue));
 
 							}
 							else
@@ -474,20 +485,20 @@ void FMeshDescriptionToDynamicMesh::Convert(const FMeshDescription* MeshIn, FDyn
 								if (TriUV[0] == TriUV[1])
 								{
 									const FVector2D UVvalue = UVCoordinates[UVIndices[0]];
-									TriUV[0] = UVIndexMap.Add(UVOverlay->AppendElement(FVector2f(UVvalue)));
-									TriUV[1] = UVIndexMap.Add(UVOverlay->AppendElement(FVector2f(UVvalue)));
+									TriUV[0] = UVOverlay->AppendElement(FVector2f(UVvalue));
+									TriUV[1] = UVOverlay->AppendElement(FVector2f(UVvalue));
 								}
 								if (TriUV[0] == TriUV[2])
 								{
 									const FVector2D UVvalue = UVCoordinates[UVIndices[0]];
-									TriUV[0] = UVIndexMap.Add(UVOverlay->AppendElement(FVector2f(UVvalue)));
-									TriUV[2] = UVIndexMap.Add(UVOverlay->AppendElement(FVector2f(UVvalue)));
+									TriUV[0] = UVOverlay->AppendElement(FVector2f(UVvalue));
+									TriUV[2] = UVOverlay->AppendElement(FVector2f(UVvalue));
 								}
 								if (TriUV[1] == TriUV[2])
 								{
 									const FVector2D UVvalue = UVCoordinates[UVIndices[1]];
-									TriUV[1] = UVIndexMap.Add(UVOverlay->AppendElement(FVector2f(UVvalue)));
-									TriUV[2] = UVIndexMap.Add(UVOverlay->AppendElement(FVector2f(UVvalue)));
+									TriUV[1] = UVOverlay->AppendElement(FVector2f(UVvalue));
+									TriUV[2] = UVOverlay->AppendElement(FVector2f(UVvalue));
 								}
 							}
 						}
@@ -495,6 +506,31 @@ void FMeshDescriptionToDynamicMesh::Convert(const FMeshDescription* MeshIn, FDyn
 						// set the triangle in the overlay
 						UVOverlay->SetTriangle(TriID, TriUV);
 					}
+				
+					// If some of the mesh description triangles were missing shared UVs, 
+					// use the per-vertex UVs on the mesh description and weld these.
+					// NB: This can happen when multiple meshes with different UV layer count are "combined" during import
+					// NB: these mesh description per-vertex UVs always exist, but may default to zero
+					// we may want to revisit this in the future and not set UVs in the Overlay for these triangles,
+					// but currently some of our tools (e.g. inspector) assume the every mesh triangle has an overlay triangle.
+					if (TrisMissingUVs.Num() != 0)
+					{
+						FUVWelder UVWelder;
+						UVWelder.UVOverlay = UVOverlay;
+						for (int32 TriangleID : TrisMissingUVs)
+						{
+							FIndex3i Tri = MeshOut.GetTriangle(TriangleID);
+							const FTriData& TriData = AddedTriangles[TriangleID];
+							FIndex3i TriUV;
+							for (int j = 0; j < 3; ++j)
+							{
+								FVector2D UV = InstanceUVs.Get(TriData.TriInstances[j], UVLayerIndex);
+								TriUV[j] = UVWelder.FindOrAddUnique(UV, Tri[j]);
+							}
+							UVOverlay->SetTriangle(TriangleID, TriUV);
+						}
+
+					}
 				}
 			});
 			Pending.Add(MoveTemp(UVFuture));
