Geometry Tools:  QEM simplification - make seam quadrics memoryless.  This was done because the accumulation of seam quadrics with very large meshes ( measured triangle count ) can lead to bad results during edge collapse.

#rnx

diff --git a/Engine/Plugins/Experimental/GeometryProcessing/Source/DynamicMesh/Private/MeshSimplification.cpp b/Engine/Plugins/Experimental/GeometryProcessing/Source/DynamicMesh/Private/MeshSimplification.cpp
--- a/Engine/Plugins/Experimental/GeometryProcessing/Source/DynamicMesh/Private/MeshSimplification.cpp
+++ b/Engine/Plugins/Experimental/GeometryProcessing/Source/DynamicMesh/Private/MeshSimplification.cpp
@@ -4,6 +4,7 @@
 #include "DynamicMeshAttributeSet.h"
 #include "Util/IndexUtil.h"
 #include "Async/ParallelFor.h"
+#include "Templates/UnrealTypeTraits.h"
 
 
 
@@ -160,22 +161,49 @@ void TMeshSimplification<QuadricErrorType>::InitializeVertexQuadrics()
 		//check(TMathUtil.EpsilonEqual(0, vertQuadrics[i].Evaluate(Mesh->GetVertex(i)), TMathUtil.Epsilon * 10));
 	}
 
-	// for each seam edge, add the seam quadric to its verts.
-	for (auto& seamQuadric : seamQuadrics)
+	// Fully memoryless simplifier directly adds seam quadrics to the vertex quadrics so it can more effiently construct
+	// edge quadrics from vertex quadrics only. 
+	// We don't do this for the standard QEM (with memory) because the memory only applies to the vertex quadrics and not the seam Qs
+	if (TIsSame<QuadricErrorType, FAttrBasedQuadricErrord>::Value)
 	{
-		int eid = seamQuadric.Key;
-		FIndex2i vids = Mesh->GetEdgeV(eid);
+		// for each seam edge, add the seam quadric to its verts.
+		for (auto& seamQuadric : seamQuadrics)
+		{
+			int eid = seamQuadric.Key;
+			FIndex2i vids = Mesh->GetEdgeV(eid);
 
-		vertQuadrics[vids[0]].AddSeamQuadric(seamQuadric.Value);
-		vertQuadrics[vids[1]].AddSeamQuadric(seamQuadric.Value);
+			vertQuadrics[vids[0]].AddSeamQuadric(seamQuadric.Value);
+			vertQuadrics[vids[1]].AddSeamQuadric(seamQuadric.Value);
+		}
 	}
-
 }
 
 template <typename QuadricErrorType>
 QuadricErrorType TMeshSimplification<QuadricErrorType>::AssembleEdgeQuadric(const FDynamicMesh3::FEdge& edge) const
 {
-	return QuadricErrorType(vertQuadrics[edge.Vert.A], vertQuadrics[edge.Vert.B]);
+	//  form standard edge quadric as sum of the vertex quadrics for the edge endpoints
+	QuadricErrorType EdgeQuadric(vertQuadrics[edge.Vert.A], vertQuadrics[edge.Vert.B]);
+	
+	if (bAllowSeamCollapse)
+	{ 
+		// lambda that adds any adjacent seam quadrics to the edge quadric
+		auto AddSeamQuadricsToEdge = [&, this](int vid)
+		{
+			for (int eid : Mesh->VtxEdgesItr(vid))
+			{
+				if (const FSeamQuadricType* seamQuadric =  seamQuadrics.Find(eid))
+				{
+					EdgeQuadric.AddSeamQuadric(*seamQuadric);
+				}
+			}
+		};
+	
+		// accumulate any adjacent seam quadrics onto this edge quadric.
+		AddSeamQuadricsToEdge(edge.Vert.A);
+		AddSeamQuadricsToEdge(edge.Vert.B);
+	}
+
+	return EdgeQuadric;
 }
 
 template<>
@@ -349,31 +377,29 @@ FVector3d TMeshSimplification<QuadricErrorType>::OptimalPoint(int eid, const FQu
 
 
 
-// update queue weight for each edge in vertex one-ring
+// update queue weight for each edge in vertex one-ring. Retains vertex Quadric memory
 template <>
-void DYNAMICMESH_API TMeshSimplification<FQuadricErrord>::UpdateNeighbours(const FDynamicMesh3::FEdgeCollapseInfo& collapseInfo)
+void DYNAMICMESH_API TMeshSimplification<FQuadricErrord>::UpdateNeighborhood(const FDynamicMesh3::FEdgeCollapseInfo& collapseInfo)
 {
-	int vid = collapseInfo.KeptVertex;
+	int kvid = collapseInfo.KeptVertex;
+	int rvid = collapseInfo.RemovedVertex;
 
+	// Quadric "memory"  the retained vertex quadric is the sum of the two vert quadrics
+	vertQuadrics[kvid] = FQuadricErrord(vertQuadrics[kvid], vertQuadrics[rvid]);
 
 	double EdgeWeight = this->SeamEdgeWeight;
 
-	for (int eid : Mesh->VtxEdgesItr(vid))
+	for (int eid : Mesh->VtxEdgesItr(kvid))
 	{
 		FDynamicMesh3::FEdge ne = Mesh->GetEdge(eid);
 
-		// update the seam quadric and vert quadric to reflect the new seams
+		// update the seam quadric to reflect changes in the seams
 		if (bAllowSeamCollapse)
 		{
 			
 			// need to recompute this seam quadric
 			if (FSeamQuadricType* seamQuadric = seamQuadrics.Find(eid))
 			{
-				// subtract the old seam quadric from adj verts
-
-				vertQuadrics[ne.Vert[0]].SubtractSeamQuadric(*seamQuadric);
-				vertQuadrics[ne.Vert[1]].SubtractSeamQuadric(*seamQuadric);
-
 				// rebuild the seam quadric
 
 				FVector3d p0 = Mesh->GetVertex(ne.Vert[0]);
@@ -392,10 +418,6 @@ void DYNAMICMESH_API TMeshSimplification<FQuadricErrord>::UpdateNeighbours(const
 				}
 
 				seamQuadric->Scale(EdgeWeight);
-
-				// add the seam quadric to the adj verts
-				vertQuadrics[ne.Vert[0]].AddSeamQuadric(*seamQuadric);
-				vertQuadrics[ne.Vert[1]].AddSeamQuadric(*seamQuadric);
 			}
 
 		}
@@ -418,7 +440,7 @@ void DYNAMICMESH_API TMeshSimplification<FQuadricErrord>::UpdateNeighbours(const
 
 // update queue weight for each edge in vertex one-ring.  Memoryless
 template <typename QuadricErrorType>
-void TMeshSimplification<QuadricErrorType>::UpdateNeighbours(const FDynamicMesh3::FEdgeCollapseInfo& collapseInfo)
+void TMeshSimplification<QuadricErrorType>::UpdateNeighborhood(const FDynamicMesh3::FEdgeCollapseInfo& collapseInfo)
 {
 	double EdgeWeight = this->SeamEdgeWeight;
 
@@ -704,8 +726,7 @@ void TMeshSimplification<QuadricErrorType>::DoSimplify()
 		ESimplificationResult result = CollapseEdge(eid, EdgeQuadrics[eid].collapse_pt, collapseInfo);
 		if (result == ESimplificationResult::Ok_Collapsed)
 		{
-			vertQuadrics[collapseInfo.KeptVertex] = EdgeQuadrics[eid].q;
-			UpdateNeighbours(collapseInfo);
+			UpdateNeighborhood(collapseInfo);
 		}
 	}
 	ProfileEndCollapse();
diff --git a/Engine/Plugins/Experimental/GeometryProcessing/Source/DynamicMesh/Public/MeshSimplification.h b/Engine/Plugins/Experimental/GeometryProcessing/Source/DynamicMesh/Public/MeshSimplification.h
--- a/Engine/Plugins/Experimental/GeometryProcessing/Source/DynamicMesh/Public/MeshSimplification.h
+++ b/Engine/Plugins/Experimental/GeometryProcessing/Source/DynamicMesh/Public/MeshSimplification.h
@@ -240,8 +240,8 @@ protected:
 	}
 
 
-	// update queue weight for each edge in vertex one-ring
-	virtual void UpdateNeighbours(const FDynamicMesh3::FEdgeCollapseInfo& collapseInfo);
+	// update queue weight for each edge in vertex one-ring and rebuild and quadrics necessary 
+	virtual void UpdateNeighborhood(const FDynamicMesh3::FEdgeCollapseInfo& collapseInfo);
 
 
 	virtual void Reproject() 
