Make DynamicMeshTriangleAttributes auto-initialize by default when constructed w/ a mesh, just like DynamicVertexAttributes already do

This fixes a crash bug w/ loading a dynamic mesh that has polygroup layers.  We could fix it more directly, but I can't think of a reason why it didn't auto-initialize by default, so this seems like a good direction ...

#rnx

diff --git a/Engine/Plugins/Experimental/GeometryProcessing/Source/DynamicMesh/Public/DynamicMeshTriangleAttribute.h b/Engine/Plugins/Experimental/GeometryProcessing/Source/DynamicMesh/Public/DynamicMeshTriangleAttribute.h
--- a/Engine/Plugins/Experimental/GeometryProcessing/Source/DynamicMesh/Public/DynamicMeshTriangleAttribute.h
+++ b/Engine/Plugins/Experimental/GeometryProcessing/Source/DynamicMesh/Public/DynamicMeshTriangleAttribute.h
@@ -67,9 +67,13 @@ public:
 	}
 
 	/** Create an overlay for the given parent mesh */
-	TDynamicMeshTriangleAttribute(FDynamicMesh3* ParentMeshIn)
+	TDynamicMeshTriangleAttribute(FDynamicMesh3* ParentMeshIn, bool bAutoInit = true)
 	{
 		ParentMesh = ParentMeshIn;
+		if (bAutoInit)
+		{
+			Initialize();
+		}
 	}
 
 private:
