Geometry Tools:  Allow DynamicMeshAttributeSet to have more than one "Normal" Attribute. This is flushing out the stubs of code that was already there.

#rnx

diff --git a/Engine/Plugins/Experimental/GeometryProcessing/Source/DynamicMesh/Private/DynamicMeshAttributeSet.cpp b/Engine/Plugins/Experimental/GeometryProcessing/Source/DynamicMesh/Private/DynamicMeshAttributeSet.cpp
--- a/Engine/Plugins/Experimental/GeometryProcessing/Source/DynamicMesh/Private/DynamicMeshAttributeSet.cpp
+++ b/Engine/Plugins/Experimental/GeometryProcessing/Source/DynamicMesh/Private/DynamicMeshAttributeSet.cpp
@@ -14,7 +14,11 @@ void FDynamicMeshAttributeSet::Copy(const FDynamicMeshAttributeSet& Copy)
 	{
 		UVLayers[UVIdx].Copy(Copy.UVLayers[UVIdx]);
 	}
-	Normals0.Copy(Copy.Normals0);
+	SetNumNormalLayers(Copy.NumNormalLayers());
+	for (int NormalLayerIndex = 0; NormalLayerIndex < NumNormalLayers(); NormalLayerIndex++)
+	{
+		NormalLayers[NormalLayerIndex].Copy(Copy.NormalLayers[NormalLayerIndex]);
+	}
 
 	if (Copy.MaterialIDAttrib)
 	{
@@ -52,10 +56,14 @@ bool FDynamicMeshAttributeSet::IsCompact()
 			return false;
 		}
 	}
-	if (!Normals0.IsCompact())
+	for (int NormalLayerIndex = 0; NormalLayerIndex < NumNormalLayers(); NormalLayerIndex++)
 	{
-		return false;
+		if (!NormalLayers[NormalLayerIndex].IsCompact())
+		{
+			return false;
+		}
 	}
+	
 	// material ID and generic per-element attributes currently cannot be non-compact
 	return true;
 }
@@ -69,7 +77,11 @@ void FDynamicMeshAttributeSet::CompactCopy(const FCompactMaps& CompactMaps, cons
 	{
 		UVLayers[UVIdx].CompactCopy(CompactMaps, Copy.UVLayers[UVIdx]);
 	}
-	Normals0.CompactCopy(CompactMaps, Copy.Normals0);
+	SetNumNormalLayers(Copy.NumNormalLayers());
+	for (int NormalLayerIndex = 0; NormalLayerIndex < NumNormalLayers(); NormalLayerIndex++)
+	{
+		NormalLayers[NormalLayerIndex].CompactCopy(CompactMaps, Copy.NormalLayers[NormalLayerIndex]);
+	}
 
 	if (Copy.MaterialIDAttrib)
 	{
@@ -106,7 +118,10 @@ void FDynamicMeshAttributeSet::CompactInPlace(const FCompactMaps& CompactMaps)
 	{
 		UVLayers[UVIdx].CompactInPlace(CompactMaps);
 	}
-	Normals0.CompactInPlace(CompactMaps);
+	for (int NormalLayerIndex = 0; NormalLayerIndex < NumNormalLayers(); NormalLayerIndex++)
+	{
+		NormalLayers[NormalLayerIndex].CompactInPlace(CompactMaps);
+	}
 
 	if (MaterialIDAttrib.IsValid())
 	{
@@ -129,15 +144,17 @@ void FDynamicMeshAttributeSet::CompactInPlace(const FCompactMaps& CompactMaps)
 void FDynamicMeshAttributeSet::SplitAllBowties(bool bParallel)
 {
 	int32 UVLayerCount = NumUVLayers();
-	ParallelFor(UVLayerCount + 1, [&](int32 idx)
+	int32 NormalLayerCount = NumNormalLayers();
+
+	ParallelFor(UVLayerCount + NormalLayerCount, [&](int32 idx)
 	{
-		if (idx == 0)
+		if (idx < NormalLayerCount)
 		{
-			PrimaryNormals()->SplitBowties();
+			GetNormalLayer(idx)->SplitBowties();
 		}
 		else
 		{
-			GetUVLayer(idx - 1)->SplitBowties();
+			GetUVLayer(idx - NormalLayerCount)->SplitBowties();
 		}
 	}, (bParallel) ? EParallelForFlags::Unbalanced : EParallelForFlags::ForceSingleThread );
 }
@@ -151,7 +168,11 @@ void FDynamicMeshAttributeSet::EnableMatchingAttributes(const FDynamicMeshAttrib
 	{
 		UVLayers[UVIdx].ClearElements();
 	}
-	Normals0.ClearElements();
+	SetNumNormalLayers(ToMatch.NumNormalLayers());
+	for (int NormalLayerIndex = 0; NormalLayerIndex < NumNormalLayers(); NormalLayerIndex++)
+	{
+		NormalLayers[NormalLayerIndex].ClearElements();
+	}
 
 	if (ToMatch.MaterialIDAttrib)
 	{
@@ -186,7 +207,10 @@ void FDynamicMeshAttributeSet::Reparent(FDynamicMesh3* NewParent)
 	{
 		UVLayers[UVIdx].Reparent(NewParent);
 	}
-	Normals0.Reparent(NewParent);
+	for (int NormalLayerIndex = 0; NormalLayerIndex < NumNormalLayers(); NormalLayerIndex++)
+	{
+		NormalLayers[NormalLayerIndex].Reparent(NewParent);
+	}
 
 	if (MaterialIDAttrib)
 	{
@@ -228,6 +252,26 @@ void FDynamicMeshAttributeSet::SetNumUVLayers(int Num)
 
 
 
+void FDynamicMeshAttributeSet::SetNumNormalLayers(int Num)
+{
+	if (NormalLayers.Num() == Num)
+	{
+		return;
+	}
+	if (Num >= NormalLayers.Num())
+	{
+		for (int32 i = NormalLayers.Num(); i < Num; ++i)
+		{
+			NormalLayers.Add(new FDynamicMeshNormalOverlay(ParentMesh));
+		}
+	}
+	else
+	{
+		NormalLayers.RemoveAt(Num, UVLayers.Num() - Num);
+	}
+	check(NormalLayers.Num() == Num);
+}
+
 
 
 int32 FDynamicMeshAttributeSet::NumPolygroupLayers() const
@@ -293,7 +337,16 @@ bool FDynamicMeshAttributeSet::IsSeamEdge(int eid) const
 			return true;
 		}
 	}
-	return Normals0.IsSeamEdge(eid);
+
+	for (const FDynamicMeshNormalOverlay& NormalLayer : NormalLayers)
+	{
+		if (NormalLayer.IsSeamEdge(eid))
+		{
+			return true;
+		}
+	}
+
+	return false;
 }
 
 
@@ -307,7 +360,16 @@ bool FDynamicMeshAttributeSet::IsSeamEdge(int EdgeID, bool& bIsUVSeamOut, bool&
 			bIsUVSeamOut = true;
 		}
 	}
-	bIsNormalSeamOut = Normals0.IsSeamEdge(EdgeID);
+
+	bIsNormalSeamOut = false;
+	for (const FDynamicMeshNormalOverlay& NormalLayer : NormalLayers)
+	{
+		if (NormalLayer.IsSeamEdge(EdgeID))
+		{
+			bIsNormalSeamOut = true;
+		}
+	}
+
 	return (bIsUVSeamOut || bIsNormalSeamOut);
 }
 
@@ -321,7 +383,14 @@ bool FDynamicMeshAttributeSet::IsSeamVertex(int VID, bool bBoundaryIsSeam) const
 			return true;
 		}
 	}
-	return Normals0.IsSeamVertex(VID, bBoundaryIsSeam);
+	for (const FDynamicMeshNormalOverlay& NormalLayer : NormalLayers)
+	{
+		if (NormalLayer.IsSeamVertex(VID, bBoundaryIsSeam))
+		{
+			return true;
+		}
+	}
+	return false;
 }
 
 bool FDynamicMeshAttributeSet::IsMaterialBoundaryEdge(int EdgeID) const
@@ -363,7 +432,10 @@ void FDynamicMeshAttributeSet::OnNewTriangle(int TriangleID, bool bInserted)
 	{
 		UVLayer.InitializeNewTriangle(TriangleID);
 	}
-	Normals0.InitializeNewTriangle(TriangleID);
+	for (FDynamicMeshNormalOverlay& NormalLayer : NormalLayers)
+	{
+		NormalLayer.InitializeNewTriangle(TriangleID);
+	}
 
 	if (MaterialIDAttrib)
 	{
@@ -387,7 +459,10 @@ void FDynamicMeshAttributeSet::OnRemoveTriangle(int TriangleID)
 	{
 		UVLayer.OnRemoveTriangle(TriangleID);
 	}
-	Normals0.OnRemoveTriangle(TriangleID);
+	for (FDynamicMeshNormalOverlay& NormalLayer : NormalLayers)
+	{
+		NormalLayer.OnRemoveTriangle(TriangleID);
+	}
 
 	// has no effect on MaterialIDAttrib
 }
@@ -400,7 +475,10 @@ void FDynamicMeshAttributeSet::OnReverseTriOrientation(int TriangleID)
 	{
 		UVLayer.OnReverseTriOrientation(TriangleID);
 	}
-	Normals0.OnReverseTriOrientation(TriangleID);
+	for (FDynamicMeshNormalOverlay& NormalLayer : NormalLayers)
+	{
+		NormalLayer.OnReverseTriOrientation(TriangleID);
+	}
 
 	// has no effect on MaterialIDAttrib
 }
@@ -413,7 +491,10 @@ void FDynamicMeshAttributeSet::OnSplitEdge(const FDynamicMesh3::FEdgeSplitInfo&
 	{
 		UVLayer.OnSplitEdge(SplitInfo);
 	}
-	Normals0.OnSplitEdge(SplitInfo);
+	for (FDynamicMeshNormalOverlay& NormalLayer : NormalLayers)
+	{
+		NormalLayer.OnSplitEdge(SplitInfo);
+	}
 
 	if (MaterialIDAttrib)
 	{
@@ -434,7 +515,10 @@ void FDynamicMeshAttributeSet::OnFlipEdge(const FDynamicMesh3::FEdgeFlipInfo & f
 	{
 		UVLayer.OnFlipEdge(flipInfo);
 	}
-	Normals0.OnFlipEdge(flipInfo);
+	for (FDynamicMeshNormalOverlay& NormalLayer : NormalLayers)
+	{
+		NormalLayer.OnFlipEdge(flipInfo);
+	}
 
 	if (MaterialIDAttrib)
 	{
@@ -456,7 +540,10 @@ void FDynamicMeshAttributeSet::OnCollapseEdge(const FDynamicMesh3::FEdgeCollapse
 	{
 		UVLayer.OnCollapseEdge(collapseInfo);
 	}
-	Normals0.OnCollapseEdge(collapseInfo);
+	for (FDynamicMeshNormalOverlay& NormalLayer : NormalLayers)
+	{
+		NormalLayer.OnCollapseEdge(collapseInfo);
+	}
 
 	if (MaterialIDAttrib)
 	{
@@ -477,7 +564,10 @@ void FDynamicMeshAttributeSet::OnPokeTriangle(const FDynamicMesh3::FPokeTriangle
 	{
 		UVLayer.OnPokeTriangle(pokeInfo);
 	}
-	Normals0.OnPokeTriangle(pokeInfo);
+	for (FDynamicMeshNormalOverlay& NormalLayer : NormalLayers)
+	{
+		NormalLayer.OnPokeTriangle(pokeInfo);
+	}
 
 	if (MaterialIDAttrib)
 	{
@@ -498,7 +588,10 @@ void FDynamicMeshAttributeSet::OnMergeEdges(const FDynamicMesh3::FMergeEdgesInfo
 	{
 		UVLayer.OnMergeEdges(mergeInfo);
 	}
-	Normals0.OnMergeEdges(mergeInfo);
+	for (FDynamicMeshNormalOverlay& NormalLayer : NormalLayers)
+	{
+		NormalLayer.OnMergeEdges(mergeInfo);
+	}
 
 	if (MaterialIDAttrib)
 	{
@@ -519,7 +612,10 @@ void FDynamicMeshAttributeSet::OnSplitVertex(const DynamicMeshInfo::FVertexSplit
 	{
 		UVLayer.OnSplitVertex(SplitInfo, TrianglesToUpdate);
 	}
-	Normals0.OnSplitVertex(SplitInfo, TrianglesToUpdate);
+	for (FDynamicMeshNormalOverlay& NormalLayer : NormalLayers)
+	{
+		NormalLayer.OnSplitVertex(SplitInfo, TrianglesToUpdate);
+	}
 
 	if (MaterialIDAttrib)
 	{
diff --git a/Engine/Plugins/Experimental/GeometryProcessing/Source/DynamicMesh/Private/DynamicMeshEditor.cpp b/Engine/Plugins/Experimental/GeometryProcessing/Source/DynamicMesh/Private/DynamicMeshEditor.cpp
--- a/Engine/Plugins/Experimental/GeometryProcessing/Source/DynamicMesh/Private/DynamicMeshEditor.cpp
+++ b/Engine/Plugins/Experimental/GeometryProcessing/Source/DynamicMesh/Private/DynamicMeshEditor.cpp
@@ -1241,17 +1241,18 @@ void FDynamicMeshEditor::InvertTriangleNormals(const TArray<int>& Triangles)
 
 	if (Mesh->HasAttributes())
 	{
-		for (FDynamicMeshNormalOverlay* Normals : Mesh->Attributes()->GetAllNormalLayers())
+		for (int NormalLayerIndex = 0; NormalLayerIndex < Mesh->Attributes()->NumNormalLayers(); NormalLayerIndex++)
 		{
-			TBitArray<FDefaultBitArrayAllocator> DoneNormals(false, Normals->MaxElementID());
+			FDynamicMeshNormalOverlay* NormalOverlay = Mesh->Attributes()->GetNormalLayer(NormalLayerIndex);
+			TBitArray<FDefaultBitArrayAllocator> DoneNormals(false, NormalOverlay->MaxElementID());
 			for (int TriangleID : Triangles)
 			{
-				FIndex3i ElemTri = Normals->GetTriangle(TriangleID);
+				FIndex3i ElemTri = NormalOverlay->GetTriangle(TriangleID);
 				for (int j = 0; j < 3; ++j)
 				{
-					if (Normals->IsElement(ElemTri[j]) && DoneNormals[ElemTri[j]] == false)
+					if (NormalOverlay->IsElement(ElemTri[j]) && DoneNormals[ElemTri[j]] == false)
 					{
-						Normals->SetElement(ElemTri[j], -Normals->GetElement(ElemTri[j]));
+						NormalOverlay->SetElement(ElemTri[j], -NormalOverlay->GetElement(ElemTri[j]));
 						DoneNormals[ElemTri[j]] = true;
 					}
 				}
@@ -1466,22 +1467,24 @@ void FDynamicMeshEditor::AppendMesh(const FDynamicMesh3* AppendMesh,
 	}
 
 
-	// @todo support multiple UV/normal layer copying
 	// @todo can we have a template fn that does this?
 
 	if (AppendMesh->HasAttributes() && Mesh->HasAttributes())
 	{
-		const FDynamicMeshNormalOverlay* FromNormals = AppendMesh->Attributes()->PrimaryNormals();
-		FDynamicMeshNormalOverlay* ToNormals = Mesh->Attributes()->PrimaryNormals();
-		if (FromNormals != nullptr && ToNormals != nullptr)
-		{
-			FIndexMapi& NormalMap = IndexMapsOut.GetNormalMap(0);
-			NormalMap.Reserve(FromNormals->ElementCount());
-			AppendNormals(AppendMesh, FromNormals, ToNormals,
-				VertexMap, TriangleMap, NormalTransform, NormalMap);
+		int NumNormalLayers = FMath::Min(Mesh->Attributes()->NumNormalLayers(), AppendMesh->Attributes()->NumNormalLayers());
+		for (int NormalLayerIndex = 0; NormalLayerIndex < NumNormalLayers; NormalLayerIndex++)
+		{ 
+			const FDynamicMeshNormalOverlay* FromNormals = AppendMesh->Attributes()->GetNormalLayer(NormalLayerIndex);
+			FDynamicMeshNormalOverlay* ToNormals = Mesh->Attributes()->GetNormalLayer(NormalLayerIndex);
+			if (FromNormals != nullptr && ToNormals != nullptr)
+			{
+				FIndexMapi& NormalMap = IndexMapsOut.GetNormalMap(0);
+				NormalMap.Reserve(FromNormals->ElementCount());
+				AppendNormals(AppendMesh, FromNormals, ToNormals,
+					VertexMap, TriangleMap, NormalTransform, NormalMap);
+			}
 		}
 
-
 		int NumUVLayers = FMath::Min(Mesh->Attributes()->NumUVLayers(), AppendMesh->Attributes()->NumUVLayers());
 		for (int UVLayerIndex = 0; UVLayerIndex < NumUVLayers; UVLayerIndex++)
 		{
@@ -1656,8 +1659,6 @@ static void AppendTriangleAttributes(const FDynamicMesh3* FromMesh, int FromTria
 		return;
 	}
 
-	// todo: if we ever support multiple normal layers, copy them all
-	check(FromMesh->Attributes()->NumNormalLayers() == 1);
 
 	for (int UVLayerIndex = 0; UVLayerIndex < FMath::Min(FromMesh->Attributes()->NumUVLayers(), ToMesh->Attributes()->NumUVLayers()); UVLayerIndex++)
 	{
@@ -1678,10 +1679,10 @@ static void AppendTriangleAttributes(const FDynamicMesh3* FromMesh, int FromTria
 	}
 
 
-	const FDynamicMeshNormalOverlay* FromNormalOverlay = FromMesh->Attributes()->PrimaryNormals();
-	FDynamicMeshNormalOverlay* ToNormalOverlay = ToMesh->Attributes()->PrimaryNormals();
-
+	for (int NormalLayerIndex = 0; NormalLayerIndex < FMath::Min(FromMesh->Attributes()->NumNormalLayers(), ToMesh->Attributes()->NumNormalLayers()); NormalLayerIndex++)
 	{
+		const FDynamicMeshNormalOverlay* FromNormalOverlay = FromMesh->Attributes()->GetNormalLayer(NormalLayerIndex);
+		FDynamicMeshNormalOverlay* ToNormalOverlay = ToMesh->Attributes()->GetNormalLayer(NormalLayerIndex);
 		if (FromNormalOverlay->IsSetTriangle(FromTriangleID))
 		{
 			FIndex3i FromElemTri = FromNormalOverlay->GetTriangle(FromTriangleID);
diff --git a/Engine/Plugins/Experimental/GeometryProcessing/Source/DynamicMesh/Public/DynamicMeshAttributeSet.h b/Engine/Plugins/Experimental/GeometryProcessing/Source/DynamicMesh/Public/DynamicMeshAttributeSet.h
--- a/Engine/Plugins/Experimental/GeometryProcessing/Source/DynamicMesh/Public/DynamicMeshAttributeSet.h
+++ b/Engine/Plugins/Experimental/GeometryProcessing/Source/DynamicMesh/Public/DynamicMeshAttributeSet.h
@@ -24,7 +24,8 @@ typedef TDynamicMeshScalarTriangleAttribute<int32> FDynamicMeshPolygroupAttribut
  * FDynamicMeshAttributeSet manages a set of extended attributes for a FDynamicMesh3.
  * This includes UV and Normal overlays, etc.
  * 
- * Currently the default is to always have one UV layer and one Normal layer
+ * Currently the default is to always have one UV layer and one Normal layer, 
+ * but the number of layers can be requested on construction.
  * 
  * @todo current internal structure is a work-in-progress
  */
@@ -33,10 +34,17 @@ class DYNAMICMESH_API FDynamicMeshAttributeSet : public FDynamicMeshAttributeSet
 public:
 
 	FDynamicMeshAttributeSet(FDynamicMesh3* Mesh)
-		: ParentMesh(Mesh), Normals0(Mesh)
+		: ParentMesh(Mesh)
 	{
 		SetNumUVLayers(1);
-		NormalLayers.Add(&Normals0);
+		SetNumNormalLayers(1);
+	}
+
+	FDynamicMeshAttributeSet(FDynamicMesh3* Mesh, int32 NumUVLayers, int32 NumNormalLayers)
+		: ParentMesh(Mesh)
+	{
+		SetNumUVLayers(NumUVLayers);
+		SetNumNormalLayers(NumNormalLayers);
 	}
 
 	virtual ~FDynamicMeshAttributeSet()
@@ -99,10 +107,11 @@ public:
 	/** @return number of Normals layers */
 	virtual int NumNormalLayers() const
 	{
-		checkSlow(NormalLayers.Num() == 1);
-		return 1;
+		return NormalLayers.Num();
 	}
 
+	virtual void SetNumNormalLayers(int Num);
+
 	/** @return true if the given edge is a seam edge in any overlay */
 	virtual bool IsSeamEdge(int EdgeID) const;
 
@@ -119,27 +128,35 @@ public:
 	// UV Layers 
 	//
 
-	/** @return the UV layer at the given Index */
+	/** @return the UV layer at the given Index  if exists, else nullptr*/
 	FDynamicMeshUVOverlay* GetUVLayer(int Index)
 	{
-		return &UVLayers[Index];
+		if (Index < UVLayers.Num() && Index > -1)
+		{ 
+			return &UVLayers[Index];
+		}
+		return nullptr;
 	}
 
-	/** @return the UV layer at the given Index */
+	/** @return the UV layer at the given Index if exists, else nullptr */
 	const FDynamicMeshUVOverlay* GetUVLayer(int Index) const
 	{
-		return &UVLayers[Index];
+		if (Index < UVLayers.Num() && Index > -1)
+		{
+			return &UVLayers[Index];
+		}
+		return nullptr;
 	}
 
 	/** @return the primary UV layer (layer 0) */
 	FDynamicMeshUVOverlay* PrimaryUV()
 	{
-		return &UVLayers[0];
+		return GetUVLayer(0);
 	}
 	/** @return the primary UV layer (layer 0) */
 	const FDynamicMeshUVOverlay* PrimaryUV() const
 	{
-		return &UVLayers[0];
+		return GetUVLayer(0);
 	}
 
 
@@ -147,35 +164,63 @@ public:
 	// Normal Layers 
 	//
 
-	/** @return the Normal layer at the given Index */
+	/** @return the Normal layer at the given Index if exists, else nullptr */
 	FDynamicMeshNormalOverlay* GetNormalLayer(int Index)
 	{
-		return (Index == 0) ? &Normals0 : nullptr;
+		if (Index < NormalLayers.Num() && Index > -1)
+		{ 
+			return &NormalLayers[Index];
+		}
+		return nullptr;
 	}
 
-	/** @return the Normal layer at the given Index */
+	/** @return the Normal layer at the given Index if exists, else nullptr */
 	const FDynamicMeshNormalOverlay* GetNormalLayer(int Index) const
 	{
-		return (Index == 0) ? &Normals0 : nullptr;
+		if (Index < NormalLayers.Num() && Index > -1)
+		{
+			return &NormalLayers[Index];
+		}
+		return nullptr;
 	}
 
-	/** @return list of all Normal layers */
-	const TArray<FDynamicMeshNormalOverlay*>& GetAllNormalLayers() const
-	{
-		return NormalLayers;
-	}
 
-	/** @return the primary Normal layer (layer 0) */
+	/** @return the primary Normal layer (normal layer 0) if exists, else nullptr */
 	FDynamicMeshNormalOverlay* PrimaryNormals()
 	{
-		return &Normals0;
+		return GetNormalLayer(0);
 	}
-	/** @return the primary Normal layer (layer 0) */
+	/** @return the primary Normal layer (normal layer 0) if exists, else nullptr */
 	const FDynamicMeshNormalOverlay* PrimaryNormals() const
 	{
-		return &Normals0;
+		return GetNormalLayer(0);
+	}
+	/** @return the primary tangent layer ( normal layer 1) if exists, else nullptr */
+	FDynamicMeshNormalOverlay* PrimaryTangents()
+	{
+		return GetNormalLayer(1);
+	}
+	/** @return the primary tangent layer ( normal layer 1) if exists, else nullptr */
+	const FDynamicMeshNormalOverlay* PrimaryTangents() const 
+	{
+		return GetNormalLayer(1);
+	}
+	/** @return the primary biTangent layer ( normal layer 2) if exists, else nullptr */
+	FDynamicMeshNormalOverlay* PrimaryBiTangents()
+	{
+		return GetNormalLayer(2);
+	}
+	/** @return the primary biTangent layer ( normal layer 2) if exists, else nullptr */
+	const FDynamicMeshNormalOverlay* PrimaryBiTangents() const
+	{
+		return GetNormalLayer(2);
 	}
 
+	/** @return true if normal layers exist for the normal, tangent, and bitangent */
+	bool HasTangentSpace() const
+	{
+		return (PrimaryNormals() != nullptr && PrimaryTangents()  != nullptr && PrimaryBiTangents() != nullptr);
+	}
 
 
 	//
@@ -265,11 +310,9 @@ protected:
 	/** Parent mesh of this attribute set */
 	FDynamicMesh3* ParentMesh;
 
-	/** Default Normals layer */
-	FDynamicMeshNormalOverlay Normals0;
-
+	
 	TIndirectArray<FDynamicMeshUVOverlay> UVLayers;
-	TArray<FDynamicMeshNormalOverlay*> NormalLayers;
+	TIndirectArray<FDynamicMeshNormalOverlay> NormalLayers;
 
 	TUniquePtr<FDynamicMeshMaterialAttribute> MaterialIDAttrib;
 
@@ -289,7 +332,10 @@ protected:
 		{
 			UVLayer.InitializeTriangles(MaxTriangleID);
 		}
-		Normals0.InitializeTriangles(MaxTriangleID);
+		for (FDynamicMeshNormalOverlay& NormalLayer : NormalLayers)
+		{
+			NormalLayer.InitializeTriangles(MaxTriangleID);
+		}
 	}
 
 	// These functions are called by the FDynamicMesh3 to update the various
