add a number of robustness improvements to avoid crashes in voxel blend, voxel morphology, marching cubes and SDF code

includes
 + fix to AABB::Contain() function to not become an infinite bounding box when Contain() is called on an empty AABB.
 + ensure() and skip processing on invalid inputs to SDF, marching cubes, blend and morphology code
 + approx max resolution enforced for marching cubes and SDF grids
 + put range restrictions on offset for voxel tools

diff --git a/Engine/Plugins/Experimental/GeometryProcessing/Source/GeometricObjects/Public/BoxTypes.h b/Engine/Plugins/Experimental/GeometryProcessing/Source/GeometricObjects/Public/BoxTypes.h
--- a/Engine/Plugins/Experimental/GeometryProcessing/Source/GeometricObjects/Public/BoxTypes.h
+++ b/Engine/Plugins/Experimental/GeometryProcessing/Source/GeometricObjects/Public/BoxTypes.h
@@ -365,9 +365,12 @@ struct TAxisAlignedBox3
 
 	void Contain(const TAxisAlignedBox3<RealType>& Other)
 	{
-		// todo: can be optimized
-		Contain(Other.Min);
-		Contain(Other.Max);
+		Min.X = Min.X < Other.Min.X ? Min.X : Other.Min.X;
+		Min.Y = Min.Y < Other.Min.Y ? Min.Y : Other.Min.Y;
+		Min.Z = Min.Z < Other.Min.Z ? Min.Z : Other.Min.Z;
+		Max.X = Max.X > Other.Max.X ? Max.X : Other.Max.X;
+		Max.Y = Max.Y > Other.Max.Y ? Max.Y : Other.Max.Y;
+		Max.Z = Max.Z > Other.Max.Z ? Max.Z : Other.Max.Z;
 	}
 
 	bool Contains(const FVector3<RealType>& V) const
@@ -600,9 +603,10 @@ struct TAxisAlignedBox2
 
 	inline void Contain(const TAxisAlignedBox2<RealType>& Other)
 	{
-		// todo: can be optimized
-		Contain(Other.Min);
-		Contain(Other.Max);
+		Min.X = Min.X < Other.Min.X ? Min.X : Other.Min.X;
+		Min.Y = Min.Y < Other.Min.Y ? Min.Y : Other.Min.Y;
+		Max.X = Max.X > Other.Max.X ? Max.X : Other.Max.X;
+		Max.Y = Max.Y > Other.Max.Y ? Max.Y : Other.Max.Y;
 	}
 
 	void Contain(const TArray<FVector2<RealType>>& Pts)
diff --git a/Engine/Plugins/Experimental/GeometryProcessing/Source/GeometricObjects/Public/Generators/MarchingCubes.h b/Engine/Plugins/Experimental/GeometryProcessing/Source/GeometricObjects/Public/Generators/MarchingCubes.h
--- a/Engine/Plugins/Experimental/GeometryProcessing/Source/GeometricObjects/Public/Generators/MarchingCubes.h
+++ b/Engine/Plugins/Experimental/GeometryProcessing/Source/GeometricObjects/Public/Generators/MarchingCubes.h
@@ -57,7 +57,10 @@ public:
 	 */
 	bool bParallelCompute = true;
 
-
+	/**
+	 * Max number of cells on any dimension; if exceeded, CubeSize will be automatically increased to fix
+	 */
+	int SafetyMaxDimension = 4096;
 
 	/**
 	 *  Which rootfinding method will be used to converge on surface along edges
@@ -91,20 +94,25 @@ public:
 	{
 	}
 
-
+	bool Validate()
+	{
+		return CubeSize > 0 && FMath::IsFinite(CubeSize) && !Bounds.IsEmpty() && FMath::IsFinite(Bounds.MaxDim());
+	}
 
 	/**
 	*  Run MC algorithm and generate Output mesh
 	*/
 	FMeshShapeGenerator& Generate() override
 	{
-		int32 nx = (int32)(Bounds.Width() / CubeSize) + 1;
-		int32 ny = (int32)(Bounds.Height() / CubeSize) + 1;
-		int32 nz = (int32)(Bounds.Depth() / CubeSize) + 1;
-		CellDimensions = FVector3i(nx, ny, nz);
+		if (!ensure(Validate()))
+		{
+			return *this;
+		}
+
+		SetDimensions();
 		GridBounds = TAxisAlignedBox3<int>(FVector3i::Zero(), CellDimensions - FVector3i(1,1,1)); // grid bounds are inclusive
 
-		corner_values_grid = FDenseGrid3f(nx+1, ny+1, nz+1, FMathf::MaxReal);
+		corner_values_grid = FDenseGrid3f(CellDimensions.X+1, CellDimensions.Y+1, CellDimensions.Z+1, FMathf::MaxReal);
 		edge_vertices.Reset();
 		corner_values.Reset();
 
@@ -120,17 +128,17 @@ public:
 
 	FMeshShapeGenerator& GenerateContinuation(TArrayView<const FVector3<double>> Seeds)
 	{
-		int nx = (int)(Bounds.Width() / CubeSize) + 1;
-		int ny = (int)(Bounds.Height() / CubeSize) + 1;
-		int nz = (int)(Bounds.Depth() / CubeSize) + 1;
-		CellDimensions = FVector3i(nx, ny, nz);
-		GridBounds = TAxisAlignedBox3<int>(FVector3i::Zero(), CellDimensions - FVector3i(1,1,1)); // grid bounds are inclusive
+		if (!ensure(Validate()))
+		{
+			return *this;
+		}
 
-		UE_LOG(LogTemp, Warning, TEXT("marching cubes dims = %d %d %d"), nx, ny, nz);
+		SetDimensions();
+		GridBounds = TAxisAlignedBox3<int>(FVector3i::Zero(), CellDimensions - FVector3i(1,1,1)); // grid bounds are inclusive
 
 		if (LastGridBounds != GridBounds)
 		{
-			corner_values_grid = FDenseGrid3f(nx + 1, ny + 1, nz + 1, FMathf::MaxReal);
+			corner_values_grid = FDenseGrid3f(CellDimensions.X + 1, CellDimensions.Y + 1, CellDimensions.Z + 1, FMathf::MaxReal);
 			edge_vertices.Reset();
 			corner_values.Reset();
 			if (bParallelCompute)
@@ -176,7 +184,21 @@ protected:
 		double f[8];      // field values at corners
 	};
 
-
+	void SetDimensions()
+	{
+		int NX = (int)(Bounds.Width() / CubeSize) + 1;
+		int NY = (int)(Bounds.Height() / CubeSize) + 1;
+		int NZ = (int)(Bounds.Depth() / CubeSize) + 1;
+		int MaxDim = FMath::Max3(NX, NY, NZ);
+		if (!ensure(MaxDim <= SafetyMaxDimension))
+		{
+			CubeSize = Bounds.MaxDim() / double(SafetyMaxDimension - 1);
+			NX = (int)(Bounds.Width() / CubeSize) + 1;
+			NY = (int)(Bounds.Height() / CubeSize) + 1;
+			NZ = (int)(Bounds.Depth() / CubeSize) + 1;
+		}
+		CellDimensions = FVector3i(NX, NY, NZ);
+	}
 
 	void corner_pos(const FVector3i& IJK, FVector3<double>& Pos)
 	{
diff --git a/Engine/Plugins/Experimental/GeometryProcessing/Source/GeometricObjects/Public/Implicit/Blend.h b/Engine/Plugins/Experimental/GeometryProcessing/Source/GeometricObjects/Public/Implicit/Blend.h
--- a/Engine/Plugins/Experimental/GeometryProcessing/Source/GeometricObjects/Public/Implicit/Blend.h
+++ b/Engine/Plugins/Experimental/GeometryProcessing/Source/GeometricObjects/Public/Implicit/Blend.h
@@ -78,7 +78,7 @@ public:
 			}
 		}
 
-		bool bValidParams = BlendPower > 0 && BlendFalloff > 0 && GridCellSize > 0 && MeshCellSize > 0;
+		bool bValidParams = BlendPower > 0 && BlendFalloff > 0 && GridCellSize > 0 && MeshCellSize > 0 && FMath::IsFinite(GridCellSize) && FMath::IsFinite(MeshCellSize);
 		return bHasSourcesWithBounds && bValidParams;
 	}
 
diff --git a/Engine/Plugins/Experimental/GeometryProcessing/Source/GeometricObjects/Public/Implicit/CachingMeshSDF.h b/Engine/Plugins/Experimental/GeometryProcessing/Source/GeometricObjects/Public/Implicit/CachingMeshSDF.h
--- a/Engine/Plugins/Experimental/GeometryProcessing/Source/GeometricObjects/Public/Implicit/CachingMeshSDF.h
+++ b/Engine/Plugins/Experimental/GeometryProcessing/Source/GeometricObjects/Public/Implicit/CachingMeshSDF.h
@@ -49,6 +49,9 @@ public:
 	// should we try to compute signs? if not, Grid remains unsigned
 	bool bComputeSigns = true;
 
+	// If the number of cells in any dimension may exceed this, CellSize will be automatically increased to keep cell count reasonable
+	int ApproxMaxCellsPerDimension = 4096;
+
 	// What counts as "inside" the mesh. Crossing count does not use triangle
 	// Orientation, so inverted faces are fine, but overlapping shells or self intersections
 	// will be filled using even/odd rules (as seen along X axis...)
@@ -101,14 +104,49 @@ protected:
 	double MaxDistQueryDist;
 
 public:
+
+	bool Validate()
+	{
+		FAxisAlignedBox3d Bounds = Spatial->GetBoundingBox();
+
+		if (Bounds.IsEmpty() || !FMath::IsFinite(Bounds.MaxDim()))
+		{
+			return false;
+		}
+		if (CellSize <= 0 || !FMath::IsFinite(CellSize))
+		{
+			return false;
+		}
+		if (ApproxMaxCellsPerDimension < 5)
+		{
+			return false;
+		}
+		return true;
+	}
+
 	void Initialize()
 	{
+		if (!ensure(Validate()))
+		{
+			return;
+		}
+
 		// figure out Origin & dimensions
-		FAxisAlignedBox3d bounds = Spatial->GetBoundingBox();
+		FAxisAlignedBox3d Bounds = Spatial->GetBoundingBox();
+		
+		float MaxDim = (Bounds.Max - Bounds.Min + ExpandBounds * 2).MaxElement() + 4 * MaxOffsetDistance;
+		if (!ensureMsgf(MaxDim / CellSize <= ApproxMaxCellsPerDimension - 4, TEXT("SDF resolution clamped to avoid excessive memory use")))
+		{
+			CellSize = MaxDim / (ApproxMaxCellsPerDimension - 4);
+			if (!ensure(CellSize > 0 && FMath::IsFinite(CellSize)))
+			{
+				return;
+			}
+		}
 
 		float fBufferWidth = (float)FMath::Max(4 * CellSize, 2 * MaxOffsetDistance + 2 * CellSize);
-		GridOrigin = (FVector3f)bounds.Min - fBufferWidth * FVector3f::One() - (FVector3f)ExpandBounds;
-		FVector3f max = (FVector3f)bounds.Max + fBufferWidth * FVector3f::One() + (FVector3f)ExpandBounds;
+		GridOrigin = (FVector3f)Bounds.Min - fBufferWidth * FVector3f::One() - (FVector3f)ExpandBounds;
+		FVector3f max = (FVector3f)Bounds.Max + fBufferWidth * FVector3f::One() + (FVector3f)ExpandBounds;
 		int NI = (int)((max.X - GridOrigin.X) / CellSize) + 1;
 		int NJ = (int)((max.Y - GridOrigin.Y) / CellSize) + 1;
 		int NK = (int)((max.Z - GridOrigin.Z) / CellSize) + 1;
diff --git a/Engine/Plugins/Experimental/GeometryProcessing/Source/GeometricObjects/Public/Implicit/Morphology.h b/Engine/Plugins/Experimental/GeometryProcessing/Source/GeometricObjects/Public/Implicit/Morphology.h
--- a/Engine/Plugins/Experimental/GeometryProcessing/Source/GeometricObjects/Public/Implicit/Morphology.h
+++ b/Engine/Plugins/Experimental/GeometryProcessing/Source/GeometricObjects/Public/Implicit/Morphology.h
@@ -83,7 +83,7 @@ public:
 	bool Validate()
 	{
 		bool bValidMeshAndSpatial = Source != nullptr && SourceSpatial != nullptr && SourceSpatial->IsValid();
-		bool bValidParams = Distance > 0 && GridCellSize > 0 && MeshCellSize > 0;
+		bool bValidParams = Distance > 0 && GridCellSize > 0 && MeshCellSize > 0 && FMath::IsFinite(MeshCellSize);
 		return bValidMeshAndSpatial && bValidParams;
 	}
 
diff --git a/Engine/Plugins/Experimental/GeometryProcessing/Source/GeometricObjects/Public/Implicit/Solidify.h b/Engine/Plugins/Experimental/GeometryProcessing/Source/GeometricObjects/Public/Implicit/Solidify.h
--- a/Engine/Plugins/Experimental/GeometryProcessing/Source/GeometricObjects/Public/Implicit/Solidify.h
+++ b/Engine/Plugins/Experimental/GeometryProcessing/Source/GeometricObjects/Public/Implicit/Solidify.h
@@ -76,7 +76,7 @@ public:
 	{
 		bool bValidMeshAndSpatial = Source != nullptr && SourceSpatial != nullptr && SourceSpatial->IsValid();
 		bool bValidWinding = SourceWinding != nullptr;
-		bool bValidParams = SurfaceSearchSteps >= 0 && MeshCellSize > 0;
+		bool bValidParams = SurfaceSearchSteps >= 0 && MeshCellSize > 0 && FMath::IsFinite(MeshCellSize);
 		return bValidMeshAndSpatial && bValidWinding && bValidParams;
 	}
 
diff --git a/Engine/Plugins/Experimental/GeometryProcessing/Source/GeometricObjects/Public/Implicit/SweepingMeshSDF.h b/Engine/Plugins/Experimental/GeometryProcessing/Source/GeometricObjects/Public/Implicit/SweepingMeshSDF.h
--- a/Engine/Plugins/Experimental/GeometryProcessing/Source/GeometricObjects/Public/Implicit/SweepingMeshSDF.h
+++ b/Engine/Plugins/Experimental/GeometryProcessing/Source/GeometricObjects/Public/Implicit/SweepingMeshSDF.h
@@ -57,6 +57,10 @@ public:
 	// Most of this parallelizes very well, makes a huge speed difference
 	bool bUseParallel = true;
 
+	// If the number of cells in any dimension may exceed this, CellSize will be automatically increased to keep cell count reasonable
+	int ApproxMaxCellsPerDimension = 4096;
+
+
 	// The narrow band is always computed exactly, and the full Grid is always signed.
 	// Can also fill in the rest of the full Grid with fast sweeping. This is 
 	// quite computationally intensive, though, and not parallelizable 
@@ -157,6 +161,23 @@ public:
 		return ExactCells > 1 && EdgesPerBand >= .25;
 	}
 
+	bool Validate(const FAxisAlignedBox3d& Bounds)
+	{
+		if (Bounds.IsEmpty() || !FMath::IsFinite(Bounds.MaxDim()))
+		{
+			return false;
+		}
+		if (CellSize <= 0 || !FMath::IsFinite(CellSize))
+		{
+			return false;
+		}
+		if (ApproxMaxCellsPerDimension < 5)
+		{
+			return false;
+		}
+		return true;
+	}
+
 	/**
 	 * Compute the SDF
 	 * @param Bounds Bounding box for the mesh data (passed as param it is usually already available, depending on TriangleMeshType and whether an AABB tree was provided)
@@ -164,6 +185,26 @@ public:
 	 */
 	bool Compute(FAxisAlignedBox3d Bounds)
 	{
+		if (!ensure(Validate(Bounds)))
+		{
+			return false;
+		}
+
+		if (!ensureMsgf(ExactBandWidth*2 < ApproxMaxCellsPerDimension, TEXT("Cannot request band wider than half the max cells per dimension")))
+		{
+			ExactBandWidth = FMath::Max(ApproxMaxCellsPerDimension/2-1, 1);
+		}
+
+		float MaxDim = (Bounds.Max - Bounds.Min + ExpandBounds * 2).MaxElement();
+		if (!ensureMsgf(MaxDim / CellSize <= ApproxMaxCellsPerDimension - 2 * ExactBandWidth, TEXT("SDF resolution clamped to avoid excessive memory use")))
+		{
+			CellSize = MaxDim / (ApproxMaxCellsPerDimension - 2 * ExactBandWidth);
+			if (!ensure(CellSize > 0 && FMath::IsFinite(CellSize)))
+			{
+				return false;
+			}
+		}
+
 		float fBufferWidth = ExactBandWidth * CellSize;
 		if (ComputeMode == EComputeModes::NarrowBand_SpatialFloodFill)
 		{
diff --git a/Engine/Plugins/Experimental/MeshModelingToolset/Source/MeshModelingTools/Public/VoxelBlendMeshesTool.h b/Engine/Plugins/Experimental/MeshModelingToolset/Source/MeshModelingTools/Public/VoxelBlendMeshesTool.h
--- a/Engine/Plugins/Experimental/MeshModelingToolset/Source/MeshModelingTools/Public/VoxelBlendMeshesTool.h
+++ b/Engine/Plugins/Experimental/MeshModelingToolset/Source/MeshModelingTools/Public/VoxelBlendMeshesTool.h
@@ -35,7 +35,7 @@ public:
 	bool bRemoveInternalsAfterSolidify = false;
 
 	/** Offset surface to create when solidifying any open-boundary inputs; if 0 then no offset surfaces are created */
-	UPROPERTY(EditAnywhere, Category = Solidify, meta = (EditCondition = "bSolidifyInput == true"))
+	UPROPERTY(EditAnywhere, Category = Solidify, meta = (EditCondition = "bSolidifyInput == true", UIMin = "0", UIMax = "100", ClampMin = "0", ClampMax = "1000"))
 	double OffsetSolidifySurface = 0.0;
 };
 
diff --git a/Engine/Plugins/Experimental/MeshModelingToolset/Source/MeshModelingTools/Public/VoxelMorphologyMeshesTool.h b/Engine/Plugins/Experimental/MeshModelingToolset/Source/MeshModelingTools/Public/VoxelMorphologyMeshesTool.h
--- a/Engine/Plugins/Experimental/MeshModelingToolset/Source/MeshModelingTools/Public/VoxelMorphologyMeshesTool.h
+++ b/Engine/Plugins/Experimental/MeshModelingToolset/Source/MeshModelingTools/Public/VoxelMorphologyMeshesTool.h
@@ -36,7 +36,7 @@ public:
 	bool bRemoveInternalsAfterSolidify = false;
 
 	/** Offset surface to create when solidifying any open-boundary inputs; if 0 then no offset surfaces are created */
-	UPROPERTY(EditAnywhere, Category = Solidify, meta = (EditCondition = "bSolidifyInput == true"))
+	UPROPERTY(EditAnywhere, Category = Solidify, meta = (EditCondition = "bSolidifyInput == true", UIMin = "0", UIMax = "100", ClampMin = "0", ClampMax = "1000"))
 	double OffsetSolidifySurface = 0.0;
 };
 
