Polygroup helper improvments:

 - Change MaxGroupID on PolygroupSet to match how it's computed on FDynamicMesh3, to be 1+the actual max (i.e. to be a bound, so we can write < in the for loop, same as MaxVertexID and friends)
 - Update PolygroupUtil functions to work on const meshes where possible
 - Expand PolygroupUtil functions some so it's easier to work with polygroups in cases where we can't use PolygroupSet (mainly if the mesh is const)

#rnx

diff --git a/Engine/Plugins/Experimental/GeometryProcessing/Source/DynamicMesh/Private/Polygroups/PolygroupSet.cpp b/Engine/Plugins/Experimental/GeometryProcessing/Source/DynamicMesh/Private/Polygroups/PolygroupSet.cpp
--- a/Engine/Plugins/Experimental/GeometryProcessing/Source/DynamicMesh/Private/Polygroups/PolygroupSet.cpp
+++ b/Engine/Plugins/Experimental/GeometryProcessing/Source/DynamicMesh/Private/Polygroups/PolygroupSet.cpp
@@ -63,14 +63,14 @@ void FPolygroupSet::RecalculateMaxGroupID()
 	{
 		for (int32 tid : Mesh->TriangleIndicesItr())
 		{
-			MaxGroupID = FMath::Max(MaxGroupID, PolygroupAttrib->GetValue(tid));
+			MaxGroupID = FMath::Max(MaxGroupID, PolygroupAttrib->GetValue(tid) + 1);
 		}
 	}
 	else
 	{
 		for (int32 tid : Mesh->TriangleIndicesItr())
 		{
-			MaxGroupID = FMath::Max(MaxGroupID, Mesh->GetTriangleGroup(tid));
+			MaxGroupID = FMath::Max(MaxGroupID, Mesh->GetTriangleGroup(tid) + 1);
 		}
 	}
 }
\ No newline at end of file
diff --git a/Engine/Plugins/Experimental/GeometryProcessing/Source/DynamicMesh/Private/Polygroups/PolygroupUtil.cpp b/Engine/Plugins/Experimental/GeometryProcessing/Source/DynamicMesh/Private/Polygroups/PolygroupUtil.cpp
--- a/Engine/Plugins/Experimental/GeometryProcessing/Source/DynamicMesh/Private/Polygroups/PolygroupUtil.cpp
+++ b/Engine/Plugins/Experimental/GeometryProcessing/Source/DynamicMesh/Private/Polygroups/PolygroupUtil.cpp
@@ -20,9 +20,25 @@ FDynamicMeshPolygroupAttribute* UE::Geometry::FindPolygroupLayerByName(FDynamicM
 }
 
 
-int32 UE::Geometry::FindPolygroupLayerIndex(FDynamicMesh3& Mesh, FDynamicMeshPolygroupAttribute* Layer)
+int32 UE::Geometry::FindPolygroupLayerIndexByName(const FDynamicMesh3& Mesh, FName Name)
 {
-	FDynamicMeshAttributeSet* AttributeSet = Mesh.Attributes();
+	const FDynamicMeshAttributeSet* AttributeSet = Mesh.Attributes();
+	if (AttributeSet == nullptr) return -1;
+	int32 NumPolygroupLayers = AttributeSet->NumPolygroupLayers();
+	for (int32 k = 0; k < NumPolygroupLayers; ++k)
+	{
+		if (AttributeSet->GetPolygroupLayer(k)->GetName() == Name)
+		{
+			return k;
+		}
+	}
+	return -1;
+}
+
+
+int32 UE::Geometry::FindPolygroupLayerIndex(const FDynamicMesh3& Mesh, const FDynamicMeshPolygroupAttribute* Layer)
+{
+	const FDynamicMeshAttributeSet* AttributeSet = Mesh.Attributes();
 	if (AttributeSet == nullptr) return -1;
 	int32 NumPolygroupLayers = AttributeSet->NumPolygroupLayers();
 	for (int32 k = 0; k < NumPolygroupLayers; ++k)
@@ -33,4 +49,15 @@ int32 UE::Geometry::FindPolygroupLayerIndex(FDynamicMesh3& Mesh, FDynamicMeshPol
 		}
 	}
 	return -1;
+}
+
+
+int32 UE::Geometry::ComputeGroupIDBound(const FDynamicMesh3& Mesh, const FDynamicMeshPolygroupAttribute* Layer)
+{
+	int Bound = 0;
+	for (int TID : Mesh.TriangleIndicesItr())
+	{
+		Bound = FMath::Max(Bound, Layer->GetValue(TID) + 1);
+	}
+	return Bound;
 }
\ No newline at end of file
diff --git a/Engine/Plugins/Experimental/GeometryProcessing/Source/DynamicMesh/Public/Polygroups/PolygroupSet.h b/Engine/Plugins/Experimental/GeometryProcessing/Source/DynamicMesh/Public/Polygroups/PolygroupSet.h
--- a/Engine/Plugins/Experimental/GeometryProcessing/Source/DynamicMesh/Public/Polygroups/PolygroupSet.h
+++ b/Engine/Plugins/Experimental/GeometryProcessing/Source/DynamicMesh/Public/Polygroups/PolygroupSet.h
@@ -21,6 +21,7 @@ namespace Geometry
  * 
  * To support unique Polygroup ID allocation, FPolygroupSet calculates the maximum GroupID on creation, and 
  * updates this maximum across SetGroup() calls. AllocateNewGroupID() can be used to provide new unused GroupIDs.
+ * For consistency with FDynamicMesh3, MaxGroupID is set such that all GroupIDs are less than MaxGroupID
  *
  */
 struct DYNAMICMESH_API FPolygroupSet
@@ -28,7 +29,7 @@ struct DYNAMICMESH_API FPolygroupSet
 	FDynamicMesh3* Mesh = nullptr;
 	FDynamicMeshPolygroupAttribute* PolygroupAttrib = nullptr;
 	int32 GroupLayerIndex = -1;
-	int32 MaxGroupID = 0;
+	int32 MaxGroupID = 0; // Note: all group IDs are less than MaxGroupID
 
 	/** Initialize a PolygroupSet for the given Mesh, and standard triangle group layer */
 	explicit FPolygroupSet(FDynamicMesh3* MeshIn);
@@ -80,7 +81,7 @@ struct DYNAMICMESH_API FPolygroupSet
 				Mesh->SetTriangleGroup(TriangleID, NewGroupID);
 			}
 		}
-		MaxGroupID = FMath::Max(MaxGroupID, NewGroupID);
+		MaxGroupID = FMath::Max(MaxGroupID, NewGroupID + 1);
 	}
 
 	/**
@@ -93,7 +94,7 @@ struct DYNAMICMESH_API FPolygroupSet
 	 */
 	int32 AllocateNewGroupID()
 	{
-		return ++MaxGroupID;
+		return MaxGroupID++;
 	}
 };
 
diff --git a/Engine/Plugins/Experimental/GeometryProcessing/Source/DynamicMesh/Public/Polygroups/PolygroupUtil.h b/Engine/Plugins/Experimental/GeometryProcessing/Source/DynamicMesh/Public/Polygroups/PolygroupUtil.h
--- a/Engine/Plugins/Experimental/GeometryProcessing/Source/DynamicMesh/Public/Polygroups/PolygroupUtil.h
+++ b/Engine/Plugins/Experimental/GeometryProcessing/Source/DynamicMesh/Public/Polygroups/PolygroupUtil.h
@@ -17,11 +17,20 @@ namespace Geometry
  */
 DYNAMICMESH_API FDynamicMeshPolygroupAttribute* FindPolygroupLayerByName(FDynamicMesh3& Mesh, FName Name);
 
+/**
+ * @return index of the first Layer with the given FName in Mesh AttributeSet, or -1 if not found
+ */
+DYNAMICMESH_API int32 FindPolygroupLayerIndexByName(const FDynamicMesh3& Mesh, FName Name);
+
 /**
  * @return index of Layer in Mesh AttributeSet, or -1 if not found
  */
-DYNAMICMESH_API int32 FindPolygroupLayerIndex(FDynamicMesh3& Mesh, FDynamicMeshPolygroupAttribute* Layer);
+DYNAMICMESH_API int32 FindPolygroupLayerIndex(const FDynamicMesh3& Mesh, const FDynamicMeshPolygroupAttribute* Layer);
 
+/**
+ * @return Compute the upper bound of the group IDs for the given layer (one greater than the highest group ID)
+ */
+DYNAMICMESH_API int32 ComputeGroupIDBound(const FDynamicMesh3& Mesh, const FDynamicMeshPolygroupAttribute* Layer);
 
 }	// end namespace Geometry
 }	// end namespace UE
\ No newline at end of file
