Fixing PVS 7.7 Issues:

warning V1047: Lifetime of the lambda is greater than lifetime of the local variable 'XYZ' captured by reference.

diff --git a/Engine/Plugins/Experimental/GeometryProcessing/Source/DynamicMesh/Private/QueueRemesher.cpp b/Engine/Plugins/Experimental/GeometryProcessing/Source/DynamicMesh/Private/QueueRemesher.cpp
--- a/Engine/Plugins/Experimental/GeometryProcessing/Source/DynamicMesh/Private/QueueRemesher.cpp
+++ b/Engine/Plugins/Experimental/GeometryProcessing/Source/DynamicMesh/Private/QueueRemesher.cpp
@@ -85,7 +85,7 @@ int FQueueRemesher::FastSplitIteration()
 	double MaxEdgeLengthSqr = MaxEdgeLength * MaxEdgeLength;
 
 	// Callback after an edge is split: enqueue new edges whose lengths are long enough
-	PostEdgeSplitFunction = [this, &MaxEdgeLengthSqr](int EdgeID, int VertexA, int VertexB, int NewVertex)
+	PostEdgeSplitFunction = [this, &MaxEdgeLengthSqr](int EdgeID, int VertexA, int VertexB, int NewVertex) //-V1047 - This lambda is cleared before routine exit
 	{
 		FVector3d NewVertexPosition = Mesh->GetVertex(NewVertex);
 		for (auto eid : Mesh->VtxEdgesItr(NewVertex))
diff --git a/Engine/Plugins/Experimental/GeometryProcessing/Source/GeometricObjects/Public/Implicit/Blend.h b/Engine/Plugins/Experimental/GeometryProcessing/Source/GeometricObjects/Public/Implicit/Blend.h
--- a/Engine/Plugins/Experimental/GeometryProcessing/Source/GeometricObjects/Public/Implicit/Blend.h
+++ b/Engine/Plugins/Experimental/GeometryProcessing/Source/GeometricObjects/Public/Implicit/Blend.h
@@ -219,7 +219,7 @@ protected:
 			}
 		}
 
-		MarchingCubes.Implicit = [&Blend](const FVector3d& Pt) {return Blend.Value(Pt);};
+		MarchingCubes.Implicit = [&Blend](const FVector3d& Pt) {return Blend.Value(Pt);}; //-V1047 - This lambda is cleared before routine exit
 
 		MarchingCubes.GenerateContinuation(Seeds);
 
