MeshModelingTools: Improving the local fit frame in group topology selection to have one of its axes in the global XY plane.

#rnx

diff --git a/Engine/Plugins/Experimental/GeometryProcessing/Source/DynamicMesh/Private/GroupTopology.cpp b/Engine/Plugins/Experimental/GeometryProcessing/Source/DynamicMesh/Private/GroupTopology.cpp
--- a/Engine/Plugins/Experimental/GeometryProcessing/Source/DynamicMesh/Private/GroupTopology.cpp
+++ b/Engine/Plugins/Experimental/GeometryProcessing/Source/DynamicMesh/Private/GroupTopology.cpp
@@ -678,13 +678,13 @@ FFrame3d FGroupTopology::GetSelectionFrame(const FGroupTopologySelection& Select
 		return StartFrame;
 	}
 
+	FVector3d AccumulatedOrigin = FVector3d::Zero();
 	FVector3d AccumulatedNormal = FVector3d::Zero();
 
-	FFrame3d Accumulated;
 	int AccumCount = 0;
 	for (int32 CornerID : Selection.SelectedCornerIDs)
 	{
-		Accumulated.Origin += Mesh->GetVertex(GetCornerVertexID(CornerID));
+		AccumulatedOrigin += Mesh->GetVertex(GetCornerVertexID(CornerID));
 		AccumulatedNormal += FVector3d::UnitZ();
 		AccumCount++;
 	}
@@ -694,7 +694,7 @@ FFrame3d FGroupTopology::GetSelectionFrame(const FGroupTopologySelection& Select
 		const FGroupEdge& Edge = Edges[EdgeID];
 		FVector3d StartPos = Mesh->GetVertex(Edge.Span.Vertices[0]);
 		FVector3d EndPos = Mesh->GetVertex(Edge.Span.Vertices[Edge.Span.Vertices.Num() - 1]);
-		Accumulated.Origin +=  0.5*(StartPos + EndPos);
+		AccumulatedOrigin +=  0.5*(StartPos + EndPos);
 		AccumulatedNormal += FVector3d::UnitZ();
 		AccumCount++;
 	}
@@ -704,19 +704,36 @@ FFrame3d FGroupTopology::GetSelectionFrame(const FGroupTopologySelection& Select
 		if (FindGroupByID(GroupID) != nullptr)
 		{
 			FFrame3d GroupFrame = GetGroupFrame(GroupID);
-			Accumulated.Origin += GroupFrame.Origin;
+			AccumulatedOrigin += GroupFrame.Origin;
 			AccumulatedNormal += GroupFrame.Z();
 			AccumCount++;
 		}
 	}
 
+	FFrame3d AccumulatedFrame;
 	if (AccumCount > 0)
 	{
-		Accumulated.Origin /= (double)AccumCount;
-		Accumulated.AlignAxis(2, AccumulatedNormal);
+		AccumulatedOrigin /= (double)AccumCount;
+		AccumulatedNormal.Normalize();
+
+		// We set our frame Z to be accumulated normal, and the other two axes are unconstrained, so
+		// we want to set them to something that will make our frame generally more useful. If the normal
+		// is aligned with world Z, then the entire frame might as well be aligned with world.
+		if (1 - AccumulatedNormal.Dot(FVector3d::UnitZ()) < KINDA_SMALL_NUMBER)
+		{
+			AccumulatedFrame = FFrame3d(AccumulatedOrigin, FQuaterniond::Identity());
+		}
+		else
+		{
+			// Otherwise, let's place one of the other axes into the XY plane so that the frame is more
+			// useful for translation. We somewhat arbitrarily choose Y for this. 
+			FVector3d FrameY = AccumulatedNormal.Cross(FVector3d::UnitZ()).Normalized(); // orthogonal to world Z and frame Z 
+			FVector3d FrameX = FrameY.Cross(AccumulatedNormal); // safe to not normalize because already orthogonal
+			AccumulatedFrame = FFrame3d(AccumulatedOrigin, FrameX, FrameY, AccumulatedNormal);
+		}
 	}
 
-	return Accumulated;
+	return AccumulatedFrame;
 }
 
 
