Speed up the non-exact version of MinVolumeBox3 by using FConvexHull3d instead of gte::ConvexHull3

#rnx

diff --git a/Engine/Plugins/Experimental/GeometryProcessing/Source/DynamicMesh/Private/ShapeApproximation/MeshSimpleShapeApproximation.cpp b/Engine/Plugins/Experimental/GeometryProcessing/Source/DynamicMesh/Private/ShapeApproximation/MeshSimpleShapeApproximation.cpp
--- a/Engine/Plugins/Experimental/GeometryProcessing/Source/DynamicMesh/Private/ShapeApproximation/MeshSimpleShapeApproximation.cpp
+++ b/Engine/Plugins/Experimental/GeometryProcessing/Source/DynamicMesh/Private/ShapeApproximation/MeshSimpleShapeApproximation.cpp
@@ -104,7 +104,7 @@ struct FSimpleShapeFitsResult
 
 
 static void ComputeSimpleShapeFits(const FDynamicMesh3& Mesh,
-	bool bSphere, bool bBox, bool bCapsule, bool bConvex, FSimpleShapeFitsResult& FitResult)
+	bool bSphere, bool bBox, bool bCapsule, bool bConvex, bool bUseExactComputationForBox, FSimpleShapeFitsResult& FitResult)
 {
 	TArray<int32> ToLinear, FromLinear;
 	if (bSphere || bBox || bCapsule)
@@ -122,7 +122,7 @@ static void ComputeSimpleShapeFits(const FDynamicMesh3& Mesh,
 	{
 		FMinVolumeBox3d MinBoxCalc;
 		bool bMinBoxOK = MinBoxCalc.Solve(FromLinear.Num(),
-			[&](int32 Index) { return Mesh.GetVertex(FromLinear[Index]); });
+			[&](int32 Index) { return Mesh.GetVertex(FromLinear[Index]); }, bUseExactComputationForBox);
 		if (bMinBoxOK && MinBoxCalc.IsSolutionAvailable())
 		{
 			FitResult.bHaveBox = true;
@@ -202,7 +202,7 @@ void FMeshSimpleShapeApproximation::Generate_OrientedBoxes(FSimpleShapeSet3d& Sh
 
 		const FDynamicMesh3& SourceMesh = *SourceMeshes[idx];
 		FSimpleShapeFitsResult FitResult;
-		ComputeSimpleShapeFits(SourceMesh, false, true, false, false, FitResult);
+		ComputeSimpleShapeFits(SourceMesh, false, true, false, false, bUseExactComputationForBox, FitResult);
 
 		if (FitResult.bHaveBox)
 		{
@@ -225,7 +225,7 @@ void FMeshSimpleShapeApproximation::Generate_MinimalSpheres(FSimpleShapeSet3d& S
 
 		const FDynamicMesh3& SourceMesh = *SourceMeshes[idx];
 		FSimpleShapeFitsResult FitResult;
-		ComputeSimpleShapeFits(SourceMesh, true, false, false, false, FitResult);
+		ComputeSimpleShapeFits(SourceMesh, true, false, false, false, bUseExactComputationForBox, FitResult);
 
 		if (FitResult.bHaveSphere)
 		{
@@ -248,7 +248,7 @@ void FMeshSimpleShapeApproximation::Generate_Capsules(FSimpleShapeSet3d& ShapeSe
 
 		const FDynamicMesh3& SourceMesh = *SourceMeshes[idx];
 		FSimpleShapeFitsResult FitResult;
-		ComputeSimpleShapeFits(SourceMesh, false, false, true, false, FitResult);
+		ComputeSimpleShapeFits(SourceMesh, false, false, true, false, bUseExactComputationForBox, FitResult);
 
 		if (FitResult.bHaveCapsule)
 		{
@@ -374,7 +374,7 @@ void FMeshSimpleShapeApproximation::Generate_MinVolume(FSimpleShapeSet3d& ShapeS
 		FOrientedBox3d AlignedBox = FOrientedBox3d(SourceMesh.GetBounds());
 
 		FSimpleShapeFitsResult FitResult;
-		ComputeSimpleShapeFits(SourceMesh, true, true, true, false, FitResult);
+		ComputeSimpleShapeFits(SourceMesh, true, true, true, false, bUseExactComputationForBox, FitResult);
 
 		double Volumes[4];
 		Volumes[0] = AlignedBox.Volume();
diff --git a/Engine/Plugins/Experimental/GeometryProcessing/Source/DynamicMesh/Public/ShapeApproximation/MeshSimpleShapeApproximation.h b/Engine/Plugins/Experimental/GeometryProcessing/Source/DynamicMesh/Public/ShapeApproximation/MeshSimpleShapeApproximation.h
--- a/Engine/Plugins/Experimental/GeometryProcessing/Source/DynamicMesh/Public/ShapeApproximation/MeshSimpleShapeApproximation.h
+++ b/Engine/Plugins/Experimental/GeometryProcessing/Source/DynamicMesh/Public/ShapeApproximation/MeshSimpleShapeApproximation.h
@@ -64,6 +64,8 @@ public:
 	/** simplification tolerance when simplifying 2D convex hulls, eg for swept/projected hulls */
 	double HullSimplifyTolerance = 1.0;
 
+	bool bUseExactComputationForBox = false;
+
 	//
 	// setup/initialization
 	//
diff --git a/Engine/Plugins/Experimental/GeometryProcessing/Source/GeometryAlgorithms/Private/MinVolumeBox3.cpp b/Engine/Plugins/Experimental/GeometryProcessing/Source/GeometryAlgorithms/Private/MinVolumeBox3.cpp
--- a/Engine/Plugins/Experimental/GeometryProcessing/Source/GeometryAlgorithms/Private/MinVolumeBox3.cpp
+++ b/Engine/Plugins/Experimental/GeometryProcessing/Source/GeometryAlgorithms/Private/MinVolumeBox3.cpp
@@ -18,6 +18,8 @@ THIRD_PARTY_INCLUDES_END
 #pragma warning(pop)
 #endif
 
+#include "ConvexHull3.h"
+
 #include "GteUtil.h"
 
 template <typename RealType>
@@ -27,7 +29,6 @@ struct TMinVolumeBox3Internal
 	using PreciseBoxNumberType = gte::BSRational<gte::UIntegerAP32>;
 	using DVector3 = gte::Vector3<double>;
 
-	bool bUseExactHull;
 	bool bUseExactBox;
 	TArray<DVector3> DoubleInput;
 
@@ -41,37 +42,40 @@ struct TMinVolumeBox3Internal
 
 	bool ComputeResult()
 	{
-		gte::OrientedBox3<double> MinimalBox;
+		gte::OrientedBox3<double> MinimalBox = gte::OrientedBox3<double>();
+
+		FConvexHull3d HullCompute;
+		bSolutionOK = HullCompute.Solve(DoubleInput.Num(),
+										[this](int32 Index) 
+		{
+			return FVector3d{ DoubleInput[Index][0], DoubleInput[Index][1], DoubleInput[Index][2]};
+		});
+
+		if (!bSolutionOK)
+		{
+			return false;
+		}
+
+		int NumIndices = 3 * HullCompute.GetTriangles().Num();
+		if (NumIndices < 1)
+		{
+			bSolutionOK = false;
+			return false;
+		}
+
+		const int* Indices = static_cast<const int*>(&(HullCompute.GetTriangles()[0][0]));		// Eww
 
-		if (bUseExactHull)
+		if (bUseExactBox)
 		{
-			// compute convex hull 
-			gte::ConvexHull3<double, PreciseHullNumberType> HullCompute;
-			bSolutionOK = HullCompute(DoubleInput.Num(), &DoubleInput[0], 0.0);
-			if (bSolutionOK)
-			{
-				// compute minimal box for convex hull
-				std::vector<gte::TriangleKey<true>> const& HullTriangles = HullCompute.GetHullUnordered();
-				int const numIndices = static_cast<int>(3 * HullTriangles.size());
-				int const* indices = static_cast<int const*>(&HullTriangles[0].V[0]);
-				if (bUseExactBox)
-				{
-					gte::MinimumVolumeBox3<double, PreciseBoxNumberType> BoxCompute;
-					MinimalBox = BoxCompute(DoubleInput.Num(), &DoubleInput[0], numIndices, indices);
-				}
-				else
-				{
-					gte::MinimumVolumeBox3<double, double> BoxCompute;
-					MinimalBox = BoxCompute(DoubleInput.Num(), &DoubleInput[0], numIndices, indices);
-				}
-			}
+			gte::MinimumVolumeBox3<double, PreciseBoxNumberType> BoxCompute;
+			MinimalBox = BoxCompute(DoubleInput.Num(), &DoubleInput[0], NumIndices, Indices);
 		}
 		else
 		{
 			gte::MinimumVolumeBox3<double, double> DoubleCompute;
-			MinimalBox = DoubleCompute(DoubleInput.Num(), &DoubleInput[0], true);
-			bSolutionOK = true;
+			MinimalBox = DoubleCompute(DoubleInput.Num(), &DoubleInput[0], NumIndices, Indices);
 		}
+		bSolutionOK = true;
 
 		// if resulting box is not finite, something went wrong, just return an empty box
 		FVector3d Extents = Convert(MinimalBox.extent);
@@ -96,9 +100,9 @@ struct TMinVolumeBox3Internal
 
 
 template<typename RealType>
-bool TMinVolumeBox3<RealType>::Solve(int32 NumPoints, TFunctionRef<FVector3<RealType>(int32)> GetPointFunc, bool bUseExactHull, bool bUseExactBox)
+bool TMinVolumeBox3<RealType>::Solve(int32 NumPoints, TFunctionRef<FVector3<RealType>(int32)> GetPointFunc, bool bUseExactBox)
 {
-	Initialize(NumPoints, bUseExactHull, bUseExactBox);
+	Initialize(NumPoints, bUseExactBox);
 	check(Internal);
 
 	for (int32 k = 0; k < NumPoints; ++k)
@@ -125,10 +129,9 @@ void TMinVolumeBox3<RealType>::GetResult(TOrientedBox3<RealType>& BoxOut)
 
 
 template<typename RealType>
-void TMinVolumeBox3<RealType>::Initialize(int32 NumPoints, bool bUseExactHull, bool bUseExactBox)
+void TMinVolumeBox3<RealType>::Initialize(int32 NumPoints, bool bUseExactBox)
 {
 	Internal = MakePimpl<TMinVolumeBox3Internal<RealType>>();
-	Internal->bUseExactHull = bUseExactHull;
 	Internal->bUseExactBox = bUseExactBox;
 	Internal->DoubleInput.SetNum(NumPoints);
 }
diff --git a/Engine/Plugins/Experimental/GeometryProcessing/Source/GeometryAlgorithms/Public/MinVolumeBox3.h b/Engine/Plugins/Experimental/GeometryProcessing/Source/GeometryAlgorithms/Public/MinVolumeBox3.h
--- a/Engine/Plugins/Experimental/GeometryProcessing/Source/GeometryAlgorithms/Public/MinVolumeBox3.h
+++ b/Engine/Plugins/Experimental/GeometryProcessing/Source/GeometryAlgorithms/Public/MinVolumeBox3.h
@@ -21,11 +21,10 @@ class TMinVolumeBox3
 public:
 	/**
 	 * Calculate the minimal box for the given point set.
-	 * @param bUseExactHull If true, high-precision number types are used for the convex hull calculation, rather than doubles. This is slower but more reliable.
 	 * @param bUseExactBox If true, high-precision number types are used for the minimal-box calculation, rather than doubles. This is *much* slower but more accurate (but not recommended).
 	 * @return true if minimal box was found
 	 */
-	bool Solve(int32 NumPoints, TFunctionRef<FVector3<RealType>(int32)> GetPointFunc, bool bUseExactHull = true, bool bUseExactBox = false);
+	bool Solve(int32 NumPoints, TFunctionRef<FVector3<RealType>(int32)> GetPointFunc, bool bUseExactBox = false);
 
 	/** @return true if minimal box is available */
 	bool IsSolutionAvailable() const;
@@ -34,7 +33,7 @@ public:
 	void GetResult(TOrientedBox3<RealType>& BoxOut);
 
 protected:
-	void Initialize(int32 NumPoints, bool bUseExactHull, bool bUseExactBox);
+	void Initialize(int32 NumPoints, bool bUseExactBox);
 
 	TPimplPtr<TMinVolumeBox3Internal<RealType>> Internal;
 };
