GeometryTools: MeshDescription To DynamicMesh.  Add option to convert the tangents when creating FDynamicMesh from FMeshDescription.  Normals are converted by default.

#rnx

diff --git a/Engine/Plugins/Experimental/GeometryProcessing/Source/MeshConversion/Private/MeshDescriptionToDynamicMesh.cpp b/Engine/Plugins/Experimental/GeometryProcessing/Source/MeshConversion/Private/MeshDescriptionToDynamicMesh.cpp
--- a/Engine/Plugins/Experimental/GeometryProcessing/Source/MeshConversion/Private/MeshDescriptionToDynamicMesh.cpp
+++ b/Engine/Plugins/Experimental/GeometryProcessing/Source/MeshConversion/Private/MeshDescriptionToDynamicMesh.cpp
@@ -123,7 +123,7 @@ public:
 
 
 
-void FMeshDescriptionToDynamicMesh::Convert(const FMeshDescription* MeshIn, FDynamicMesh3& MeshOut)
+void FMeshDescriptionToDynamicMesh::Convert(const FMeshDescription* MeshIn, FDynamicMesh3& MeshOut, bool bCopyTangents )
 {
 	TriIDMap.Reset();
 	VertIDMap.Reset();
@@ -147,6 +147,8 @@ void FMeshDescriptionToDynamicMesh::Convert(const FMeshDescription* MeshIn, FDyn
 	FStaticMeshConstAttributes Attributes(*MeshIn);
 	TVertexInstanceAttributesConstRef<FVector2D> InstanceUVs = Attributes.GetVertexInstanceUVs();
 	TVertexInstanceAttributesConstRef<FVector> InstanceNormals = Attributes.GetVertexInstanceNormals();
+	TVertexInstanceAttributesConstRef<FVector> InstanceTangents = Attributes.GetVertexInstanceTangents();
+	TVertexInstanceAttributesConstRef<float> InstanceBiTangentSign = Attributes.GetVertexInstanceBinormalSigns();
 
 	TPolygonAttributesConstRef<int> PolyGroups =
 		MeshIn->PolygonAttributes().GetAttributesRef<int>(ExtendedMeshAttribute::PolyTriGroups);
@@ -155,6 +157,8 @@ void FMeshDescriptionToDynamicMesh::Convert(const FMeshDescription* MeshIn, FDyn
 		MeshOut.EnableTriangleGroups(0);
 	}
 
+	bCopyTangents = bCopyTangents && InstanceTangents.IsValid() && InstanceBiTangentSign.IsValid();
+	
 	// base triangle groups will track polygons
 	int NumVertices = MeshIn->Vertices().Num();
 	int NumPolygons = MeshIn->Polygons().Num();
@@ -333,26 +337,45 @@ void FMeshDescriptionToDynamicMesh::Convert(const FMeshDescription* MeshIn, FDyn
 	}
 
 
-	TArray<FDynamicMeshUVOverlay*> UVOverlays;
-	TArray<FUVWelder> UVWelders;
 	FDynamicMeshNormalOverlay* NormalOverlay = nullptr;
-	FNormalWelder NormalWelder;
+	FDynamicMeshNormalOverlay* TangentOverlay = nullptr;
+	FDynamicMeshNormalOverlay* BiTangentOverlay = nullptr;
 	FDynamicMeshMaterialAttribute* MaterialIDAttrib = nullptr;
+
+	
+
 	if (!bDisableAttributes)
 	{
-		MeshOut.EnableAttributes();
+		MeshOut.EnableAttributes(); // by default 1-UV layer and 1-normal layer
+
+		// Normals
+		// set up for the tangent plane vectors if required
+		int32 NumRequiredNormalLayers = (bCopyTangents) ? 3 : 1;
+		if (MeshOut.Attributes()->NumNormalLayers() <  NumRequiredNormalLayers)
+		{	
+			// add additional normal layers and reserve space in them
+			MeshOut.Attributes()->SetNumNormalLayers(NumRequiredNormalLayers);
+		}
+		for (int32 i = 0; i < NumRequiredNormalLayers; ++i)
+		{ 
+			MeshOut.Attributes()->GetNormalLayer(i)->InitializeTriangles(MeshOut.MaxTriangleID());
+		}
 
+		NormalOverlay = MeshOut.Attributes()->PrimaryNormals();
+		if (bCopyTangents)
+		{ 
+			TangentOverlay = MeshOut.Attributes()->PrimaryTangents();
+			BiTangentOverlay = MeshOut.Attributes()->PrimaryBiTangents();
+		}
+		
+		// UVs 
 		MeshOut.Attributes()->SetNumUVLayers(NumUVLayers);
-		UVOverlays.SetNum(NumUVLayers);
-		UVWelders.SetNum(NumUVLayers);
-		for (int j = 0; j < NumUVLayers; ++j)
+		// reserve space in any new UV layers.
+		for (int32 i = 1; i < NumUVLayers; ++i)
 		{
-			UVOverlays[j] = MeshOut.Attributes()->GetUVLayer(j);
-			UVWelders[j].UVOverlay = UVOverlays[j];
+			MeshOut.Attributes()->GetUVLayer(i)->InitializeTriangles(MeshOut.MaxTriangleID());
 		}
-
-		NormalOverlay = MeshOut.Attributes()->PrimaryNormals();
-		NormalWelder.NormalOverlay = NormalOverlay;
+	
 
 		// always enable Material ID if there are any attributes
 		MeshOut.Attributes()->EnableMaterialID();
@@ -391,9 +414,16 @@ void FMeshDescriptionToDynamicMesh::Convert(const FMeshDescription* MeshIn, FDyn
 		{
 			auto UVFuture = Async(EAsyncExecution::ThreadPool, [&, UVLayerIndex, bUseSharedUVs]() // must copy UVLayerIndex here!
 			{
+				// the overlay to fill.
+				FDynamicMeshUVOverlay* UVOverlay = MeshOut.Attributes()->GetUVLayer(UVLayerIndex);
+				if (UVOverlay == nullptr) return;
 
 				if (!bUseSharedUVs) // have to rely on welding the per-instance uvs 
 				{
+					
+					FUVWelder UVWelder;
+					UVWelder.UVOverlay = UVOverlay;
+
 					for (int32 TriangleID : MeshOut.TriangleIndicesItr())
 					{
 						FIndex3i Tri = MeshOut.GetTriangle(TriangleID);
@@ -402,15 +432,13 @@ void FMeshDescriptionToDynamicMesh::Convert(const FMeshDescription* MeshIn, FDyn
 						for (int j = 0; j < 3; ++j)
 						{
 							FVector2D UV = InstanceUVs.Get(TriData.TriInstances[j], UVLayerIndex);
-							TriUV[j] = UVWelders[UVLayerIndex].FindOrAddUnique(UV, Tri[j]);
+							TriUV[j] = UVWelder.FindOrAddUnique(UV, Tri[j]);
 						}
-						UVOverlays[UVLayerIndex]->SetTriangle(TriangleID, TriUV);
+						UVOverlay->SetTriangle(TriangleID, TriUV);
 					}
 				}
 				else
 				{
-					// the overlay to fill.
-					FDynamicMeshUVOverlay* UVOverlay = UVOverlays[UVLayerIndex];
 
 					// copy uv "vertex buffer"
 					const FUVArray& UVs = MeshIn->UVs(UVLayerIndex);
@@ -519,12 +547,12 @@ void FMeshDescriptionToDynamicMesh::Convert(const FMeshDescription* MeshIn, FDyn
 						UVWelder.UVOverlay = UVOverlay;
 						for (int32 TriangleID : TrisMissingUVs)
 						{
-							FIndex3i Tri = MeshOut.GetTriangle(TriangleID);
+							const FIndex3i Tri = MeshOut.GetTriangle(TriangleID);
 							const FTriData& TriData = AddedTriangles[TriangleID];
 							FIndex3i TriUV;
 							for (int j = 0; j < 3; ++j)
 							{
-								FVector2D UV = InstanceUVs.Get(TriData.TriInstances[j], UVLayerIndex);
+								const FVector2D UV = InstanceUVs.Get(TriData.TriInstances[j], UVLayerIndex);
 								TriUV[j] = UVWelder.FindOrAddUnique(UV, Tri[j]);
 							}
 							UVOverlay->SetTriangle(TriangleID, TriUV);
@@ -541,14 +569,16 @@ void FMeshDescriptionToDynamicMesh::Convert(const FMeshDescription* MeshIn, FDyn
 		{
 			auto NormalFuture = Async(EAsyncExecution::ThreadPool, [&]()
 			{
+				FNormalWelder NormalWelder;
+				NormalWelder.NormalOverlay = NormalOverlay;
 				for (int32 TriangleID : MeshOut.TriangleIndicesItr())
 				{
-					FIndex3i Tri = MeshOut.GetTriangle(TriangleID);
+					const FIndex3i Tri = MeshOut.GetTriangle(TriangleID);
 					const FTriData& TriData = AddedTriangles[TriangleID];
 					FIndex3i TriNormals;
 					for (int j = 0; j < 3; ++j)
 					{
-						FVector Normal = InstanceNormals.Get(TriData.TriInstances[j]);
+						const FVector Normal = InstanceNormals.Get(TriData.TriInstances[j]);
 						TriNormals[j] = NormalWelder.FindOrAddUnique(Normal, Tri[j]);
 					}
 					NormalOverlay->SetTriangle(TriangleID, TriNormals);
@@ -558,6 +588,58 @@ void FMeshDescriptionToDynamicMesh::Convert(const FMeshDescription* MeshIn, FDyn
 		}
 
 
+		if (TangentOverlay != nullptr)
+		{
+			auto TangentFuture = Async(EAsyncExecution::ThreadPool, [&]()
+			{
+				FNormalWelder TangentWelder;
+				TangentWelder.NormalOverlay = TangentOverlay;
+				for (int32 TriangleID : MeshOut.TriangleIndicesItr())
+				{
+					const FIndex3i Tri = MeshOut.GetTriangle(TriangleID);
+					const FTriData& TriData = AddedTriangles[TriangleID];
+					FIndex3i TriVector;
+					for (int j = 0; j < 3; ++j)
+					{
+						const FVector Vector = InstanceTangents.Get(TriData.TriInstances[j]);
+						TriVector[j] = TangentWelder.FindOrAddUnique(Vector, Tri[j]);
+					}
+					TangentOverlay->SetTriangle(TriangleID, TriVector);
+				}
+			});
+			Pending.Add(MoveTemp(TangentFuture));
+		}
+
+		if (BiTangentOverlay != nullptr)
+		{
+			auto BiTangentFuture = Async(EAsyncExecution::ThreadPool, [&]()
+			{
+				FNormalWelder BiTangentWelder;
+				BiTangentWelder.NormalOverlay = BiTangentOverlay;
+				for (int32 TriangleID : MeshOut.TriangleIndicesItr())
+				{
+					const FIndex3i Tri = MeshOut.GetTriangle(TriangleID);
+					const FTriData& TriData = AddedTriangles[TriangleID];
+					FIndex3i TriVector;
+					for (int j = 0; j < 3; ++j)
+					{
+						// compute the bi tangent.
+						const FVertexInstanceID VertexInstanceID = TriData.TriInstances[j];
+						const FVector NormalVector = InstanceNormals.Get(VertexInstanceID);
+						const FVector TangentVector = InstanceTangents.Get(VertexInstanceID);
+						const float BiSign  = InstanceBiTangentSign.Get(VertexInstanceID);
+
+						FVector BiTangentVector = BiSign * FVector::CrossProduct(NormalVector, TangentVector);
+						BiTangentVector.Normalize();
+
+						TriVector[j] = BiTangentWelder.FindOrAddUnique(BiTangentVector, Tri[j]);
+					}
+					BiTangentOverlay->SetTriangle(TriangleID, TriVector);
+				}
+			});
+			Pending.Add(MoveTemp(BiTangentFuture));
+		}
+
 		if (MaterialIDAttrib != nullptr)
 		{
 			auto MaterialFuture = Async(EAsyncExecution::ThreadPool, [&]()
diff --git a/Engine/Plugins/Experimental/GeometryProcessing/Source/MeshConversion/Public/MeshDescriptionToDynamicMesh.h b/Engine/Plugins/Experimental/GeometryProcessing/Source/MeshConversion/Public/MeshDescriptionToDynamicMesh.h
--- a/Engine/Plugins/Experimental/GeometryProcessing/Source/MeshConversion/Public/MeshDescriptionToDynamicMesh.h
+++ b/Engine/Plugins/Experimental/GeometryProcessing/Source/MeshConversion/Public/MeshDescriptionToDynamicMesh.h
@@ -62,8 +62,10 @@ public:
 
 	/**
 	 * Default conversion of MeshDescription to DynamicMesh
+	 * @param bCopyTangents  - if bDisableAttributes is false, this requests the tangent plane vectors (tangent and bitangent) 
+	 *                          be stored as overlays in the MeshOut DynamicAttributeSet, provided they exist on the MeshIn
 	 */
-	void Convert(const FMeshDescription* MeshIn, FDynamicMesh3& MeshOut);
+	void Convert(const FMeshDescription* MeshIn, FDynamicMesh3& MeshOut, bool bCopyTangents = false);
 
 	/**
 	 * Copy tangents from MeshDescription to a FMeshTangents instance.
