Geometry Tools: DynamicMeshToMeshDescription.  Add support for tangents in the conversion (and update) from DynamicMesh to MeshDescription.

#rnx

diff --git a/Engine/Plugins/Experimental/GeometryProcessing/Source/MeshConversion/Private/DynamicMeshToMeshDescription.cpp b/Engine/Plugins/Experimental/GeometryProcessing/Source/MeshConversion/Private/DynamicMeshToMeshDescription.cpp
--- a/Engine/Plugins/Experimental/GeometryProcessing/Source/MeshConversion/Private/DynamicMeshToMeshDescription.cpp
+++ b/Engine/Plugins/Experimental/GeometryProcessing/Source/MeshConversion/Private/DynamicMeshToMeshDescription.cpp
@@ -34,7 +34,7 @@ namespace DynamicMeshToMeshDescriptionConversionHelper
 }
 
 
-void FDynamicMeshToMeshDescription::Update(const FDynamicMesh3* MeshIn, FMeshDescription& MeshOut, bool bUpdateNormals, bool bUpdateUVs)
+void FDynamicMeshToMeshDescription::Update(const FDynamicMesh3* MeshIn, FMeshDescription& MeshOut, bool bUpdateNormals, bool bUpdateTangents, bool bUpdateUVs)
 {
 	FMeshDescriptionBuilder Builder;
 	Builder.SetMeshDescription(&MeshOut);
@@ -54,12 +54,12 @@ void FDynamicMeshToMeshDescription::Update(const FDynamicMesh3* MeshIn, FMeshDes
 		Builder.SetPosition(FVertexID(VertID), (FVector)MeshIn->GetVertex(VertID));
 	}
 
-	UpdateAttributes(MeshIn, MeshOut, bUpdateNormals, bUpdateUVs);
+	UpdateAttributes(MeshIn, MeshOut, bUpdateNormals, bUpdateTangents, bUpdateUVs);
 }
 
 
 
-void FDynamicMeshToMeshDescription::UpdateAttributes(const FDynamicMesh3* MeshIn, FMeshDescription& MeshOut, bool bUpdateNormals, bool bUpdateUVs)
+void FDynamicMeshToMeshDescription::UpdateAttributes(const FDynamicMesh3* MeshIn, FMeshDescription& MeshOut, bool bUpdateNormals, bool bUpdateTangents, bool bUpdateUVs)
 {
 	check(MeshIn->IsCompactV());
 
@@ -73,11 +73,13 @@ void FDynamicMeshToMeshDescription::UpdateAttributes(const FDynamicMesh3* MeshIn
 
 	FStaticMeshAttributes Attributes(MeshOut);
 
+
 	if (bUpdateNormals)
 	{
 		TVertexInstanceAttributesRef<FVector> InstanceAttrib = Attributes.GetVertexInstanceNormals();
-		ensureMsgf(InstanceAttrib.IsValid(), TEXT("Trying to update normals on a MeshDescription that has no normal attributes"));
-		if (InstanceAttrib.IsValid())
+		bool bIsValidDst = InstanceAttrib.IsValid();
+		ensureMsgf(bIsValidDst, TEXT("Trying to update normals on a MeshDescription that has no normal attributes"));
+		if (bIsValidDst)
 		{
 			const FDynamicMeshNormalOverlay* Overlay = MeshIn->HasAttributes() ? MeshIn->Attributes()->PrimaryNormals() : nullptr;
 			if (Overlay)
@@ -100,6 +102,12 @@ void FDynamicMeshToMeshDescription::UpdateAttributes(const FDynamicMesh3* MeshIn
 		}
 	}
 
+	if (bUpdateTangents)
+	{
+		UpdateTangents(MeshIn, MeshOut);
+	}
+	
+
 	if (bUpdateUVs)
 	{
 		TVertexInstanceAttributesRef<FVector2D> InstanceAttrib = Attributes.GetVertexInstanceUVs();
@@ -167,14 +175,69 @@ void FDynamicMeshToMeshDescription::UpdateTangents(const FDynamicMesh3* MeshIn,
 	}
 }
 
+void FDynamicMeshToMeshDescription::UpdateTangents(const FDynamicMesh3* MeshIn, FMeshDescription& MeshOut)
+{
+	if (!ensureMsgf(MeshIn->IsCompactT(), TEXT("Trying to update MeshDescription Tangents from a non-compact DynamicMesh"))) return;
+	if (!ensureMsgf(MeshIn->TriangleCount() == MeshOut.Triangles().Num(), TEXT("Trying to update MeshDescription Tangents from Mesh that does not have same triangle count"))) return;
+	if (!ensureMsgf(MeshIn->HasAttributes(), TEXT("Trying to update MeshDescription Tangents from a DynamicMesh that has no attributes, e.g. normals"))) return;
+
+	
+	// src
+	const FDynamicMeshNormalOverlay* NormalOverlay = MeshIn->Attributes()->PrimaryNormals();
+	const FDynamicMeshNormalOverlay* TangentOverlay = MeshIn->Attributes()->PrimaryTangents();
+	const FDynamicMeshNormalOverlay* BiTangentOverlay = MeshIn->Attributes()->PrimaryBiTangents();
+
+	bool bHasValidSrc = (NormalOverlay != nullptr) && (TangentOverlay != nullptr) && (BiTangentOverlay != nullptr);
+
+	if(!ensureMsgf(bHasValidSrc, TEXT("Trying to update MeshDescription Tangents from a DynamicMesh that does not have all three tangent space attributes"))) return;
+
+	// dst
+	FStaticMeshAttributes Attributes(MeshOut);
+	TVertexInstanceAttributesRef<FVector> TangentAttrib = Attributes.GetVertexInstanceTangents();
+	TVertexInstanceAttributesRef<float> BiTangentSignAttrib = Attributes.GetVertexInstanceBinormalSigns();
+
+	if (!ensureMsgf(TangentAttrib.IsValid(), TEXT("Trying to update Tangents on a MeshDescription that has no Tangent Vertex Instance attribute"))) return;
+	if (!ensureMsgf(BiTangentSignAttrib.IsValid(), TEXT("Trying to update Tangents on a MeshDescription that has no BinormalSign Vertex Instance attribute"))) return;
+
+	const int32 NumTriangles = MeshIn->TriangleCount();
+	for (int32 t = 0; t < NumTriangles; ++t)
+	{
+		const bool bTriHasTangentSpace = NormalOverlay->IsSetTriangle(t) && TangentOverlay->IsSetTriangle(t) && BiTangentOverlay->IsSetTriangle(t);
+
+		if (!bTriHasTangentSpace) continue;
+		
+		// get data from dynamic mesh overlays
+		FVector3f TriNormals[3];
+		NormalOverlay->GetTriElements(t, TriNormals[0], TriNormals[1], TriNormals[2]);
+	
+
+		FVector3f TriTangents[3];
+		TangentOverlay->GetTriElements(t, TriTangents[0], TriTangents[1], TriTangents[2]);
+		
+
+		FVector3f TriBiTangents[3];
+		BiTangentOverlay->GetTriElements(t, TriBiTangents[0], TriBiTangents[1], TriBiTangents[2]);
+		
+		// set data in the mesh description per-vertex attributes
+		TArrayView<const FVertexInstanceID> TriInstances = MeshOut.GetTriangleVertexInstances(FTriangleID(t));
+		for (int32 i = 0; i < 3; ++i)
+		{
+			const float BiTangentSign = VectorUtil::BitangentSign(TriNormals[i], TriTangents[i], TriBiTangents[i]);
+
+			TangentAttrib.Set(TriInstances[i], (FVector)TriTangents[i]);
+			BiTangentSignAttrib.Set(TriInstances[i], BiTangentSign);
+		}
+	}
+}
+
 
 
-void FDynamicMeshToMeshDescription::Convert(const FDynamicMesh3* MeshIn, FMeshDescription& MeshOut)
+void FDynamicMeshToMeshDescription::Convert(const FDynamicMesh3* MeshIn, FMeshDescription& MeshOut, bool bCopyTangents)
 {
 	if (MeshIn->HasAttributes())
 	{
 		//Convert_SharedInstances(MeshIn, MeshOut);
-		Convert_NoSharedInstances(MeshIn, MeshOut);
+		Convert_NoSharedInstances(MeshIn, MeshOut, bCopyTangents);
 	}
 	else
 	{
@@ -234,7 +297,7 @@ void FDynamicMeshToMeshDescription::Convert_NoAttributes(const FDynamicMesh3* Me
 				FVector Normal = MeshIn->HasVertexNormals() ? FVector(MeshIn->GetVertexNormal(Triangle[j])) : FVector::UpVector;
 				Builder.SetInstanceNormal(NewInstanceID, Normal);
 
-				// Add UV to MeshDecription UVvertex buffer
+				// Add UV to MeshDescription UVvertex buffer
 				FVector2D UV = MeshIn->HasVertexUVs() ? FVector2D(MeshIn->GetVertexUV(Triangle[j])) : FVector2D::ZeroVector;
 				FUVID UVID = Builder.AppendUV(UV, UVLayerIndex);
 				
@@ -426,7 +489,7 @@ void FDynamicMeshToMeshDescription::Convert_SharedInstances(const FDynamicMesh3*
 
 
 
-void FDynamicMeshToMeshDescription::Convert_NoSharedInstances(const FDynamicMesh3* MeshIn, FMeshDescription& MeshOut)
+void FDynamicMeshToMeshDescription::Convert_NoSharedInstances(const FDynamicMesh3* MeshIn, FMeshDescription& MeshOut, bool bCopyTangents)
 {
 	
 	const bool bHasAttributes = MeshIn->HasAttributes();
@@ -434,8 +497,10 @@ void FDynamicMeshToMeshDescription::Convert_NoSharedInstances(const FDynamicMesh
 	// check if we have per-triangle material ID
 	const FDynamicMeshMaterialAttribute* MaterialIDAttrib = (bHasAttributes && MeshIn->Attributes()->HasMaterialID()) ? MeshIn->Attributes()->GetMaterialID() : nullptr;
 
-	// cache normal and UV overlay info
+	// cache tangent space and UV overlay info
 	const FDynamicMeshNormalOverlay* NormalOverlay = bHasAttributes ? MeshIn->Attributes()->PrimaryNormals() : nullptr;
+	const FDynamicMeshNormalOverlay* TangentOverlay = bHasAttributes ? MeshIn->Attributes()->PrimaryTangents() : nullptr;
+	const FDynamicMeshNormalOverlay* BiTangentOverlay = bHasAttributes ? MeshIn->Attributes()->PrimaryBiTangents() : nullptr;
 
 	const int32 NumUVLayers = bHasAttributes ? MeshIn->Attributes()->NumUVLayers() : 0;
 	
@@ -472,7 +537,7 @@ void FDynamicMeshToMeshDescription::Convert_NoSharedInstances(const FDynamicMesh
 		MapV[VertID] = Builder.AppendVertex((FVector)MeshIn->GetVertex(VertID));
 	}
 
-	// create UV vertex buffer in MeshDecription
+	// create UV vertex buffer in MeshDescription
 	Builder.SetNumUVLayers(NumUVLayers);
 
 	TArray<TArray<FUVID>> MapUVArray; MapUVArray.SetNum(NumUVLayers);
@@ -495,6 +560,57 @@ void FDynamicMeshToMeshDescription::Convert_NoSharedInstances(const FDynamicMesh
 
 	FPolygonGroupID ZeroPolygonGroupID = Builder.AppendPolygonGroup();
 
+	
+	// construct a function that will transfer tangent space data. 
+	// if the DynamicMesh has a full tangent space: a normal, tangent and bitangent sign will be transfered
+	//                                            otherwise just transfer the normal if it exists
+	const bool bCopyFullTangentSpace = bCopyTangents && (NormalOverlay != nullptr && TangentOverlay != nullptr && BiTangentOverlay != nullptr);
+
+	TFunction<void(int TriID, FVertexInstanceID TriVertInstances[3])> TangetSpaceInstanceSetter;
+	if (bCopyFullTangentSpace) // create function that sets the tangent space
+	{
+		TangetSpaceInstanceSetter = [NormalOverlay, TangentOverlay, BiTangentOverlay, &Builder](int TriID, FVertexInstanceID TriVertInstances[3])
+		{
+			FIndex3i NormalTri = NormalOverlay->GetTriangle(TriID);
+			FIndex3i TangentTri = TangentOverlay->GetTriangle(TriID);
+			FIndex3i BiTangentTri = BiTangentOverlay->GetTriangle(TriID);
+
+			for (int32 j = 0; j < 3; ++j)
+			{
+				const FVertexInstanceID CornerInstanceID = TriVertInstances[j];
+				
+				const FVector3f TriVertNormal = (NormalOverlay->IsElement(NormalTri[j]))  ? NormalOverlay->GetElement(NormalTri[j]) : FVector3f(FVector::UpVector);
+				const FVector3f TriVertTangent = (TangentOverlay->IsElement(TangentTri[j])) ? TangentOverlay->GetElement(TangentTri[j]) :  FVector3f(FVector::ForwardVector);
+				const FVector3f TriVertBiTangent = (BiTangentOverlay->IsElement(BiTangentTri[j])) ? BiTangentOverlay->GetElement(BiTangentTri[j]) : FVector3f(FVector::RightVector);
+
+				// infer sign
+				float BiTangentSign = VectorUtil::BitangentSign(TriVertNormal, TriVertTangent, TriVertBiTangent);
+
+				// set the tangent space
+				Builder.SetInstanceTangentSpace(CornerInstanceID, (FVector)TriVertNormal, (FVector)TriVertTangent, BiTangentSign);
+			}
+			
+		};
+	}
+	else if (NormalOverlay != nullptr) // create function that just sets the normals
+	{
+		TangetSpaceInstanceSetter = [NormalOverlay, &Builder](int TriID, FVertexInstanceID TriVertInstances[3])
+		{
+			FIndex3i NormalTri = NormalOverlay->GetTriangle(TriID);
+			for (int32 j = 0; j < 3; ++j)
+			{
+				FVertexInstanceID CornerInstanceID = TriVertInstances[j];
+				FVector TriVertNormal = (NormalOverlay->IsElement(NormalTri[j])) ? (FVector)NormalOverlay->GetElement(NormalTri[j]) : FVector::UpVector;
+				Builder.SetInstanceNormal(CornerInstanceID, TriVertNormal);
+			}
+		};
+	}
+	else  // dummy function that does nothing 
+	{
+		TangetSpaceInstanceSetter = [](int TriID, FVertexInstanceID TriVertInstances[3]){ return;};
+	}
+
+
 	// need to know max material index value so we can reserve groups in MeshDescription
 	int32 MaxPolygonGroupID = 0;
 	if (MaterialIDAttrib)
@@ -583,22 +699,11 @@ void FDynamicMeshToMeshDescription::Convert_NoSharedInstances(const FDynamicMesh
 
 		}
 		
-		// look up normal triangles
-		FIndex3i NormalTri = (NormalOverlay) ? NormalOverlay->GetTriangle(TriID) : FIndex3i::Invalid();
-
-		// transfer Normals.   
-		// NB: only per-instance normals are supported in MeshDescription at this time
-		// NB: will need to be updated to follow the pattern used in UVs above when MeshDecription supports shared normals 
-		for (int32 j = 0; j < 3; ++j)
-		{
-			FVertexInstanceID CornerInstanceID = TriVertInstances[j];
-			FVector TriVertNormal = FVector::UpVector;
-			if (NormalOverlay && NormalOverlay->IsElement(NormalTri[j]))
-			{
-				TriVertNormal = (FVector)NormalOverlay->GetElement(NormalTri[j]);
-			}
-			Builder.SetInstanceNormal(CornerInstanceID, TriVertNormal);
-		}
+		// transfer tangent space. 
+		// NB: MeshDescription doesn't store and explicit bitangent, so this conversion isn't perfect.
+		// NB: only per-instance normals , tangents, bitangent sign are supported in MeshDescription at this time
+		// NB: will need to be updated to follow the pattern used in UVs above when MeshDescription supports shared tangent space elements. 
+		TangetSpaceInstanceSetter(TriID, TriVertInstances);
 
 
 		if (bCopyGroupToPolyGroup)
diff --git a/Engine/Plugins/Experimental/GeometryProcessing/Source/MeshConversion/Private/MeshDescriptionBuilder.cpp b/Engine/Plugins/Experimental/GeometryProcessing/Source/MeshConversion/Private/MeshDescriptionBuilder.cpp
--- a/Engine/Plugins/Experimental/GeometryProcessing/Source/MeshConversion/Private/MeshDescriptionBuilder.cpp
+++ b/Engine/Plugins/Experimental/GeometryProcessing/Source/MeshConversion/Private/MeshDescriptionBuilder.cpp
@@ -33,6 +33,8 @@ void FMeshDescriptionBuilder::SetMeshDescription(FMeshDescription* Description)
 	this->VertexPositions = Attributes.GetVertexPositions();
 	this->InstanceUVs = Attributes.GetVertexInstanceUVs();
 	this->InstanceNormals = Attributes.GetVertexInstanceNormals();
+	this->InstanceTangents = Attributes.GetVertexInstanceTangents();
+	this->InstanceBiTangentSign = Attributes.GetVertexInstanceBinormalSigns();
 	this->InstanceColors = Attributes.GetVertexInstanceColors();
 }
 
@@ -160,6 +162,24 @@ void FMeshDescriptionBuilder::SetInstanceNormal(const FVertexInstanceID& Instanc
 	}
 }
 
+void FMeshDescriptionBuilder::SetInstanceTangentSpace(const FVertexInstanceID& InstanceID, const FVector& Normal, const FVector& Tangent, float Sign)
+{
+	// set the normal
+	SetInstanceNormal(InstanceID, Normal);
+
+	if (InstanceTangents.IsValid())
+	{
+		FVector4 TangentExtended;  
+		TangentExtended.X = Tangent.X; 
+		TangentExtended.Y = Tangent.Y; 
+		TangentExtended.Z = Tangent.Z; 
+		TangentExtended.W =  Sign;
+		InstanceTangents.Set(InstanceID, TangentExtended);
+	}
+
+
+}
+
 
 void FMeshDescriptionBuilder::SetInstanceUV(const FVertexInstanceID& InstanceID, const FVector2D& InstanceUV, int32 UVLayerIndex)
 {
diff --git a/Engine/Plugins/Experimental/GeometryProcessing/Source/MeshConversion/Public/DynamicMeshToMeshDescription.h b/Engine/Plugins/Experimental/GeometryProcessing/Source/MeshConversion/Public/DynamicMeshToMeshDescription.h
--- a/Engine/Plugins/Experimental/GeometryProcessing/Source/MeshConversion/Public/DynamicMeshToMeshDescription.h
+++ b/Engine/Plugins/Experimental/GeometryProcessing/Source/MeshConversion/Public/DynamicMeshToMeshDescription.h
@@ -70,30 +70,51 @@ public:
 
 	/**
 	 * Default conversion of DynamicMesh to MeshDescription. Calls functions below depending on mesh state
+	 *
+	 * @param bCopyTangents - should tangent, bitangent overlays exist on the MeshIn, this will use them when producing the MeshDescription MeshOut
+	 *                         
+	 * Note: Don't copy tangents if the resulting MeshDescription corresponds to a StaticMesh with autogenerated tangents
+	 *       in this case the MeshDescription is expected to have empty tangents.
+	 *       
 	 */
-	void Convert(const FDynamicMesh3* MeshIn, FMeshDescription& MeshOut);
-
+	void Convert(const FDynamicMesh3* MeshIn, FMeshDescription& MeshOut, bool bCopyTangents = false);
 
 	/**
 	 * Update existing MeshDescription based on DynamicMesh. Assumes mesh topology has not changed.
-	 * Copies positions and normals; does not update UVs
+	 * Copies positions 
+	 * optionally, normals, tangents and UVs
 	 */
-	void Update(const FDynamicMesh3* MeshIn, FMeshDescription& MeshOut, bool bUpdateNormals = true, bool bUpdateUVs = false);
-
+	void Update(const FDynamicMesh3* MeshIn, FMeshDescription& MeshOut, bool bUpdateNormals = true, bool bUpdateTangents = false, bool bUpdateUVs = false);
+	
 
 	/**
 	 * Update only attributes, assuming the mesh topology has not changed.  Does not touch positions.
 	 *	NOTE: assumes the order of triangles in the MeshIn correspond to the ordering you'd get by iterating over triangles, on MeshOut
 	 *		  This matches conversion currently used in MeshDescriptionToDynamicMesh.cpp, but if that changes we will need to change this function to match!
+	 *  @param bUpdateNormals   Specifies if the normals should be transfered from the MeshIn PrimaryNormal overlay if it exists
+	 *  @param bUpdateTangents  Specifies if the tangent and bitangent sign should be populated from the MeshIn PrimaryTangents and PrimaryBiTangents. 
+	 *                          this requires the PrimaryNormals to exist as well. 
+	 *  @param bUpdateUVs       Specifices if the UV layers should be transfered from the MeshIn overlays.
 	 */
-	void UpdateAttributes(const FDynamicMesh3* MeshIn, FMeshDescription& MeshOut, bool bUpdateNormals, bool bUpdateUVs);
+	void UpdateAttributes(const FDynamicMesh3* MeshIn, FMeshDescription& MeshOut, bool bUpdateNormals, bool bUpdateTangents, bool bUpdateUVs);
 
 	/**
-	 * Update the Tangent and BinormalSign attributes of the MeshDescription, assuming mesh topology has not changed. Does not modify any other attributes.
+	 * Use the TMeshTangents to update the Tangent and BinormalSign attributes of the MeshDescription, assuming mesh topology has not changed. Does not modify any other attributes.
+	 *  NOTE: this ignores any tangent or bitangent overlays on the MeshIn, and instead uses the tangent and bitangent information
+	 *        stored in the TMeshTangents
 	 *	NOTE: assumes the order of triangles in the MeshIn correspond to the ordering you'd get by iterating over triangles, on MeshOut
 	 *		  This matches conversion currently used in MeshDescriptionToDynamicMesh.cpp, but if that changes we will need to change this function to match!
 	 */
-	void UpdateTangents(const FDynamicMesh3* MeshIn, FMeshDescription& MeshOut, const TMeshTangents<double>* Tangents);
+	void UpdateTangents(const FDynamicMesh3* MeshIn, FMeshDescription& MeshOut, const TMeshTangents<double>* SrcTangents);
+
+
+	/**
+	 * Use the MeshIn Overlays to update the Tangent and BinormalSign attributes of the MeshDescription, assuming mesh topology has not changed. Does not modify any other attributes.
+	 *	NOTE: assumes the order of triangles in the MeshIn correspond to the ordering you'd get by iterating over triangles, on MeshOut
+	 *		  This matches conversion currently used in MeshDescriptionToDynamicMesh.cpp, but if that changes we will need to change this function to match!
+	 */
+	void UpdateTangents(const FDynamicMesh3* MeshIn, FMeshDescription& MeshOut);
+
 
 	//
 	// Internal functions that you can also call directly
@@ -108,6 +129,9 @@ public:
 	/**
 	 * Convert while minimizing VertexInstance count, IE new VertexInstances are only created 
 	 * if a unique UV or Normal is required.
+	 * 
+	 * Note: This doesn't copy any tangents from the FDynamicMesh
+	 * Note: This conversion is not currently being used. It is unclear if all consumers of FMeshDescription can handle such shared vertex instances 
 	 */
 	void Convert_SharedInstances(const FDynamicMesh3* MeshIn, FMeshDescription& MeshOut);
 
@@ -115,8 +139,13 @@ public:
 	 * Convert with no shared VertexInstances. A new VertexInstance is created for
 	 * each triangle vertex (ie corner). However vertex positions are shared and
      * the shared UVs structures are populated.
+	 * 
+	 * @param bCopyTangents - should tangent, bitangent overlays exist on the MeshIn, this will use them when producing the MeshDescription MeshOut
+	 *
+	 * Note: Don't copy tangents if the resulting MeshDescription corresponds to a StaticMesh with autogenerated tangents
+	 *       in this case the MeshDescription is expected to have empty tangents.
 	 */
-	void Convert_NoSharedInstances(const FDynamicMesh3* MeshIn, FMeshDescription& MeshOut);
+	void Convert_NoSharedInstances(const FDynamicMesh3* MeshIn, FMeshDescription& MeshOut, bool bCopyTangents);
 
 
 
diff --git a/Engine/Plugins/Experimental/GeometryProcessing/Source/MeshConversion/Public/MeshConversionOptions.h b/Engine/Plugins/Experimental/GeometryProcessing/Source/MeshConversion/Public/MeshConversionOptions.h
--- a/Engine/Plugins/Experimental/GeometryProcessing/Source/MeshConversion/Public/MeshConversionOptions.h
+++ b/Engine/Plugins/Experimental/GeometryProcessing/Source/MeshConversion/Public/MeshConversionOptions.h
@@ -18,6 +18,9 @@ public:
 	/** Should normals of MeshDescription be updated, if available and relevant */
 	bool bUpdateNormals = true;
 
+	/** Should Tangents and BiTangentSign of MeshDescription be updated, if available and relevant */
+	bool bUpdateTangents = false;
+
 	/** Should UVs of MeshDescription be updated, if available and relevant */
 	bool bUpdateUVs = false;
 };
\ No newline at end of file
diff --git a/Engine/Plugins/Experimental/GeometryProcessing/Source/MeshConversion/Public/MeshDescriptionBuilder.h b/Engine/Plugins/Experimental/GeometryProcessing/Source/MeshConversion/Public/MeshDescriptionBuilder.h
--- a/Engine/Plugins/Experimental/GeometryProcessing/Source/MeshConversion/Public/MeshDescriptionBuilder.h
+++ b/Engine/Plugins/Experimental/GeometryProcessing/Source/MeshConversion/Public/MeshDescriptionBuilder.h
@@ -63,6 +63,9 @@ public:
 	/** Set the Normal of a vertex instance*/
 	void SetInstanceNormal(const FVertexInstanceID& InstanceID, const FVector& Normal);
 
+	/** Set the full tangent space of a vertex instance, this is stored as a vec3 normal and a vec4 tangent */
+	void SetInstanceTangentSpace(const FVertexInstanceID& InstanceID, const FVector& Normal, const FVector& Tangent, float Sign);
+
 	/** Set the Color of a vertex instance*/
 	void SetInstanceColor(const FVertexInstanceID& InstanceID, const FVector4& Color);
 
@@ -120,6 +123,8 @@ protected:
 	TVertexAttributesRef<FVector> VertexPositions;
 	TVertexInstanceAttributesRef<FVector2D> InstanceUVs;
 	TVertexInstanceAttributesRef<FVector> InstanceNormals;
+	TVertexInstanceAttributesRef<FVector> InstanceTangents;
+	TVertexInstanceAttributesRef<float> InstanceBiTangentSign;
 	TVertexInstanceAttributesRef<FVector4> InstanceColors;
 
 	TArray<TUVAttributesRef<FVector2D>> UVCoordinateLayers; 
diff --git a/Engine/Plugins/Experimental/MeshModelingToolset/Source/MeshModelingTools/Private/EditNormalsTool.cpp b/Engine/Plugins/Experimental/MeshModelingToolset/Source/MeshModelingTools/Private/EditNormalsTool.cpp
--- a/Engine/Plugins/Experimental/MeshModelingToolset/Source/MeshModelingTools/Private/EditNormalsTool.cpp
+++ b/Engine/Plugins/Experimental/MeshModelingToolset/Source/MeshModelingTools/Private/EditNormalsTool.cpp
@@ -304,7 +304,10 @@ void UEditNormalsTool::GenerateAsset(const TArray<FDynamicMeshOpResult>& Results
 			else
 			{
 				// otherwise just copy attributes
-				Converter.UpdateAttributes(Results[ComponentIdx].Mesh.Get(), *CommitParams.MeshDescription, true, false);
+				const bool bUpdateNormals = true;
+				const bool bCopyOverlayTangents = false;
+				const bool bCopyOverlayUVs = false;
+				Converter.UpdateAttributes(Results[ComponentIdx].Mesh.Get(), *CommitParams.MeshDescription, bUpdateNormals, bCopyOverlayTangents, bCopyOverlayUVs);
 			}
 		});
 	}
diff --git a/Engine/Plugins/Experimental/MeshModelingToolset/Source/MeshModelingToolsEditorOnly/Private/MeshTangentsTool.cpp b/Engine/Plugins/Experimental/MeshModelingToolset/Source/MeshModelingToolsEditorOnly/Private/MeshTangentsTool.cpp
--- a/Engine/Plugins/Experimental/MeshModelingToolset/Source/MeshModelingToolsEditorOnly/Private/MeshTangentsTool.cpp
+++ b/Engine/Plugins/Experimental/MeshModelingToolset/Source/MeshModelingToolsEditorOnly/Private/MeshTangentsTool.cpp
@@ -135,7 +135,12 @@ void UMeshTangentsTool::Shutdown(EToolShutdownType ShutdownType)
 		ComponentTarget->CommitMesh([&](const FPrimitiveComponentTarget::FCommitParams& CommitParams)
 		{
 			FDynamicMeshToMeshDescription Converter;
-			Converter.UpdateAttributes(InputMesh.Get(), *CommitParams.MeshDescription, true, false);
+			{
+				const bool bCopyNormals = true; 
+				const bool bCopyOverlayTangents = false;
+				const bool bCopyOverlayUVs = false;
+				Converter.UpdateAttributes(InputMesh.Get(), *CommitParams.MeshDescription, bCopyNormals, bCopyOverlayTangents, bCopyOverlayUVs);
+			}
 			Converter.UpdateTangents(InputMesh.Get(), *CommitParams.MeshDescription, Tangents.Get());
 		});
 
diff --git a/Engine/Plugins/Experimental/MeshModelingToolset/Source/ModelingComponents/Private/SimpleDynamicMeshComponent.cpp b/Engine/Plugins/Experimental/MeshModelingToolset/Source/ModelingComponents/Private/SimpleDynamicMeshComponent.cpp
--- a/Engine/Plugins/Experimental/MeshModelingToolset/Source/ModelingComponents/Private/SimpleDynamicMeshComponent.cpp
+++ b/Engine/Plugins/Experimental/MeshModelingToolset/Source/ModelingComponents/Private/SimpleDynamicMeshComponent.cpp
@@ -176,11 +176,11 @@ void USimpleDynamicMeshComponent::Bake(FMeshDescription* MeshDescription, bool b
 	{
 		if (ConversionOptions.bUpdatePositions)
 		{
-			Converter.Update(Mesh.Get(), *MeshDescription, ConversionOptions.bUpdateNormals, ConversionOptions.bUpdateUVs);
+			Converter.Update(Mesh.Get(), *MeshDescription, ConversionOptions.bUpdateNormals, ConversionOptions.bUpdateTangents, ConversionOptions.bUpdateUVs);
 		}
 		else
 		{
-			Converter.UpdateAttributes(Mesh.Get(), *MeshDescription, ConversionOptions.bUpdateNormals, ConversionOptions.bUpdateUVs);
+			Converter.UpdateAttributes(Mesh.Get(), *MeshDescription, ConversionOptions.bUpdateNormals, ConversionOptions.bUpdateTangents, ConversionOptions.bUpdateUVs);
 		}
 	}
 	else
