Add a n-ary subtract option to the voxel blend tool

diff --git a/Engine/Plugins/Experimental/GeometryProcessing/Source/GeometricObjects/Public/Implicit/Blend.h b/Engine/Plugins/Experimental/GeometryProcessing/Source/GeometricObjects/Public/Implicit/Blend.h
--- a/Engine/Plugins/Experimental/GeometryProcessing/Source/GeometricObjects/Public/Implicit/Blend.h
+++ b/Engine/Plugins/Experimental/GeometryProcessing/Source/GeometricObjects/Public/Implicit/Blend.h
@@ -40,6 +40,9 @@ public:
 	// size of the cells used when meshing the output (marching cubes' cube size)
 	double MeshCellSize = 1.0;
 
+	// if true, perform a smooth subtraction instead of a smooth union
+	bool bSubtract = false;
+
 	// Set cell sizes to hit the target voxel counts along the max dimension of the bounds
 	void SetCellSizesAndFalloff(FAxisAlignedBox3d Bounds, double BlendFalloffIn, int TargetInputVoxelCount, int TargetOutputVoxelCount)
 	{
@@ -180,6 +183,7 @@ protected:
 
 		TSkeletalRicciNaryBlend3<TDistanceFieldToSkeletalField<TTriLinearGridInterpolant<TCachingMeshSDF<TriangleMeshType>>, double>, double> Blend;
 		Blend.BlendPower = BlendPower;
+		Blend.bSubtract = bSubtract;
 		Blend.Children.Reserve(Sources.Num());
 		TArray<TTriLinearGridInterpolant<TCachingMeshSDF<TriangleMeshType>>> Interpolants; Interpolants.Reserve(Sources.Num());
 		TArray<TDistanceFieldToSkeletalField<TTriLinearGridInterpolant<TCachingMeshSDF<TriangleMeshType>>, double>> SkeletalFields; SkeletalFields.Reserve(Sources.Num());
diff --git a/Engine/Plugins/Experimental/GeometryProcessing/Source/GeometricObjects/Public/Implicit/ImplicitFunctions.h b/Engine/Plugins/Experimental/GeometryProcessing/Source/GeometricObjects/Public/Implicit/ImplicitFunctions.h
--- a/Engine/Plugins/Experimental/GeometryProcessing/Source/GeometricObjects/Public/Implicit/ImplicitFunctions.h
+++ b/Engine/Plugins/Experimental/GeometryProcessing/Source/GeometricObjects/Public/Implicit/ImplicitFunctions.h
@@ -161,43 +161,50 @@ struct TDistanceFieldToSkeletalField
 
 /**
  * Boolean Union of N implicit functions (F_1 or F_2 or ... or F_N)
- * Assumption is that all have surface at zero isocontour and 
- * negative is inside.
+ * Or if bSubtract=true, Subtraction: (F_1 - (F_2 or F_3 or ... or F_N))
+ * We expect this to be used with "Skeletal Field" inputs (e.g. TDistanceFieldToSkeletalField, TSkeletalImplicitLine3, etc)
+ * where values are in the range [0,1] and the isosurface is ~ in the middle of that range
 **/
 template<typename InputBoundedImplicitType, typename RealType>
 struct TSkeletalRicciNaryBlend3
 {
 	TArray<InputBoundedImplicitType*> Children;
 	RealType BlendPower = 2.0;
+	bool bSubtract = false;
 
 	RealType Value(const FVector3<RealType> Pt)
 	{
 		int N = Children.Num();
-		RealType f = 0;
+		checkSlow(N > 0);
+		RealType f = Children[0]->Value(Pt);
+		RealType Scale = bSubtract ? -1 : 1;
 		if (BlendPower == 1.0)
 		{
-			for (int k = 0; k < N; ++k)
+			for (int k = 1; k < N; ++k)
 			{
-				f += Children[k]->Value(Pt);
+				f += Children[k]->Value(Pt) * Scale;
 			}
+			f = TMathUtil<RealType>::Max(0, f);
 		}
 		else if (BlendPower == 2.0)
 		{
-			for (int k = 0; k < N; ++k)
+			f = f * f;
+			for (int k = 1; k < N; ++k)
 			{
 				RealType v = Children[k]->Value(Pt);
-				f += v * v;
+				f += v * v * Scale;
 			}
-			f = TMathUtil<RealType>::Sqrt(f);
+			f = TMathUtil<RealType>::Sqrt(TMathUtil<RealType>::Max(0, f));
 		}
 		else
 		{
-			for (int k = 0; k < N; ++k)
+			f = TMathUtil<RealType>::Pow(f, BlendPower);
+			for (int k = 1; k < N; ++k)
 			{
 				RealType v = Children[k]->Value(Pt);
-				f += TMathUtil<RealType>::Pow(v, BlendPower);
+				f += TMathUtil<RealType>::Pow(v, BlendPower) * Scale;
 			}
-			f = TMathUtil<RealType>::Pow(f, 1.0 / BlendPower);
+			f = TMathUtil<RealType>::Pow(TMathUtil<RealType>::Max(0, f), 1.0 / BlendPower);
 		}
 		return f;
 	}
diff --git a/Engine/Plugins/Experimental/MeshModelingToolset/Source/MeshModelingTools/Private/VoxelBlendMeshesTool.cpp b/Engine/Plugins/Experimental/MeshModelingToolset/Source/MeshModelingTools/Private/VoxelBlendMeshesTool.cpp
--- a/Engine/Plugins/Experimental/MeshModelingToolset/Source/MeshModelingTools/Private/VoxelBlendMeshesTool.cpp
+++ b/Engine/Plugins/Experimental/MeshModelingToolset/Source/MeshModelingTools/Private/VoxelBlendMeshesTool.cpp
@@ -69,6 +69,7 @@ TUniquePtr<FDynamicMeshOperator> UVoxelBlendMeshesTool::MakeNewOperator()
 
 	Op->BlendFalloff = BlendProperties->BlendFalloff;
 	Op->BlendPower = BlendProperties->BlendPower;
+	Op->bSubtract = BlendProperties->Operation == EVoxelBlendOperation::Subtract;
 	Op->bVoxWrap = BlendProperties->bVoxWrap;
 	Op->bRemoveInternalsAfterVoxWrap = BlendProperties->bRemoveInternalsAfterVoxWrap;
 	Op->ThickenShells = BlendProperties->ThickenShells;
diff --git a/Engine/Plugins/Experimental/MeshModelingToolset/Source/MeshModelingTools/Public/VoxelBlendMeshesTool.h b/Engine/Plugins/Experimental/MeshModelingToolset/Source/MeshModelingTools/Public/VoxelBlendMeshesTool.h
--- a/Engine/Plugins/Experimental/MeshModelingToolset/Source/MeshModelingTools/Public/VoxelBlendMeshesTool.h
+++ b/Engine/Plugins/Experimental/MeshModelingToolset/Source/MeshModelingTools/Public/VoxelBlendMeshesTool.h
@@ -9,6 +9,18 @@
 #include "VoxelBlendMeshesTool.generated.h"
 
 
+/** CSG-style blend operations */
+UENUM()
+enum class EVoxelBlendOperation : uint8
+{
+	/** Smoothed union of all shapes */
+	Union = 0,
+
+	/** Smoothed subtraction of all shapes from the first selected shape */
+	Subtract = 1,
+};
+
+
 /**
  * Properties of the blend operation
  */
@@ -26,6 +38,12 @@ public:
 	UPROPERTY(EditAnywhere, Category = Blend, meta = (UIMin = ".1", UIMax = "100", ClampMin = ".001", ClampMax = "1000"))
 	double BlendFalloff = 10;
 
+	/** How to combine the shapes */
+	UPROPERTY(EditAnywhere, Category = Blend)
+	EVoxelBlendOperation Operation = EVoxelBlendOperation::Union;
+
+
+
 	/** Apply a "VoxWrap" operation to input mesh(es) before computing the blend.  Fixes results for inputs with holes and/or self-intersections. */
 	UPROPERTY(EditAnywhere, Category = VoxWrapPreprocess)
 	bool bVoxWrap = false;
diff --git a/Engine/Plugins/Experimental/MeshModelingToolset/Source/ModelingOperators/Private/CompositionOps/VoxelBlendMeshesOp.cpp b/Engine/Plugins/Experimental/MeshModelingToolset/Source/ModelingOperators/Private/CompositionOps/VoxelBlendMeshesOp.cpp
--- a/Engine/Plugins/Experimental/MeshModelingToolset/Source/ModelingOperators/Private/CompositionOps/VoxelBlendMeshesOp.cpp
+++ b/Engine/Plugins/Experimental/MeshModelingToolset/Source/ModelingOperators/Private/CompositionOps/VoxelBlendMeshesOp.cpp
@@ -33,6 +33,7 @@ void FVoxelBlendMeshesOp::CalculateResult(FProgressCancel* Progress)
 	}
 
 	TImplicitBlend<FDynamicMesh3> ImplicitBlend;
+	ImplicitBlend.bSubtract = bSubtract;
 
 	TArray<FDynamicMesh3> TransformedMeshes; TransformedMeshes.Reserve(Meshes.Num());
 	FAxisAlignedBox3d CombinedBounds = FAxisAlignedBox3d::Empty();
diff --git a/Engine/Plugins/Experimental/MeshModelingToolset/Source/ModelingOperators/Public/CompositionOps/VoxelBlendMeshesOp.h b/Engine/Plugins/Experimental/MeshModelingToolset/Source/ModelingOperators/Public/CompositionOps/VoxelBlendMeshesOp.h
--- a/Engine/Plugins/Experimental/MeshModelingToolset/Source/ModelingOperators/Public/CompositionOps/VoxelBlendMeshesOp.h
+++ b/Engine/Plugins/Experimental/MeshModelingToolset/Source/ModelingOperators/Public/CompositionOps/VoxelBlendMeshesOp.h
@@ -26,6 +26,8 @@ public:
 	bool bRemoveInternalsAfterVoxWrap = false;
 	double ThickenShells = 0.0;
 
+	bool bSubtract = false;
+
 	void SetTransform(const FTransform& Transform);
 
 	//
