GeometryProcessing: refactor SmallListSet::ValueIterator into two separate iterators, one that has a map function (MappedValueIterator) and one that does not (ValueIterator). Having a null TFunction and testing in the operator* function appears to introduce significant overhead, splitting to two separate iterators results in FDynamicMesh3::VtxEdgesItr() edge iteration being about 30% faster in my profiling. Presumably this will improve all other Values() iterations, of which there are many.
Changed many check() to checkSlow() in FDynamicMesh3. These are not stripped in shipping builds, and are a high cost to pay on virtually every mesh access, when the errors are unrecoverable anyway, and generally require running in a Debug build to have a usable stack.
#rnx

diff --git a/Engine/Plugins/Experimental/GeometryProcessing/Source/DynamicMesh/Private/MeshIndexUtil.cpp b/Engine/Plugins/Experimental/GeometryProcessing/Source/DynamicMesh/Private/MeshIndexUtil.cpp
--- a/Engine/Plugins/Experimental/GeometryProcessing/Source/DynamicMesh/Private/MeshIndexUtil.cpp
+++ b/Engine/Plugins/Experimental/GeometryProcessing/Source/DynamicMesh/Private/MeshIndexUtil.cpp
@@ -59,35 +59,17 @@ void MeshIndexUtil::VertexToTriangleOneRing(const FDynamicMesh3* Mesh, const TAr
 }
 void UE::MeshIndexUtil::VertexToTriangleOneRing(const FDynamicMesh3* Mesh, const TArray<int>& VertexIDs, TSet<int>& TriangleIDsOut)
 {
-	int NumVerts = VertexIDs.Num();
+	// for a TSet it is more efficient to just try to add each triangle twice, than it is to
+	// try to avoid duplicate adds with more complex mesh queries
+	int32 NumVerts = VertexIDs.Num();
 	TriangleIDsOut.Reserve( (NumVerts < 5) ? NumVerts*6 : NumVerts*4);
-	TArray<int32> OneRingT;  // temp storage
-	OneRingT.Reserve(16);
-
-	//auto AddFunc = [&](int32 TriID)
-	//{
-	//	TriangleIDsOut.Add(TriID);
-	//};
-
-	for (int k = 0; k < NumVerts; ++k)
+	for (int32 vid : VertexIDs)
 	{
-		if (Mesh->IsVertex(VertexIDs[k]))
+		Mesh->EnumerateVertexEdges(vid, [&](int32 eid) 
 		{
-			// slowest
-			//for (int TriID : Mesh->VtxTrianglesItr(VertexIDs[k]))
-			//{
-			//	TriangleIDsOut.Add(TriID);
-			//}
-
-			OneRingT.SetNum(0, false);
-			Mesh->GetVertexOneRingTriangles(VertexIDs[k], OneRingT);
-			for (int TriID : OneRingT)
-			{
-				TriangleIDsOut.Add(TriID);
-			}
-
-			// slightly slower?
-			//Mesh->EnumerateVertexTriangles(VertexIDs[k], AddFunc);
-		}
+			FIndex2i EdgeT = Mesh->GetEdgeT(eid);
+			TriangleIDsOut.Add(EdgeT.A);
+			if (EdgeT.B != IndexConstants::InvalidID) TriangleIDsOut.Add(EdgeT.B);
+		});
 	}
 }
\ No newline at end of file
diff --git a/Engine/Plugins/Experimental/GeometryProcessing/Source/DynamicMesh/Public/DynamicMesh3.h b/Engine/Plugins/Experimental/GeometryProcessing/Source/DynamicMesh/Public/DynamicMesh3.h
--- a/Engine/Plugins/Experimental/GeometryProcessing/Source/DynamicMesh/Public/DynamicMesh3.h
+++ b/Engine/Plugins/Experimental/GeometryProcessing/Source/DynamicMesh/Public/DynamicMesh3.h
@@ -358,24 +358,43 @@ public:
 	}
 
 	/** @return enumerable object for one-ring vertex neighbours of a vertex, suitable for use with range-based for, ie for ( int i : VtxVerticesItr(VertexID) ) */
-	FSmallListSet::ValueEnumerable VtxVerticesItr(int VertexID) const
+	FSmallListSet::MappedValueEnumerable VtxVerticesItr(int VertexID) const
 	{
-		check(VertexRefCounts.IsValid(VertexID));
-		return VertexEdgeLists.Values(VertexID,
-		                              [VertexID, this](int eid) { return GetOtherEdgeVertex(eid, VertexID); });
+		checkSlow(VertexRefCounts.IsValid(VertexID));
+		return VertexEdgeLists.MappedValues(VertexID,
+			[VertexID, this](int eid) { return GetOtherEdgeVertex(eid, VertexID); });
 	}
 
+	/** Call VertexFunc for each one-ring vertex neighbour of a vertex. Currently this is more efficient than VtxVerticesItr() due to overhead in the Values() enumerable */
+	void EnumerateVertexVertices(int32 VertexID, TFunctionRef<void(int32)> VertexFunc) const
+	{
+		checkSlow(VertexRefCounts.IsValid(VertexID));
+		VertexEdgeLists.Enumerate(VertexID, [this, &VertexFunc, VertexID](int32 eid)
+		{
+			VertexFunc(GetOtherEdgeVertex(eid, VertexID));
+		});
+	}
+
+
 	/** @return enumerable object for one-ring edges of a vertex, suitable for use with range-based for, ie for ( int i : VtxEdgesItr(VertexID) ) */
 	FSmallListSet::ValueEnumerable VtxEdgesItr(int VertexID) const
 	{
-		check(VertexRefCounts.IsValid(VertexID));
+		checkSlow(VertexRefCounts.IsValid(VertexID));
 		return VertexEdgeLists.Values(VertexID);
 	}
 
+	/** Call EdgeFunc for each one-ring edge of a vertex. Currently this is more efficient than VtxEdgesItr() due to overhead in the Values() enumerable */
+	void EnumerateVertexEdges(int32 VertexID, TFunctionRef<void(int32)> EdgeFunc) const
+	{
+		checkSlow(VertexRefCounts.IsValid(VertexID));
+		VertexEdgeLists.Enumerate(VertexID, EdgeFunc);
+	}
+
+
 	/** @return enumerable object for one-ring triangles of a vertex, suitable for use with range-based for, ie for ( int i : VtxTrianglesItr(VertexID) ) */
 	vtx_triangles_enumerable VtxTrianglesItr(int VertexID) const
 	{
-		check(VertexRefCounts.IsValid(VertexID));
+		checkSlow(VertexRefCounts.IsValid(VertexID));
 		return vtx_triangles_enumerable(VertexEdgeLists.Values(VertexID), [this, VertexID](int EdgeID) {
 			return GetOrderedOneRingEdgeTris(VertexID, EdgeID);
 		});
@@ -460,7 +479,7 @@ public:
 	/** @return the vertex position */
 	inline FVector3d GetVertex(int VertexID) const
 	{
-		check(IsVertex(VertexID));
+		checkSlow(IsVertex(VertexID));
 		return Vertices[VertexID];
 	}
 
@@ -468,7 +487,7 @@ public:
 	inline void SetVertex(int VertexID, const FVector3d& vNewPos)
 	{
 		checkSlow(VectorUtil::IsFinite(vNewPos));
-		check(IsVertex(VertexID));
+		checkSlow(IsVertex(VertexID));
 		if (VectorUtil::IsFinite(vNewPos))
 		{
 			Vertices[VertexID] = vNewPos;
@@ -484,7 +503,7 @@ public:
 	void SetVertex_NoTimeStampUpdate(int VertexID, const FVector3d& vNewPos)
 	{
 		checkSlow(VectorUtil::IsFinite(vNewPos));
-		check(IsVertex(VertexID));
+		checkSlow(IsVertex(VertexID));
 
 		if (VectorUtil::IsFinite(vNewPos))
 		{
@@ -531,21 +550,21 @@ public:
 	/** Get triangle vertices */
 	inline FIndex3i GetTriangle(int TriangleID) const
 	{
-		check(IsTriangle(TriangleID));
+		checkSlow(IsTriangle(TriangleID));
 		return Triangles[TriangleID];
 	}
 
 	/** Get triangle edges */
 	inline FIndex3i GetTriEdges(int TriangleID) const
 	{
-		check(IsTriangle(TriangleID));
+		checkSlow(IsTriangle(TriangleID));
 		return TriangleEdges[TriangleID];
 	}
 
 	/** Get one of the edges of a triangle */
 	inline int GetTriEdge(int TriangleID, int j) const
 	{
-		check(IsTriangle(TriangleID));
+		checkSlow(IsTriangle(TriangleID));
 		return TriangleEdges[TriangleID][j];
 	}
 
@@ -570,21 +589,21 @@ public:
 	/** Get the vertices and triangles of an edge, returned as [v0,v1,t0,t1], where t1 may be InvalidID */
 	inline FEdge GetEdge(int EdgeID) const
 	{
-		check(IsEdge(EdgeID));
+		checkSlow(IsEdge(EdgeID));
 		return Edges[EdgeID];
 	}
 
 	/** Get the vertex pair for an edge */
 	inline FIndex2i GetEdgeV(int EdgeID) const
 	{
-		check(IsEdge(EdgeID));
+		checkSlow(IsEdge(EdgeID));
 		return Edges[EdgeID].Vert;
 	}
 
 	/** Get the vertex positions of an edge */
 	inline bool GetEdgeV(int EdgeID, FVector3d& a, FVector3d& b) const
 	{
-		check(IsEdge(EdgeID));
+		checkSlow(IsEdge(EdgeID));
 
 		const FIndex2i Verts = Edges[EdgeID].Vert;
 
@@ -597,7 +616,7 @@ public:
 	/** Get the triangle pair for an edge. The second triangle may be InvalidID */
 	inline FIndex2i GetEdgeT(int EdgeID) const
 	{
-		check(IsEdge(EdgeID));
+		checkSlow(IsEdge(EdgeID));
 		return Edges[EdgeID].Tri;
 	}
 
@@ -737,7 +756,7 @@ public:
 	/** Returns true if edge is on the mesh boundary, ie only connected to one triangle */
 	inline bool IsBoundaryEdge(int EdgeID) const
 	{
-		check(IsEdge(EdgeID));
+		checkSlow(IsEdge(EdgeID));
 		return Edges[EdgeID].Tri[1] == InvalidID;
 	}
 
diff --git a/Engine/Plugins/Experimental/GeometryProcessing/Source/GeometricObjects/Public/Curve/DynamicGraph.h b/Engine/Plugins/Experimental/GeometryProcessing/Source/GeometricObjects/Public/Curve/DynamicGraph.h
--- a/Engine/Plugins/Experimental/GeometryProcessing/Source/GeometricObjects/Public/Curve/DynamicGraph.h
+++ b/Engine/Plugins/Experimental/GeometryProcessing/Source/GeometricObjects/Public/Curve/DynamicGraph.h
@@ -113,10 +113,10 @@ public:
 	/**
 	 * Enumerate "other" vertices of edges connected to vertex (i.e. vertex one-ring)
 	 */
-	FSmallListSet::ValueEnumerable VtxVerticesItr(int VID) const
+	FSmallListSet::MappedValueEnumerable VtxVerticesItr(int VID) const
 	{
 		check(vertices_refcount.IsValid(VID));
-		return vertex_edges.Values(VID, [VID, this](int EID) { return edge_other_v(EID, VID); });
+		return vertex_edges.MappedValues(VID, [VID, this](int EID) { return edge_other_v(EID, VID); });
 	}
 
 	/**
diff --git a/Engine/Plugins/Experimental/GeometryProcessing/Source/GeometricObjects/Public/Util/SmallListSet.h b/Engine/Plugins/Experimental/GeometryProcessing/Source/GeometricObjects/Public/Util/SmallListSet.h
--- a/Engine/Plugins/Experimental/GeometryProcessing/Source/GeometricObjects/Public/Util/SmallListSet.h
+++ b/Engine/Plugins/Experimental/GeometryProcessing/Source/GeometricObjects/Public/Util/SmallListSet.h
@@ -170,47 +170,29 @@ public:
 
 
 	friend class ValueIterator;
+	friend class BaseValueIterator;
 
 	/**
-	 * ValueIterator iterates over the values of a small list
-	 * An optional mapping function can be provided which will then be applied to the values returned by the * operator
+	 * BaseValueIterator is a base class for ValueIterator and MappedValueIterator below.
 	 */
-	class ValueIterator
+	class BaseValueIterator
 	{
 	public:
-		inline ValueIterator() 
+		BaseValueIterator()
 		{ 
 			ListSet = nullptr;
-			MapFunc = nullptr; 
 			ListIndex = 0; 
 		}
 
-		inline bool operator==(const ValueIterator& Other) const 
+		inline bool operator==(const BaseValueIterator& Other) const
 		{
 			return ListSet == Other.ListSet && ListIndex == Other.ListIndex;
 		}
-		inline bool operator!=(const ValueIterator& Other) const 
+		inline bool operator!=(const BaseValueIterator& Other) const
 		{
 			return ListSet != Other.ListSet || ListIndex != Other.ListIndex || iCur != Other.iCur || cur_ptr != Other.cur_ptr;
 		}
 
-		inline int32 operator*() const 
-		{
-			return (MapFunc == nullptr) ? cur_value : MapFunc(cur_value);
-		}
-
-		inline const ValueIterator& operator++() 		// prefix
-		{
-			this->GotoNext();
-			return *this;
-		}
-		//inline ValueIterator operator++(int32) {		// postfix
-		//	index_iterator copy(*this);
-		//	this->GotoNext();
-		//	return copy;
-		//}
-
-
 	protected:
 		inline void GotoNext() 
 		{
@@ -240,14 +222,12 @@ public:
 			}
 		}
 
-		inline ValueIterator(
+		BaseValueIterator(
 			const FSmallListSet* ListSetIn,
 			int32 ListIndex, 
-			bool is_end,
-			const TFunction<int32(int32)>& MapFuncIn = nullptr)
+			bool is_end)
 		{
 			this->ListSet = ListSetIn;
-			this->MapFunc = MapFuncIn;
 			this->ListIndex = ListIndex;
 			if (is_end) 
 			{
@@ -280,7 +260,6 @@ public:
 		}
 
 		const FSmallListSet * ListSet;
-		TFunction<int32(int32)> MapFunc;
 		int32 ListIndex;
 		int32 block_ptr;
 		int32 N;
@@ -291,20 +270,39 @@ public:
 		friend class FSmallListSet;
 	};
 
+
 	/**
-	 * @return iterator for start of list at ListIndex
+	 * ValueIterator iterates over the values of a small list
 	 */
-	inline ValueIterator BeginValues(int32 ListIndex) const 
+	class ValueIterator : public BaseValueIterator
 	{
-		return ValueIterator(this, ListIndex, false);
-	}
+	public:
+		ValueIterator() : BaseValueIterator() {}
+
+		inline int32 operator*() const
+		{
+			return cur_value;
+		}
+
+		inline const ValueIterator& operator++() 		// prefix
+		{
+			this->GotoNext();
+			return *this;
+		}
+
+	protected:
+		ValueIterator(const FSmallListSet* ListSetIn, int32 ListIndex, bool is_end) 
+			: BaseValueIterator(ListSetIn, ListIndex, is_end) {}
+
+		friend class FSmallListSet;
+	};
 
 	/**
-	 * @return iterator for start of list at ListIndex, with given value mapping function
+	 * @return iterator for start of list at ListIndex
 	 */
-	inline ValueIterator BeginValues(int32 ListIndex, const TFunction<int32(int32)>& MapFunc) const 
+	inline ValueIterator BeginValues(int32 ListIndex) const 
 	{
-		return ValueIterator(this, ListIndex, false, MapFunc);
+		return ValueIterator(this, ListIndex, false);
 	}
 
 	/**
@@ -323,15 +321,13 @@ public:
 	public:
 		const FSmallListSet* ListSet;
 		int32 ListIndex;
-		TFunction<int32(int32)> MapFunc;
 		ValueEnumerable() {}
-		ValueEnumerable(const FSmallListSet* ListSetIn, int32 ListIndex, TFunction<int32(int32)> MapFunc = nullptr)
+		ValueEnumerable(const FSmallListSet* ListSetIn, int32 ListIndex)
 		{
 			this->ListSet = ListSetIn;
 			this->ListIndex = ListIndex;
-			this->MapFunc = MapFunc;
 		}
-		typename FSmallListSet::ValueIterator begin() const { return ListSet->BeginValues(ListIndex, MapFunc); }
+		typename FSmallListSet::ValueIterator begin() const { return ListSet->BeginValues(ListIndex); }
 		typename FSmallListSet::ValueIterator end() const { return ListSet->EndValues(ListIndex); }
 	};
 
@@ -343,17 +339,101 @@ public:
 		return ValueEnumerable(this, ListIndex);
 	}
 
+
+
+
+
+
+	//
+	// mapped iterator support - mapped iterator applies an arbitrary function to the iterator value
+	// 
+
+	friend class MappedValueIterator;
+
+	/**
+	 * MappedValueIterator iterates over the values of a small list
+	 * An optional mapping function can be provided which will then be applied to the values returned by the * operator
+	 */
+	class MappedValueIterator : public BaseValueIterator
+	{
+	public:
+		MappedValueIterator() : BaseValueIterator() 
+		{
+			MapFunc = [](int32 value) { return value; };
+		}
+
+		inline int32 operator*() const
+		{
+			return MapFunc(cur_value);
+		}
+
+		inline const MappedValueIterator& operator++() 		// prefix
+		{
+			this->GotoNext();
+			return *this;
+		}
+
+	protected:
+		MappedValueIterator(const FSmallListSet* ListSetIn, int32 ListIndex, bool is_end, TFunction<int32(int32)> MapFuncIn)
+			: BaseValueIterator(ListSetIn, ListIndex, is_end) 
+		{
+			MapFunc = MapFuncIn;
+		}
+
+		TFunction<int32(int32)> MapFunc;
+		friend class FSmallListSet;
+	};
+
+	/**
+	 * @return iterator for start of list at ListIndex, with given value mapping function
+	 */
+	inline MappedValueIterator BeginMappedValues(int32 ListIndex, const TFunction<int32(int32)>& MapFunc) const
+	{
+		return MappedValueIterator(this, ListIndex, false, MapFunc);
+	}
+
+	/**
+	 * @return iterator for end of list at ListIndex, with given value mapping function
+	 */
+	inline MappedValueIterator EndMappedValues(int32 ListIndex, const TFunction<int32(int32)>& MapFunc) const
+	{
+		return MappedValueIterator(this, ListIndex, true, MapFunc);
+	}
+
+	/**
+	 * MappedValueEnumerable is an object that provides begin/end semantics for a small list, suitable for use with a range-based for loop
+	 */
+	class MappedValueEnumerable
+	{
+	public:
+		const FSmallListSet* ListSet;
+		int32 ListIndex;
+		TFunction<int32(int32)> MapFunc;
+		MappedValueEnumerable() {}
+		MappedValueEnumerable(const FSmallListSet* ListSetIn, int32 ListIndex, TFunction<int32(int32)> MapFunc)
+		{
+			this->ListSet = ListSetIn;
+			this->ListIndex = ListIndex;
+			this->MapFunc = MoveTemp(MapFunc);
+		}
+		typename FSmallListSet::MappedValueIterator begin() const { return ListSet->BeginMappedValues(ListIndex, MapFunc); }
+		typename FSmallListSet::MappedValueIterator end() const { return ListSet->EndMappedValues(ListIndex, MapFunc); }
+	};
+
 	/**
 	 * @return a value enumerable for the given ListIndex, with the given value mapping function
 	 */
-	inline ValueEnumerable Values(int32 ListIndex, const TFunction<int32(int32)>& MapFunc) const 
+	inline MappedValueEnumerable MappedValues(int32 ListIndex, TFunction<int32(int32)> MapFunc) const
 	{
-		return ValueEnumerable(this, ListIndex, MapFunc);
+		return MappedValueEnumerable(this, ListIndex, MapFunc);
 	}
 
 
 
 
+
+
+
 protected:
 
 
@@ -380,5 +460,6 @@ public:
 			ListBlocks.GetLength(), (LinkedListElements.GetLength() * sizeof(int32) / 1024));
 	}
 
-
 };
+
+
