GeometryProcessing: minor optimizations of functions on dynamic sculpt critical paths. Use mesh Enumerate functions instead of enumerables. Add IndexUtil::FindTriOtherVtxUnsafe() which is more efficient.
#rnx

diff --git a/Engine/Plugins/Experimental/GeometryProcessing/Source/DynamicMesh/Private/DynamicMesh3_Queries.cpp b/Engine/Plugins/Experimental/GeometryProcessing/Source/DynamicMesh/Private/DynamicMesh3_Queries.cpp
--- a/Engine/Plugins/Experimental/GeometryProcessing/Source/DynamicMesh/Private/DynamicMesh3_Queries.cpp
+++ b/Engine/Plugins/Experimental/GeometryProcessing/Source/DynamicMesh/Private/DynamicMesh3_Queries.cpp
@@ -5,13 +5,15 @@
 
 FIndex2i FDynamicMesh3::GetEdgeOpposingV(int eID) const
 {
-	// [TODO] there was a comment here saying this does more work than necessary??
+	const FEdge& Edge = Edges[eID];
+	int a = Edge.Vert[0];
+	int b = Edge.Vert[1];
+
 	// ** it is important that verts returned maintain [c,d] order!!
-	const FEdge Edge = Edges[eID];
-	int c = IndexUtil::FindTriOtherVtx(Edge.Vert[0], Edge.Vert[1], Triangles, Edge.Tri[0]);
+	int c = IndexUtil::FindTriOtherVtxUnsafe(a, b, Triangles[Edge.Tri[0]]);
 	if (Edge.Tri[1] != InvalidID)
 	{
-		int d = IndexUtil::FindTriOtherVtx(Edge.Vert[0], Edge.Vert[1], Triangles, Edge.Tri[1]);
+		int d = IndexUtil::FindTriOtherVtxUnsafe(a, b, Triangles[Edge.Tri[1]]);
 		return FIndex2i(c, d);
 	}
 	else
diff --git a/Engine/Plugins/Experimental/GeometryProcessing/Source/DynamicMesh/Private/MeshRefinerBase.cpp b/Engine/Plugins/Experimental/GeometryProcessing/Source/DynamicMesh/Private/MeshRefinerBase.cpp
--- a/Engine/Plugins/Experimental/GeometryProcessing/Source/DynamicMesh/Private/MeshRefinerBase.cpp
+++ b/Engine/Plugins/Experimental/GeometryProcessing/Source/DynamicMesh/Private/MeshRefinerBase.cpp
@@ -306,10 +306,10 @@ void FMeshRefinerBase::SaveVertexTrianglesBeforeModify(int32 VertexID)
 {
 	if (ActiveChangeTracker)
 	{
-		for (int32 TriangleID : Mesh->VtxTrianglesItr(VertexID))
+		Mesh->EnumerateVertexTriangles(VertexID, [&](int32 TriangleID)
 		{
-			ActiveChangeTracker->SaveTriangle(TriangleID, true );
-		}
+			ActiveChangeTracker->SaveTriangle(TriangleID, true);
+		});
 	}
 }
 
diff --git a/Engine/Plugins/Experimental/GeometryProcessing/Source/DynamicMesh/Public/DynamicMeshChangeTracker.h b/Engine/Plugins/Experimental/GeometryProcessing/Source/DynamicMesh/Public/DynamicMeshChangeTracker.h
--- a/Engine/Plugins/Experimental/GeometryProcessing/Source/DynamicMesh/Public/DynamicMeshChangeTracker.h
+++ b/Engine/Plugins/Experimental/GeometryProcessing/Source/DynamicMesh/Public/DynamicMeshChangeTracker.h
@@ -383,10 +383,10 @@ void FDynamicMeshChangeTracker::SaveTrianglesAndNeighbourTris(EnumerableType Tri
 
 void FDynamicMeshChangeTracker::SaveVertexOneRingTriangles(int32 VertexID, bool bSaveVertices)
 {
-	for (int32 TriangleID : Mesh->VtxTrianglesItr(VertexID))
+	Mesh->EnumerateVertexTriangles(VertexID, [this, bSaveVertices](int32 TriangleID)
 	{
 		SaveTriangle(TriangleID, bSaveVertices);
-	}
+	});
 }
 
 template<typename EnumerableType>
diff --git a/Engine/Plugins/Experimental/GeometryProcessing/Source/GeometricObjects/Public/Util/IndexUtil.h b/Engine/Plugins/Experimental/GeometryProcessing/Source/GeometricObjects/Public/Util/IndexUtil.h
--- a/Engine/Plugins/Experimental/GeometryProcessing/Source/GeometricObjects/Public/Util/IndexUtil.h
+++ b/Engine/Plugins/Experimental/GeometryProcessing/Source/GeometricObjects/Public/Util/IndexUtil.h
@@ -87,6 +87,27 @@ namespace IndexUtil
 		return IndexConstants::InvalidID;
 	}
 
+	/**
+	 * Find third vertex of triangle that is not VertexID1 or VertexID2. Triangle must contain VertexID1 and VertexID2 or result will be incorrect.
+	 * @return VertexID of third vertex, or incorrect result if precondition is not met
+	 */
+	template<typename T, typename Vec>
+	inline int FindTriOtherVtxUnsafe(T VertexID1, T VertexID2, const Vec& TriangleVerts)
+	{
+		if (TriangleVerts[0] == VertexID1)
+		{
+			return (TriangleVerts[1] == VertexID2) ? TriangleVerts[2] : TriangleVerts[1];
+		}
+		else if (TriangleVerts[1] == VertexID1)
+		{
+			return (TriangleVerts[0] == VertexID2) ? TriangleVerts[2] : TriangleVerts[0];
+		}
+		else
+		{
+			return (TriangleVerts[0] == VertexID2) ? TriangleVerts[1] : TriangleVerts[0];
+		}
+	}
+
 
 	/**
 	 * Find ordered edge [VertexID1,VertexID2] in TriangleVerts and then return the remaining third vertex
