GeometryProcessing: add ParallelRangeQuery to Octree. Replace FSmallListSet Values() enumerations with Enumerate() calls.
#rnx

diff --git a/Engine/Plugins/Experimental/GeometryProcessing/Source/DynamicMesh/Public/DynamicMeshOctree3.h b/Engine/Plugins/Experimental/GeometryProcessing/Source/DynamicMesh/Public/DynamicMeshOctree3.h
--- a/Engine/Plugins/Experimental/GeometryProcessing/Source/DynamicMesh/Public/DynamicMeshOctree3.h
+++ b/Engine/Plugins/Experimental/GeometryProcessing/Source/DynamicMesh/Public/DynamicMeshOctree3.h
@@ -112,36 +112,20 @@ public:
 		return RemoveObject(TriangleID);		// will ignore if we do not contain this triangle
 	}
 
-	/**
-	 * Remove a list of triangles into the tree
-	 */
-	void RemoveTriangles(const TArray<int>& Triangles)
-	{
-		int N = Triangles.Num();
-		for ( int i = 0; i < N; ++i )
-		{
-			if (Mesh->IsTriangle(Triangles[i]))
-			{
-				FAxisAlignedBox3d Bounds = Mesh->GetTriBounds(Triangles[i]);
-				ModifiedBounds.Contain(Bounds);
-			}
-			RemoveObject(Triangles[i]);		// will ignore if we do not contain this triangle
-		}
-	}
 
 	/**
 	 * Remove a set of triangles into the tree
 	 */
-	void RemoveTriangles(const TSet<int>& Triangles)
+	template<typename EnumerableType>
+	void RemoveTriangles(const EnumerableType& Triangles, bool bMarkModifiedBounds = true)
 	{
 		for (int TriangleID : Triangles)
 		{
-			if (Mesh->IsTriangle(TriangleID))
+			if (RemoveObject(TriangleID) && bMarkModifiedBounds && Mesh->IsTriangle(TriangleID))
 			{
 				FAxisAlignedBox3d Bounds = Mesh->GetTriBounds(TriangleID);
 				ModifiedBounds.Contain(Bounds);
 			}
-			RemoveObject(TriangleID);		// will ignore if we do not contain this triangle
 		}
 	}
 
@@ -149,7 +133,8 @@ public:
 	/**
 	 * Reinsert a set of triangles into the tree
 	 */
-	void ReinsertTriangles(const TSet<int>& Triangles)
+	template<typename EnumerableType>
+	void ReinsertTriangles(const EnumerableType& Triangles)
 	{
 		for (int TriangleID : Triangles)
 		{
@@ -182,7 +167,8 @@ public:
 	/**
 	 * Include the current bounds of a set of triangles in the ModifiedBounds box
 	 */
-	void NotifyPendingModification(const TSet<int>& Triangles)
+	template<typename EnumerableType>
+	void NotifyPendingModification(const EnumerableType& Triangles)
 	{
 		for (int TriangleID : Triangles)
 		{
diff --git a/Engine/Plugins/Experimental/GeometryProcessing/Source/GeometricObjects/Private/Spatial/SparseDynamicOctree3.cpp b/Engine/Plugins/Experimental/GeometryProcessing/Source/GeometricObjects/Private/Spatial/SparseDynamicOctree3.cpp
--- a/Engine/Plugins/Experimental/GeometryProcessing/Source/GeometricObjects/Private/Spatial/SparseDynamicOctree3.cpp
+++ b/Engine/Plugins/Experimental/GeometryProcessing/Source/GeometricObjects/Private/Spatial/SparseDynamicOctree3.cpp
@@ -2,6 +2,8 @@
 
 
 #include "Spatial/SparseDynamicOctree3.h"
+#include "Async/Async.h"
+#include "Async/ParallelFor.h"
 
 
 // NB: These have to be here until C++17 allows inline variables
@@ -240,7 +242,7 @@ int32 FSparseDynamicOctree3::FindNearestHitObject(const FRay3d& Ray,
 		const FSparseOctreeCell* CurCell = Queue.Pop(false);
 		
 		// process elements
-		for (int ObjectID : CellObjectLists.Values(CurCell->CellID))
+		CellObjectLists.Enumerate(CurCell->CellID, [&](int32 ObjectID)
 		{
 			double HitDist = HitObjectDistFunc(ObjectID, Ray);
 			if (HitDist < MaxDistance)
@@ -248,7 +250,7 @@ int32 FSparseDynamicOctree3::FindNearestHitObject(const FRay3d& Ray,
 				MaxDistance = HitDist;
 				HitObjectID = ObjectID;
 			}
-		}
+		});
 
 		// descend to child cells
 		// sort by distance? use DDA?
@@ -303,10 +305,10 @@ void FSparseDynamicOctree3::RangeQuery(
 		const FSparseOctreeCell* CurCell = Queue.Pop(false);
 
 		// process elements
-		for (int ObjectID : CellObjectLists.Values(CurCell->CellID))
+		CellObjectLists.Enumerate(CurCell->CellID, [&](int32 ObjectID)
 		{
 			ObjectIDFunc(ObjectID);
-		}
+		});
 
 		for (int k = 0; k < 8; ++k)
 		{
@@ -352,16 +354,15 @@ void FSparseDynamicOctree3::RangeQuery(
 		}
 	});
 
-
 	while (Queue.Num() > 0)
 	{
 		const FSparseOctreeCell* CurCell = Queue.Pop(false);
 
 		// process elements
-		for (int ObjectID : CellObjectLists.Values(CurCell->CellID))
+		CellObjectLists.Enumerate(CurCell->CellID, [&](int32 ObjectID)
 		{
 			ObjectIDs.Add(ObjectID);
-		}
+		});
 
 		for (int k = 0; k < 8; ++k)
 		{
@@ -380,6 +381,82 @@ void FSparseDynamicOctree3::RangeQuery(
 
 
 
+void FSparseDynamicOctree3::ParallelRangeQuery(
+	const FAxisAlignedBox3d& Bounds,
+	TArray<int>& ObjectIDs) const
+{
+	// always collect spill objects
+	for (int ObjectID : SpillObjectSet)
+	{
+		ObjectIDs.Add(ObjectID);
+	}
+
+	TArray<const FSparseOctreeCell*> Queue;
+
+	// start at root cells
+	RootCells.AllocatedIteration([&](const uint32* RootCellID)
+	{
+		const FSparseOctreeCell* RootCell = &Cells[*RootCellID];
+		if (GetCellBox(*RootCell, MaxExpandFactor).Intersects(Bounds))
+		{
+			Queue.Add(&Cells[*RootCellID]);
+		}
+	});
+
+	// parallel strategy here is to collect each rootcell subtree into a separate
+	// array, and then merge them. Although this means we do some dynamic memory
+	// allocation, it's much faster than trying to lock ObjectIDs on every Add(),
+	// which results in crazy lock contention
+	FCriticalSection ObjectIDsLock;
+	ParallelFor(Queue.Num(), [&](int32 qi)
+	{
+		TArray<int32> LocalObjectIDs;
+		const FSparseOctreeCell* RootCell = Queue[qi];
+		BranchRangeQuery(RootCell, Bounds, LocalObjectIDs);
+		
+		ObjectIDsLock.Lock();
+		for (int32 ObjectID : LocalObjectIDs)
+		{
+			ObjectIDs.Add(ObjectID);
+		}
+		ObjectIDsLock.Unlock();
+	});
+}
+
+
+void FSparseDynamicOctree3::BranchRangeQuery(
+	const FSparseOctreeCell* ParentCell,
+	const FAxisAlignedBox3d& Bounds,
+	TArray<int>& ObjectIDs) const
+{
+	TArray<const FSparseOctreeCell*, TInlineAllocator<32>> Queue;
+	Queue.Add(ParentCell);
+
+	while (Queue.Num() > 0)
+	{
+		const FSparseOctreeCell* CurCell = Queue.Pop(false);
+
+		// process elements
+		CellObjectLists.Enumerate(CurCell->CellID, [&](int32 ObjectID)
+		{
+			ObjectIDs.Add(ObjectID);
+		});
+
+		for (int k = 0; k < 8; ++k)
+		{
+			if (CurCell->HasChild(k))
+			{
+				const FSparseOctreeCell* ChildCell = &Cells[CurCell->GetChildCellID(k)];
+				if (GetCellBox(*ChildCell, MaxExpandFactor).Intersects(Bounds))
+				{
+					Queue.Add(ChildCell);
+				}
+			}
+		}
+	}
+}
+
+
 
 
 FSparseOctreeCell FSparseDynamicOctree3::FindCurrentContainingCell(const FAxisAlignedBox3d& Bounds) const
diff --git a/Engine/Plugins/Experimental/GeometryProcessing/Source/GeometricObjects/Public/Spatial/SparseDynamicOctree3.h b/Engine/Plugins/Experimental/GeometryProcessing/Source/GeometricObjects/Public/Spatial/SparseDynamicOctree3.h
--- a/Engine/Plugins/Experimental/GeometryProcessing/Source/GeometricObjects/Public/Spatial/SparseDynamicOctree3.h
+++ b/Engine/Plugins/Experimental/GeometryProcessing/Source/GeometricObjects/Public/Spatial/SparseDynamicOctree3.h
@@ -252,6 +252,15 @@ public:
 	void RangeQuery(const FAxisAlignedBox3d& Bounds,
 		TArray<int>& ObjectIDsOut ) const;
 
+	/**
+	 * Collect ObjectIDs from all the cells with bounding boxes that intersect Bounds.
+	 * Current implementation creates a separate TArray for each thread.
+	 * @param Bounds query box
+	 * @param ObjectIDsOut collected ObjectIDs are stored here
+	 */
+	void ParallelRangeQuery(const FAxisAlignedBox3d& Bounds,
+		TArray<int>& ObjectIDsOut) const;
+
 
 	/**
 	 * Check that the octree is internally valid
@@ -395,5 +404,10 @@ protected:
 
 	double FindNearestRayCellIntersection(const FSparseOctreeCell& Cell, const FRay3d& Ray) const;
 
+
+	void BranchRangeQuery(const FSparseOctreeCell* ParentCell,
+		const FAxisAlignedBox3d& Bounds,
+		TArray<int>& ObjectIDs) const;
+
 };
 
