GeometryProcessing: replace TIndirectArray in TDynamicVector with local array-of-arrays type that does not bounds-check on every access.
#rnx

diff --git a/Engine/Plugins/Experimental/GeometryProcessing/Source/DynamicMesh/Public/DynamicMeshAttributeSet.h b/Engine/Plugins/Experimental/GeometryProcessing/Source/DynamicMesh/Public/DynamicMeshAttributeSet.h
--- a/Engine/Plugins/Experimental/GeometryProcessing/Source/DynamicMesh/Public/DynamicMeshAttributeSet.h
+++ b/Engine/Plugins/Experimental/GeometryProcessing/Source/DynamicMesh/Public/DynamicMeshAttributeSet.h
@@ -8,6 +8,7 @@
 #include "DynamicAttribute.h"
 #include "GeometryTypes.h"
 #include "InfoTypes.h"
+#include "Containers/IndirectArray.h"
 
 /** Standard UV overlay type - 2-element float */
 typedef TDynamicMeshVectorOverlay<float, 2, FVector2f> FDynamicMeshUVOverlay;
diff --git a/Engine/Plugins/Experimental/GeometryProcessing/Source/GeometricObjects/Public/Util/DynamicVector.h b/Engine/Plugins/Experimental/GeometryProcessing/Source/GeometricObjects/Public/Util/DynamicVector.h
--- a/Engine/Plugins/Experimental/GeometryProcessing/Source/GeometricObjects/Public/Util/DynamicVector.h
+++ b/Engine/Plugins/Experimental/GeometryProcessing/Source/GeometricObjects/Public/Util/DynamicVector.h
@@ -6,7 +6,6 @@
 
 #include <CoreMinimal.h>
 #include "Containers/StaticArray.h"
-#include "Containers/IndirectArray.h"
 #include "VectorTypes.h"
 #include "IndexTypes.h"
 
@@ -212,7 +211,79 @@ private:
 	unsigned int CurBlockUsed{0};
 
 	using BlockType = TStaticArray<Type, BlockSize>;
-	TIndirectArray<BlockType> Blocks;
+	
+
+	/**
+	 * This is used to store/manage an array of BlockType*. Basically a stripped
+	 * down TIndirectArray that does not do range-checking on operator[] in non-debug builds
+	 */
+	template<typename ArrayType>
+	class TBlockVector
+	{
+	protected:
+		TArray<ArrayType*> Elements;
+
+	public:
+		TBlockVector()
+		{
+		}
+		TBlockVector(const TBlockVector& Copy)
+		{
+			*this = Copy;
+		}
+		TBlockVector(TBlockVector&& Moved)
+		{
+			Elements = MoveTemp(Moved.Elements);
+		}
+		TBlockVector& operator=(const TBlockVector& Copy)
+		{
+			int32 N = Copy.Num();
+			Elements.SetNum(N);
+			for (int32 k = 0; k < N; ++k)
+			{
+				Elements[k] = new ArrayType;
+				*Elements[k] = *Copy.Elements[k];
+			}
+			return *this;
+		}
+		TBlockVector& operator=(TBlockVector&& Moved)
+		{
+			Elements = MoveTemp(Moved.Elements);
+			return *this;
+		}
+		~TBlockVector()
+		{
+			Truncate(0, false);
+		}
+
+		int32 Num() const { return Elements.Num(); }
+
+#if UE_BUILD_DEBUG
+		ArrayType& operator[](int32 Index) { return *Elements[Index]; }
+		const ArrayType& operator[](int32 Index) const { return *Elements[Index]; }
+#else
+		ArrayType& operator[](int32 Index) { return *Elements.GetData()[Index]; }
+		const ArrayType& operator[](int32 Index) const { return *Elements.GetData()[Index]; }
+#endif
+
+		void Add(ArrayType* NewElement)
+		{
+			Elements.Add(NewElement);
+		}
+
+		void Truncate(int32 NewElementCount, bool bAllowShrinking)
+		{
+			for (int32 k = NewElementCount; k < Elements.Num(); ++k)
+			{
+				delete Elements[k];
+				Elements[k] = nullptr;
+			}
+			Elements.RemoveAt(NewElementCount, Elements.Num() - NewElementCount, bAllowShrinking);
+		}
+	};
+
+	// memory chunks for dynamic vector
+	TBlockVector<BlockType> Blocks;
 };
 
 template <class Type, int N>
@@ -370,7 +441,7 @@ using  TDynamicVector2i = TDynamicVectorN<int, 2>;
 template <class Type>
 void TDynamicVector<Type>::Clear()
 {
-	Blocks.Empty();
+	Blocks.Truncate(0, false);
 	CurBlock = 0;
 	CurBlockUsed = 0;
 	Blocks.Add(new BlockType());
@@ -414,7 +485,7 @@ void TDynamicVector<Type>::Resize(size_t Count)
 	}
 	else
 	{
-		Blocks.RemoveAt(nNumSegs, nCurCount - nNumSegs);
+		Blocks.Truncate(nNumSegs, false);
 	}
 
 	// mark last segment
