add a helper to convexhull3 to get the hull representation that containmentqueries3 expects
 + minor cleanup and add coverage of new function

#rnx

diff --git a/Engine/Plugins/Experimental/GeometryProcessing/Source/GeometricObjects/Public/HalfspaceTypes.h b/Engine/Plugins/Experimental/GeometryProcessing/Source/GeometricObjects/Public/HalfspaceTypes.h
--- a/Engine/Plugins/Experimental/GeometryProcessing/Source/GeometricObjects/Public/HalfspaceTypes.h
+++ b/Engine/Plugins/Experimental/GeometryProcessing/Source/GeometricObjects/Public/HalfspaceTypes.h
@@ -31,6 +31,13 @@ public:
 	THalfspace3(const FVector3<T>& PlaneNormalIn, const FVector3<T>& PlanePoint)
 		: Normal(PlaneNormalIn), Constant(Normal.Dot(PlanePoint)) {}
 
+	/** Construct a Halfspace from three points */
+	THalfspace3(const FVector3<T>& P0, const FVector3<T>& P1, const FVector3<T>& P2)
+	{
+		Normal = VectorUtil::Normal(P0, P1, P2);
+		Constant = Normal.Dot(P0);
+	}
+
 
 	/** @return true if Halfspace contains given Point */
 	bool Contains(const FVector3<T>& Point) const
diff --git a/Engine/Plugins/Experimental/GeometryProcessing/Source/GeometryAlgorithms/Public/ConvexHull3.h b/Engine/Plugins/Experimental/GeometryProcessing/Source/GeometryAlgorithms/Public/ConvexHull3.h
--- a/Engine/Plugins/Experimental/GeometryProcessing/Source/GeometryAlgorithms/Public/ConvexHull3.h
+++ b/Engine/Plugins/Experimental/GeometryProcessing/Source/GeometryAlgorithms/Public/ConvexHull3.h
@@ -7,7 +7,7 @@
 #include "IndexTypes.h"
 #include "LineTypes.h"
 #include "PlaneTypes.h"
-#include "Templates/PimplPtr.h"
+#include "HalfspaceTypes.h"
 
 template <typename RealType> struct TConvexHull3Internal;
 
@@ -24,7 +24,8 @@ public:
 	 * If input is degenerate, this will return false, and caller can call GetDimension()
 	 * to determine whether the points were coplanar, collinear, or all the same point
 	 *
-	 * @param Points Array of points to consider
+	 * @param NumPoints Number of points to consider
+	 * @param GetPointFunc Function providing array-style access into points
 	 * @param Filter Optional filter to include only a subset of the points in the output hull
 	 * @return true if hull was generated, false if points span < 2 dimensions
 	 */
@@ -35,8 +36,7 @@ public:
 	 * If input is degenerate, this will return false, and caller can call GetDimension()
 	 * to determine whether the points were collinear, or all the same point
 	 *
-	 * @param NumPoints Number of points to consider
-	 * @param GetPointFunc Function providing array-style access into points
+	 * @param Points Array of points to consider
 	 * @param FilterFunc Optional filter to include only a subset of the points in the output hull
 	 * @return true if hull was generated, false if points span < 2 dimensions
 	 */
@@ -54,8 +54,7 @@ public:
 	 * If input is degenerate, this will return false, and caller can call GetDimension()
 	 * to determine whether the points were collinear, or all the same point
 	 *
-	 * @param NumPoints Number of points to consider
-	 * @param GetPointFunc Function providing array-style access into points
+	 * @param Points Array of points to consider
 	 * @return true if hull was generated, false if points span < 2 dimensions
 	 */
 	bool Solve(TArrayView<const FVector3<RealType>> Points)
@@ -89,6 +88,44 @@ public:
 		return Hull;
 	}
 
+	/** 
+	 * Convert an already-computed convex hull into a halfspace representation
+	 * Following the logic of ContainmentQueries3.h, all halfspaces are oriented "outwards,"
+	 * so a point is inside the convex hull if it is outside of all halfspaces in the array
+	 * 
+	 * @param Points Array of points to consider
+	 * @return Array of halfspaces
+	 */
+	TArray<THalfspace3<RealType>> GetAsHalfspaces(TArrayView<const FVector3<RealType>> Points) const
+	{
+		TArray<THalfspace3<RealType>> Halfspaces;
+		for (FIndex3i Tri : Hull)
+		{
+			THalfspace3<RealType> TriHalfspace(Points[Tri.A], Points[Tri.B], Points[Tri.C]);
+			Halfspaces.Add(TriHalfspace);
+		}
+		return Halfspaces;
+	}
+
+	/**
+	 * Convert an already-computed convex hull into a halfspace representation
+	 * Following the logic of ContainmentQueries3.h, all halfspaces are oriented "outwards,"
+	 * so a point is inside the convex hull if it is outside of all halfspaces in the array
+	 *
+	 * @param GetPointFunc Function providing array-style access into points
+	 * @return Array of halfspaces
+	 */
+	TArray<THalfspace3<RealType>> GetAsHalfspaces(TFunctionRef<FVector3<RealType>(int32)> GetPointFunc) const
+	{
+		TArray<THalfspace3<RealType>> Halfspaces;
+		for (FIndex3i Tri : Hull)
+		{
+			THalfspace3<RealType> TriHalfspace(GetPointFunc(Tri.A), GetPointFunc(Tri.B), GetPointFunc(Tri.C));
+			Halfspaces.Add(TriHalfspace);
+		}
+		return Halfspaces;
+	}
+
 	/**
 	 * Empty any previously-computed convex hull data.  Frees the hull memory.
 	 * Note: You do not need to call this before calling Generate() with new data.
