GeometryProcessing: support vertex colors in DynamicMesh/MeshDescription conversion. Add flag in FConversionToMeshDescriptionOptions. Support in USimpleDynamicMeshComponent::Bake.
#rnx

diff --git a/Engine/Plugins/Experimental/GeometryProcessing/Source/MeshConversion/Private/DynamicMeshToMeshDescription.cpp b/Engine/Plugins/Experimental/GeometryProcessing/Source/MeshConversion/Private/DynamicMeshToMeshDescription.cpp
--- a/Engine/Plugins/Experimental/GeometryProcessing/Source/MeshConversion/Private/DynamicMeshToMeshDescription.cpp
+++ b/Engine/Plugins/Experimental/GeometryProcessing/Source/MeshConversion/Private/DynamicMeshToMeshDescription.cpp
@@ -231,6 +231,31 @@ void FDynamicMeshToMeshDescription::UpdateTangents(const FDynamicMesh3* MeshIn,
 }
 
 
+void FDynamicMeshToMeshDescription::UpdateVertexColors(const FDynamicMesh3* MeshIn, FMeshDescription& MeshOut)
+{
+	check(MeshIn->IsCompactV() && MeshIn->HasVertexColors());
+
+	FStaticMeshAttributes Attributes(MeshOut);
+	TVertexInstanceAttributesRef<FVector4> InstanceColors = Attributes.GetVertexInstanceColors();
+	bool bIsValidDst = InstanceColors.IsValid();
+	ensureMsgf(bIsValidDst, TEXT("Trying to update colors on a MeshDescription that has no color attributes"));
+	if (bIsValidDst)
+	{
+		check(MeshIn->VertexCount() == MeshOut.Vertices().Num());
+		for (int VertID : MeshIn->VertexIndicesItr())
+		{
+			FVector3f Color3f = MeshIn->GetVertexColor(VertID);
+			FVector4 Color4(Color3f.X, Color3f.Y, Color3f.Z, 1.0f);
+			for (FVertexInstanceID InstanceID : MeshOut.GetVertexVertexInstanceIDs(FVertexID(VertID)))
+			{
+				InstanceColors.Set(InstanceID, 0, Color4);
+			}
+		}
+	}
+}
+
+
+
 
 void FDynamicMeshToMeshDescription::Convert(const FDynamicMesh3* MeshIn, FMeshDescription& MeshOut, bool bCopyTangents)
 {
@@ -524,6 +549,8 @@ void FDynamicMeshToMeshDescription::Convert_NoSharedInstances(const FDynamicMesh
 		bCopyGroupToPolyGroup = true;
 	}
 
+	bool bCopyVertexColors = MeshIn->HasVertexColors();		// always copy when we are baking new mesh? should this be a config option?
+
 	// disable indexing during the full build of the mesh
 	Builder.SuspendMeshDescriptionIndexing();
 
@@ -654,6 +681,16 @@ void FDynamicMeshToMeshDescription::Convert_NoSharedInstances(const FDynamicMesh
 			TriVertInstances[j] = Builder.AppendInstance(TriVertex);
 		}
 
+		// copy vertex colors. This writes each color multiple times, but that is not expensive in this context
+		if (bCopyVertexColors)
+		{
+			for (int32 j = 0; j < 3; ++j)
+			{
+				FVector3f Color = MeshIn->GetVertexColor(Triangle[j]);
+				Builder.SetInstanceColor(TriVertInstances[j], FVector4(Color.X, Color.Y, Color.Z, 1.0f));
+			}
+		}
+
 		// transfer material index to MeshDescription polygon group (by convention)
 		FPolygonGroupID UsePolygonGroupID = ZeroPolygonGroupID;
 		if (MaterialIDAttrib)
diff --git a/Engine/Plugins/Experimental/GeometryProcessing/Source/MeshConversion/Private/MeshDescriptionToDynamicMesh.cpp b/Engine/Plugins/Experimental/GeometryProcessing/Source/MeshConversion/Private/MeshDescriptionToDynamicMesh.cpp
--- a/Engine/Plugins/Experimental/GeometryProcessing/Source/MeshConversion/Private/MeshDescriptionToDynamicMesh.cpp
+++ b/Engine/Plugins/Experimental/GeometryProcessing/Source/MeshConversion/Private/MeshDescriptionToDynamicMesh.cpp
@@ -150,6 +150,17 @@ void FMeshDescriptionToDynamicMesh::Convert(const FMeshDescription* MeshIn, FDyn
 	TVertexInstanceAttributesConstRef<FVector> InstanceTangents = Attributes.GetVertexInstanceTangents();
 	TVertexInstanceAttributesConstRef<float> InstanceBiTangentSign = Attributes.GetVertexInstanceBinormalSigns();
 
+	TVertexInstanceAttributesConstRef<FVector4> InstanceColors = Attributes.GetVertexInstanceColors();
+	if (bEnableOutputVertexColors && InstanceColors.IsValid())
+	{
+		MeshOut.EnableVertexColors(FVector3f::One());
+	}
+	else
+	{
+		bEnableOutputVertexColors = false;
+	}
+	bool bFoundNonDefaultVertexColor = false;
+
 	TPolygonAttributesConstRef<int> PolyGroups =
 		MeshIn->PolygonAttributes().GetAttributesRef<int>(ExtendedMeshAttribute::PolyTriGroups);
 	if (bEnableOutputGroups)
@@ -306,6 +317,23 @@ void FMeshDescriptionToDynamicMesh::Convert(const FMeshDescription* MeshIn, FDyn
 		// .Insert(Value, Index)
 		TriIDMap.Insert(TriangleID, NewTriangleID);
 	
+		// copy triangle instance colors to per-vertex colors. Currently this just uses last-instance as
+		// we do not support per-triangle-vertex colors.
+		if (bEnableOutputVertexColors)
+		{
+			for (int32 j = 0; j < 3; ++j)
+			{
+				FVector4 InstanceColor4 = InstanceColors.Get(TriData.TriInstances[j], 0);
+				FVector3f InstanceColor3(InstanceColor4.X, InstanceColor4.Y, InstanceColor4.Z);
+				bFoundNonDefaultVertexColor |= (InstanceColor3 != FVector3f::One());
+				MeshOut.SetVertexColor(VertexIDs[j], InstanceColor3);
+			}
+		}
+	}
+
+	if (bFoundNonDefaultVertexColor == false)
+	{
+		MeshOut.DiscardVertexColors();
 	}
 
 	FDateTime Time_AfterTriangles = FDateTime::Now();
diff --git a/Engine/Plugins/Experimental/GeometryProcessing/Source/MeshConversion/Public/DynamicMeshToMeshDescription.h b/Engine/Plugins/Experimental/GeometryProcessing/Source/MeshConversion/Public/DynamicMeshToMeshDescription.h
--- a/Engine/Plugins/Experimental/GeometryProcessing/Source/MeshConversion/Public/DynamicMeshToMeshDescription.h
+++ b/Engine/Plugins/Experimental/GeometryProcessing/Source/MeshConversion/Public/DynamicMeshToMeshDescription.h
@@ -116,6 +116,15 @@ public:
 	void UpdateTangents(const FDynamicMesh3* MeshIn, FMeshDescription& MeshOut);
 
 
+	/**
+	 * Update only vertex colors, assuming the mesh topology has not changed.  Does not touch positions or other attributes.
+	 *	NOTE: assumes the order of triangles in the MeshIn correspond to the ordering you'd get by iterating over triangles, on MeshOut
+	 *		  This matches conversion currently used in MeshDescriptionToDynamicMesh.cpp, but if that changes we will need to change this function to match!
+	 */
+	void UpdateVertexColors(const FDynamicMesh3* MeshIn, FMeshDescription& MeshOut);
+
+
+
 	//
 	// Internal functions that you can also call directly
 	//
diff --git a/Engine/Plugins/Experimental/GeometryProcessing/Source/MeshConversion/Public/MeshConversionOptions.h b/Engine/Plugins/Experimental/GeometryProcessing/Source/MeshConversion/Public/MeshConversionOptions.h
--- a/Engine/Plugins/Experimental/GeometryProcessing/Source/MeshConversion/Public/MeshConversionOptions.h
+++ b/Engine/Plugins/Experimental/GeometryProcessing/Source/MeshConversion/Public/MeshConversionOptions.h
@@ -23,4 +23,18 @@ public:
 
 	/** Should UVs of MeshDescription be updated, if available and relevant */
 	bool bUpdateUVs = false;
+
+	/** Should Vertex Colors of MeshDescription be updated, if available and relevant */
+	bool bUpdateVtxColors = false;
+
+
+	//
+	// utility functions for common configuration cases
+	//
+
+	void SetToVertexColorsOnly()
+	{
+		bSetPolyGroups = bUpdatePositions = bUpdateNormals = bUpdateTangents = bUpdateUVs = false;
+		bUpdateVtxColors = true;
+	}
 };
\ No newline at end of file
diff --git a/Engine/Plugins/Experimental/GeometryProcessing/Source/MeshConversion/Public/MeshDescriptionToDynamicMesh.h b/Engine/Plugins/Experimental/GeometryProcessing/Source/MeshConversion/Public/MeshDescriptionToDynamicMesh.h
--- a/Engine/Plugins/Experimental/GeometryProcessing/Source/MeshConversion/Public/MeshDescriptionToDynamicMesh.h
+++ b/Engine/Plugins/Experimental/GeometryProcessing/Source/MeshConversion/Public/MeshDescriptionToDynamicMesh.h
@@ -26,6 +26,9 @@ public:
 	/** Should we initialize triangle groups on output mesh */
 	bool bEnableOutputGroups = true;
 
+	/** Should we initialize vertex colors on output mesh */
+	bool bEnableOutputVertexColors = true;
+
 	/** Should we calculate conversion index maps */
 	bool bCalculateMaps = true;
 
diff --git a/Engine/Plugins/Experimental/MeshModelingToolset/Source/ModelingComponents/Private/SimpleDynamicMeshComponent.cpp b/Engine/Plugins/Experimental/MeshModelingToolset/Source/ModelingComponents/Private/SimpleDynamicMeshComponent.cpp
--- a/Engine/Plugins/Experimental/MeshModelingToolset/Source/ModelingComponents/Private/SimpleDynamicMeshComponent.cpp
+++ b/Engine/Plugins/Experimental/MeshModelingToolset/Source/ModelingComponents/Private/SimpleDynamicMeshComponent.cpp
@@ -178,10 +178,15 @@ void USimpleDynamicMeshComponent::Bake(FMeshDescription* MeshDescription, bool b
 		{
 			Converter.Update(Mesh.Get(), *MeshDescription, ConversionOptions.bUpdateNormals, ConversionOptions.bUpdateTangents, ConversionOptions.bUpdateUVs);
 		}
-		else
+		else if (ConversionOptions.bUpdateNormals || ConversionOptions.bUpdateTangents || ConversionOptions.bUpdateUVs)
 		{
 			Converter.UpdateAttributes(Mesh.Get(), *MeshDescription, ConversionOptions.bUpdateNormals, ConversionOptions.bUpdateTangents, ConversionOptions.bUpdateUVs);
 		}
+
+		if (ConversionOptions.bUpdateVtxColors)
+		{
+			Converter.UpdateVertexColors(Mesh.Get(), *MeshDescription);
+		}
 	}
 	else
 	{
