GeometryProcessing: Add FSparseDynamicOctree3::CheckIfObjectNeedsReinsert(), precomputes the check in ReinsertObject(), which can be done in parallel. Allow ReinsertObject() to take output from check function. Use these in new FDynamicMeshOctree3::ReinsertTrianglesParallel() function that is 20-30% faster than non-parallel Reinsert (possibly more for very small mesh changes, like vertex smoothing)
#rnx

diff --git a/Engine/Plugins/Experimental/GeometryProcessing/Source/DynamicMesh/Public/DynamicMeshOctree3.h b/Engine/Plugins/Experimental/GeometryProcessing/Source/DynamicMesh/Public/DynamicMeshOctree3.h
--- a/Engine/Plugins/Experimental/GeometryProcessing/Source/DynamicMesh/Public/DynamicMeshOctree3.h
+++ b/Engine/Plugins/Experimental/GeometryProcessing/Source/DynamicMesh/Public/DynamicMeshOctree3.h
@@ -6,6 +6,7 @@
 #include "Spatial/SparseDynamicOctree3.h"
 #include "MeshQueries.h"
 #include "DynamicMesh3.h"
+#include "Async/ParallelFor.h"
 
 
 /**
@@ -152,6 +153,36 @@ public:
 	}
 
 
+	/**
+	 * Reinsert a set of triangles into the tree. Internally precomputes which triangles need
+	 * re-inserting, which can be done in parallel and generally saves time as some triangles can be skipped.
+	 */
+	void ReinsertTrianglesParallel(const TArray<int32>& Triangles, TArray<uint32>& TempBuffer, TArray<bool>& TempFlagBuffer)
+	{
+		int32 NumTriangles = Triangles.Num();
+		TempBuffer.SetNum(NumTriangles, false);
+		TempFlagBuffer.SetNum(NumTriangles, false);
+
+		// can check which triangles need reinsertion in parallel. This will also return which
+		// CellID the triangle is in, which saves time in the Reinsert function
+		ParallelFor(NumTriangles, [&](int k)
+		{
+			FAxisAlignedBox3d Bounds = Mesh->GetTriBounds(Triangles[k]);
+			TempFlagBuffer[k] = CheckIfObjectNeedsReinsert(Triangles[k], Bounds, TempBuffer[k]);
+		});
+
+		// now reinsert all necessary triangles
+		for (int32 k = 0; k < NumTriangles; ++k)
+		{
+			if (TempFlagBuffer[k])
+			{
+				ReinsertObject(Triangles[k], Mesh->GetTriBounds(Triangles[k]), TempBuffer[k]);
+			}
+		}
+	}
+
+
+
 	/**
 	 * Include the current bounds of a triangle in the ModifiedBounds box
 	 */
diff --git a/Engine/Plugins/Experimental/GeometryProcessing/Source/GeometricObjects/Private/Spatial/SparseDynamicOctree3.cpp b/Engine/Plugins/Experimental/GeometryProcessing/Source/GeometricObjects/Private/Spatial/SparseDynamicOctree3.cpp
--- a/Engine/Plugins/Experimental/GeometryProcessing/Source/GeometricObjects/Private/Spatial/SparseDynamicOctree3.cpp
+++ b/Engine/Plugins/Experimental/GeometryProcessing/Source/GeometricObjects/Private/Spatial/SparseDynamicOctree3.cpp
@@ -161,25 +161,63 @@ bool FSparseDynamicOctree3::RemoveObject(int32 ObjectID)
 
 
 
-
-void FSparseDynamicOctree3::ReinsertObject(int32 ObjectID, const FAxisAlignedBox3d& NewBounds)
+bool FSparseDynamicOctree3::CheckIfObjectNeedsReinsert(int32 ObjectID, const FAxisAlignedBox3d& NewBounds, uint32& CellIDOut) const
 {
+	CellIDOut = InvalidCellID;
 	if (ContainsObject(ObjectID))
 	{
-		uint32 CellID = GetCellForObject(ObjectID);
+		CellIDOut = GetCellForObject(ObjectID);
+		if (CellIDOut != SpillCellID && CellIDOut != InvalidCellID)
+		{
+			const FSparseOctreeCell& CurrentCell = Cells[CellIDOut];
+			if (CanFit(CurrentCell, NewBounds))
+			{
+				return false;		// everything is fine
+			}
+
+		}
+	}
+	return true;
+}
+
+
+bool FSparseDynamicOctree3::ReinsertObject(int32 ObjectID, const FAxisAlignedBox3d& NewBounds, uint32 CellIDHint)
+{
+	uint32 CellID = CellIDHint;
+
+	// check if this object still fits in current cell. If so, we can ignore it
+	if ( CellID == InvalidCellID && ContainsObject(ObjectID) )
+	{
+		CellID = GetCellForObject(ObjectID);
 		if (CellID != SpillCellID && CellID != InvalidCellID)
 		{
 			FSparseOctreeCell& CurrentCell = Cells[CellID];
 			if (CanFit(CurrentCell, NewBounds))
 			{
-				return;		// everything is fine
+				return false;		// everything is fine
 			}
+		}
+	}
 
+	// remove object
+	if (CellID != InvalidCellID)
+	{
+		if (CellID == SpillCellID)
+		{
+			SpillObjectSet.Remove(ObjectID);
 		}
+		else
+		{
+			ObjectIDToCellMap[ObjectID] = InvalidCellID;
+			CellObjectLists.Remove(CellID, ObjectID);
+		}
+		ValidObjectIDs.Set(ObjectID, false);
 	}
 
-	RemoveObject(ObjectID);
+	// reinsert
 	InsertObject(ObjectID, NewBounds);
+
+	return true;
 }
 
 
diff --git a/Engine/Plugins/Experimental/GeometryProcessing/Source/GeometricObjects/Public/Spatial/SparseDynamicOctree3.h b/Engine/Plugins/Experimental/GeometryProcessing/Source/GeometricObjects/Public/Spatial/SparseDynamicOctree3.h
--- a/Engine/Plugins/Experimental/GeometryProcessing/Source/GeometricObjects/Public/Spatial/SparseDynamicOctree3.h
+++ b/Engine/Plugins/Experimental/GeometryProcessing/Source/GeometricObjects/Public/Spatial/SparseDynamicOctree3.h
@@ -219,9 +219,20 @@ public:
 	/**
 	 * Update the position of an object in the octree. This is more efficient than doing a remove+insert
 	 * @param ObjectID ID of the object
-	 * @param NewBounds enw bounding box of the object
+	 * @param NewBounds new bounding box of the object
+	 * @return true if object was reinserted, false if the object was fine in it's existing cell
 	 */
-	void ReinsertObject(int32 ObjectID, const FAxisAlignedBox3d& NewBounds);
+	bool ReinsertObject(int32 ObjectID, const FAxisAlignedBox3d& NewBounds, uint32 CellIDHint = InvalidCellID);
+
+	/**
+	 * Check if the object needs to be reinserted, if it has NewBounds. 
+	 * @param ObjectID ID of the object
+	 * @param NewBounds new bounding box of the object
+	 * @param CellIDOut returned CellID of the object, if it is in the tree. This can be passed to ReinsertObject() to save some computation.
+	 * @return false if ObjectID is already in the tree, and NewBounds fits within the cell that currently contains ObjectID
+	 */
+	bool CheckIfObjectNeedsReinsert(int32 ObjectID, const FAxisAlignedBox3d& NewBounds, uint32& CellIDOut) const;
+
 
 	/**
 	 * Find nearest ray-hit point with objects in tree
@@ -401,7 +412,6 @@ protected:
 	void Insert_ToCell(int32 ObjectID, const FAxisAlignedBox3d& Bounds, const FSparseOctreeCell& ExistingCell);
 	void Insert_NewChildCell(int32 ObjectID, const FAxisAlignedBox3d& Bounds, int ParentCellID, FSparseOctreeCell NewChildCell, int ChildIdx);
 
-
 	double FindNearestRayCellIntersection(const FSparseOctreeCell& Cell, const FRay3d& Ray) const;
 
 
