GeometryProcessing: Small optimizations for inner loops in mesh processing. Add FDynamicMesh3::GetVertexRef, and similar for triangle/edges/etc, that returns const-ref to internal members. Add parallel-compute option to FDynamicMesh3::GetBounds(), and unit test.
#rnx

diff --git a/Engine/Plugins/Experimental/GeometryProcessing/Source/DynamicMesh/Private/DynamicMesh3_Queries.cpp b/Engine/Plugins/Experimental/GeometryProcessing/Source/DynamicMesh/Private/DynamicMesh3_Queries.cpp
--- a/Engine/Plugins/Experimental/GeometryProcessing/Source/DynamicMesh/Private/DynamicMesh3_Queries.cpp
+++ b/Engine/Plugins/Experimental/GeometryProcessing/Source/DynamicMesh/Private/DynamicMesh3_Queries.cpp
@@ -1,6 +1,7 @@
 // Copyright Epic Games, Inc. All Rights Reserved.
 
 #include "DynamicMesh3.h"
+#include "Async/ParallelFor.h"
 
 
 FIndex2i FDynamicMesh3::GetEdgeOpposingV(int eID) const
@@ -591,21 +592,49 @@ int FDynamicMesh3::FindTriangle(int a, int b, int c) const
 /**
  * Computes bounding box of all vertices.
  */
-FAxisAlignedBox3d FDynamicMesh3::GetBounds() const
+FAxisAlignedBox3d FDynamicMesh3::GetBounds(bool bParallel) const
 {
-	if (VertexCount() == 0) 
+	if (VertexCount() == 0)
 	{
 		return FAxisAlignedBox3d::Empty();
 	}
 
-	FVector3d MinVec = Vertices[*(VertexIndicesItr().begin())];
-	FVector3d MaxVec = MinVec;
-	for (int vi : VertexIndicesItr())
+	FAxisAlignedBox3d ResultBounds = FAxisAlignedBox3d::Empty();
+	int32 MaxVID = MaxVertexID();
+	if (bParallel == false || MaxVID < 50000)
 	{
-		MinVec = Min(MinVec, Vertices[vi]);
-		MaxVec = Max(MaxVec, Vertices[vi]);
+		FVector3d MinVec = Vertices[*(VertexIndicesItr().begin())];
+		FVector3d MaxVec = MinVec;
+		for (int vi : VertexIndicesItr())
+		{
+			MinVec = Min(MinVec, Vertices[vi]);
+			MaxVec = Max(MaxVec, Vertices[vi]);
+		}
+		return FAxisAlignedBox3d(MinVec, MaxVec);
 	}
-	return FAxisAlignedBox3d(MinVec, MaxVec);
+	else
+	{
+		constexpr int ChunkSize = 10000;
+		int32 NumChunks = (MaxVID / ChunkSize) + 1;
+		FCriticalSection BoundsLock;
+		ParallelFor(NumChunks, [&](int ci)
+		{
+			FAxisAlignedBox3d ChunkBounds = FAxisAlignedBox3d::Empty();
+			int Start = ci * ChunkSize;
+			for (int k = 0; k < ChunkSize; ++k)
+			{
+				int vid = Start + k;
+				if (IsVertex(vid))
+				{
+					ChunkBounds.Contain(Vertices[vid]);
+				}
+			}
+			BoundsLock.Lock();
+			ResultBounds.Contain(ChunkBounds);
+			BoundsLock.Unlock();
+		});
+	}
+	return ResultBounds;
 }
 
 
@@ -811,13 +840,14 @@ void FDynamicMesh3::GetTriBaryPoint(int tID, double bary0, double bary1, double
 	}
 }
 
+
 FAxisAlignedBox3d FDynamicMesh3::GetTriBounds(int tID) const
 {
 	const FIndex3i& tIDs = Triangles[tID];
-	const FVector3d TriVerts[3]       = {Vertices[tIDs[0]], Vertices[tIDs[1]], Vertices[tIDs[2]]};
-	const FVector3d MinVec            = Min(Min(TriVerts[0], TriVerts[1]), TriVerts[2]);
-	const FVector3d MaxVec            = Max(Max(TriVerts[0], TriVerts[1]), TriVerts[2]);
-	return FAxisAlignedBox3d(MinVec, MaxVec);
+	const FVector3d& A = Vertices[tIDs.A];
+	const FVector3d& B = Vertices[tIDs.B];
+	const FVector3d& C = Vertices[tIDs.C];
+	return FAxisAlignedBox3d(A, B, C);
 }
 
 FFrame3d FDynamicMesh3::GetTriFrame(int tID, int nEdge) const
diff --git a/Engine/Plugins/Experimental/GeometryProcessing/Source/DynamicMesh/Public/DynamicMesh3.h b/Engine/Plugins/Experimental/GeometryProcessing/Source/DynamicMesh/Public/DynamicMesh3.h
--- a/Engine/Plugins/Experimental/GeometryProcessing/Source/DynamicMesh/Public/DynamicMesh3.h
+++ b/Engine/Plugins/Experimental/GeometryProcessing/Source/DynamicMesh/Public/DynamicMesh3.h
@@ -483,6 +483,13 @@ public:
 		return Vertices[VertexID];
 	}
 
+	/** @return the vertex position */
+	inline const FVector3d& GetVertexRef(int VertexID) const
+	{
+		checkSlow(IsVertex(VertexID));
+		return Vertices[VertexID];
+	}
+
 	/** Set vertex position */
 	inline void SetVertex(int VertexID, const FVector3d& vNewPos)
 	{
@@ -554,6 +561,13 @@ public:
 		return Triangles[TriangleID];
 	}
 
+	/** Get triangle vertices */
+	inline const FIndex3i& GetTriangleRef(int TriangleID) const
+	{
+		checkSlow(IsTriangle(TriangleID));
+		return Triangles[TriangleID];
+	}
+
 	/** Get triangle edges */
 	inline FIndex3i GetTriEdges(int TriangleID) const
 	{
@@ -561,6 +575,13 @@ public:
 		return TriangleEdges[TriangleID];
 	}
 
+	/** Get triangle edges */
+	inline const FIndex3i& GetTriEdgesRef(int TriangleID) const
+	{
+		checkSlow(IsTriangle(TriangleID));
+		return TriangleEdges[TriangleID];
+	}
+
 	/** Get one of the edges of a triangle */
 	inline int GetTriEdge(int TriangleID, int j) const
 	{
@@ -593,6 +614,13 @@ public:
 		return Edges[EdgeID];
 	}
 
+	/** Get the vertices and triangles of an edge, returned as [v0,v1,t0,t1], where t1 may be InvalidID */
+	inline const FEdge& GetEdgeRef(int EdgeID) const
+	{
+		checkSlow(IsEdge(EdgeID));
+		return Edges[EdgeID];
+	}
+
 	/** Get the vertex pair for an edge */
 	inline FIndex2i GetEdgeV(int EdgeID) const
 	{
@@ -643,7 +671,7 @@ public:
 		{
 			return FVector3f::UnitY();
 		}
-		check(IsVertex(vID));
+		checkSlow(IsVertex(vID));
 		const TDynamicVector<FVector3f>& Normals = VertexNormals.GetValue();
 		return Normals[vID];
 	}
@@ -652,7 +680,7 @@ public:
 	{
 		if (HasVertexNormals())
 		{
-			check(IsVertex(vID));
+			checkSlow(IsVertex(vID));
 			TDynamicVector<FVector3f>& Normals = VertexNormals.GetValue();
 			Normals[vID]                       = vNewNormal;
 			UpdateTimeStamp(false, false);
@@ -669,7 +697,7 @@ public:
 		{
 			return FVector3f::One();
 		}
-		check(IsVertex(vID));
+		checkSlow(IsVertex(vID));
 
 		const TDynamicVector<FVector3f>& Colors = VertexColors.GetValue();
 		return Colors[vID];
@@ -679,7 +707,7 @@ public:
 	{
 		if (HasVertexColors())
 		{
-			check(IsVertex(vID));
+			checkSlow(IsVertex(vID));
 			TDynamicVector<FVector3f>& Colors = VertexColors.GetValue();
 			Colors[vID]                       = vNewColor;
 			UpdateTimeStamp(false, false);
@@ -695,7 +723,7 @@ public:
 		{
 			return FVector2f::Zero();
 		}
-		check(IsVertex(vID));
+		checkSlow(IsVertex(vID));
 		const TDynamicVector<FVector2f>& UVs = VertexUVs.GetValue();
 		return UVs[vID];
 	}
@@ -704,7 +732,7 @@ public:
 	{
 		if (HasVertexUVs())
 		{
-			check(IsVertex(vID));
+			checkSlow(IsVertex(vID));
 			TDynamicVector<FVector2f>& UVs = VertexUVs.GetValue();
 			UVs[vID]                       = vNewUV;
 			UpdateTimeStamp(false, false);
@@ -729,7 +757,7 @@ public:
 	{
 		if (HasTriangleGroups())
 		{
-			check(IsTriangle(tid));
+			checkSlow(IsTriangle(tid));
 			TriangleGroups.GetValue()[tid] = group_id;
 			GroupIDCounter                 = FMath::Max(GroupIDCounter, group_id + 1);
 			UpdateTimeStamp(false, false);
@@ -905,7 +933,7 @@ public:
 	//
 public:
 	/** Returns bounding box of all mesh vertices (including unreferenced vertices) */
-	FAxisAlignedBox3d GetBounds() const;
+	FAxisAlignedBox3d GetBounds(bool bParallel = false) const;
 
 	/** Returns GetBounds() and saves result, cache is invalidated and recomputed if topology has changed since last call */
 	FAxisAlignedBox3d GetCachedBounds();
@@ -1384,7 +1412,7 @@ protected:
 		}
 		if (bTopologyChange)
 		{
-			check(bShapeChange); // we consider topology change to be a shape change!
+			checkSlow(bShapeChange); // we consider topology change to be a shape change!
 			TopologyTimestamp++;
 		}
 	}
diff --git a/Engine/Plugins/Experimental/GeometryProcessing/Source/GeometricObjects/Private/Util/SmallListSet.cpp b/Engine/Plugins/Experimental/GeometryProcessing/Source/GeometricObjects/Private/Util/SmallListSet.cpp
--- a/Engine/Plugins/Experimental/GeometryProcessing/Source/GeometricObjects/Private/Util/SmallListSet.cpp
+++ b/Engine/Plugins/Experimental/GeometryProcessing/Source/GeometricObjects/Private/Util/SmallListSet.cpp
@@ -20,7 +20,7 @@ void FSmallListSet::Resize(int32 NewSize)
 
 void FSmallListSet::AllocateAt(int32 ListIndex)
 {
-	check(ListIndex >= 0);
+	checkSlow(ListIndex >= 0);
 	if (ListIndex >= (int)ListHeads.GetLength())
 	{
 		int32 j = (int32)ListHeads.GetLength();
@@ -43,7 +43,7 @@ void FSmallListSet::AllocateAt(int32 ListIndex)
 
 void FSmallListSet::Insert(int32 ListIndex, int32 Value)
 {
-	check(ListIndex >= 0);
+	checkSlow(ListIndex >= 0);
 	int32 block_ptr = ListHeads[ListIndex];
 	if (block_ptr == NullValue)
 	{
@@ -90,7 +90,7 @@ void FSmallListSet::Insert(int32 ListIndex, int32 Value)
 
 bool FSmallListSet::Remove(int32 ListIndex, int32 Value)
 {
-	check(ListIndex >= 0);
+	checkSlow(ListIndex >= 0);
 	int32 block_ptr = ListHeads[ListIndex];
 	int32 N = ListBlocks[block_ptr];
 
@@ -141,9 +141,9 @@ bool FSmallListSet::Remove(int32 ListIndex, int32 Value)
 
 void FSmallListSet::Move(int32 FromIndex, int32 ToIndex)
 {
-	check(FromIndex >= 0);
-	check(ToIndex >= 0);
-	check(ListHeads[ToIndex] == NullValue);
+	checkSlow(FromIndex >= 0);
+	checkSlow(ToIndex >= 0);
+	checkSlow(ListHeads[ToIndex] == NullValue);
 	ListHeads[ToIndex] = ListHeads[FromIndex];
 	ListHeads[FromIndex] = NullValue;
 }
@@ -153,7 +153,7 @@ void FSmallListSet::Move(int32 FromIndex, int32 ToIndex)
 
 void FSmallListSet::Clear(int32 ListIndex)
 {
-	check(ListIndex >= 0);
+	checkSlow(ListIndex >= 0);
 	int32 block_ptr = ListHeads[ListIndex];
 	if (block_ptr != NullValue)
 	{
@@ -184,7 +184,7 @@ void FSmallListSet::Clear(int32 ListIndex)
 
 bool FSmallListSet::Contains(int32 ListIndex, int32 Value) const
 {
-	check(ListIndex >= 0);
+	checkSlow(ListIndex >= 0);
 	int32 block_ptr = ListHeads[ListIndex];
 	if (block_ptr != NullValue)
 	{
diff --git a/Engine/Plugins/Experimental/GeometryProcessing/Source/GeometricObjects/Public/BoxTypes.h b/Engine/Plugins/Experimental/GeometryProcessing/Source/GeometricObjects/Public/BoxTypes.h
--- a/Engine/Plugins/Experimental/GeometryProcessing/Source/GeometricObjects/Public/BoxTypes.h
+++ b/Engine/Plugins/Experimental/GeometryProcessing/Source/GeometricObjects/Public/BoxTypes.h
@@ -221,11 +221,18 @@ struct TAxisAlignedBox3
 		this->Max = Max;
 	}
 
-	TAxisAlignedBox3(const FVector3<RealType>& V0, const FVector3<RealType>& V1, const FVector3<RealType>& V2)
-	{
-		TMathUtil<RealType>::MinMax(V0.X, V1.X, V2.X, Min.X, Max.X);
-		TMathUtil<RealType>::MinMax(V0.Y, V1.Y, V2.Y, Min.Y, Max.Y);
-		TMathUtil<RealType>::MinMax(V0.Z, V1.Z, V2.Z, Min.Z, Max.Z);
+	TAxisAlignedBox3(const FVector3<RealType>& A, const FVector3<RealType>& B, const FVector3<RealType>& C)
+	{
+		// TMathUtil::MinMax could be used here, but it generates worse code because the Min3's below will be
+		// turned into SSE instructions by the optimizer, while MinMax will not
+		Min = FVector3<RealType>(
+			TMathUtil<RealType>::Min3(A.X, B.X, C.X),
+			TMathUtil<RealType>::Min3(A.Y, B.Y, C.Y),
+			TMathUtil<RealType>::Min3(A.Z, B.Z, C.Z));
+		Max = FVector3<RealType>(
+			TMathUtil<RealType>::Max3(A.X, B.X, C.X),
+			TMathUtil<RealType>::Max3(A.Y, B.Y, C.Y),
+			TMathUtil<RealType>::Max3(A.Z, B.Z, C.Z));
 	}
 
 	TAxisAlignedBox3(const TAxisAlignedBox3& OtherBox) = default;
diff --git a/Engine/Plugins/Experimental/GeometryProcessing/Source/GeometricObjects/Public/Util/SmallListSet.h b/Engine/Plugins/Experimental/GeometryProcessing/Source/GeometricObjects/Public/Util/SmallListSet.h
--- a/Engine/Plugins/Experimental/GeometryProcessing/Source/GeometricObjects/Public/Util/SmallListSet.h
+++ b/Engine/Plugins/Experimental/GeometryProcessing/Source/GeometricObjects/Public/Util/SmallListSet.h
@@ -118,7 +118,7 @@ public:
 	 */
 	inline int32 GetCount(int32 ListIndex) const
 	{
-		check(ListIndex >= 0);
+		checkSlow(ListIndex >= 0);
 		int32 block_ptr = ListHeads[ListIndex];
 		return (block_ptr == NullValue) ? 0 : ListBlocks[block_ptr];
 	}
@@ -130,7 +130,7 @@ public:
 	 */
 	inline int32 First(int32 ListIndex) const
 	{
-		check(ListIndex >= 0);
+		checkSlow(ListIndex >= 0);
 		int32 block_ptr = ListHeads[ListIndex];
 		return ListBlocks[block_ptr + 1];
 	}
diff --git a/Engine/Plugins/Experimental/GeometryProcessing/Source/GeometricObjects/Public/Util/UniqueIndexSet.h b/Engine/Plugins/Experimental/GeometryProcessing/Source/GeometricObjects/Public/Util/UniqueIndexSet.h
--- a/Engine/Plugins/Experimental/GeometryProcessing/Source/GeometricObjects/Public/Util/UniqueIndexSet.h
+++ b/Engine/Plugins/Experimental/GeometryProcessing/Source/GeometricObjects/Public/Util/UniqueIndexSet.h
@@ -60,7 +60,7 @@ public:
 	/**
 	 * Add Index to the set
 	 */
-	inline void Add(int32 Index)
+	inline bool Add(int32 Index)
 	{
 		int64& Word = Bits[(int64)(Index / 64)];
 		int64 Offset = (int64)(Index % 64);
@@ -69,7 +69,9 @@ public:
 		{
 			Word |= Mask;
 			Values.Add(Index);
+			return true;
 		}
+		return false;
 	}
 
 	/**
