GeometryProcessing: re-normalize quaternion in TFrame3 rotation, otherwise it can degenerate over time as rotations are accumulated. Const fix in MeshNormals.
#rnx

diff --git a/Engine/Plugins/Experimental/GeometryProcessing/Source/DynamicMesh/Private/MeshNormals.cpp b/Engine/Plugins/Experimental/GeometryProcessing/Source/DynamicMesh/Private/MeshNormals.cpp
--- a/Engine/Plugins/Experimental/GeometryProcessing/Source/DynamicMesh/Private/MeshNormals.cpp
+++ b/Engine/Plugins/Experimental/GeometryProcessing/Source/DynamicMesh/Private/MeshNormals.cpp
@@ -260,7 +260,7 @@ FVector3d FMeshNormals::ComputeVertexNormal(const FDynamicMesh3& Mesh, int32 Ver
 
 
 
-FVector3d FMeshNormals::ComputeOverlayNormal(const FDynamicMesh3& Mesh, FDynamicMeshNormalOverlay* NormalOverlay, int ElemIdx)
+FVector3d FMeshNormals::ComputeOverlayNormal(const FDynamicMesh3& Mesh, const FDynamicMeshNormalOverlay* NormalOverlay, int ElemIdx)
 {
 	int ParentVertexID = NormalOverlay->GetParentVertex(ElemIdx);
 	FVector3d SumNormal = FVector3d::Zero();
diff --git a/Engine/Plugins/Experimental/GeometryProcessing/Source/DynamicMesh/Public/MeshNormals.h b/Engine/Plugins/Experimental/GeometryProcessing/Source/DynamicMesh/Public/MeshNormals.h
--- a/Engine/Plugins/Experimental/GeometryProcessing/Source/DynamicMesh/Public/MeshNormals.h
+++ b/Engine/Plugins/Experimental/GeometryProcessing/Source/DynamicMesh/Public/MeshNormals.h
@@ -124,7 +124,7 @@ public:
 	/**
 	 * @return the computed overlay normal at an element of the overlay (ie based on normals of triangles connected to this element)
 	 */
-	static FVector3d ComputeOverlayNormal(const FDynamicMesh3& Mesh, FDynamicMeshNormalOverlay* NormalOverlay, int ElemIdx);
+	static FVector3d ComputeOverlayNormal(const FDynamicMesh3& Mesh, const FDynamicMeshNormalOverlay* NormalOverlay, int ElemIdx);
 
 	/**
 	 * Initialize the given NormalOverlay with per-vertex normals, ie single overlay element for each mesh vertex.
diff --git a/Engine/Plugins/Experimental/GeometryProcessing/Source/GeometricObjects/Public/FrameTypes.h b/Engine/Plugins/Experimental/GeometryProcessing/Source/GeometricObjects/Public/FrameTypes.h
--- a/Engine/Plugins/Experimental/GeometryProcessing/Source/GeometricObjects/Public/FrameTypes.h
+++ b/Engine/Plugins/Experimental/GeometryProcessing/Source/GeometricObjects/Public/FrameTypes.h
@@ -319,7 +319,15 @@ struct TFrame3
 	 */
 	void Rotate(const TQuaternion<RealType>& Quat)
 	{
-		Rotation = Quat * Rotation;
+		TQuaternion<RealType> NewRotation = Quat * Rotation;
+		if (NewRotation.Normalize() > 0 )
+		{
+			Rotation = NewRotation;
+		}
+		else
+		{
+			checkSlow(false);
+		}
 	}
 
 
@@ -329,7 +337,7 @@ struct TFrame3
 	void Transform(const FTransform& XForm)
 	{
 		Origin = FVector3<RealType>(XForm.TransformPosition((FVector)Origin));
-		Rotation = TQuaternion<RealType>(XForm.GetRotation()) * Rotation;
+		Rotate(TQuaternion<RealType>(XForm.GetRotation()));
 	}
 
 
@@ -339,7 +347,7 @@ struct TFrame3
 	void Transform(const TTransform3<RealType>& XForm)
 	{
 		Origin = XForm.TransformPosition(Origin);
-		Rotation = XForm.GetRotation() * Rotation;
+		Rotate(XForm.GetRotation());
 	}
 
 
