GeometryProcessing: IWYU fix
#rnx

diff --git a/Engine/Plugins/Experimental/GeometryProcessing/Source/GeometricObjects/Public/Util/ProgressCancel.h b/Engine/Plugins/Experimental/GeometryProcessing/Source/GeometricObjects/Public/Util/ProgressCancel.h
--- a/Engine/Plugins/Experimental/GeometryProcessing/Source/GeometricObjects/Public/Util/ProgressCancel.h
+++ b/Engine/Plugins/Experimental/GeometryProcessing/Source/GeometricObjects/Public/Util/ProgressCancel.h
@@ -2,6 +2,8 @@
 
 // Port of geometry3cpp ProgressCancel
 
+#include "Templates/Function.h"
+
 #pragma once
 
 
