add option to improve geometry on mesh boolean during the boolean operation itself, rather than as a post-process, which should be faster and work even if the edges aren't welded
(plus same for mesh self-union)

#rnx

diff --git a/Engine/Plugins/Experimental/GeometryProcessing/Source/DynamicMesh/Private/Operations/MeshBoolean.cpp b/Engine/Plugins/Experimental/GeometryProcessing/Source/DynamicMesh/Private/Operations/MeshBoolean.cpp
--- a/Engine/Plugins/Experimental/GeometryProcessing/Source/DynamicMesh/Private/Operations/MeshBoolean.cpp
+++ b/Engine/Plugins/Experimental/GeometryProcessing/Source/DynamicMesh/Private/Operations/MeshBoolean.cpp
@@ -637,6 +637,11 @@ bool FMeshBoolean::Compute()
 		}
 	}
 
+	if (bSimplifyAlongNewEdges)
+	{
+		SimplifyAlongNewEdges(NumMeshesToProcess, CutMesh, CutBoundaryEdges, AllVIDMatches);
+	}
+
 	if (Operation == EBooleanOp::Difference)
 	{
 		// TODO: implement a way to flip all the triangles in the mesh without building this AllTID array
@@ -691,12 +696,424 @@ bool FMeshBoolean::Compute()
 }
 
 
+
+bool FMeshBoolean::IsFlat(const FDynamicMesh3& Mesh, int VID, double DotTolerance, FVector3d& OutFirstNormal)
+{
+	bool bHasFirst = false;
+	bool bIsFlat = true;
+	Mesh.EnumerateVertexTriangles(VID, [&Mesh, DotTolerance, &OutFirstNormal, &bHasFirst, &bIsFlat](int32 TID)
+		{
+			if (!bIsFlat)
+			{
+				return;
+			}
+			FVector3d Normal = Mesh.GetTriNormal(TID);
+			if (!bHasFirst)
+			{
+				OutFirstNormal = Normal;
+				bHasFirst = true;
+			}
+			else
+			{
+				bIsFlat = bIsFlat && Normal.Dot(OutFirstNormal) >= DotTolerance;
+			}
+		});
+	return bIsFlat;
+}
+
+/**
+ * Test if a given edge collapse would cause a triangle flip or other unacceptable decrease in mesh quality
+ */
+bool FMeshBoolean::CollapseWouldHurtTriangleQuality(
+	const FDynamicMesh3& Mesh, const FVector3d& ExpectNormal,
+	int32 RemoveV, const FVector3d& RemoveVPos, int32 KeepV, const FVector3d& KeepVPos
+	)
+{
+	bool bIsHurt = false;
+	Mesh.EnumerateVertexTriangles(RemoveV,
+		[&Mesh, &bIsHurt, &KeepVPos, RemoveV, KeepV, &ExpectNormal](int32 TID)
+		{
+			if (bIsHurt)
+			{
+				return;
+			}
+			FIndex3i Tri = Mesh.GetTriangle(TID);
+			FVector3d Verts[3];
+			for (int Idx = 0; Idx < 3; Idx++)
+			{
+				int VID = Tri[Idx];
+				if (VID == KeepV)
+				{
+					// this tri has both RemoveV and KeepV, so it'll be removed and we don't need to consider it
+					return;
+				}
+				else if (VID == RemoveV)
+				{
+					// anything at RemoveV is reconnected to KeepV's position
+					Verts[Idx] = KeepVPos;
+				}
+				else
+				{
+					// it's not on the collapsed edge so it stays still
+					Verts[Idx] = Mesh.GetVertex(Tri[Idx]);
+				}
+			}
+			FVector3d Edge1(Verts[1] - Verts[0]);
+			FVector3d Edge2(Verts[2] - Verts[0]);
+			FVector3d VCross(Edge2.Cross(Edge1));
+
+			// TODO: does this tolerance make a difference?  if not, set to zero and remove the VCross.Normalize()
+			double EdgeFlipTolerance = 1.e-5;
+			double Area2 = VCross.Normalize();
+
+			bIsHurt = VCross.Dot(ExpectNormal) <= EdgeFlipTolerance;
+		}
+	);
+	return bIsHurt;
+}
+
+
+/**
+ * Test if a given edge collapse would change the mesh shape or UVs unacceptably
+ */
+bool FMeshBoolean::CollapseWouldChangeShapeOrUVs(
+	const FDynamicMesh3& Mesh, const TSet<int> CutBoundaryEdgeSet, double DotTolerance, int SourceEID,
+	int32 RemoveV, const FVector3d& RemoveVPos, int32 KeepV, const FVector3d& KeepVPos, const FVector3d& EdgeDir,
+	bool bPreserveUVsForMesh, bool bPreserveVertexUVs, bool bPreserveOverlayUVs, float UVEqualThresholdSq)
+{
+	// Search the edges connected to the vertex to find one in the boundary set that points in the opposite direction
+	// If we don't find that edge, or if there are other boundary/seam edges attached, we can't remove this vertex
+	// We also can't remove the vertex if doing so would distort the UVs
+	bool bHasBadEdge = false;
+
+	int OpposedEdge = -1;
+	Mesh.EnumerateVertexEdges(RemoveV,
+		[&Mesh, &CutBoundaryEdgeSet, &DotTolerance,
+		&RemoveVPos, &KeepVPos, &EdgeDir, &bHasBadEdge,
+		&RemoveV, &KeepV, &SourceEID, &OpposedEdge,
+		bPreserveUVsForMesh, bPreserveVertexUVs, bPreserveOverlayUVs, UVEqualThresholdSq](int32 VertEID)
+		{
+			if (bHasBadEdge || VertEID == SourceEID)
+			{
+				return;
+			}
+
+			// it's a known boundary edge; check if it's the opposite-facing one we need
+			if (CutBoundaryEdgeSet.Contains(VertEID))
+			{
+				if (OpposedEdge != -1)
+				{
+					bHasBadEdge = true;
+					return;
+				}
+				FIndex2i OtherEdgeV = Mesh.GetEdgeV(VertEID);
+				int OtherV = IndexUtil::FindEdgeOtherVertex(OtherEdgeV, RemoveV);
+				FVector3d OtherVPos = Mesh.GetVertex(OtherV);
+				FVector3d OtherEdgeDir = OtherVPos - RemoveVPos;
+				if (OtherEdgeDir.Normalize() == 0)
+				{
+					// collapsing degenerate edges above should prevent this
+					bHasBadEdge = true;
+					return; // break instead of continue to skip the whole edge
+				}
+				if (OtherEdgeDir.Dot(EdgeDir) <= -1 + DotTolerance)
+				{
+					OpposedEdge = VertEID;
+				}
+				else
+				{
+					bHasBadEdge = true;
+					return;
+				}
+
+				// test that UVs are not too distorted through the collapse
+				if (!bPreserveUVsForMesh)
+				{
+					return;
+				}
+				float LerpT = (RemoveVPos - OtherVPos).Dot(OtherEdgeDir) / (KeepVPos - OtherVPos).Dot(OtherEdgeDir);
+				if (bPreserveVertexUVs && Mesh.HasVertexUVs())
+				{
+					FVector2f OtherUV = Mesh.GetVertexUV(OtherV);
+					FVector2f RemoveUV = Mesh.GetVertexUV(RemoveV);
+					FVector2f KeepUV = Mesh.GetVertexUV(KeepV);
+					if (FVector2f::Lerp(OtherUV, KeepUV, LerpT).DistanceSquared(RemoveUV) > UVEqualThresholdSq)
+					{
+						bHasBadEdge = true;
+						return;
+					}
+				}
+				if (bPreserveOverlayUVs && Mesh.HasAttributes())
+				{
+					int NumLayers = Mesh.Attributes()->NumUVLayers();
+					FIndex2i SourceEdgeTris = Mesh.GetEdgeT(SourceEID);
+					FIndex2i OppEdgeTris = Mesh.GetEdgeT(VertEID);
+
+					// special handling of seam edge when the edges aren't boundary edges
+					// -- if they're seams, we'd need to check both sides of the seams for a UV match
+					//    but this is complicated and should be quite rare, so we just don't collapse these
+					if (SourceEdgeTris.B != -1 || OppEdgeTris.B != -1)
+					{
+						if (Mesh.Attributes()->IsSeamEdge(SourceEID) ||
+							Mesh.Attributes()->IsSeamEdge(VertEID))
+						{
+							bHasBadEdge = true;
+							return;
+						}
+					}
+
+					FIndex3i SourceBaseTri = Mesh.GetTriangle(SourceEdgeTris.A);
+					FIndex3i OppBaseTri = Mesh.GetTriangle(OppEdgeTris.A);
+					int KeepSourceIdx = IndexUtil::FindTriIndex(KeepV, SourceBaseTri);
+					int RemoveSourceIdx = IndexUtil::FindTriIndex(RemoveV, SourceBaseTri);
+					int OtherOppIdx = IndexUtil::FindTriIndex(OtherV, OppBaseTri);
+					if (!ensure(KeepSourceIdx != -1 && RemoveSourceIdx != -1 && OtherOppIdx != -1))
+					{
+						bHasBadEdge = true;
+						return;
+					}
+
+					// get the UVs per overlay off the triangle(s) attached the two edges
+					for (int UVLayerIdx = 0; UVLayerIdx < NumLayers; UVLayerIdx++)
+					{
+						const FDynamicMeshUVOverlay* UVs = Mesh.Attributes()->GetUVLayer(UVLayerIdx);
+						if (UVs->ElementCount() < 3)
+						{
+							// overlay is not actually in use; skip it
+							continue;
+						}
+						FIndex3i SourceT = UVs->GetTriangle(SourceEdgeTris.A);
+						FIndex3i OppT = UVs->GetTriangle(OppEdgeTris.A);
+						int KeepE = SourceT[KeepSourceIdx];
+						int RemoveE = SourceT[RemoveSourceIdx];
+						int OtherE = OppT[OtherOppIdx];
+						if (KeepE == -1 || RemoveE == -1 || OtherE == -1)
+						{
+							// overlay is not set on relevant triangles; skip it
+							continue;
+						}
+						FVector2f OtherUV = UVs->GetElement(OtherE);
+						FVector2f RemoveUV = UVs->GetElement(RemoveE);
+						FVector2f KeepUV = UVs->GetElement(KeepE);
+						if (FVector2f::Lerp(OtherUV, KeepUV, LerpT).DistanceSquared(RemoveUV) > UVEqualThresholdSq)
+						{
+							bHasBadEdge = true;
+							return;
+						}
+					}
+				}
+			}
+			else // it wasn't in the boundary edge set; check if it's one that would prevent us from safely removing the vertex
+			{
+				if (Mesh.IsBoundaryEdge(VertEID) || (Mesh.HasAttributes() && Mesh.Attributes()->IsSeamEdge(VertEID)))
+				{
+					bHasBadEdge = true;
+				}
+			}
+		});
+
+	return bHasBadEdge;
+}
+
+
+void FMeshBoolean::SimplifyAlongNewEdges(int NumMeshesToProcess, FDynamicMesh3* CutMesh[2], TArray<int> CutBoundaryEdges[2], TMap<int, int>& AllVIDMatches)
+{
+	double DotTolerance = FMathd::Cos(SimplificationAngleTolerance * FMathd::DegToRad);
+
+	TSet<int> CutBoundaryEdgeSets[2]; // set versions of CutBoundaryEdges, for faster membership tests
+	for (int MeshIdx = 0; MeshIdx < NumMeshesToProcess; MeshIdx++)
+	{
+		CutBoundaryEdgeSets[MeshIdx].Append(CutBoundaryEdges[MeshIdx]);
+	}
+
+	int NumCollapses = 0, CollapseIters = 0;
+	int MaxCollapseIters = 1; // TODO: is there a case where we need more iterations?  Perhaps if we add some triangle quality criteria?
+	while (CollapseIters < MaxCollapseIters)
+	{
+		int LastNumCollapses = NumCollapses;
+		for (int EID : CutBoundaryEdges[0])
+		{
+			// this can happen if a collapse removes another cut boundary edge
+			// (which can happen e.g. if you have a degenerate (colinear) tri flat on the cut boundary)
+			if (!CutMesh[0]->IsEdge(EID))
+			{
+				continue;
+			}
+			// don't allow collapses if we somehow get down to our last triangle on either mesh
+			if (CutMesh[0]->TriangleCount() <= 1 || (NumMeshesToProcess == 2 && CutMesh[1]->TriangleCount() <= 1))
+			{
+				break;
+			}
+
+			FDynamicMesh3::FEdge Edge = CutMesh[0]->GetEdge(EID);
+			int Matches[2]{ -1, -1 };
+			bool bHasMatches = NumMeshesToProcess == 2;
+			if (bHasMatches)
+			{
+				for (int MatchIdx = 0; MatchIdx < 2; MatchIdx++)
+				{
+					int* Match = AllVIDMatches.Find(Edge.Vert[MatchIdx]);
+					if (Match)
+					{
+						Matches[MatchIdx] = *Match;
+					}
+					else
+					{
+						bHasMatches = false;
+						// TODO: if we switch to allow collapse on unmatched edges, we shouldn't break here
+						//        b/c we may be partially matched, and need to track which is matched.
+						break;
+					}
+				}
+				if (!bHasMatches)
+				{
+					continue; // edge wasn't matched up on the other mesh; can't collapse it?
+					// TODO: consider supporting collapses in this case?
+				}
+			}
+			// if we have matched vertices, we also need a matched edge to collapse
+			int OtherEID = -1;
+			if (bHasMatches)
+			{
+				OtherEID = CutMesh[1]->FindEdge(Matches[0], Matches[1]);
+				if (OtherEID == -1)
+				{
+					continue;
+				}
+			}
+			// track whether the neighborhood of the vertex is flat (and likewise its matched pair's neighborhood, if present)
+			bool Flat[2]{ false, false };
+			// normals for each flat vertex, and each "side" (mesh 0 and mesh 1, if mesh 1 is present)
+			FVector3d FlatNormals[2][2]{ {FVector3d::Zero(), FVector3d::Zero()}, {FVector3d::Zero(), FVector3d::Zero()} };
+			int NumFlat = 0;
+			for (int VIdx = 0; VIdx < 2; VIdx++)
+			{
+				if (IsFlat(*CutMesh[0], Edge.Vert[VIdx], DotTolerance, FlatNormals[VIdx][0]))
+				{
+					Flat[VIdx] = (Matches[VIdx] == -1) || IsFlat(*CutMesh[1], Matches[VIdx], DotTolerance, FlatNormals[VIdx][1]);
+				}
+
+				if (Flat[VIdx])
+				{
+					NumFlat++;
+				}
+			}
+
+			if (NumFlat == 0)
+			{
+				continue;
+			}
+
+			// see if we can collapse to remove either vertex
+			for (int RemoveVIdx = 0; RemoveVIdx < 2; RemoveVIdx++)
+			{
+				if (!Flat[RemoveVIdx])
+				{
+					continue;
+				}
+				int KeepVIdx = 1 - RemoveVIdx;
+				FVector3d RemoveVPos = CutMesh[0]->GetVertex(Edge.Vert[RemoveVIdx]);
+				FVector3d KeepVPos = CutMesh[0]->GetVertex(Edge.Vert[KeepVIdx]);
+				FVector3d EdgeDir = KeepVPos - RemoveVPos;
+				if (EdgeDir.Normalize() == 0) // 0 is returned as a special case when the edge was too short to normalize
+				{
+					// collapsing degenerate edges above should prevent this
+					ensure(!bCollapseDegenerateEdgesOnCut);
+					// Just skip these edges, because in practice we generally have bCollapseDegenerateEdgesOnCut enabled
+					break; // break instead of continue to skip the whole edge
+				}
+
+				bool bHasBadEdge = false; // will be set if either mesh can't collapse the edge
+				for (int MeshIdx = 0; !bHasBadEdge && MeshIdx < NumMeshesToProcess; MeshIdx++)
+				{
+					int RemoveV = MeshIdx == 0 ? Edge.Vert[RemoveVIdx] : Matches[RemoveVIdx];
+					int KeepV = MeshIdx == 0 ? Edge.Vert[KeepVIdx] : Matches[KeepVIdx];
+					int SourceEID = MeshIdx == 0 ? EID : OtherEID;
+
+					bHasBadEdge = bHasBadEdge || CollapseWouldHurtTriangleQuality(*CutMesh[MeshIdx], FlatNormals[RemoveVIdx][MeshIdx], RemoveV, RemoveVPos, KeepV, KeepVPos);
+					
+					bHasBadEdge = bHasBadEdge || CollapseWouldChangeShapeOrUVs(
+						*CutMesh[MeshIdx], CutBoundaryEdgeSets[MeshIdx], DotTolerance,
+						SourceEID, RemoveV, RemoveVPos, KeepV, KeepVPos, EdgeDir,
+						PreserveUVsOnlyForMesh == -1 || MeshIdx == PreserveUVsOnlyForMesh,
+						bPreserveVertexUVs, bPreserveOverlayUVs, UVDistortTolerance* UVDistortTolerance);
+				};
+
+				if (bHasBadEdge)
+				{
+					continue;
+				}
+
+				FDynamicMesh3::FEdgeCollapseInfo CollapseInfo;
+				int RemoveV = Edge.Vert[RemoveVIdx];
+				int KeepV = Edge.Vert[KeepVIdx];
+				EMeshResult CollapseResult = CutMesh[0]->CollapseEdge(KeepV, RemoveV, 0, CollapseInfo);
+				if (ensure(CollapseResult == EMeshResult::Ok))
+				{
+					bool bSuccessfulCollapse = true;
+					if (bHasMatches)
+					{
+						int OtherRemoveV = Matches[RemoveVIdx];
+						int OtherKeepV = Matches[KeepVIdx];
+						FDynamicMesh3::FEdgeCollapseInfo OtherCollapseInfo;
+						EMeshResult OtherCollapseResult = CutMesh[1]->CollapseEdge(OtherKeepV, OtherRemoveV, 0, OtherCollapseInfo);
+						if (!ensure(OtherCollapseResult == EMeshResult::Ok))
+						{
+							// if we get here, we've somehow managed to collapse the first edge but failed on the second (matched) edge
+							// which will leave a crack in the result unless we can somehow undo the first collapse, which would require a bunch of extra work
+							// but the only case where I could see this happen is if the second edge is on an isolated triangle, which means there is a hole anyway
+							// or if the mesh topology is somehow invalid
+							bSuccessfulCollapse = false;
+						}
+						else
+						{
+							AllVIDMatches.Remove(RemoveV);
+							CutBoundaryEdgeSets[1].Remove(OtherCollapseInfo.CollapsedEdge);
+							CutBoundaryEdgeSets[1].Remove(OtherCollapseInfo.RemovedEdges[0]);
+							if (OtherCollapseInfo.RemovedEdges[1] != -1)
+							{
+								CutBoundaryEdgeSets[1].Remove(OtherCollapseInfo.RemovedEdges[1]);
+							}
+						}
+					}
+					if (bSuccessfulCollapse)
+					{
+						NumCollapses++;
+						CutBoundaryEdgeSets[0].Remove(CollapseInfo.CollapsedEdge);
+						CutBoundaryEdgeSets[0].Remove(CollapseInfo.RemovedEdges[0]);
+						if (CollapseInfo.RemovedEdges[1] != -1)
+						{
+							CutBoundaryEdgeSets[0].Remove(CollapseInfo.RemovedEdges[1]);
+						}
+					}
+				}
+				break; // if we got through to trying to collapse the edge, don't try to collapse from the other vertex.
+			}
+		}
+
+		CutBoundaryEdges[0] = CutBoundaryEdgeSets[0].Array();
+		CutBoundaryEdges[1] = CutBoundaryEdgeSets[1].Array();
+
+		if (NumCollapses == LastNumCollapses)
+		{
+			break;
+		}
+
+		CollapseIters++;
+	}
+}
+
+
 bool FMeshBoolean::MergeEdges(const FMeshIndexMappings& IndexMaps, FDynamicMesh3* CutMesh[2], const TArray<int> CutBoundaryEdges[2], const TMap<int, int>& AllVIDMatches)
 {
 	// translate the edge IDs from CutMesh[1] over to Result mesh edge IDs
 	TArray<int> OtherMeshEdges;
 	for (int OldMeshEID : CutBoundaryEdges[1])
 	{
+		if (!ensure(CutMesh[1]->IsEdge(OldMeshEID)))
+		{
+			continue;
+		}
 		FIndex2i OtherEV = CutMesh[1]->GetEdgeV(OldMeshEID);
 		int MappedEID = Result->FindEdge(IndexMaps.GetNewVertex(OtherEV.A), IndexMaps.GetNewVertex(OtherEV.B));
 		if (ensure(Result->IsBoundaryEdge(MappedEID)))
@@ -765,7 +1182,7 @@ bool FMeshBoolean::MergeEdges(const FMeshIndexMappings& IndexMaps, FDynamicMesh3
 		}));
 
 	// see if we can match anything else
-	bool bAllMatched = CutBoundaryEdges[0].Num() == CutBoundaryEdges[1].Num();
+	bool bAllMatched = true;
 	if (UnmatchedEdges.Num() > 0)
 	{
 		// greedily match within snap tolerance
diff --git a/Engine/Plugins/Experimental/GeometryProcessing/Source/DynamicMesh/Private/Operations/MeshSelfUnion.cpp b/Engine/Plugins/Experimental/GeometryProcessing/Source/DynamicMesh/Private/Operations/MeshSelfUnion.cpp
--- a/Engine/Plugins/Experimental/GeometryProcessing/Source/DynamicMesh/Private/Operations/MeshSelfUnion.cpp
+++ b/Engine/Plugins/Experimental/GeometryProcessing/Source/DynamicMesh/Private/Operations/MeshSelfUnion.cpp
@@ -437,6 +437,12 @@ bool FMeshSelfUnion::Compute()
 		Mesh->SetVertex(Match.Value, Mesh->GetVertex(Match.Key));
 	}
 
+
+	if (bSimplifyAlongNewEdges)
+	{
+		SimplifyAlongNewEdges(CutBoundaryEdges, FoundMatches);
+	}
+
 	if (Cancelled())
 	{
 		return false;
@@ -462,6 +468,177 @@ bool FMeshSelfUnion::Compute()
 }
 
 
+void FMeshSelfUnion::SimplifyAlongNewEdges(TArray<int>& CutBoundaryEdges, TMap<int, int>& FoundMatches)
+{
+	double DotTolerance = FMathd::Cos(SimplificationAngleTolerance * FMathd::DegToRad);
+
+	TSet<int> CutBoundaryEdgeSet; // set version of CutBoundaryEdges, for faster membership tests
+	CutBoundaryEdgeSet.Append(CutBoundaryEdges);
+
+	int NumCollapses = 0, CollapseIters = 0;
+	int MaxCollapseIters = 1; // TODO: is there a case where we need more iterations?  Perhaps if we add some triangle quality criteria?
+	while (CollapseIters < MaxCollapseIters)
+	{
+		int LastNumCollapses = NumCollapses;
+		for (int EID : CutBoundaryEdges)
+		{
+			// this can happen if a collapse removes another cut boundary edge
+			// (which can happen e.g. if you have a degenerate (colinear) tri flat on the cut boundary)
+			if (!Mesh->IsEdge(EID))
+			{
+				continue;
+			}
+
+			FDynamicMesh3::FEdge Edge = Mesh->GetEdge(EID);
+
+			int Matches[2]{ -1, -1 };
+			bool bHasMatches = true;
+			for (int MatchIdx = 0; MatchIdx < 2; MatchIdx++)
+			{
+				int* Match = FoundMatches.Find(Edge.Vert[MatchIdx]);
+				if (Match)
+				{
+					Matches[MatchIdx] = *Match;
+				}
+				else
+				{
+					bHasMatches = false;
+					// TODO: if we switch to allow collapse on unmatched edges, we shouldn't break here
+					//        b/c we may be partially matched, and need to track which is matched.
+					break;
+				}
+			}
+			if (!bHasMatches)
+			{
+				continue; // edge wasn't matched up; can't collapse it?
+				// TODO: consider supporting collapses in this case?
+			}
+
+			// if we have matched vertices, we also need a matched edge to collapse
+			int MatchEID = Mesh->FindEdge(Matches[0], Matches[1]);
+			if (MatchEID == -1)
+			{
+				continue;
+			}
+
+			// track whether the neighborhood of the vertex is flat (and likewise its matched pair's neighborhood, if present)
+			bool Flat[2]{ false, false };
+			// normals for each flat vertex, and each "side" (EID side and MatchEID side)
+			FVector3d FlatNormals[2][2]{ {FVector3d::Zero(), FVector3d::Zero()}, {FVector3d::Zero(), FVector3d::Zero()} };
+			int NumFlat = 0;
+			for (int VIdx = 0; VIdx < 2; VIdx++)
+			{
+				Flat[VIdx] = FMeshBoolean::IsFlat(*Mesh, Edge.Vert[VIdx], DotTolerance, FlatNormals[VIdx][0]) 
+						  && FMeshBoolean::IsFlat(*Mesh, Matches[VIdx], DotTolerance, FlatNormals[VIdx][1]);
+
+				if (Flat[VIdx])
+				{
+					NumFlat++;
+				}
+			}
+
+			if (NumFlat == 0)
+			{
+				continue;
+			}
+
+			// see if we can collapse to remove either vertex
+			for (int RemoveVIdx = 0; RemoveVIdx < 2; RemoveVIdx++)
+			{
+				if (!Flat[RemoveVIdx])
+				{
+					continue;
+				}
+				int KeepVIdx = 1 - RemoveVIdx;
+				// Note: positions are exactly the same on matched edges because snapping has already happened
+				FVector3d RemoveVPos = Mesh->GetVertex(Edge.Vert[RemoveVIdx]);
+				FVector3d KeepVPos = Mesh->GetVertex(Edge.Vert[KeepVIdx]);
+				FVector3d EdgeDir = KeepVPos - RemoveVPos;
+				if (EdgeDir.Normalize() == 0) // 0 is returned as a special case when the edge was too short to normalize
+				{
+					// collapsing degenerate edges above should prevent this
+					ensure(!bCollapseDegenerateEdgesOnCut);
+					// Just skip these edges, because in practice we generally have bCollapseDegenerateEdgesOnCut enabled
+					break; // break instead of continue to skip the whole edge
+				}
+
+				bool bHasBadEdge = false; // will be set if either mesh can't collapse the edge
+				for (int WhichEdge = 0; !bHasBadEdge && WhichEdge < 2; WhichEdge++) // same processing on EID and on MatchEID
+				{
+					int RemoveV = WhichEdge == 0 ? Edge.Vert[RemoveVIdx] : Matches[RemoveVIdx];
+					int KeepV = WhichEdge == 0 ? Edge.Vert[KeepVIdx] : Matches[KeepVIdx];
+					int SourceEID = WhichEdge == 0 ? EID : MatchEID;
+
+					bHasBadEdge = bHasBadEdge || FMeshBoolean::CollapseWouldHurtTriangleQuality(
+						*Mesh, FlatNormals[RemoveVIdx][WhichEdge], RemoveV, RemoveVPos, KeepV, KeepVPos);
+
+					bHasBadEdge = bHasBadEdge || FMeshBoolean::CollapseWouldChangeShapeOrUVs(
+						*Mesh, CutBoundaryEdgeSet, DotTolerance,
+						SourceEID, RemoveV, RemoveVPos, KeepV, KeepVPos, EdgeDir,
+						true, bPreserveVertexUVs, bPreserveOverlayUVs, UVDistortTolerance * UVDistortTolerance);
+				}
+
+				if (bHasBadEdge)
+				{
+					continue;
+				}
+
+				FDynamicMesh3::FEdgeCollapseInfo CollapseInfo;
+				int RemoveV = Edge.Vert[RemoveVIdx];
+				int KeepV = Edge.Vert[KeepVIdx];
+				EMeshResult CollapseResult = Mesh->CollapseEdge(KeepV, RemoveV, 0, CollapseInfo);
+				if (ensure(CollapseResult == EMeshResult::Ok))
+				{
+					int OtherRemoveV = Matches[RemoveVIdx];
+					int OtherKeepV = Matches[KeepVIdx];
+					FDynamicMesh3::FEdgeCollapseInfo OtherCollapseInfo;
+					EMeshResult OtherCollapseResult = Mesh->CollapseEdge(OtherKeepV, OtherRemoveV, 0, OtherCollapseInfo);
+					if (!ensure(OtherCollapseResult == EMeshResult::Ok))
+					{
+						// if we get here, we've somehow managed to collapse the first edge but failed on the second (matched) edge
+						// which will leave a crack in the result unless we can somehow undo the first collapse, which would require a bunch of extra work
+						// but the only case where I could see this happen is if the second edge is on an isolated triangle, which means there is a hole anyway
+						// or if the mesh topology is somehow invalid
+						break;
+					}
+					else
+					{
+						NumCollapses++;
+
+						FoundMatches.Remove(RemoveV);
+						FoundMatches.Remove(OtherRemoveV);
+
+						CutBoundaryEdgeSet.Remove(CollapseInfo.CollapsedEdge);
+						CutBoundaryEdgeSet.Remove(CollapseInfo.RemovedEdges[0]);
+						if (CollapseInfo.RemovedEdges[1] != -1)
+						{
+							CutBoundaryEdgeSet.Remove(CollapseInfo.RemovedEdges[1]);
+						}
+
+						CutBoundaryEdgeSet.Remove(OtherCollapseInfo.CollapsedEdge);
+						CutBoundaryEdgeSet.Remove(OtherCollapseInfo.RemovedEdges[0]);
+						if (OtherCollapseInfo.RemovedEdges[1] != -1)
+						{
+							CutBoundaryEdgeSet.Remove(OtherCollapseInfo.RemovedEdges[1]);
+						}
+					}
+				}
+				break; // if we got through to trying to collapse the edge, don't try to collapse from the other vertex.
+			}
+		}
+
+		CutBoundaryEdges = CutBoundaryEdgeSet.Array();
+
+		if (NumCollapses == LastNumCollapses)
+		{
+			break;
+		}
+
+		CollapseIters++;
+	}
+}
+
+
 bool FMeshSelfUnion::MergeEdges(const TArray<int>& CutBoundaryEdges, const TMap<int, int>& FoundMatches)
 {	
 	// find "easy" match candidates using the already-made vertex correspondence
diff --git a/Engine/Plugins/Experimental/GeometryProcessing/Source/DynamicMesh/Public/Operations/MeshBoolean.h b/Engine/Plugins/Experimental/GeometryProcessing/Source/DynamicMesh/Public/Operations/MeshBoolean.h
--- a/Engine/Plugins/Experimental/GeometryProcessing/Source/DynamicMesh/Public/Operations/MeshBoolean.h
+++ b/Engine/Plugins/Experimental/GeometryProcessing/Source/DynamicMesh/Public/Operations/MeshBoolean.h
@@ -70,6 +70,23 @@ public:
 	/** Set this to be able to cancel running operation */
 	FProgressCancel* Progress = nullptr;
 
+	/** Control whether we attempt to auto-simplify the small planar triangles that the boolean operation tends to generate */
+	bool bSimplifyAlongNewEdges = false;
+	//
+    // Simplification-specific settings (only relevant if bSimplifyAlongNewEdges==true):
+	//
+	/** Degrees of deviation from coplanar that we will still simplify */
+	double SimplificationAngleTolerance = .1;
+	/** Prevent simplification from distorting vertex UVs */
+	bool bPreserveVertexUVs = true;
+	/** Prevent simplification from distorting overlay UVs */
+	bool bPreserveOverlayUVs = true;
+	/** When preserving UVs, sets maximum allowed change in UV coordinates from collapsing an edge, measured at the removed vertex */
+	float UVDistortTolerance = .0001;
+	/** If > -1, then only preserve the UVs of one of the input meshes.  Useful when cutting an artist-created asset w/ procedural geometry. */
+	int PreserveUVsOnlyForMesh = -1;
+
+
 
 	//
 	// Input & Output (to be modified by algorithm)
@@ -147,4 +164,37 @@ private:
 
 	bool MergeEdges(const FMeshIndexMappings& IndexMaps, FDynamicMesh3* CutMesh[2], const TArray<int> CutBoundaryEdges[2], const TMap<int, int>& AllVIDMatches);
 
+	void SimplifyAlongNewEdges(int NumMeshesToProcess, FDynamicMesh3* CutMesh[2], TArray<int> CutBoundaryEdges[2], TMap<int, int>& AllVIDMatches);
+
+public:
+	
+	// specialized helper functions useful for implementing mesh boolean-like features (shared with the extremely similar FMeshSelfUnion algorithm)
+
+	/**
+	 * Test if the triangles connected to a vertex are all coplanar
+	 * @param Mesh The mesh to query
+	 * @param VID The vertex to query
+	 * @param DotTolerance If the dot product of two normals are >= this tolerance, the normals are considered equivalent
+	 * @param The normal of the first triangle attached to the vertex.
+	 * @return Whether all the triangles were coplanar
+	 */
+	static bool IsFlat(const FDynamicMesh3& Mesh, int VID, double DotTolerance, FVector3d& OutFirstNormal);
+
+	/**
+	 * Test if a given edge collapse would cause a triangle flip or other unacceptable decrease in mesh quality
+	 * Specialized for collapsing at flat triangles
+	 * TODO: Currently this only detects triangle flips; need to extend it to also detect other mesh quality issues
+	 */
+	static bool CollapseWouldHurtTriangleQuality(
+		const FDynamicMesh3& Mesh, const FVector3d& ExpectNormal,
+		int32 RemoveV, const FVector3d& RemoveVPos, int32 KeepV, const FVector3d& KeepVPos
+	);
+
+	/**
+	 * Test if a given edge collapse would change the mesh shape or UVs unacceptably
+	 */
+	static bool CollapseWouldChangeShapeOrUVs(
+		const FDynamicMesh3& Mesh, const TSet<int> CutBoundaryEdgeSet, double DotTolerance, int SourceEID,
+		int32 RemoveV, const FVector3d& RemoveVPos, int32 KeepV, const FVector3d& KeepVPos,
+		const FVector3d& EdgeDir, bool bPreserveUVsForMesh, bool bPreserveVertexUVs, bool bPreserveOverlayUVs, float UVToleranceSq);
 };
diff --git a/Engine/Plugins/Experimental/GeometryProcessing/Source/DynamicMesh/Public/Operations/MeshSelfUnion.h b/Engine/Plugins/Experimental/GeometryProcessing/Source/DynamicMesh/Public/Operations/MeshSelfUnion.h
--- a/Engine/Plugins/Experimental/GeometryProcessing/Source/DynamicMesh/Public/Operations/MeshSelfUnion.h
+++ b/Engine/Plugins/Experimental/GeometryProcessing/Source/DynamicMesh/Public/Operations/MeshSelfUnion.h
@@ -14,7 +14,7 @@
 
 #include "Util/ProgressCancel.h"
 
-
+#include "Operations/MeshBoolean.h" // for shared utility functions
 
 
 
@@ -56,6 +56,20 @@ public:
 	/** Set this to be able to cancel running operation */
 	FProgressCancel* Progress = nullptr;
 
+	/** Control whether we attempt to auto-simplify the small planar triangles that the boolean operation tends to generate */
+	bool bSimplifyAlongNewEdges = false;
+	//
+	// Simplification-specific settings (only relevant if bSimplifyAlongNewEdges==true):
+	//
+	/** Degrees of deviation from coplanar that we will still simplify */
+	double SimplificationAngleTolerance = .1;
+	/** Prevent simplification from distorting vertex UVs */
+	bool bPreserveVertexUVs = true;
+	/** Prevent simplification from distorting overlay UVs */
+	bool bPreserveOverlayUVs = true;
+	/** When preserving UVs, sets maximum allowed change in UV coordinates from collapsing an edge, measured at the removed vertex */
+	float UVDistortTolerance = .0001;
+
 
 	//
 	// Input & Output (to be modified by algorithm)
@@ -114,4 +128,6 @@ private:
 
 	bool MergeEdges(const TArray<int>& CutBoundaryEdges, const TMap<int, int>& AllVIDMatches);
 
+	void SimplifyAlongNewEdges(TArray<int>& CutBoundaryEdges, TMap<int, int>& FoundMatches);
+
 };
diff --git a/Engine/Plugins/Experimental/MeshModelingToolset/Source/ModelingOperators/Private/CompositionOps/BooleanMeshesOp.cpp b/Engine/Plugins/Experimental/MeshModelingToolset/Source/ModelingOperators/Private/CompositionOps/BooleanMeshesOp.cpp
--- a/Engine/Plugins/Experimental/MeshModelingToolset/Source/ModelingOperators/Private/CompositionOps/BooleanMeshesOp.cpp
+++ b/Engine/Plugins/Experimental/MeshModelingToolset/Source/ModelingOperators/Private/CompositionOps/BooleanMeshesOp.cpp
@@ -72,7 +72,7 @@ void FBooleanMeshesOp::CalculateResult(FProgressCancel* Progress)
 	}
 
 	MeshBoolean.bPutResultInInputSpace = false;
-	MeshBoolean.bTrackAllNewEdges = (bTryCollapseExtraEdges);
+	MeshBoolean.bSimplifyAlongNewEdges = bTryCollapseExtraEdges;
 	MeshBoolean.Progress = Progress;
 	bool bSuccess = MeshBoolean.Compute();
 	ResultTransform = MeshBoolean.ResultTransform;
@@ -84,40 +84,6 @@ void FBooleanMeshesOp::CalculateResult(FProgressCancel* Progress)
 
 	CreatedBoundaryEdges = MeshBoolean.CreatedBoundaryEdges;
 
-	// Boolean operation is based on edge splits, which results in spurious vertices
-	// along straight intersection edges. Try to collapse away those extra vertices.
-	if (bTryCollapseExtraEdges)
-	{
-		FDynamicMesh3* TargetMesh = MeshBoolean.Result;
-
-		FQEMSimplification Simplifier(TargetMesh);
-		Simplifier.bAllowSeamCollapse = true;
-		if (TargetMesh->Attributes())
-		{
-			TargetMesh->Attributes()->SplitAllBowties();		// eliminate any bowties that might have formed on UV seams.
-		}
-
-		FMeshConstraints Constraints;
-		FMeshConstraintsUtil::ConstrainAllBoundariesAndSeams(Constraints, *TargetMesh,
-			EEdgeRefineFlags::NoConstraint, EEdgeRefineFlags::NoConstraint, EEdgeRefineFlags::NoConstraint,
-			true, true, true);
-		Simplifier.SetExternalConstraints(MoveTemp(Constraints));
-
-		Simplifier.SimplifyToMinimalPlanar( TryCollapseExtraEdgesPlanarThresh,
-			[&MeshBoolean](int32 eid) { return MeshBoolean.AllNewEdges.Contains(eid); } );
-
-		// update boundary-edge set
-		TArray<int32> UpdatedBoundaryEdges;
-		for (int32 eid : CreatedBoundaryEdges)
-		{
-			if (MeshBoolean.Result->IsEdge(eid))
-			{
-				UpdatedBoundaryEdges.Add(eid);
-			}
-		}
-		CreatedBoundaryEdges = MoveTemp(UpdatedBoundaryEdges);
-	}
-
 	// try to fill cracks/holes in boolean result
 	if (CreatedBoundaryEdges.Num() > 0 && bAttemptFixHoles)
 	{
diff --git a/Engine/Plugins/Experimental/MeshModelingToolset/Source/ModelingOperators/Private/CompositionOps/SelfUnionMeshesOp.cpp b/Engine/Plugins/Experimental/MeshModelingToolset/Source/ModelingOperators/Private/CompositionOps/SelfUnionMeshesOp.cpp
--- a/Engine/Plugins/Experimental/MeshModelingToolset/Source/ModelingOperators/Private/CompositionOps/SelfUnionMeshesOp.cpp
+++ b/Engine/Plugins/Experimental/MeshModelingToolset/Source/ModelingOperators/Private/CompositionOps/SelfUnionMeshesOp.cpp
@@ -30,7 +30,7 @@ void FSelfUnionMeshesOp::CalculateResult(FProgressCancel* Progress)
 	FMeshSelfUnion Union(ResultMesh.Get());
 	Union.WindingThreshold = WindingNumberThreshold;
 	Union.bTrimFlaps = bTrimFlaps;
-	Union.bTrackAllNewEdges = (bTryCollapseExtraEdges);
+	Union.bSimplifyAlongNewEdges = bTryCollapseExtraEdges;
 	bool bSuccess = Union.Compute();
 
 	if (Progress && Progress->Cancelled())
@@ -40,40 +40,6 @@ void FSelfUnionMeshesOp::CalculateResult(FProgressCancel* Progress)
 
 	CreatedBoundaryEdges = Union.CreatedBoundaryEdges;
 
-	// Boolean operation is based on edge splits, which results in spurious vertices
-	// along straight intersection edges. Try to collapse away those extra vertices.
-	if (bTryCollapseExtraEdges)
-	{
-		FDynamicMesh3* TargetMesh = ResultMesh.Get();
-
-		FQEMSimplification Simplifier(TargetMesh);
-		Simplifier.bAllowSeamCollapse = true;
-		if (TargetMesh->Attributes())
-		{
-			TargetMesh->Attributes()->SplitAllBowties();		// eliminate any bowties that might have formed on UV seams.
-		}
-
-		FMeshConstraints Constraints;
-		FMeshConstraintsUtil::ConstrainAllBoundariesAndSeams(Constraints, *TargetMesh,
-			EEdgeRefineFlags::NoConstraint, EEdgeRefineFlags::NoConstraint, EEdgeRefineFlags::NoConstraint,
-			true, true, true);
-		Simplifier.SetExternalConstraints(MoveTemp(Constraints));
-
-		Simplifier.SimplifyToMinimalPlanar(TryCollapseExtraEdgesPlanarThresh,
-			[&Union](int32 eid) { return Union.AllNewEdges.Contains(eid); });
-
-		// update boundary-edge set as we may have collapsed some during this process
-		TArray<int32> UpdatedBoundaryEdges;
-		for (int32 eid : CreatedBoundaryEdges)
-		{
-			if (ResultMesh->IsEdge(eid))
-			{
-				UpdatedBoundaryEdges.Add(eid);
-			}
-		}
-		CreatedBoundaryEdges = MoveTemp(UpdatedBoundaryEdges);
-	}
-
 	if (!bSuccess && bAttemptFixHoles)
 	{
 		FMeshBoundaryLoops OpenBoundary(ResultMesh.Get(), false);
diff --git a/Engine/Plugins/Experimental/PlanarCutPlugin/Source/PlanarCut/Private/PlanarCut.cpp b/Engine/Plugins/Experimental/PlanarCutPlugin/Source/PlanarCut/Private/PlanarCut.cpp
--- a/Engine/Plugins/Experimental/PlanarCutPlugin/Source/PlanarCut/Private/PlanarCut.cpp
+++ b/Engine/Plugins/Experimental/PlanarCutPlugin/Source/PlanarCut/Private/PlanarCut.cpp
@@ -1551,6 +1551,8 @@ struct FDynamicMeshCollection
 						Op = FMeshBoolean::EBooleanOp::Difference;
 					}
 					FMeshBoolean Boolean(&Surface.AugMesh, &Cell.AugMesh, &BoolResults[CellIdx]->AugMesh, Op);
+					Boolean.bSimplifyAlongNewEdges = true;
+					Boolean.PreserveUVsOnlyForMesh = 0; // slight warping of the autogenerated cell UVs generally doesn't matter
 					Boolean.bWeldSharedEdges = false;
 					if (!Boolean.Compute())
 					{
@@ -1742,6 +1744,8 @@ struct FDynamicMeshCollection
 							Op = FMeshBoolean::EBooleanOp::Difference;
 						}
 						FMeshBoolean Boolean(&Surface.AugMesh, &Cell.AugMesh, &AugBoolResult, Op);
+						Boolean.bSimplifyAlongNewEdges = true;
+						Boolean.PreserveUVsOnlyForMesh = 0; // slight warping of the autogenerated cell UVs generally doesn't matter
 						Boolean.bWeldSharedEdges = false;
 						if (!Boolean.Compute())
 						{
