fix memory leak w/ copy/assignment/move on geometryobjects dynamicvector

#rnx

diff --git a/Engine/Plugins/Experimental/GeometryProcessing/Source/GeometricObjects/Public/Util/DynamicVector.h b/Engine/Plugins/Experimental/GeometryProcessing/Source/GeometricObjects/Public/Util/DynamicVector.h
--- a/Engine/Plugins/Experimental/GeometryProcessing/Source/GeometricObjects/Public/Util/DynamicVector.h
+++ b/Engine/Plugins/Experimental/GeometryProcessing/Source/GeometricObjects/Public/Util/DynamicVector.h
@@ -224,36 +224,44 @@ private:
 		TArray<ArrayType*> Elements;
 
 	public:
-		TBlockVector()
-		{
-		}
+		TBlockVector() = default;
+		TBlockVector(TBlockVector&& Moved) = default;
+
 		TBlockVector(const TBlockVector& Copy)
 		{
-			*this = Copy;
-		}
-		TBlockVector(TBlockVector&& Moved)
-		{
-			Elements = MoveTemp(Moved.Elements);
+			int32 N = Copy.Num();
+			Elements.Reserve(N);
+			for (int32 k = 0; k < N; ++k)
+			{
+				Elements.Add(new ArrayType(*Copy.Elements[k]));
+			}
 		}
 		TBlockVector& operator=(const TBlockVector& Copy)
 		{
-			int32 N = Copy.Num();
-			Elements.SetNum(N);
-			for (int32 k = 0; k < N; ++k)
+			if (&Copy != this)
 			{
-				Elements[k] = new ArrayType;
-				*Elements[k] = *Copy.Elements[k];
+				int32 N = Copy.Num();
+				Empty(N);
+				Elements.Reserve(N);
+				for (int32 k = 0; k < N; ++k)
+				{
+					Elements.Add(new ArrayType(*Copy.Elements[k]));
+				}
 			}
 			return *this;
 		}
 		TBlockVector& operator=(TBlockVector&& Moved)
 		{
-			Elements = MoveTemp(Moved.Elements);
+			if (&Moved != this)
+			{
+				Empty();
+				Elements = MoveTemp(Moved.Elements);
+			}
 			return *this;
 		}
 		~TBlockVector()
 		{
-			Truncate(0, false);
+			Empty();
 		}
 
 		int32 Num() const { return Elements.Num(); }
@@ -280,6 +288,15 @@ private:
 			}
 			Elements.RemoveAt(NewElementCount, Elements.Num() - NewElementCount, bAllowShrinking);
 		}
+
+		void Empty(int32 NewReservedSize = 0)
+		{
+			for (int32 k = 0, N = Elements.Num(); k < N; ++k)
+			{
+				delete Elements[k];
+			}
+			Elements.Empty(NewReservedSize);
+		}
 	};
 
 	// memory chunks for dynamic vector
@@ -441,10 +458,13 @@ using  TDynamicVector2i = TDynamicVectorN<int, 2>;
 template <class Type>
 void TDynamicVector<Type>::Clear()
 {
-	Blocks.Truncate(0, false);
+	Blocks.Truncate(1, false);
 	CurBlock = 0;
 	CurBlockUsed = 0;
-	Blocks.Add(new BlockType());
+	if (Blocks.Num() == 0)
+	{
+		Blocks.Add(new BlockType());
+	}
 }
 
 template <class Type>
