Fix non-unity compile error -- missing include.

#rnx

diff --git a/Engine/Plugins/Experimental/GeometryProcessing/Source/DynamicMesh/Public/Selections/MeshConnectedComponents.h b/Engine/Plugins/Experimental/GeometryProcessing/Source/DynamicMesh/Public/Selections/MeshConnectedComponents.h
--- a/Engine/Plugins/Experimental/GeometryProcessing/Source/DynamicMesh/Public/Selections/MeshConnectedComponents.h
+++ b/Engine/Plugins/Experimental/GeometryProcessing/Source/DynamicMesh/Public/Selections/MeshConnectedComponents.h
@@ -6,6 +6,7 @@
 
 #include "CoreMinimal.h"
 #include "DynamicMesh3.h"
+#include "Containers/IndirectArray.h"
 
 /**
  * FMeshConnectedComponents calculates Connected Components of a Mesh, or sub-regions of a Mesh.
