Optimize FDynamicMesh3 to do less wasteful allocating/better copying

#rnx

diff --git a/Engine/Plugins/Experimental/GeometryProcessing/Source/DynamicMesh/Private/DynamicMesh3.cpp b/Engine/Plugins/Experimental/GeometryProcessing/Source/DynamicMesh/Private/DynamicMesh3.cpp
--- a/Engine/Plugins/Experimental/GeometryProcessing/Source/DynamicMesh/Private/DynamicMesh3.cpp
+++ b/Engine/Plugins/Experimental/GeometryProcessing/Source/DynamicMesh/Private/DynamicMesh3.cpp
@@ -144,23 +144,13 @@ void FDynamicMesh3::Copy(const FMeshShapeGenerator* Generator)
 {
 	Clear();
 
-	Vertices = TDynamicVector<FVector3d>();
-	VertexEdgeLists = FSmallListSet();
-	VertexRefCounts = FRefCountVector();
-	Triangles = TDynamicVector<FIndex3i>();
-	TriangleEdges = TDynamicVector<FIndex3i>();
-	TriangleRefCounts = FRefCountVector();
-	TriangleGroups = TDynamicVector<int>();
-	GroupIDCounter = 0;
-	Edges = TDynamicVector<FEdge>();
-	EdgeRefCounts = FRefCountVector();
+	EnableTriangleGroups();
 
 	if (Generator->HasAttributes())
 	{
 		EnableAttributes();
 	}
 
-
 	int NumVerts = Generator->Vertices.Num();
 	for (int i = 0; i < NumVerts; ++i)
 	{
@@ -341,7 +331,32 @@ void FDynamicMesh3::CompactCopy(const FDynamicMesh3& copy, bool bNormals, bool b
 
 void FDynamicMesh3::Clear()
 {
-	*this = FDynamicMesh3();
+	Vertices.Clear();
+	VertexRefCounts.Clear();
+	VertexNormals.Reset();
+	VertexColors.Reset(); 
+	VertexUVs.Reset();
+	VertexEdgeLists.Reset();
+
+	Triangles.Clear();
+	TriangleRefCounts.Clear();
+	TriangleEdges.Clear();
+	TriangleGroups.Reset();
+	GroupIDCounter = 0;
+
+	Edges.Clear();
+	EdgeRefCounts.Clear();
+
+	AttributeSet.Reset();
+	Timestamp = 0;
+	ShapeTimestamp = 0;
+	TopologyTimestamp = 0;
+
+	CachedBoundingBox = FAxisAlignedBox3d();
+	CachedBoundingBoxTimestamp = -1;
+
+	bIsClosedCached = false;
+	CachedIsClosedTimestamp = -1;
 }
 
 int FDynamicMesh3::GetComponentsFlags() const
diff --git a/Engine/Plugins/Experimental/GeometryProcessing/Source/DynamicMesh/Public/DynamicMesh3.h b/Engine/Plugins/Experimental/GeometryProcessing/Source/DynamicMesh/Public/DynamicMesh3.h
--- a/Engine/Plugins/Experimental/GeometryProcessing/Source/DynamicMesh/Public/DynamicMesh3.h
+++ b/Engine/Plugins/Experimental/GeometryProcessing/Source/DynamicMesh/Public/DynamicMesh3.h
@@ -164,7 +164,7 @@ protected:
 
 public:
 	/** Default constructor */
-	FDynamicMesh3() : FDynamicMesh3(true, false, false, false) {}
+	FDynamicMesh3() : FDynamicMesh3(false, false, false, false) {}
 
 	/** Copy/Move construction */
 	FDynamicMesh3(const FDynamicMesh3& CopyMesh);
diff --git a/Engine/Plugins/Experimental/GeometryProcessing/Source/GeometricObjects/Public/Util/DynamicVector.h b/Engine/Plugins/Experimental/GeometryProcessing/Source/GeometricObjects/Public/Util/DynamicVector.h
--- a/Engine/Plugins/Experimental/GeometryProcessing/Source/GeometricObjects/Public/Util/DynamicVector.h
+++ b/Engine/Plugins/Experimental/GeometryProcessing/Source/GeometricObjects/Public/Util/DynamicVector.h
@@ -202,9 +202,9 @@ public:
 	}
 
 private:
-	static constexpr int nShiftBits = 11;
+	static constexpr int nShiftBits = 9;
 	static constexpr int BlockSize = 1 << nShiftBits;
-	static constexpr int BlockIndexBitmask = BlockSize - 1; // low 11 bits
+	static constexpr int BlockIndexBitmask = BlockSize - 1; // low 9 bits
 	static_assert( BlockSize && ((BlockSize & (BlockSize - 1)) == 0), "DynamicVector: BlockSize must be a power of two");
 
 	unsigned int CurBlock{0};
diff --git a/Engine/Plugins/Experimental/GeometryProcessing/Source/GeometricObjects/Public/Util/RefCountVector.h b/Engine/Plugins/Experimental/GeometryProcessing/Source/GeometricObjects/Public/Util/RefCountVector.h
--- a/Engine/Plugins/Experimental/GeometryProcessing/Source/GeometricObjects/Public/Util/RefCountVector.h
+++ b/Engine/Plugins/Experimental/GeometryProcessing/Source/GeometricObjects/Public/Util/RefCountVector.h
@@ -224,7 +224,7 @@ public:
 
 	void RebuildFreeList()
 	{
-		FreeIndices = TDynamicVector<int>();
+		FreeIndices.Clear();
 		UsedCount = 0;
 
 		int N = (int)RefCounts.GetLength();
@@ -243,11 +243,18 @@ public:
 
 	void Trim(int maxIndex)
 	{
-		FreeIndices = TDynamicVector<int>();
+		FreeIndices.Clear();
 		RefCounts.Resize(maxIndex);
 		UsedCount = maxIndex;
 	}
 
+	void Clear()
+	{
+		FreeIndices.Clear();
+		RefCounts.Clear();
+		UsedCount = 0;
+	}
+
 	//
 	// Iterators
 	//
diff --git a/Engine/Plugins/Experimental/GeometryProcessing/Source/GeometricObjects/Public/Util/SmallListSet.h b/Engine/Plugins/Experimental/GeometryProcessing/Source/GeometricObjects/Public/Util/SmallListSet.h
--- a/Engine/Plugins/Experimental/GeometryProcessing/Source/GeometricObjects/Public/Util/SmallListSet.h
+++ b/Engine/Plugins/Experimental/GeometryProcessing/Source/GeometricObjects/Public/Util/SmallListSet.h
@@ -70,6 +70,19 @@ public:
 	 */
 	void Resize(int32 NewSize);
 
+	/**
+	 * Reset to initial state
+	 */
+	void Reset()
+	{
+		ListHeads.Clear();
+		ListBlocks.Clear();
+		FreeBlocks.Clear();
+		AllocatedCount = 0;
+		LinkedListElements.Clear();
+		FreeHeadIndex = NullValue;
+	}
+
 
 	/**
 	 * @return true if a list has been allocated at the given ListIndex
