Geometry Tools:  Mesh Simplification.  Improve collapse behavior with attribute seams.  Previously edge collapses were more limited in the neighborhood of seams.

Additionally this CL:
Introduces a simplifier-level bool allow for easy switching between simplification with and without quadric memory and adds logic to remove isolated triangles

Extend MeshOverlay to identify seam-end edges  -  DynamicMeshAttributeSet.{h,cpp} DynamicMeshOverlay.{h,cpp}
Expose ability to recompute edge constraints   -  MeshConstraintUtil.{h,cpp}  MeshConstraints.h

diff --git a/Engine/Plugins/Experimental/GeometryProcessing/Source/DynamicMesh/Private/DynamicMeshAttributeSet.cpp b/Engine/Plugins/Experimental/GeometryProcessing/Source/DynamicMesh/Private/DynamicMeshAttributeSet.cpp
--- a/Engine/Plugins/Experimental/GeometryProcessing/Source/DynamicMesh/Private/DynamicMeshAttributeSet.cpp
+++ b/Engine/Plugins/Experimental/GeometryProcessing/Source/DynamicMesh/Private/DynamicMeshAttributeSet.cpp
@@ -349,6 +349,26 @@ bool FDynamicMeshAttributeSet::IsSeamEdge(int eid) const
 	return false;
 }
 
+bool FDynamicMeshAttributeSet::IsSeamEndEdge(int eid) const
+{
+	for (const FDynamicMeshUVOverlay& UVLayer : UVLayers)
+	{
+		if (UVLayer.IsSeamEndEdge(eid))
+		{
+			return true;
+		}
+	}
+
+	for (const FDynamicMeshNormalOverlay& NormalLayer : NormalLayers)
+	{
+		if (NormalLayer.IsSeamEndEdge(eid))
+		{
+			return true;
+		}
+	}
+
+	return false;
+}
 
 bool FDynamicMeshAttributeSet::IsSeamEdge(int EdgeID, bool& bIsUVSeamOut, bool& bIsNormalSeamOut) const
 {
diff --git a/Engine/Plugins/Experimental/GeometryProcessing/Source/DynamicMesh/Private/DynamicMeshOverlay.cpp b/Engine/Plugins/Experimental/GeometryProcessing/Source/DynamicMesh/Private/DynamicMeshOverlay.cpp
--- a/Engine/Plugins/Experimental/GeometryProcessing/Source/DynamicMesh/Private/DynamicMeshOverlay.cpp
+++ b/Engine/Plugins/Experimental/GeometryProcessing/Source/DynamicMesh/Private/DynamicMeshOverlay.cpp
@@ -493,6 +493,54 @@ bool TDynamicMeshOverlay<RealType,ElementSize>::IsSeamEdge(int eid) const
 }
 
 
+template<typename RealType, int ElementSize>
+bool TDynamicMeshOverlay<RealType, ElementSize>::IsSeamEndEdge(int eid) const
+{
+	FIndex2i et = ParentMesh->GetEdgeT(eid);
+	if (et.B == FDynamicMesh3::InvalidID)
+	{
+		return false;
+	}
+
+	FIndex2i ev = ParentMesh->GetEdgeV(eid);
+	int base_a = ev.A, base_b = ev.B;
+
+	bool bASet = IsSetTriangle(et.A), bBSet = IsSetTriangle(et.B);
+	if (!bASet || !bBSet) 
+	{
+		return false; 
+	}
+
+	FIndex3i Triangle0 = GetTriangle(et.A);
+	FIndex3i BaseTriangle0(ParentVertices[Triangle0.A], ParentVertices[Triangle0.B], ParentVertices[Triangle0.C]);
+	int idx_base_a0 = BaseTriangle0.IndexOf(base_a);
+	int idx_base_b0 = BaseTriangle0.IndexOf(base_b);
+
+	FIndex3i Triangle1 = GetTriangle(et.B);
+	FIndex3i BaseTriangle1(ParentVertices[Triangle1.A], ParentVertices[Triangle1.B], ParentVertices[Triangle1.C]);
+	int idx_base_a1 = BaseTriangle1.IndexOf(base_a);
+	int idx_base_b1 = BaseTriangle1.IndexOf(base_b);
+
+	int el_a_tri0 = Triangle0[idx_base_a0];
+	int el_b_tri0 = Triangle0[idx_base_b0];
+	int el_a_tri1 = Triangle1[idx_base_a1];
+	int el_b_tri1 = Triangle1[idx_base_b1];
+
+	bool bIsSeam = !IndexUtil::SamePairUnordered(el_a_tri0, el_b_tri0, el_a_tri1, el_b_tri1);
+
+	bool bIsSeamEnd = false;
+	if (bIsSeam)
+	{
+		// is only one of elements split?
+		if ((el_a_tri0 == el_a_tri1 || el_a_tri0 == el_b_tri1) || (el_b_tri0 == el_b_tri1 || el_b_tri0 == el_a_tri1))
+		{
+			bIsSeamEnd = true;
+		}
+	}
+
+	return bIsSeamEnd;
+}
+
 template<typename RealType, int ElementSize>
 bool TDynamicMeshOverlay<RealType, ElementSize>::HasInteriorSeamEdges() const
 {
@@ -882,39 +930,40 @@ void TDynamicMeshOverlay<RealType, ElementSize>::OnCollapseEdge(const FDynamicMe
 	
 	// look up triangle 0
 	FIndex3i Triangle0(-1,-1,-1), BaseTriangle0(-1,-1,-1);
-	int idx_removed0_a = -1, idx_removed0_b = -1;
+	int idx_removed_tri0 = -1, idx_kept_tri0 = -1;
 	if (bT0Set)
 	{
 		Triangle0 = GetTriangle(tid_removed0);
 		BaseTriangle0 = FIndex3i(ParentVertices[Triangle0.A], ParentVertices[Triangle0.B], ParentVertices[Triangle0.C]);
-		idx_removed0_a = BaseTriangle0.IndexOf(vid_base_kept);
-		idx_removed0_b = BaseTriangle0.IndexOf(vid_base_removed);
+		idx_kept_tri0 = BaseTriangle0.IndexOf(vid_base_kept);
+		idx_removed_tri0    = BaseTriangle0.IndexOf(vid_base_removed);
 	}
 
 	// look up triangle 1 if this is not a boundary edge
 	FIndex3i Triangle1(-1,-1,-1), BaseTriangle1(-1,-1,-1);
+	int idx_removed_tri1 = -1, idx_kept_tri1 = -1;
 	if (collapseInfo.bIsBoundary == false && bT1Set)
 	{
 		Triangle1 = GetTriangle(tid_removed1);
 		BaseTriangle1 = FIndex3i(ParentVertices[Triangle1.A], ParentVertices[Triangle1.B], ParentVertices[Triangle1.C]);
 
-		int idx_removed1_a = BaseTriangle1.IndexOf(vid_base_kept);
-		int idx_removed1_b = BaseTriangle1.IndexOf(vid_base_removed);
+		idx_kept_tri1 = BaseTriangle1.IndexOf(vid_base_kept);
+		idx_removed_tri1 = BaseTriangle1.IndexOf(vid_base_removed);
 
 		if (bT0Set)
 		{
-			int el_a_tri0 = Triangle0[idx_removed0_a];
-			int el_b_tri0 = Triangle0[idx_removed0_b];
-			int el_a_tri1 = Triangle1[idx_removed1_a];
-			int el_b_tri1 = Triangle1[idx_removed1_b];
+			int el_kept_tri0    = Triangle0[idx_kept_tri0];
+			int el_removed_tri0 = Triangle0[idx_removed_tri0];
+			int el_kept_tri1    = Triangle1[idx_kept_tri1];
+			int el_removed_tri1 = Triangle1[idx_removed_tri1];
 
 			// is this a seam?
-			bIsSeam = !IndexUtil::SamePairUnordered(Triangle0[idx_removed0_a], Triangle0[idx_removed0_b], Triangle1[idx_removed1_a], Triangle1[idx_removed1_b]);
+			bIsSeam = !IndexUtil::SamePairUnordered(el_kept_tri0, el_removed_tri0, el_kept_tri1, el_removed_tri1);
 
 			if (bIsSeam)
 			{
 				// is only one of elements split?
-				if ((el_a_tri0 == el_a_tri1 || el_a_tri0 == el_b_tri1) || (el_b_tri0 == el_a_tri1 || el_b_tri0 == el_b_tri1))
+				if ((el_kept_tri0 == el_kept_tri1 || el_kept_tri0 == el_removed_tri1) || (el_removed_tri0 == el_kept_tri1 || el_removed_tri0 == el_removed_tri1))
 				{
 					bIsSeamEnd = true;
 				}
@@ -922,6 +971,9 @@ void TDynamicMeshOverlay<RealType, ElementSize>::OnCollapseEdge(const FDynamicMe
 		}
 	}
 
+	// this should be protected against by calling code.
+	//checkSlow(!bIsSeamEnd);
+
 	// need to find the elementid for the "kept" and "removed" vertices that are connected by the edges of T0 and T1.
 	// If this edge is :
 	//       not a seam - just one kept and one removed element.
@@ -937,42 +989,28 @@ void TDynamicMeshOverlay<RealType, ElementSize>::OnCollapseEdge(const FDynamicMe
 	bool bFoundKeptElement[2]    = { false, false };
 	if (bT0Set)
 	{
-		for (int j = 0; j < 3; ++j)
-		{
-			if (BaseTriangle0[j] == vid_base_kept)
-			{
-				kept_elemid[0] = Triangle0[j];
-			}
-			if (BaseTriangle0[j] == vid_base_removed)
-			{
-				removed_elemid[0] = Triangle0[j];
-			}
-		}
-		check(kept_elemid[0] != FDynamicMesh3::InvalidID);
-		check(removed_elemid[0] != FDynamicMesh3::InvalidID);
-		bFoundKeptElement[0]= bFoundRemovedElement[0] = true;
+		kept_elemid[0] = Triangle0[idx_kept_tri0];
+		removed_elemid[0] = Triangle0[idx_removed_tri0];
+		bFoundKeptElement[0] = bFoundRemovedElement[0] = true;
+
+		checkSlow(kept_elemid[0] != FDynamicMesh3::InvalidID);
+		checkSlow(removed_elemid[0] != FDynamicMesh3::InvalidID);
+		
 	}
 	if ((bIsSeam || !bT0Set) && bT1Set)
 	{
-		for (int j = 0; j < 3; ++j)
-		{
-			if (BaseTriangle1[j] == vid_base_kept)
-			{
-				kept_elemid[1] = Triangle1[j];
-			}
-			if (BaseTriangle1[j] == vid_base_removed)
-			{
-				removed_elemid[1] = Triangle1[j];
-			}
-		}
-		check(kept_elemid[1] != FDynamicMesh3::InvalidID);
-		check(removed_elemid[1] != FDynamicMesh3::InvalidID);
+		kept_elemid[1] = Triangle1[idx_kept_tri1];
+		removed_elemid[1] = Triangle1[idx_removed_tri1];
 		bFoundKeptElement[1] = bFoundRemovedElement[1] = true;
+
+		checkSlow(kept_elemid[1] != FDynamicMesh3::InvalidID);
+		checkSlow(removed_elemid[1] != FDynamicMesh3::InvalidID);
+		
 	}
 	if (!bT0Set && !bT1Set)
 	{
 		// if neither triangle was set, still need to check surrounding triangles; one or both elements may still be used
-		check(!bFoundRemovedElement[0] && !bFoundRemovedElement[1]);
+		checkSlow(!bFoundRemovedElement[0] && !bFoundRemovedElement[1]);
 		
 		// note this may do the wrong thing in the case that at least one of the verts has split elements.
 		// higher-level code should make sure this doesn't happen.
@@ -1021,7 +1059,7 @@ void TDynamicMeshOverlay<RealType, ElementSize>::OnCollapseEdge(const FDynamicMe
 				{
 					if (elem_id == removed_elemid[0])
 					{
-					ElementsRefCounts.Decrement(elem_id);
+						ElementsRefCounts.Decrement(elem_id);
 						ElementTriangles[3 * onering_tid + j] = kept_elemid[0];
 						if (bFoundKeptElement[0])
 						{
diff --git a/Engine/Plugins/Experimental/GeometryProcessing/Source/DynamicMesh/Private/MeshConstraintsUtil.cpp b/Engine/Plugins/Experimental/GeometryProcessing/Source/DynamicMesh/Private/MeshConstraintsUtil.cpp
--- a/Engine/Plugins/Experimental/GeometryProcessing/Source/DynamicMesh/Private/MeshConstraintsUtil.cpp
+++ b/Engine/Plugins/Experimental/GeometryProcessing/Source/DynamicMesh/Private/MeshConstraintsUtil.cpp
@@ -48,140 +48,56 @@ FMeshConstraintsUtil::ConstrainAllBoundariesAndSeams(FMeshConstraints& Constrain
 	const FDynamicMeshAttributeSet* Attributes = Mesh.Attributes();
 
 	// Seam edge can never flip, it is never fully unconstrained 
-	EEdgeRefineFlags SeamEdgeContraint = EEdgeRefineFlags::NoFlip;
+	EEdgeRefineFlags SeamEdgeConstraint = EEdgeRefineFlags::NoFlip;
 	if (!bAllowSeamSplits)
 	{
-		SeamEdgeContraint = EEdgeRefineFlags((uint8)SeamEdgeContraint | (uint8)EEdgeRefineFlags::NoSplit);
+		SeamEdgeConstraint = EEdgeRefineFlags((int)SeamEdgeConstraint | (int)EEdgeRefineFlags::NoSplit);
 	}
 	if (!bAllowSeamCollapse)
 	{
-		SeamEdgeContraint = EEdgeRefineFlags((uint8)SeamEdgeContraint | (uint8)EEdgeRefineFlags::NoCollapse);
+		SeamEdgeConstraint = EEdgeRefineFlags((int)SeamEdgeConstraint | (int)EEdgeRefineFlags::NoCollapse);
 	}
 
 	FCriticalSection ConstraintSetLock;
 
 	int32 NumEdges = Mesh.MaxEdgeID();
 	bool bHaveGroups = Mesh.HasTriangleGroups();
+
+
 	ParallelFor(NumEdges, [&](int EdgeID)
 	{
-		bool bIsEdge = Mesh.IsEdge(EdgeID);
-		if (Mesh.IsEdge(EdgeID))
+		FVertexConstraint VtxConstraintA = FVertexConstraint::Unconstrained();
+		FVertexConstraint VtxConstraintB = FVertexConstraint::Unconstrained();
+
+		FEdgeConstraint EdgeConstraint(EEdgeRefineFlags::NoConstraint);
+
+		// compute the edge and vertex constraints.
+		bool bHasUpdate = ConstrainEdgeBoundariesAndSeams(
+		        EdgeID,
+				Mesh,
+				MeshBoundaryConstraint,
+				GroupBoundaryConstraint,
+				MaterialBoundaryConstraint,
+				SeamEdgeConstraint,
+				bAllowSeamSmoothing,
+				EdgeConstraint, VtxConstraintA, VtxConstraintB);
+		
+		if (bHasUpdate)
 		{
-			const bool bIsMeshBoundary = Mesh.IsBoundaryEdge(EdgeID);
-			const bool bIsGroupBoundary = bHaveGroups && Mesh.IsGroupBoundaryEdge(EdgeID);
-			const bool bIsMaterialBoundary = Attributes && Attributes->IsMaterialBoundaryEdge(EdgeID);
-			const bool bIsSeam = Attributes && Attributes->IsSeamEdge(EdgeID);
-			FVertexConstraint VtxConstraint = FVertexConstraint::Unconstrained();
-			EEdgeRefineFlags EdgeFlags{};
-
-			auto ApplyBoundaryConstraint = [&VtxConstraint, &EdgeFlags](uint8 BoundaryConstraint)
-				{
-				VtxConstraint.bCannotDelete = VtxConstraint.bCannotDelete ||
-						(BoundaryConstraint == (uint8)EEdgeRefineFlags::FullyConstrained) ||
-						(BoundaryConstraint == (uint8)EEdgeRefineFlags::SplitsOnly);
-				VtxConstraint.bCanMove = VtxConstraint.bCanMove &&
-						(BoundaryConstraint != (uint8)EEdgeRefineFlags::FullyConstrained) &&
-						(BoundaryConstraint != (uint8)EEdgeRefineFlags::SplitsOnly);
-					EdgeFlags = EEdgeRefineFlags((uint8)EdgeFlags |
-												 (uint8)BoundaryConstraint);
-				};
-
-			if ( bIsMeshBoundary )
-			{
-				ApplyBoundaryConstraint((uint8)MeshBoundaryConstraint);
-			}
-			if ( bIsGroupBoundary )
-			{
-				ApplyBoundaryConstraint((uint8)GroupBoundaryConstraint);
-			}
-			if ( bIsMaterialBoundary )
-			{
-				ApplyBoundaryConstraint((uint8)MaterialBoundaryConstraint);
-			}
-			if ( bIsSeam )
-			{
-
-				VtxConstraint.bCanMove = VtxConstraint.bCanMove && (bAllowSeamSmoothing || bAllowSeamCollapse);
-				VtxConstraint.bCannotDelete = VtxConstraint.bCannotDelete  || !bAllowSeamCollapse;
-				EdgeFlags = EEdgeRefineFlags((uint8)EdgeFlags |
-											 (uint8)( SeamEdgeContraint ));
-
-				// Additional logic to add the NoCollapse flag to any edge that is the start or end of a seam.
-				if (bAllowSeamCollapse)
-				{
-					FIndex2i et = Mesh.GetEdgeT(EdgeID);
-					// test if two double attribute edge shares one element: call this a seam end
-					auto IsSeamWithEnd = [&, et](auto& Overlay)->bool
-					{
-						if (et.A == -1 || et.B == -1)
-						{
-							return false;
-						}
-						bool bASet = Overlay.IsSetTriangle(et.A), bBSet = Overlay.IsSetTriangle(et.B);
-						if (!bASet || !bBSet)
-						{
-							return false;
-						}
-
-						TArray<int, TInlineAllocator<6>> UniqueElements;
-						FIndex3i Triangle0 = Overlay.GetTriangle(et.A);
-						UniqueElements.AddUnique(Triangle0[0]);
-						UniqueElements.AddUnique(Triangle0[1]);
-						UniqueElements.AddUnique(Triangle0[2]);
-						FIndex3i Triangle1 = Overlay.GetTriangle(et.B);
-						UniqueElements.AddUnique(Triangle1[0]);
-						UniqueElements.AddUnique(Triangle1[1]);
-						UniqueElements.AddUnique(Triangle1[2]);
-
-						return UniqueElements.Num() == 5;
-					};
-
-					bool bHasSeamEnd = false;
-					for (int32 i = 0; !bHasSeamEnd && (i < Attributes->NumUVLayers()); ++i)
-					{
-						bool bIsEnd = IsSeamWithEnd(*Attributes->GetUVLayer(i));
-
-						bHasSeamEnd = bHasSeamEnd || bIsEnd;
-					}
-					for (int32 i = 0; !bHasSeamEnd && (i < Attributes->NumNormalLayers()); ++i)
-					{
-						bool bIsEnd = IsSeamWithEnd(*Attributes->GetNormalLayer(i));
-
-						bHasSeamEnd = bHasSeamEnd || bIsEnd;
-					}
-
-					if (bHasSeamEnd)
-					{
-						EdgeFlags = EEdgeRefineFlags((uint8)EdgeFlags | (uint8)EEdgeRefineFlags::NoCollapse);
-					}
-				}
-			}
-			if (bIsMeshBoundary||bIsGroupBoundary||bIsMaterialBoundary||bIsSeam)
-			{
-				FIndex2i EdgeVerts = Mesh.GetEdgeV(EdgeID);
-
-				FEdgeConstraint EdgeConstraint(EdgeFlags);
-
-				// don't update with a phantom constraint (i.e. an unconstrained constraint )
-				if (!(EdgeConstraint.IsUnconstrained() && VtxConstraint.IsUnconstrained()))
-				{
-					ConstraintSetLock.Lock();
-
-					Constraints.SetOrUpdateEdgeConstraint(EdgeID, EdgeConstraint);
-
-					// If any vertex constraints exist, we can only make them more restrictive!
+			// have updates - merge with existing constraints
+			ConstraintSetLock.Lock();
+			
+			FIndex2i EdgeVerts = Mesh.GetEdgeV(EdgeID);
+			
+			Constraints.SetOrUpdateEdgeConstraint(EdgeID, EdgeConstraint);
 
-					FVertexConstraint ConstraintA = VtxConstraint;
-					ConstraintA.CombineConstraint(Constraints.GetVertexConstraint(EdgeVerts.A));
-					Constraints.SetOrUpdateVertexConstraint(EdgeVerts.A, ConstraintA);
+			VtxConstraintA.CombineConstraint(Constraints.GetVertexConstraint(EdgeVerts.A));
+			Constraints.SetOrUpdateVertexConstraint(EdgeVerts.A, VtxConstraintA);
 
-					FVertexConstraint ConstraintB = VtxConstraint;
-					ConstraintB.CombineConstraint(Constraints.GetVertexConstraint(EdgeVerts.B));
-					Constraints.SetOrUpdateVertexConstraint(EdgeVerts.B, ConstraintB);
+			VtxConstraintB.CombineConstraint(Constraints.GetVertexConstraint(EdgeVerts.B));
+			Constraints.SetOrUpdateVertexConstraint(EdgeVerts.B, VtxConstraintB);
 
-					ConstraintSetLock.Unlock();
-				}
-			}
+			ConstraintSetLock.Unlock();
 		}
 	}, (bParallel == false) );
 }
@@ -255,3 +171,108 @@ void FMeshConstraintsUtil::ConstrainROIBoundariesInEdgeROI(FMeshConstraints& Con
 		}
 	}
 }
+
+bool FMeshConstraintsUtil::ConstrainEdgeBoundariesAndSeams(const int EdgeID,
+	                                                       const FDynamicMesh3& Mesh,
+	                                                       const EEdgeRefineFlags MeshBoundaryConstraintFlags,
+	                                                       const EEdgeRefineFlags GroupBoundaryConstraintFlags,
+	                                                       const EEdgeRefineFlags MaterialBoundaryConstraintFlags,
+	                                                       const EEdgeRefineFlags SeamEdgeConstraintFlags,
+	                                                       const bool bAllowSeamSmoothing,
+	                                                       FEdgeConstraint& EdgeConstraint,
+	                                                       FVertexConstraint& VertexConstraintA,
+	                                                       FVertexConstraint& VertexConstraintB)
+{
+	const bool bAllowSeamCollapse = FEdgeConstraint::CanCollapse(SeamEdgeConstraintFlags);
+
+	// initialize constraints
+	VertexConstraintA = FVertexConstraint::Unconstrained();
+	VertexConstraintB = FVertexConstraint::Unconstrained();
+	EdgeConstraint = FEdgeConstraint::Unconstrained();
+
+	const bool bIsEdge = Mesh.IsEdge(EdgeID);
+	if (!bIsEdge)  return false;
+
+	const bool bHaveGroups = Mesh.HasTriangleGroups();
+	const FDynamicMeshAttributeSet* Attributes = Mesh.Attributes();
+
+	const bool bIsMeshBoundary = Mesh.IsBoundaryEdge(EdgeID);
+	const bool bIsGroupBoundary = bHaveGroups && Mesh.IsGroupBoundaryEdge(EdgeID);
+	const bool bIsMaterialBoundary = Attributes && Attributes->IsMaterialBoundaryEdge(EdgeID);
+	const bool bIsSeam = Attributes && Attributes->IsSeamEdge(EdgeID);
+
+	FVertexConstraint CurVtxConstraint = FVertexConstraint::Unconstrained(); // note: this is needed since the default constructor is constrained.
+	EEdgeRefineFlags EdgeFlags{};
+
+	auto ApplyBoundaryConstraint = [&CurVtxConstraint, &EdgeFlags](EEdgeRefineFlags BoundaryConstraintFlags)
+	{
+
+		CurVtxConstraint.bCannotDelete = CurVtxConstraint.bCannotDelete ||
+			(!FEdgeConstraint::CanCollapse(BoundaryConstraintFlags) &&
+				!FEdgeConstraint::CanFlip(BoundaryConstraintFlags)
+				);
+
+		CurVtxConstraint.bCanMove = CurVtxConstraint.bCanMove &&
+			(FEdgeConstraint::CanCollapse(BoundaryConstraintFlags) ||
+				FEdgeConstraint::CanFlip(BoundaryConstraintFlags)
+				);
+
+		EdgeFlags = EEdgeRefineFlags((int)EdgeFlags | (int)BoundaryConstraintFlags);
+	};
+
+
+	if (bIsMeshBoundary)
+	{
+		ApplyBoundaryConstraint(MeshBoundaryConstraintFlags);
+	}
+	if (bIsGroupBoundary)
+	{
+		ApplyBoundaryConstraint(GroupBoundaryConstraintFlags);
+	}
+	if (bIsMaterialBoundary)
+	{
+		ApplyBoundaryConstraint(MaterialBoundaryConstraintFlags);
+	}
+	if (bIsSeam)
+	{
+		CurVtxConstraint.bCannotDelete = CurVtxConstraint.bCannotDelete || !bAllowSeamCollapse;
+		CurVtxConstraint.bCanMove = CurVtxConstraint.bCanMove && (bAllowSeamSmoothing || bAllowSeamCollapse);
+
+		EdgeFlags = EEdgeRefineFlags((int)EdgeFlags | (int)(SeamEdgeConstraintFlags));
+
+		// Additional logic to add the NoCollapse flag to any edge that is the start or end of a seam.
+		if (bAllowSeamCollapse)
+		{
+
+			bool bHasSeamEnd = false;
+			for (int32 i = 0; !bHasSeamEnd && (i < Attributes->NumUVLayers()); ++i)
+			{
+				const FDynamicMeshUVOverlay* UVLayer = Attributes->GetUVLayer(i);
+				bHasSeamEnd = bHasSeamEnd || UVLayer->IsSeamEndEdge(EdgeID);
+			}
+			for (int32 i = 0; !bHasSeamEnd && (i < Attributes->NumNormalLayers()); ++i)
+			{
+				const FDynamicMeshNormalOverlay* NormalOverlay = Attributes->GetNormalLayer(i);
+				bHasSeamEnd = bHasSeamEnd || NormalOverlay->IsSeamEndEdge(EdgeID);
+			}
+
+			if (bHasSeamEnd)
+			{
+				EdgeFlags = EEdgeRefineFlags((int)EdgeFlags | (int)EEdgeRefineFlags::NoCollapse);
+			}
+		}
+	}
+	if (bIsMeshBoundary || bIsGroupBoundary || bIsMaterialBoundary || bIsSeam)
+	{
+		EdgeConstraint = FEdgeConstraint(EdgeFlags);
+
+		// only return true if we have a constraint
+		if (!EdgeConstraint.IsUnconstrained() || !CurVtxConstraint.IsUnconstrained())
+		{
+			VertexConstraintA.CombineConstraint(CurVtxConstraint);
+			VertexConstraintB.CombineConstraint(CurVtxConstraint);
+			return  true;
+		}
+	}
+	return false;
+}
\ No newline at end of file
diff --git a/Engine/Plugins/Experimental/GeometryProcessing/Source/DynamicMesh/Private/MeshSimplification.cpp b/Engine/Plugins/Experimental/GeometryProcessing/Source/DynamicMesh/Private/MeshSimplification.cpp
--- a/Engine/Plugins/Experimental/GeometryProcessing/Source/DynamicMesh/Private/MeshSimplification.cpp
+++ b/Engine/Plugins/Experimental/GeometryProcessing/Source/DynamicMesh/Private/MeshSimplification.cpp
@@ -1,6 +1,7 @@
 // Copyright Epic Games, Inc. All Rights Reserved.
 
 #include "MeshSimplification.h"
+#include "MeshConstraintsUtil.h"
 #include "DynamicMeshAttributeSet.h"
 #include "Util/IndexUtil.h"
 #include "Async/ParallelFor.h"
@@ -119,11 +120,6 @@ void TMeshSimplification<QuadricErrorType>::InitializeSeamQuadrics()
 	{
 		const FDynamicMeshAttributeSet* Attributes = Mesh->Attributes();
 
-		if (!Attributes)
-		{
-			return;
-		}
-
 		for (int eid : Mesh->EdgeIndicesItr())
 		{
 			bool bNeedsQuadric = Mesh->IsBoundaryEdge(eid);
@@ -160,22 +156,6 @@ void TMeshSimplification<QuadricErrorType>::InitializeVertexQuadrics()
 		}
 		//check(TMathUtil.EpsilonEqual(0, vertQuadrics[i].Evaluate(Mesh->GetVertex(i)), TMathUtil.Epsilon * 10));
 	}
-
-	// Fully memoryless simplifier directly adds seam quadrics to the vertex quadrics so it can more effiently construct
-	// edge quadrics from vertex quadrics only. 
-	// We don't do this for the standard QEM (with memory) because the memory only applies to the vertex quadrics and not the seam Qs
-	if (TIsSame<QuadricErrorType, FAttrBasedQuadricErrord>::Value)
-	{
-		// for each seam edge, add the seam quadric to its verts.
-		for (auto& seamQuadric : seamQuadrics)
-		{
-			int eid = seamQuadric.Key;
-			FIndex2i vids = Mesh->GetEdgeV(eid);
-
-			vertQuadrics[vids[0]].AddSeamQuadric(seamQuadric.Value);
-			vertQuadrics[vids[1]].AddSeamQuadric(seamQuadric.Value);
-		}
-	}
 }
 
 template <typename QuadricErrorType>
@@ -184,6 +164,20 @@ QuadricErrorType TMeshSimplification<QuadricErrorType>::AssembleEdgeQuadric(cons
 	//  form standard edge quadric as sum of the vertex quadrics for the edge endpoints
 	QuadricErrorType EdgeQuadric(vertQuadrics[edge.Vert.A], vertQuadrics[edge.Vert.B]);
 	
+	if (!bRetainQuadricMemory)
+	{ 
+		// the edge.Tri faces are double counted. Remove one.
+		const FIndex2i& Tris = edge.Tri;
+		if (Tris.A != FDynamicMesh3::InvalidID)
+		{
+			EdgeQuadric.Add(-triAreas[Tris.A], triQuadrics[Tris.A]);
+		}
+
+		if (Tris.B != FDynamicMesh3::InvalidID)
+		{
+			EdgeQuadric.Add(-triAreas[Tris.B], triQuadrics[Tris.B]);
+		}
+	}
 	if (bAllowSeamCollapse)
 	{ 
 		// lambda that adds any adjacent seam quadrics to the edge quadric
@@ -194,7 +188,7 @@ QuadricErrorType TMeshSimplification<QuadricErrorType>::AssembleEdgeQuadric(cons
 				if (const FSeamQuadricType* seamQuadric =  seamQuadrics.Find(eid))
 				{
 					EdgeQuadric.AddSeamQuadric(*seamQuadric);
-				}
+				} 
 			}
 		};
 	
@@ -206,26 +200,6 @@ QuadricErrorType TMeshSimplification<QuadricErrorType>::AssembleEdgeQuadric(cons
 	return EdgeQuadric;
 }
 
-template<>
-FAttrBasedQuadricErrord TMeshSimplification<FAttrBasedQuadricErrord>::AssembleEdgeQuadric(const FDynamicMesh3::FEdge& edge) const
-{
-	FAttrBasedQuadricErrord Q(vertQuadrics[edge.Vert.A], vertQuadrics[edge.Vert.B]);
-
-	// the edge.Tri faces are double counted. Remove one.
-	const FIndex2i& Tris = edge.Tri;
-	if (Tris.A != FDynamicMesh3::InvalidID)
-	{
-		Q.Add(-triAreas[Tris.A], triQuadrics[Tris.A]);
-	}
-
-	if (Tris.B != FDynamicMesh3::InvalidID)
-	{
-		Q.Add(-triAreas[Tris.B], triQuadrics[Tris.B]);
-	}
-
-	return Q;
-}
-
 
 template <typename QuadricErrorType>
 void TMeshSimplification<QuadricErrorType>::InitializeQueue()
@@ -377,26 +351,79 @@ FVector3d TMeshSimplification<QuadricErrorType>::OptimalPoint(int eid, const FQu
 
 
 
-// update queue weight for each edge in vertex one-ring. Retains vertex Quadric memory
-template <>
-void DYNAMICMESH_API TMeshSimplification<FQuadricErrord>::UpdateNeighborhood(const FDynamicMesh3::FEdgeCollapseInfo& collapseInfo)
-{
+template <typename QuadricErrorType>
+void TMeshSimplification<QuadricErrorType>::UpdateNeighborhood(const FDynamicMesh3::FEdgeCollapseInfo& collapseInfo)
+{	
 	int kvid = collapseInfo.KeptVertex;
 	int rvid = collapseInfo.RemovedVertex;
 
-	// Quadric "memory"  the retained vertex quadric is the sum of the two vert quadrics
-	vertQuadrics[kvid] = FQuadricErrord(vertQuadrics[kvid], vertQuadrics[rvid]);
-
-	double EdgeWeight = this->SeamEdgeWeight;
+	FIndex2i removedTris = collapseInfo.RemovedTris;
+	FIndex2i opposingVerts = collapseInfo.OpposingVerts;
 
-	for (int eid : Mesh->VtxEdgesItr(kvid))
+	// --- Update the seam quadrics
+	if (bAllowSeamCollapse)
 	{
-		FDynamicMesh3::FEdge ne = Mesh->GetEdge(eid);
-
-		// update the seam quadric to reflect changes in the seams
-		if (bAllowSeamCollapse)
+		
+		FIndex2i removedEdges = collapseInfo.RemovedEdges;
+		FIndex2i keptEdges = collapseInfo.KeptEdges;
+
+		// update the map between edge id and seam quadrics 
+		// if constraints exist, they define the edges with seam quadrics
+		// otherwise require kept edges to have a seam quadric if either 
+		// the kept or collapse edge had a seam quadric.
+		if (Constraints)  // quadrics on the constrained edges
 		{
+			if (Constraints->HasEdgeConstraint(keptEdges.A))
+			{
+				seamQuadrics.Add(keptEdges.A);
+			}
+			else
+			{
+				seamQuadrics.Remove(keptEdges.A);
+			}
 			
+			if (keptEdges.B != FDynamicMesh3::InvalidID)
+			{ 
+				if( Constraints->HasEdgeConstraint(keptEdges.B))
+				{
+					seamQuadrics.Add(keptEdges.B);
+				}
+				else
+				{
+					seamQuadrics.Remove(keptEdges.B);
+				}
+			}
+		}
+		else // propagate any existing seam quadric requirements.
+		{			
+			if (FSeamQuadricType* seamQuadric = seamQuadrics.Find(removedEdges.A))
+			{
+				seamQuadrics.Add(keptEdges.A);
+			}
+			if (removedEdges.B != FDynamicMesh3::InvalidID)
+			{
+				if (FSeamQuadricType* seamQuadric = seamQuadrics.Find(removedEdges.B))
+				{
+					seamQuadrics.Add(keptEdges.B);
+				}
+			}
+		}
+
+		// removed quadrics from deleted edges
+		seamQuadrics.Remove(removedEdges.A);
+		if (removedEdges.B != FDynamicMesh3::InvalidID)
+		{
+			seamQuadrics.Remove(removedEdges.B);
+		}
+		
+		// update any seam quadrics adjacent to kvid to reflect changes in the seams
+	
+		double EdgeWeight = this->SeamEdgeWeight;
+
+		for (int eid : Mesh->VtxEdgesItr(kvid))
+		{
+			FDynamicMesh3::FEdge ne = Mesh->GetEdge(eid);
+
 			// need to recompute this seam quadric
 			if (FSeamQuadricType* seamQuadric = seamQuadrics.Find(eid))
 			{
@@ -419,57 +446,24 @@ void DYNAMICMESH_API TMeshSimplification<FQuadricErrord>::UpdateNeighborhood(con
 
 				seamQuadric->Scale(EdgeWeight);
 			}
-
-		}
-
-		
-		FQuadricErrord Q = AssembleEdgeQuadric(ne);
-		FVector3d opt = OptimalPoint(eid, Q, ne.Vert.A, ne.Vert.B);
-		float err = (float)Q.Evaluate(opt);
-		EdgeQuadrics[eid] = QEdge(eid, Q, opt);
-		if (EdgeQueue.Contains(eid))
-		{
-			EdgeQueue.Update(eid, err);
-		}
-		else
-		{
-			EdgeQueue.Insert(eid, err);
 		}
 	}
-}
 
-// update queue weight for each edge in vertex one-ring.  Memoryless
-template <typename QuadricErrorType>
-void TMeshSimplification<QuadricErrorType>::UpdateNeighborhood(const FDynamicMesh3::FEdgeCollapseInfo& collapseInfo)
-{
-	double EdgeWeight = this->SeamEdgeWeight;
-
-	int vid = collapseInfo.KeptVertex;
-	FIndex2i removedTris = collapseInfo.RemovedTris;
-	FIndex2i opposingVerts = collapseInfo.OpposingVerts;
-
-	TArray<int, TInlineAllocator<15>> AdjTris;
-	for (int tid : Mesh->VtxTrianglesItr(vid))
-	{
-		AdjTris.Add(tid);
-	}
-
-	TArray<int, TInlineAllocator<15>> AdjEdges;
-	for (int eid : Mesh->VtxEdgesItr(vid))
-	{
-		AdjEdges.Add(eid);
+	// --- Update the vertex quadrics
+	if (bRetainQuadricMemory)
+	{ 
+		// Quadric "memory"  the retained vertex quadric is the sum of the two vert quadrics
+		vertQuadrics[kvid] = QuadricErrorType(vertQuadrics[kvid], vertQuadrics[rvid]);
 	}
-
-	// This is the faster version that selectively updates the one-ring
+	else
 	{
-
 		// compute the change in affected face quadrics, and then propagate 
-		// that change to the face adjacent verts.
+			// that change to the face adjacent verts.
 		FVector3d n, c;
 		double NewtriArea;
 
 		// Update the triangle areas and quadrics that will have changed
-		for (int tid : AdjTris)
+		for (int tid : Mesh->VtxTrianglesItr(kvid))
 		{
 
 			const double OldtriArea = triAreas[tid];
@@ -488,7 +482,7 @@ void TMeshSimplification<QuadricErrorType>::UpdateNeighborhood(const FDynamicMes
 			// update the vert quadrics that are adjacent to vid.
 			for (int32 i = 0; i < 3; ++i)
 			{
-				if (tri_vids[i] == vid) continue;
+				if (tri_vids[i] == kvid) continue;
 
 				// correct the adjacent vertQuadrics
 				vertQuadrics[tri_vids[i]].Add(-OldtriArea, OldtriQuadric); // subtract old quadric
@@ -518,69 +512,47 @@ void TMeshSimplification<QuadricErrorType>::UpdateNeighborhood(const FDynamicMes
 		// NB: in the version with memory this quadric took the value of the edge quadric that collapsed.
 		{
 			FQuadricErrorType vertQuadric = FQuadricErrorType::Zero();
-			for (int tid : AdjTris)
+			for (int tid : Mesh->VtxTrianglesItr(kvid))
 			{
 				vertQuadric.Add(triAreas[tid], triQuadrics[tid]);
 			}
-			vertQuadrics[vid] = vertQuadric;
+			vertQuadrics[kvid] = vertQuadric;
 		}
+	}
 
-		if (bAllowSeamCollapse)
-		{
-			for (int eid : AdjEdges)
-			{
-				// need to recompute this seam quadric
-				if (FSeamQuadricType* seamQuadric = seamQuadrics.Find(eid))
-				{
-					const FDynamicMesh3::FEdge ne = Mesh->GetEdge(eid);
-
-					// subtract the old seam quadric from adj verts
-					if (ne.Vert[0] == vid) // note we have reset the vertQuadric at the retained vertex so it doesn't have a contribution from the old seamQ
-					{
-						vertQuadrics[ne.Vert[1]].SubtractSeamQuadric(*seamQuadric);
-					}
-					else
-					{
-						check(ne.Vert[1] == vid);
-						vertQuadrics[ne.Vert[0]].SubtractSeamQuadric(*seamQuadric);
-					}
-
-					// rebuild the seam quadric
-
-					FVector3d p0 = Mesh->GetVertex(ne.Vert[0]);
-					FVector3d p1 = Mesh->GetVertex(ne.Vert[1]);
-
-					// face normal 
-					FVector3d nA = Mesh->GetTriNormal(ne.Tri.A);
+	// --- Update all edge quadrics in the nbrhood
+	// NB: this has to follow updating all potential seam quadrics adjacent to kvid 
+	// because an edge quadric gathers seam quadrics adjacent the ends 
 
-					// this constrains the point to a plane aligned with the edge and normal to the face
-					*seamQuadric = CreateSeamQuadric(p0, p1, nA);
-					// add the other side - this constrains the point to the line where the two planes intersect.
-					if (ne.Tri.B != FDynamicMesh3::InvalidID)
-					{
-						FVector3d nB = Mesh->GetTriNormal(ne.Tri.B);
-						seamQuadric->Add(CreateSeamQuadric(p0, p1, nB));
-					}
-
-					seamQuadric->Scale(EdgeWeight);
+	if (bRetainQuadricMemory)
+	{ 
+		for (int eid : Mesh->VtxEdgesItr(kvid))
+		{
+			FDynamicMesh3::FEdge ne = Mesh->GetEdge(eid);
 
-					// add the seam quadric to the adj verts
-					vertQuadrics[ne.Vert[0]].AddSeamQuadric(*seamQuadric);
-					vertQuadrics[ne.Vert[1]].AddSeamQuadric(*seamQuadric);
-				}
+			QuadricErrorType Q = AssembleEdgeQuadric(ne);
+			FVector3d opt = OptimalPoint(eid, Q, ne.Vert.A, ne.Vert.B);
+			float err = (float)Q.Evaluate(opt);
+			EdgeQuadrics[eid] = QEdge(eid, Q, opt);
+			if (EdgeQueue.Contains(eid))
+			{
+				EdgeQueue.Update(eid, err);
 			}
+			else
+			{
+				EdgeQueue.Insert(eid, err);
+			}
+		}
 	}
-	}
-
-	// Update all the edges
+	else
 	{
 		TArray<int, TInlineAllocator<64>> EdgesToUpdate;
-		for (int adjeid : AdjEdges)
+		for (int adjeid : Mesh->VtxEdgesItr(kvid))
 		{
 			EdgesToUpdate.Add(adjeid);
 
 			const FIndex2i Verts = Mesh->GetEdgeV(adjeid);
-			int adjvid = (Verts[0] == vid) ? Verts[1] : Verts[0];
+			int adjvid = (Verts[0] == kvid) ? Verts[1] : Verts[0];
 			if (adjvid != FDynamicMesh3::InvalidID)
 			{
 				for (int eid : Mesh->VtxEdgesItr(adjvid))
@@ -726,7 +698,14 @@ void TMeshSimplification<QuadricErrorType>::DoSimplify()
 		ESimplificationResult result = CollapseEdge(eid, EdgeQuadrics[eid].collapse_pt, collapseInfo);
 		if (result == ESimplificationResult::Ok_Collapsed)
 		{
+			// update the quadrics
 			UpdateNeighborhood(collapseInfo);
+
+		}
+		else if (result == ESimplificationResult::Failed_IsolatedTriangle && Mesh->TriangleCount() > 2)
+		{
+			const FDynamicMesh3::FEdge Edge = Mesh->GetEdge(eid);
+			RemoveIsolatedTriangle(Edge.Tri.A);
 		}
 	}
 	ProfileEndCollapse();
@@ -1192,7 +1171,16 @@ ESimplificationResult TMeshSimplification<QuadricErrorType>::CollapseEdge(int ed
 	// check if we should collapse, and also find which vertex we should retain
 	// in cases where we have constraints/etc
 	int collapse_to = -1;
-	bool bCanCollapse = CanCollapseEdge(edgeID, a, b, c, d, t0, t1, collapse_to);
+	bool bCanCollapse = false;
+	if (!bAllowSeamCollapse)
+	{
+		bCanCollapse = CanCollapseEdge(edgeID, a, b, c, d, t0, t1, collapse_to);
+	}
+	else
+	{
+		bCanCollapse = CanCollapseVertex(edgeID, a, b, collapse_to);
+	}
+
 	if (bCanCollapse == false)
 	{
 		return ESimplificationResult::Ignored_Constrained;
@@ -1280,10 +1268,66 @@ ESimplificationResult TMeshSimplification<QuadricErrorType>::CollapseEdge(int ed
 		if (Constraints)
 		{
 			Constraints->ClearEdgeConstraint(edgeID);
-			Constraints->ClearEdgeConstraint(collapseInfo.RemovedEdges.A);
+
+			auto ConstraintUpdator = [this](int cur_eid)->void
+			{
+				
+				// Seam edge can never flip, it is never fully unconstrained 
+				EEdgeRefineFlags SeamEdgeConstraint = EEdgeRefineFlags::NoFlip;
+				if (!bAllowSeamCollapse)
+				{
+					SeamEdgeConstraint = EEdgeRefineFlags((int)SeamEdgeConstraint | (int)EEdgeRefineFlags::NoCollapse);
+				}
+
+				FEdgeConstraint UpdatedEdgeConstraint;
+				FVertexConstraint UpdatedVertexConstraintA;
+				FVertexConstraint UpdatedVertexConstraintB;
+
+
+				bool bHaveUpdate = 
+				FMeshConstraintsUtil::ConstrainEdgeBoundariesAndSeams(cur_eid,
+					*Mesh,
+					MeshBoundaryConstraint,
+					GroupBoundaryConstraint,
+					MaterialBoundaryConstraint,
+					SeamEdgeConstraint,
+					!bAllowSeamCollapse,
+					UpdatedEdgeConstraint,
+					UpdatedVertexConstraintA,
+					UpdatedVertexConstraintB);
+
+				if (bHaveUpdate)
+				{
+					FIndex2i EdgeVerts = Mesh->GetEdgeV(cur_eid);
+
+					Constraints->SetOrUpdateEdgeConstraint(cur_eid, UpdatedEdgeConstraint);
+					UpdatedVertexConstraintA.CombineConstraint(Constraints->GetVertexConstraint(EdgeVerts.A));
+					Constraints->SetOrUpdateVertexConstraint(EdgeVerts.A, UpdatedVertexConstraintA);
+
+					UpdatedVertexConstraintB.CombineConstraint(Constraints->GetVertexConstraint(EdgeVerts.B));
+					Constraints->SetOrUpdateVertexConstraint(EdgeVerts.B, UpdatedVertexConstraintB);
+				}
+			};
+			
+			if (Constraints->HasEdgeConstraint(collapseInfo.RemovedEdges.A))
+			{
+
+				Constraints->ClearEdgeConstraint(collapseInfo.KeptEdges.A);
+				Constraints->ClearEdgeConstraint(collapseInfo.RemovedEdges.A);
+
+				ConstraintUpdator(collapseInfo.KeptEdges.A);
+			}
+			
 			if (collapseInfo.RemovedEdges.B != FDynamicMesh3::InvalidID)
 			{
-				Constraints->ClearEdgeConstraint(collapseInfo.RemovedEdges.B);
+				if (Constraints->HasEdgeConstraint(collapseInfo.RemovedEdges.B))
+				{
+
+					Constraints->ClearEdgeConstraint(collapseInfo.KeptEdges.B);
+					Constraints->ClearEdgeConstraint(collapseInfo.RemovedEdges.B);
+
+					ConstraintUpdator(collapseInfo.KeptEdges.B);
+				}
 			}
 			Constraints->ClearVertexConstraint(iCollapse);
 		}
@@ -1292,16 +1336,57 @@ ESimplificationResult TMeshSimplification<QuadricErrorType>::CollapseEdge(int ed
 
 		retVal = ESimplificationResult::Ok_Collapsed;
 	}
+	else if (result == EMeshResult::Failed_CollapseTriangle)
+	{
+		retVal = ESimplificationResult::Failed_IsolatedTriangle;
+	}
 
 	ProfileEndCollapse();
 	return retVal;
 }
 
 
+template <typename QuadricErrorType>
+bool TMeshSimplification<QuadricErrorType>::RemoveIsolatedTriangle(int tID)
+{
+	if (!Mesh->IsTriangle(tID)) return true;
 
+	FIndex3i tv = Mesh->GetTriangle(tID);
 
+	bool bIsIsolated = true;
+	for (int i = 0; i < 3; ++i)
+	{
+		for (int nbtr : Mesh->VtxTrianglesItr(tv[i]))
+		{
+			bIsIsolated = bIsIsolated && (nbtr == tID);
+		}
+	}
 
 
+	if (bIsIsolated)
+	{
+		const FIndex3i TriEdges = Mesh->GetTriEdges(tID);
+		if (Mesh->RemoveTriangle(tID) == EMeshResult::Ok)
+		{
+			if (Constraints)
+			{
+				Constraints->ClearEdgeConstraint(TriEdges.A);
+				Constraints->ClearEdgeConstraint(TriEdges.B);
+				Constraints->ClearEdgeConstraint(TriEdges.C);
+
+				Constraints->ClearVertexConstraint(tv.A);
+				Constraints->ClearVertexConstraint(tv.B);
+				Constraints->ClearVertexConstraint(tv.C);
+			}
+		}
+
+		OnRemoveIsolatedTriangle(tID);
+	}
+
+	return bIsIsolated;
+
+}
+
 
 
 // Project vertices onto projection target. 
diff --git a/Engine/Plugins/Experimental/GeometryProcessing/Source/DynamicMesh/Public/DynamicMeshAttributeSet.h b/Engine/Plugins/Experimental/GeometryProcessing/Source/DynamicMesh/Public/DynamicMeshAttributeSet.h
--- a/Engine/Plugins/Experimental/GeometryProcessing/Source/DynamicMesh/Public/DynamicMeshAttributeSet.h
+++ b/Engine/Plugins/Experimental/GeometryProcessing/Source/DynamicMesh/Public/DynamicMeshAttributeSet.h
@@ -116,6 +116,9 @@ public:
 	/** @return true if the given edge is a seam edge in any overlay */
 	virtual bool IsSeamEdge(int EdgeID) const;
 
+	/** @return true if the given edge is the termination of a seam in any overlay*/
+	virtual bool IsSeamEndEdge(int EdgeID) const;
+
 	/** @return true if the given edge is a seam edge in any overlay */
 	virtual bool IsSeamEdge(int EdgeID, bool& bIsUVSeamOut, bool& bIsNormalSeamOut) const;
 
diff --git a/Engine/Plugins/Experimental/GeometryProcessing/Source/DynamicMesh/Public/DynamicMeshOverlay.h b/Engine/Plugins/Experimental/GeometryProcessing/Source/DynamicMesh/Public/DynamicMeshOverlay.h
--- a/Engine/Plugins/Experimental/GeometryProcessing/Source/DynamicMesh/Public/DynamicMeshOverlay.h
+++ b/Engine/Plugins/Experimental/GeometryProcessing/Source/DynamicMesh/Public/DynamicMeshOverlay.h
@@ -387,6 +387,8 @@ public:
 
 	/** Returns true if the parent-mesh edge is a "Seam" in this overlay */
 	bool IsSeamEdge(int EdgeID) const;
+	/** Returns true if the parent-mesh edge is a "Seam End" in this overlay, meaning the adjacent element triangles share one element, not two */
+	bool IsSeamEndEdge(int EdgeID) const;
 	/** Returns true if the parent-mesh vertex is connected to any seam edges */
 	bool IsSeamVertex(int VertexID, bool bBoundaryIsSeam = true) const;
 
diff --git a/Engine/Plugins/Experimental/GeometryProcessing/Source/DynamicMesh/Public/MeshConstraints.h b/Engine/Plugins/Experimental/GeometryProcessing/Source/DynamicMesh/Public/MeshConstraints.h
--- a/Engine/Plugins/Experimental/GeometryProcessing/Source/DynamicMesh/Public/MeshConstraints.h
+++ b/Engine/Plugins/Experimental/GeometryProcessing/Source/DynamicMesh/Public/MeshConstraints.h
@@ -9,7 +9,7 @@
 /**
  * EEdgeRefineFlags indicate constraints on triangle mesh edges
  */
-enum class EEdgeRefineFlags
+enum class EEdgeRefineFlags 
 {
 	/** Edge is unconstrained */
 	NoConstraint = 0,
@@ -22,7 +22,11 @@ enum class EEdgeRefineFlags
 	/** Edge cannot be flipped, split, or collapsed */
 	FullyConstrained = NoFlip | NoSplit | NoCollapse,
 	/** Edge can only be split */
-	SplitsOnly = NoFlip | NoCollapse
+	SplitsOnly = NoFlip | NoCollapse,
+	/** Edge can only flip*/
+	FlipOnly = NoSplit | NoCollapse,
+	/** Edge can only collapse*/
+	CollapseOnly = NoFlip | NoSplit
 };
 
 
@@ -37,7 +41,7 @@ public:
 	/** Edge is associated with this projection target. */
 	IProjectionTarget * Target;
 
-	/** This ID is not a constraint, but can be used to find descendents of a constrained input edge after splits */
+	/** This ID is not a constraint, but can be used to find descendants of a constrained input edge after splits */
 	int TrackingSetID;
 
 	FEdgeConstraint()
@@ -70,19 +74,19 @@ public:
 	/** @return true if edge can be split */
 	bool CanSplit() const
 	{
-		return ((int)RefineFlags & (int)EEdgeRefineFlags::NoSplit) == 0;
+		return CanSplit(RefineFlags);
 	}
 
 	/** @return true if edge can be collapsed */
 	bool CanCollapse() const
 	{
-		return ((int)RefineFlags & (int)EEdgeRefineFlags::NoCollapse) == 0;
+		return CanCollapse(RefineFlags);
 	}
 
 	/** @return true if edge cannot be modified at all */
 	bool NoModifications() const
 	{
-		return ((int)RefineFlags & (int)EEdgeRefineFlags::FullyConstrained) == (int)EEdgeRefineFlags::FullyConstrained;
+		return NoModifications(RefineFlags);
 	}
 
 	/** @return true if edge is unconstrained */
@@ -102,6 +106,30 @@ public:
 
 	/** @return a fully constrained edge constraint */
 	static FEdgeConstraint FullyConstrained() { return FEdgeConstraint(EEdgeRefineFlags::FullyConstrained); }
+
+	/** @return true if flags permit flip */
+	static inline bool CanFlip(const EEdgeRefineFlags Flags)
+	{
+		return ((int)Flags & (int)EEdgeRefineFlags::NoFlip) == 0;
+	}
+
+	/** @return true if flags permit split */
+	static inline bool CanSplit(const EEdgeRefineFlags Flags)
+	{
+		return ((int)Flags & (int)EEdgeRefineFlags::NoSplit) == 0;
+	}
+
+	/** @return true if flags permit collapsed */
+	static inline bool CanCollapse(const EEdgeRefineFlags Flags)
+	{
+		return ((int)Flags & (int)EEdgeRefineFlags::NoCollapse) == 0;
+	}
+
+	/** @return true if flags permit no edge modification */
+	static inline bool NoModifications(const EEdgeRefineFlags Flags)
+	{
+		return ((int)Flags & (int)EEdgeRefineFlags::FullyConstrained) == (int)EEdgeRefineFlags::FullyConstrained;
+	}
 };
 
 
diff --git a/Engine/Plugins/Experimental/GeometryProcessing/Source/DynamicMesh/Public/MeshConstraintsUtil.h b/Engine/Plugins/Experimental/GeometryProcessing/Source/DynamicMesh/Public/MeshConstraintsUtil.h
--- a/Engine/Plugins/Experimental/GeometryProcessing/Source/DynamicMesh/Public/MeshConstraintsUtil.h
+++ b/Engine/Plugins/Experimental/GeometryProcessing/Source/DynamicMesh/Public/MeshConstraintsUtil.h
@@ -142,6 +142,30 @@ s	 * EdgeROI that are seams, as well as vertices incident to edges in EdgeROI th
 		}
 	}
 
+	/**
+	 * Constrain all attribute seams for all overlays of a mesh and different types of boundaries on the mesh
+	 * @param EdgeID the edge in Mesh to constrain
+	 * @param Mesh the mesh to constrain
+	 * @param MeshBoundaryConstraint the constraint policy for Mesh boundaries
+	 * @param GroupBoundaryConstraint the constraint policy for boundaries between polygon groups
+	 * @param MaterialBoundaryConstraint the constraint policy for boundaries between different materials
+	 * @param SeamEdgeConstraint  the constraint policy for attribute (e.g. uv, normals ) seams
+	 * @param bAllowSmoothing should we allow constrained vertices to be smoothed
+	 * @param EdgeConstraint [out] edge constraint to be updated
+	 * @param VertexConstraintA [out] vertex constraint to be updated on return
+	 * @param VertexConstraintB [out] vertex constraint to be updated on return
+	 * @return false if no constraint is needed for this edge
+	 */
+	static bool ConstrainEdgeBoundariesAndSeams(const int EdgeID,
+											    const FDynamicMesh3& Mesh,
+			                                    const EEdgeRefineFlags MeshBoundaryConstraint,
+			                                    const EEdgeRefineFlags GroupBoundaryConstraint,
+			                                    const EEdgeRefineFlags MaterialBoundaryConstraint,
+			                                    const EEdgeRefineFlags SeamEdgeConstraint,
+			                                    const bool bAllowSeamSmoothing,
+			                                    FEdgeConstraint& EdgeConstraint,
+			                                    FVertexConstraint& VertexConstraintA,
+			                                    FVertexConstraint& VertexConstraintB);
 
 
 private:
diff --git a/Engine/Plugins/Experimental/GeometryProcessing/Source/DynamicMesh/Public/MeshSimplification.h b/Engine/Plugins/Experimental/GeometryProcessing/Source/DynamicMesh/Public/MeshSimplification.h
--- a/Engine/Plugins/Experimental/GeometryProcessing/Source/DynamicMesh/Public/MeshSimplification.h
+++ b/Engine/Plugins/Experimental/GeometryProcessing/Source/DynamicMesh/Public/MeshSimplification.h
@@ -20,7 +20,8 @@ enum class ESimplificationResult
 	Ignored_Constrained = 4,
 	Ignored_CreatesFlip = 5,
 	Failed_OpNotSuccessful = 6,
-	Failed_NotAnEdge = 7
+	Failed_NotAnEdge = 7,
+	Failed_IsolatedTriangle = 8
 };
 
 /**
@@ -58,8 +59,8 @@ public:
 	 */
 	ESimplificationCollapseModes CollapseMode = ESimplificationCollapseModes::MinimalQuadricPositionError;
 
-
-
+	/** if false, face and vertex quadrics are recomputed in the neighborhood of each collapse, definitely slower but maybe higher quality*/
+	bool bRetainQuadricMemory = false;
 
 	/** if true, we try to keep boundary vertices on boundary. You probably want this. */
 	bool bPreserveBoundaryShape = true;
@@ -67,6 +68,9 @@ public:
 	/** if true, we allow UV and Normal seams to collapse during simplification.*/
 	bool bAllowSeamCollapse = true;
 
+	/** When using the constraint system, these options will apply to the appropriate boundaries. */
+	EEdgeRefineFlags MeshBoundaryConstraint, GroupBoundaryConstraint, MaterialBoundaryConstraint;
+
 	TMeshSimplification(FDynamicMesh3* m) : FMeshRefinerBase(m)
 	{
 		NormalOverlay = nullptr;
@@ -123,7 +127,7 @@ public:
 	 */
 	virtual void FastCollapsePass(double MinEdgeLength, int Rounds = 1, bool bMeshIsClosedHint = false);
 
-
+	
 
 protected:
 
@@ -175,7 +179,7 @@ protected:
 
 
 	
-	double SeamEdgeWeight = 256.;
+	double SeamEdgeWeight =  256.;
 
 	TArray<FQuadricErrorType> vertQuadrics;
 	virtual void InitializeVertexQuadrics();
@@ -191,7 +195,7 @@ protected:
 
 	FQuadricErrorType ComputeFaceQuadric(const int tid, FVector3d& nface, FVector3d& c, double& Area) const;
 	
-	// uses pre-computed vertex and face quadrics to construct the edge quadric.
+	// uses pre-computed vertex, face and seam quadrics to construct the edge quadric.
 	FQuadricErrorType AssembleEdgeQuadric(const FDynamicMesh3::FEdge& edge) const;
 
 	
@@ -277,6 +281,11 @@ protected:
 	 */
 	ESimplificationResult CollapseEdge(int edgeID, FVector3d vNewPos, FDynamicMesh3::FEdgeCollapseInfo& collapseInfo, int32 RequireKeepVert = -1);
 
+	/**
+	* Remove an isolated triangle.
+	* @return false if the triangle shares a vertex with another triangle
+	*/
+	bool RemoveIsolatedTriangle(int tID);
 
 
 	// subclasses can override these to implement custom behavior...
@@ -285,8 +294,10 @@ protected:
 		// this is for subclasses...
 	}
 
-
-
+	virtual void OnRemoveIsolatedTriangle(int tId)
+	{
+		// this is for subclasses
+	}
 
 
 	// Project vertices onto projection target. 
diff --git a/Engine/Plugins/Experimental/MeshModelingToolset/Source/ModelingOperatorsEditorOnly/Private/CleaningOps/SimplifyMeshOp.cpp b/Engine/Plugins/Experimental/MeshModelingToolset/Source/ModelingOperatorsEditorOnly/Private/CleaningOps/SimplifyMeshOp.cpp
--- a/Engine/Plugins/Experimental/MeshModelingToolset/Source/ModelingOperatorsEditorOnly/Private/CleaningOps/SimplifyMeshOp.cpp
+++ b/Engine/Plugins/Experimental/MeshModelingToolset/Source/ModelingOperatorsEditorOnly/Private/CleaningOps/SimplifyMeshOp.cpp
@@ -30,7 +30,8 @@ void ComputeSimplify(FDynamicMesh3* TargetMesh, const bool bReproject,
 					 const ESimplifyTargetType TargetMode,
 					 const float TargetPercentage, const int TargetCount, const float TargetEdgeLength,
 					 const float AngleThreshold,
-	                 typename SimplificationType::ESimplificationCollapseModes CollapseMode)
+	                 typename SimplificationType::ESimplificationCollapseModes CollapseMode,
+					 bool bUseQuadricMemory)
 {
 	SimplificationType Reducer(TargetMesh);
 
@@ -38,8 +39,10 @@ void ComputeSimplify(FDynamicMesh3* TargetMesh, const bool bReproject,
 		SimplificationType::ETargetProjectionMode::AfterRefinement : SimplificationType::ETargetProjectionMode::NoProjection;
 
 	Reducer.DEBUG_CHECK_LEVEL = 0;
+	//Reducer.ENABLE_PROFILING = true;
 
 	Reducer.bAllowSeamCollapse = bAllowSeamCollapse;
+	Reducer.bRetainQuadricMemory = bUseQuadricMemory;
 
 	if (bAllowSeamCollapse)
 	{
@@ -59,6 +62,12 @@ void ComputeSimplify(FDynamicMesh3* TargetMesh, const bool bReproject,
 														 MaterialBoundaryConstraint,
 														 true, !bPreserveSharpEdges, bAllowSeamCollapse);
 	Reducer.SetExternalConstraints(MoveTemp(constraints));
+	
+	// transfer constraint setting to the simplifier, these are used to update the constraints as edges collapse.	
+	Reducer.MeshBoundaryConstraint = MeshBoundaryConstraint;
+	Reducer.GroupBoundaryConstraint = GroupBoundaryConstraint;
+	Reducer.MaterialBoundaryConstraint = MaterialBoundaryConstraint;
+	
 
 	if (bReproject)
 	{
@@ -117,6 +126,7 @@ void FSimplifyMeshOp::CalculateResult(FProgressCancel* Progress)
 	int OriginalTriCount = OriginalMesh->TriangleCount();
 	if (SimplifierType == ESimplifyType::QEM)
 	{
+		bool bUseQuadricMemory = true;
 		ResultMesh->Copy(*OriginalMesh, true, true, true, !bDiscardAttributes);
 		ComputeSimplify<FQEMSimplification>(TargetMesh, bReproject, OriginalTriCount, *OriginalMesh, *OriginalMeshSpatial,
 											MeshBoundaryConstraint,
@@ -124,10 +134,11 @@ void FSimplifyMeshOp::CalculateResult(FProgressCancel* Progress)
 											MaterialBoundaryConstraint,
 											bPreserveSharpEdges, bAllowSeamCollapse,
 											TargetMode, TargetPercentage, TargetCount, TargetEdgeLength, MinimalPlanarAngleThresh,
-											FQEMSimplification::ESimplificationCollapseModes::MinimalQuadricPositionError);
+											FQEMSimplification::ESimplificationCollapseModes::MinimalQuadricPositionError, bUseQuadricMemory);
 	}
 	else if (SimplifierType == ESimplifyType::Attribute)
 	{
+		bool bUseQuadricMemory = false;
 		ResultMesh->Copy(*OriginalMesh, true, true, true, !bDiscardAttributes);
 		ComputeSimplify<FAttrMeshSimplification>(TargetMesh, bReproject, OriginalTriCount, *OriginalMesh, *OriginalMeshSpatial,
 													MeshBoundaryConstraint,
@@ -135,10 +146,11 @@ void FSimplifyMeshOp::CalculateResult(FProgressCancel* Progress)
 													MaterialBoundaryConstraint,
 													bPreserveSharpEdges, bAllowSeamCollapse,
 													TargetMode, TargetPercentage, TargetCount, TargetEdgeLength, MinimalPlanarAngleThresh,
-													FAttrMeshSimplification::ESimplificationCollapseModes::MinimalQuadricPositionError);
+													FAttrMeshSimplification::ESimplificationCollapseModes::MinimalQuadricPositionError, bUseQuadricMemory);
 	}
 	else if (SimplifierType == ESimplifyType::MinimalPlanar)
 	{
+		bool bUseQuadricMemory = false;
 		ResultMesh->Copy(*OriginalMesh, true, true, true, !bDiscardAttributes);
 		ComputeSimplify<FAttrMeshSimplification>(TargetMesh, bReproject, OriginalTriCount, *OriginalMesh, *OriginalMeshSpatial,
 			MeshBoundaryConstraint,
@@ -146,10 +158,11 @@ void FSimplifyMeshOp::CalculateResult(FProgressCancel* Progress)
 			MaterialBoundaryConstraint,
 			bPreserveSharpEdges, bAllowSeamCollapse,
 			ESimplifyTargetType::MinimalPlanar, TargetPercentage, TargetCount, TargetEdgeLength, MinimalPlanarAngleThresh,
-			FAttrMeshSimplification::ESimplificationCollapseModes::MinimalQuadricPositionError);
+			FAttrMeshSimplification::ESimplificationCollapseModes::MinimalQuadricPositionError, bUseQuadricMemory);
 	}
 	else if (SimplifierType == ESimplifyType::MinimalExistingVertex)
 	{
+		bool bUseQuadricMemory = true;
 		ResultMesh->Copy(*OriginalMesh, true, true, true, !bDiscardAttributes);
 		ComputeSimplify<FQEMSimplification>(TargetMesh, bReproject, OriginalTriCount, *OriginalMesh, *OriginalMeshSpatial,
 			MeshBoundaryConstraint,
@@ -157,9 +170,9 @@ void FSimplifyMeshOp::CalculateResult(FProgressCancel* Progress)
 			MaterialBoundaryConstraint,
 			bPreserveSharpEdges, bAllowSeamCollapse,
 			TargetMode, TargetPercentage, TargetCount, TargetEdgeLength, MinimalPlanarAngleThresh, 
-			FQEMSimplification::ESimplificationCollapseModes::MinimalExistingVertexError);
+			FQEMSimplification::ESimplificationCollapseModes::MinimalExistingVertexError, bUseQuadricMemory);
 	}
-	else // SimplifierType == ESimplifyType::UE4Standard
+	else // SimplifierType == ESimplifyType::UEStandard
 	{
 		const FMeshDescription* SrcMeshDescription = OriginalMeshDescription.Get();
 		FMeshDescription DstMeshDescription(*SrcMeshDescription);
