Add mesh quality criteria to boolean mesh simplification algorithm; don't collapse edges when it would make bad triangles even worse.

+ fix a bunch of smaller bugs w/ mesh boolean

#rnx

diff --git a/Engine/Plugins/Experimental/GeometryProcessing/Source/DynamicMesh/Private/Operations/MeshBoolean.cpp b/Engine/Plugins/Experimental/GeometryProcessing/Source/DynamicMesh/Private/Operations/MeshBoolean.cpp
--- a/Engine/Plugins/Experimental/GeometryProcessing/Source/DynamicMesh/Private/Operations/MeshBoolean.cpp
+++ b/Engine/Plugins/Experimental/GeometryProcessing/Source/DynamicMesh/Private/Operations/MeshBoolean.cpp
@@ -494,25 +494,20 @@ bool FMeshBoolean::Compute()
 	TMap<int, int> AllVIDMatches; // mapping of matched vertex IDs from cutmesh 0 to cutmesh 1
 	if (NumMeshesToProcess == 2)
 	{
+		TMap<int, int> FoundMatchesMaps[2]; // mappings of matched vertex IDs from mesh 1->0 and mesh 0->1
 		double SnapToleranceSq = SnapTolerance * SnapTolerance;
 
-		// Hash boundary verts for faster search
-		TArray<TPointHashGrid3d<int>> PointHashes;
-		for (int MeshIdx = 0; MeshIdx < 2; MeshIdx++)
-		{
-			PointHashes.Emplace(CutMesh[MeshIdx]->GetCachedBounds().MaxDim() / 64, -1);
-			for (int BoundaryVID : PossUnmatchedBdryVerts[MeshIdx])
-			{
-				PointHashes[MeshIdx].InsertPointUnsafe(BoundaryVID, CutMesh[MeshIdx]->GetVertex(BoundaryVID));
-			}
-		}
-
 		// ensure segments that are now on boundaries have 1:1 vertex correspondence across meshes
 		for (int MeshIdx = 0; MeshIdx < 2; MeshIdx++)
 		{
 			int OtherMeshIdx = 1 - MeshIdx;
 			FDynamicMesh3& OtherMesh = *CutMesh[OtherMeshIdx];
 
+			TPointHashGrid3d<int> OtherMeshPointHash(OtherMesh.GetCachedBounds().MaxDim() / 64, -1);
+			for (int BoundaryVID : PossUnmatchedBdryVerts[OtherMeshIdx])
+			{
+				OtherMeshPointHash.InsertPointUnsafe(BoundaryVID, OtherMesh.GetVertex(BoundaryVID));
+			}
 
 			FSparseDynamicOctree3 EdgeOctree;
 			EdgeOctree.RootDimension = .25;
@@ -553,12 +548,17 @@ bool FMeshBoolean::Compute()
 
 			// mapping from OtherMesh VIDs to ProcessMesh VIDs
 			// used to ensure we only keep the best match, in cases where multiple boundary vertices map to a given vertex on the other mesh boundary
-			TMap<int, int> FoundMatches;
+			TMap<int, int>& FoundMatches = FoundMatchesMaps[MeshIdx];
 
 			for (int BoundaryVID : PossUnmatchedBdryVerts[MeshIdx])
 			{
+				if (MeshIdx == 1 && FoundMatchesMaps[0].Contains(BoundaryVID))
+				{
+					continue; // was already snapped to a vertex
+				}
+
 				FVector3d Pos = CutMesh[MeshIdx]->GetVertex(BoundaryVID);
-				TPair<int, double> VIDDist = PointHashes[OtherMeshIdx].FindNearestInRadius(Pos, SnapTolerance, [&Pos, &OtherMesh](int VID)
+				TPair<int, double> VIDDist = OtherMeshPointHash.FindNearestInRadius(Pos, SnapTolerance, [&Pos, &OtherMesh](int VID)
 					{
 						return Pos.DistanceSquared(OtherMesh.GetVertex(VID));
 					});
@@ -726,12 +726,16 @@ bool FMeshBoolean::IsFlat(const FDynamicMesh3& Mesh, int VID, double DotToleranc
  */
 bool FMeshBoolean::CollapseWouldHurtTriangleQuality(
 	const FDynamicMesh3& Mesh, const FVector3d& ExpectNormal,
-	int32 RemoveV, const FVector3d& RemoveVPos, int32 KeepV, const FVector3d& KeepVPos
+	int32 RemoveV, const FVector3d& RemoveVPos, int32 KeepV, const FVector3d& KeepVPos,
+	double TryToImproveTriQualityThreshold
 	)
 {
+	double WorstQualityNewTriangle = FMathd::MaxReal;
+
 	bool bIsHurt = false;
 	Mesh.EnumerateVertexTriangles(RemoveV,
-		[&Mesh, &bIsHurt, &KeepVPos, RemoveV, KeepV, &ExpectNormal](int32 TID)
+		[&Mesh, &bIsHurt, &KeepVPos, RemoveV, KeepV, &ExpectNormal, 
+		TryToImproveTriQualityThreshold, &WorstQualityNewTriangle](int32 TID)
 		{
 			if (bIsHurt)
 			{
@@ -765,10 +769,44 @@ bool FMeshBoolean::CollapseWouldHurtTriangleQuality(
 			// TODO: does this tolerance make a difference?  if not, set to zero and remove the VCross.Normalize()
 			double EdgeFlipTolerance = 1.e-5;
 			double Area2 = VCross.Normalize();
+			if (TryToImproveTriQualityThreshold > 0)
+			{
+				FVector3d Edge3(Verts[2] - Verts[1]);
+				double MaxLenSq = FMathd::Max3(Edge1.SquaredLength(), Edge2.SquaredLength(), Edge3.SquaredLength());
+				double Quality = Area2 / (MaxLenSq + FMathd::ZeroTolerance);
+				WorstQualityNewTriangle = FMathd::Min(Quality, WorstQualityNewTriangle);
+			}
 
 			bIsHurt = VCross.Dot(ExpectNormal) <= EdgeFlipTolerance;
 		}
 	);
+
+	// note tri quality was computed as 2*Area / MaxEdgeLenSquared
+	//  so need to multiply it by 2/sqrt(3) to get actual aspect ratio
+	if (!bIsHurt && WorstQualityNewTriangle * 2 * FMathd::InvSqrt3 < TryToImproveTriQualityThreshold)
+	{
+		// we found a bad tri; tentatively switch to rejecting this edge collapse
+		bIsHurt = true;
+		// but if there was an even worse tri in the original neighborhood, accept the collapse after all
+		Mesh.EnumerateVertexTriangles(RemoveV, [&Mesh, &bIsHurt, WorstQualityNewTriangle](int TID)
+			{
+				if (!bIsHurt) // early out if we already found an originally-worse tri
+				{
+					return;
+				}
+				FVector3d A, B, C;
+				Mesh.GetTriVertices(TID, A, B, C);
+				FVector3d E1 = B - A, E2 = C - A, E3 = C - B;
+				double Area2 = E1.Cross(E2).Length();
+				double MaxLenSq = FMathd::Max3(E1.SquaredLength(), E2.SquaredLength(), E3.SquaredLength());
+				double Quality = Area2 / (MaxLenSq + FMathd::ZeroTolerance);
+				if (Quality < WorstQualityNewTriangle)
+				{
+					bIsHurt = false;
+				}
+			}
+		);
+	}
 	return bIsHurt;
 }
 
@@ -777,7 +815,7 @@ bool FMeshBoolean::CollapseWouldHurtTriangleQuality(
  * Test if a given edge collapse would change the mesh shape or UVs unacceptably
  */
 bool FMeshBoolean::CollapseWouldChangeShapeOrUVs(
-	const FDynamicMesh3& Mesh, const TSet<int> CutBoundaryEdgeSet, double DotTolerance, int SourceEID,
+	const FDynamicMesh3& Mesh, const TSet<int>& CutBoundaryEdgeSet, double DotTolerance, int SourceEID,
 	int32 RemoveV, const FVector3d& RemoveVPos, int32 KeepV, const FVector3d& KeepVPos, const FVector3d& EdgeDir,
 	bool bPreserveUVsForMesh, bool bPreserveVertexUVs, bool bPreserveOverlayUVs, float UVEqualThresholdSq)
 {
@@ -1030,7 +1068,8 @@ void FMeshBoolean::SimplifyAlongNewEdges(int NumMeshesToProcess, FDynamicMesh3*
 					int KeepV = MeshIdx == 0 ? Edge.Vert[KeepVIdx] : Matches[KeepVIdx];
 					int SourceEID = MeshIdx == 0 ? EID : OtherEID;
 
-					bHasBadEdge = bHasBadEdge || CollapseWouldHurtTriangleQuality(*CutMesh[MeshIdx], FlatNormals[RemoveVIdx][MeshIdx], RemoveV, RemoveVPos, KeepV, KeepVPos);
+					bHasBadEdge = bHasBadEdge || CollapseWouldHurtTriangleQuality(*CutMesh[MeshIdx],
+						FlatNormals[RemoveVIdx][MeshIdx], RemoveV, RemoveVPos, KeepV, KeepVPos, TryToImproveTriQualityThreshold);
 					
 					bHasBadEdge = bHasBadEdge || CollapseWouldChangeShapeOrUVs(
 						*CutMesh[MeshIdx], CutBoundaryEdgeSets[MeshIdx], DotTolerance,
@@ -1048,22 +1087,21 @@ void FMeshBoolean::SimplifyAlongNewEdges(int NumMeshesToProcess, FDynamicMesh3*
 				int RemoveV = Edge.Vert[RemoveVIdx];
 				int KeepV = Edge.Vert[KeepVIdx];
 				EMeshResult CollapseResult = CutMesh[0]->CollapseEdge(KeepV, RemoveV, 0, CollapseInfo);
-				if (ensure(CollapseResult == EMeshResult::Ok))
+				if (CollapseResult == EMeshResult::Ok)
 				{
-					bool bSuccessfulCollapse = true;
 					if (bHasMatches)
 					{
 						int OtherRemoveV = Matches[RemoveVIdx];
 						int OtherKeepV = Matches[KeepVIdx];
 						FDynamicMesh3::FEdgeCollapseInfo OtherCollapseInfo;
 						EMeshResult OtherCollapseResult = CutMesh[1]->CollapseEdge(OtherKeepV, OtherRemoveV, 0, OtherCollapseInfo);
-						if (!ensure(OtherCollapseResult == EMeshResult::Ok))
+						if (OtherCollapseResult != EMeshResult::Ok)
 						{
 							// if we get here, we've somehow managed to collapse the first edge but failed on the second (matched) edge
 							// which will leave a crack in the result unless we can somehow undo the first collapse, which would require a bunch of extra work
 							// but the only case where I could see this happen is if the second edge is on an isolated triangle, which means there is a hole anyway
 							// or if the mesh topology is somehow invalid
-							bSuccessfulCollapse = false;
+							ensureMsgf(OtherCollapseResult == EMeshResult::Failed_CollapseTriangle, TEXT("Collapse failed with result: %d"), (int)OtherCollapseResult);
 						}
 						else
 						{
@@ -1076,15 +1114,13 @@ void FMeshBoolean::SimplifyAlongNewEdges(int NumMeshesToProcess, FDynamicMesh3*
 							}
 						}
 					}
-					if (bSuccessfulCollapse)
+
+					NumCollapses++;
+					CutBoundaryEdgeSets[0].Remove(CollapseInfo.CollapsedEdge);
+					CutBoundaryEdgeSets[0].Remove(CollapseInfo.RemovedEdges[0]);
+					if (CollapseInfo.RemovedEdges[1] != -1)
 					{
-						NumCollapses++;
-						CutBoundaryEdgeSets[0].Remove(CollapseInfo.CollapsedEdge);
-						CutBoundaryEdgeSets[0].Remove(CollapseInfo.RemovedEdges[0]);
-						if (CollapseInfo.RemovedEdges[1] != -1)
-						{
-							CutBoundaryEdgeSets[0].Remove(CollapseInfo.RemovedEdges[1]);
-						}
+						CutBoundaryEdgeSets[0].Remove(CollapseInfo.RemovedEdges[1]);
 					}
 				}
 				break; // if we got through to trying to collapse the edge, don't try to collapse from the other vertex.
diff --git a/Engine/Plugins/Experimental/GeometryProcessing/Source/DynamicMesh/Private/Operations/MeshSelfUnion.cpp b/Engine/Plugins/Experimental/GeometryProcessing/Source/DynamicMesh/Private/Operations/MeshSelfUnion.cpp
--- a/Engine/Plugins/Experimental/GeometryProcessing/Source/DynamicMesh/Private/Operations/MeshSelfUnion.cpp
+++ b/Engine/Plugins/Experimental/GeometryProcessing/Source/DynamicMesh/Private/Operations/MeshSelfUnion.cpp
@@ -570,7 +570,7 @@ void FMeshSelfUnion::SimplifyAlongNewEdges(TArray<int>& CutBoundaryEdges, TMap<i
 					int SourceEID = WhichEdge == 0 ? EID : MatchEID;
 
 					bHasBadEdge = bHasBadEdge || FMeshBoolean::CollapseWouldHurtTriangleQuality(
-						*Mesh, FlatNormals[RemoveVIdx][WhichEdge], RemoveV, RemoveVPos, KeepV, KeepVPos);
+						*Mesh, FlatNormals[RemoveVIdx][WhichEdge], RemoveV, RemoveVPos, KeepV, KeepVPos, TryToImproveTriQualityThreshold);
 
 					bHasBadEdge = bHasBadEdge || FMeshBoolean::CollapseWouldChangeShapeOrUVs(
 						*Mesh, CutBoundaryEdgeSet, DotTolerance,
@@ -587,34 +587,23 @@ void FMeshSelfUnion::SimplifyAlongNewEdges(TArray<int>& CutBoundaryEdges, TMap<i
 				int RemoveV = Edge.Vert[RemoveVIdx];
 				int KeepV = Edge.Vert[KeepVIdx];
 				EMeshResult CollapseResult = Mesh->CollapseEdge(KeepV, RemoveV, 0, CollapseInfo);
-				if (ensure(CollapseResult == EMeshResult::Ok))
+				if (CollapseResult == EMeshResult::Ok)
 				{
 					int OtherRemoveV = Matches[RemoveVIdx];
 					int OtherKeepV = Matches[KeepVIdx];
 					FDynamicMesh3::FEdgeCollapseInfo OtherCollapseInfo;
 					EMeshResult OtherCollapseResult = Mesh->CollapseEdge(OtherKeepV, OtherRemoveV, 0, OtherCollapseInfo);
-					if (!ensure(OtherCollapseResult == EMeshResult::Ok))
+					if (OtherCollapseResult != EMeshResult::Ok)
 					{
 						// if we get here, we've somehow managed to collapse the first edge but failed on the second (matched) edge
 						// which will leave a crack in the result unless we can somehow undo the first collapse, which would require a bunch of extra work
 						// but the only case where I could see this happen is if the second edge is on an isolated triangle, which means there is a hole anyway
 						// or if the mesh topology is somehow invalid
-						break;
+						ensure(OtherCollapseResult == EMeshResult::Failed_CollapseTriangle);
 					}
 					else
 					{
-						NumCollapses++;
-
-						FoundMatches.Remove(RemoveV);
 						FoundMatches.Remove(OtherRemoveV);
-
-						CutBoundaryEdgeSet.Remove(CollapseInfo.CollapsedEdge);
-						CutBoundaryEdgeSet.Remove(CollapseInfo.RemovedEdges[0]);
-						if (CollapseInfo.RemovedEdges[1] != -1)
-						{
-							CutBoundaryEdgeSet.Remove(CollapseInfo.RemovedEdges[1]);
-						}
-
 						CutBoundaryEdgeSet.Remove(OtherCollapseInfo.CollapsedEdge);
 						CutBoundaryEdgeSet.Remove(OtherCollapseInfo.RemovedEdges[0]);
 						if (OtherCollapseInfo.RemovedEdges[1] != -1)
@@ -622,6 +611,15 @@ void FMeshSelfUnion::SimplifyAlongNewEdges(TArray<int>& CutBoundaryEdges, TMap<i
 							CutBoundaryEdgeSet.Remove(OtherCollapseInfo.RemovedEdges[1]);
 						}
 					}
+
+					NumCollapses++;
+					FoundMatches.Remove(RemoveV);
+					CutBoundaryEdgeSet.Remove(CollapseInfo.CollapsedEdge);
+					CutBoundaryEdgeSet.Remove(CollapseInfo.RemovedEdges[0]);
+					if (CollapseInfo.RemovedEdges[1] != -1)
+					{
+						CutBoundaryEdgeSet.Remove(CollapseInfo.RemovedEdges[1]);
+					}
 				}
 				break; // if we got through to trying to collapse the edge, don't try to collapse from the other vertex.
 			}
diff --git a/Engine/Plugins/Experimental/GeometryProcessing/Source/DynamicMesh/Public/Operations/MeshBoolean.h b/Engine/Plugins/Experimental/GeometryProcessing/Source/DynamicMesh/Public/Operations/MeshBoolean.h
--- a/Engine/Plugins/Experimental/GeometryProcessing/Source/DynamicMesh/Public/Operations/MeshBoolean.h
+++ b/Engine/Plugins/Experimental/GeometryProcessing/Source/DynamicMesh/Public/Operations/MeshBoolean.h
@@ -77,12 +77,18 @@ public:
 	//
 	/** Degrees of deviation from coplanar that we will still simplify */
 	double SimplificationAngleTolerance = .1;
+	/**
+	 * If triangle quality (aspect ratio) is worse than this threshold, only simplify in ways that improve quality.  If <= 0, triangle quality is ignored.
+	 *  Note: For aspect ratio we use definition: 4*TriArea / (sqrt(3)*MaxEdgeLen^2), ref: https://people.eecs.berkeley.edu/~jrs/papers/elemj.pdf p.53
+	 *  Equilateral triangles have value 1; Smaller values -> lower quality
+	 */
+	double TryToImproveTriQualityThreshold = .5;
 	/** Prevent simplification from distorting vertex UVs */
 	bool bPreserveVertexUVs = true;
 	/** Prevent simplification from distorting overlay UVs */
 	bool bPreserveOverlayUVs = true;
 	/** When preserving UVs, sets maximum allowed change in UV coordinates from collapsing an edge, measured at the removed vertex */
-	float UVDistortTolerance = .0001;
+	float UVDistortTolerance = FMathf::ZeroTolerance;
 	/** If > -1, then only preserve the UVs of one of the input meshes.  Useful when cutting an artist-created asset w/ procedural geometry. */
 	int PreserveUVsOnlyForMesh = -1;
 
@@ -187,14 +193,15 @@ public:
 	 */
 	static bool CollapseWouldHurtTriangleQuality(
 		const FDynamicMesh3& Mesh, const FVector3d& ExpectNormal,
-		int32 RemoveV, const FVector3d& RemoveVPos, int32 KeepV, const FVector3d& KeepVPos
+		int32 RemoveV, const FVector3d& RemoveVPos, int32 KeepV, const FVector3d& KeepVPos,
+		double TryToImproveTriQualityThreshold
 	);
 
 	/**
 	 * Test if a given edge collapse would change the mesh shape or UVs unacceptably
 	 */
 	static bool CollapseWouldChangeShapeOrUVs(
-		const FDynamicMesh3& Mesh, const TSet<int> CutBoundaryEdgeSet, double DotTolerance, int SourceEID,
+		const FDynamicMesh3& Mesh, const TSet<int>& CutBoundaryEdgeSet, double DotTolerance, int SourceEID,
 		int32 RemoveV, const FVector3d& RemoveVPos, int32 KeepV, const FVector3d& KeepVPos,
 		const FVector3d& EdgeDir, bool bPreserveUVsForMesh, bool bPreserveVertexUVs, bool bPreserveOverlayUVs, float UVToleranceSq);
 };
diff --git a/Engine/Plugins/Experimental/GeometryProcessing/Source/DynamicMesh/Public/Operations/MeshSelfUnion.h b/Engine/Plugins/Experimental/GeometryProcessing/Source/DynamicMesh/Public/Operations/MeshSelfUnion.h
--- a/Engine/Plugins/Experimental/GeometryProcessing/Source/DynamicMesh/Public/Operations/MeshSelfUnion.h
+++ b/Engine/Plugins/Experimental/GeometryProcessing/Source/DynamicMesh/Public/Operations/MeshSelfUnion.h
@@ -63,12 +63,18 @@ public:
 	//
 	/** Degrees of deviation from coplanar that we will still simplify */
 	double SimplificationAngleTolerance = .1;
+	/**
+	 * If triangle quality (aspect ratio) is worse than this threshold, only simplify in ways that improve quality.  If <= 0, triangle quality is ignored.
+	 *  Note: For aspect ratio we use definition: 4*TriArea / (sqrt(3)*MaxEdgeLen^2), ref: https://people.eecs.berkeley.edu/~jrs/papers/elemj.pdf p.53
+	 *  Equilateral triangles have value 1; Smaller values -> lower quality
+	 */
+	double TryToImproveTriQualityThreshold = .25;
 	/** Prevent simplification from distorting vertex UVs */
 	bool bPreserveVertexUVs = true;
 	/** Prevent simplification from distorting overlay UVs */
 	bool bPreserveOverlayUVs = true;
 	/** When preserving UVs, sets maximum allowed change in UV coordinates from collapsing an edge, measured at the removed vertex */
-	float UVDistortTolerance = .0001;
+	float UVDistortTolerance = FMathf::ZeroTolerance;
 
 
 	//
diff --git a/Engine/Plugins/Experimental/PlanarCutPlugin/Source/PlanarCut/Private/PlanarCut.cpp b/Engine/Plugins/Experimental/PlanarCutPlugin/Source/PlanarCut/Private/PlanarCut.cpp
--- a/Engine/Plugins/Experimental/PlanarCutPlugin/Source/PlanarCut/Private/PlanarCut.cpp
+++ b/Engine/Plugins/Experimental/PlanarCutPlugin/Source/PlanarCut/Private/PlanarCut.cpp
@@ -1416,6 +1416,7 @@ struct FDynamicMeshCollection
 
 			EnterProgressFrame(Planes.Num() * .2);
 			FMeshSelfUnion GroutUnion(&GroutMesh);
+			GroutUnion.bSimplifyAlongNewEdges = true;
 			GroutUnion.bWeldSharedEdges = false;
 			GroutUnion.Compute();
 
