GeometryProcessing: change handling of polygroup queries in DynamicMesh3 to return false if groups are not defined, instead of check/fail
#rnx

diff --git a/Engine/Plugins/Experimental/GeometryProcessing/Source/DynamicMesh/Private/DynamicMesh3_Queries.cpp b/Engine/Plugins/Experimental/GeometryProcessing/Source/DynamicMesh/Private/DynamicMesh3_Queries.cpp
--- a/Engine/Plugins/Experimental/GeometryProcessing/Source/DynamicMesh/Private/DynamicMesh3_Queries.cpp
+++ b/Engine/Plugins/Experimental/GeometryProcessing/Source/DynamicMesh/Private/DynamicMesh3_Queries.cpp
@@ -332,8 +332,7 @@ FIndex2i FDynamicMesh3::GetOrientedBoundaryEdgeV(int eID) const
 
 bool FDynamicMesh3::IsGroupBoundaryEdge(int eID) const
 {
-	check(IsEdge(eID));
-	check(HasTriangleGroups());
+	if (!HasTriangleGroups()) return false;
 
 	const FEdge Edge = Edges[eID];
 	int et1 = Edge.Tri[1];
@@ -349,8 +348,7 @@ bool FDynamicMesh3::IsGroupBoundaryEdge(int eID) const
 
 bool FDynamicMesh3::IsGroupBoundaryVertex(int vID) const
 {
-	check(IsVertex(vID));
-	check(HasTriangleGroups());
+	if (!HasTriangleGroups()) return false;
 
 	int group_id = InvalidID;
 	for (int eID : VertexEdgeLists.Values(vID))
@@ -386,8 +384,7 @@ bool FDynamicMesh3::IsGroupBoundaryVertex(int vID) const
 
 bool FDynamicMesh3::IsGroupJunctionVertex(int vID) const
 {
-	check(IsVertex(vID));
-	check(HasTriangleGroups());
+	if (!HasTriangleGroups()) return false;
 
 	FIndex2i groups(InvalidID, InvalidID);
 	for (int eID : VertexEdgeLists.Values(vID))
@@ -424,10 +421,8 @@ bool FDynamicMesh3::IsGroupJunctionVertex(int vID) const
 
 bool FDynamicMesh3::GetVertexGroups(int vID, FIndex4i& groups) const
 {
-	check(IsVertex(vID));
-	check(HasTriangleGroups());
-
 	groups = FIndex4i(InvalidID, InvalidID, InvalidID, InvalidID);
+	if (!HasTriangleGroups()) return false;
 	int ng = 0;
 
 	for (int eID : VertexEdgeLists.Values(vID))
@@ -465,8 +460,7 @@ bool FDynamicMesh3::GetVertexGroups(int vID, FIndex4i& groups) const
 
 bool FDynamicMesh3::GetAllVertexGroups(int vID, TArray<int>& GroupsOut) const
 {
-	check(IsVertex(vID));
-	check(HasTriangleGroups());
+	if (!HasTriangleGroups()) return false;
 
 	for (int eID : VertexEdgeLists.Values(vID))
 	{
