MeshModelingTools: Add Linear Stair generator to the Add Primitive Tool

#rnx

diff --git a/Engine/Plugins/Experimental/GeometryProcessing/Source/GeometricObjects/Public/Generators/StairGenerator.h b/Engine/Plugins/Experimental/GeometryProcessing/Source/GeometricObjects/Public/Generators/StairGenerator.h
new file mode 100644
--- /dev/null
+++ b/Engine/Plugins/Experimental/GeometryProcessing/Source/GeometricObjects/Public/Generators/StairGenerator.h
@@ -0,0 +1,742 @@
+// Copyright Epic Games, Inc. All Rights Reserved.
+#pragma once
+
+#include "Generators/MeshShapeGenerator.h"
+#include "IndexTypes.h"
+#include "MathUtil.h"
+
+/**
+ * Base Stair mesh generator class.
+ */
+class /*GEOMETRICOBJECTS_API*/ FStairGenerator : public FMeshShapeGenerator
+{
+public:
+	/** If true (default), UVs are scaled so that there is no stretching. If false, UVs are scaled to fill unit square */
+	bool bScaleUVByAspectRatio = true;
+
+	/** If true, each quad of box gets a separate polygroup */
+	bool bPolygroupPerQuad = false;
+
+	/** The width of each step. */
+	float StepWidth = 150.0f;
+
+	/** The depth of each step. */
+	float StepDepth = 30.0f;
+
+	/** The height of each step. */
+	float StepHeight = 20.0f;
+
+	/** The number of steps in this staircase. */
+	int NumSteps = 5;
+
+	enum class ESide
+	{
+		Right,	// +Y
+		Left,	// -Y
+		Front,	// -X
+		Top,	// +Z
+		Back,	// +X
+		Bottom	// -Z
+	};
+
+	enum class EBaseStyle
+	{
+		SolidBase,	// Each step connects to the floor.
+		HollowBase	// Each step only extends one step over to support the next step.
+	};
+	/** The style of the base of the stairs */
+	EBaseStyle BaseStyle = EBaseStyle::SolidBase;
+
+public:
+	/** Generate the mesh */
+	virtual FMeshShapeGenerator& Generate() override
+	{
+		Reset();
+
+		switch (BaseStyle)
+		{
+		case EBaseStyle::HollowBase:
+			return GenerateHollowBase();
+		case EBaseStyle::SolidBase:
+		default:
+			return GenerateSolidBase();
+		}
+	}
+
+	/**
+	 * Generate a stair mesh with a solid base.
+	 *
+	 * Stair topology composition. (side view cross section)
+	 *
+	 * (EBaseStyle::SolidBase)
+	 *                  .___.   
+	 *                  |   |          | z (height)
+	 *              .___|___|          | 
+	 *              |   |   |          |_____ x (depth)
+	 *          .___|___|___|           \
+	 *          |   |   |   |            \ y (width)
+	 *  Row --> |___|___|___|
+	 *
+	 *          ^-- Column
+	 *
+	 */
+	virtual FMeshShapeGenerator& GenerateSolidBase()
+	{
+		auto TriangleNumber = [](int x)
+		{
+			return (x * (x + 1) / 2);
+		};
+
+		const int NumConnectQuads = 4 * NumSteps;
+		const int NumCornerVerts = 4 * (NumSteps + 1);
+		NumVertsPerSide = TriangleNumber(NumSteps + 1) + NumSteps;
+		NumVerts = 2 * NumVertsPerSide;
+		NumQuadsPerSide = TriangleNumber(NumSteps);
+		NumQuads = 2 * NumQuadsPerSide + NumConnectQuads;
+
+		const int NumSideQuadAttrs = NumVerts;
+		const int NumConnectQuadCornerAttrs = 2 * NumCornerVerts;
+		const int NumConnectQuadInteriorAttrs = 4 * (NumSteps - 1);
+		NumAttrs = NumConnectQuadCornerAttrs + NumConnectQuadInteriorAttrs + NumSideQuadAttrs;
+		SetBufferSizes(NumVerts, 2 * NumQuads, NumAttrs, NumAttrs);
+
+		// Ordered Side lists for iteration.
+		const ESide RightLeftSides[2] = { ESide::Right, ESide::Left };
+		const ESide AllSides[6] = { ESide::Right, ESide::Left, ESide::Front, ESide::Top, ESide::Back, ESide::Bottom };
+
+		// Generate vertices by vertical column per Right/Left side.
+		LeftSideColumnId = (NumSteps + 1);
+		VertexIds.SetNum(2 * LeftSideColumnId);
+		VertexIdsToColumnRow.SetNum(NumVerts);
+		int VertexId = 0;
+		for (const ESide& Side : RightLeftSides)
+		{
+			int StartColumnId = (Side == ESide::Right ? 0 : LeftSideColumnId);
+			for (int VertexColumn = 0; VertexColumn < NumSteps + 1; VertexColumn++)
+			{
+				int NumVertices = (VertexColumn == NumSteps ? VertexColumn + 1 : VertexColumn + 2);
+				VertexIds[StartColumnId + VertexColumn].SetNum(NumVertices);
+				for (int VertexRow = 0; VertexRow < NumVertices; VertexRow++)
+				{
+					Vertices[VertexId] = GenerateVertex(Side, VertexColumn, VertexRow);
+					VertexIds[StartColumnId + VertexColumn][VertexRow] = VertexId;
+					VertexIdsToColumnRow[VertexId] = FIndex2i(StartColumnId + VertexColumn, VertexRow);
+					VertexId++;
+				}
+			}
+		}
+
+		// Generate quad representation for each side.
+		//
+		// Compute these in respective groups for later normal,
+		// UV & polygroup assignment.
+		FaceDesc.SetNum(4 * NumQuads);
+		int FaceDescId = 0;
+		for (const ESide& Side : AllSides)
+		{
+			switch (Side)
+			{
+			case ESide::Right:
+			{
+				const int StartColId = 0;
+				RightStartFaceId = FaceDescId;
+				for (int Step = 0; Step < NumSteps; Step++)
+				{
+					int CurrentColId = StartColId + Step;
+					for (int VertexRow = 0; VertexRow < VertexIds[CurrentColId].Num() - 1; VertexRow++)
+					{
+						FaceDesc[FaceDescId++] = (VertexIds[CurrentColId][VertexRow]);
+						FaceDesc[FaceDescId++] = (VertexIds[CurrentColId][VertexRow + 1]);
+						FaceDesc[FaceDescId++] = (VertexIds[CurrentColId + 1][VertexRow + 1]);
+						FaceDesc[FaceDescId++] = (VertexIds[CurrentColId + 1][VertexRow]);
+					}
+				}
+				break;
+			}
+			case ESide::Left:
+			{
+				const int StartColId = LeftSideColumnId;
+				LeftStartFaceId = FaceDescId;
+				for (int Step = 0; Step < NumSteps; Step++)
+				{
+					int CurrentColId = StartColId + Step;
+					for (int VertexRow = 0; VertexRow < VertexIds[CurrentColId].Num() - 1; VertexRow++)
+					{
+						// Flipped face order from ESide::Right
+						FaceDesc[FaceDescId++] = (VertexIds[CurrentColId][VertexRow]);
+						FaceDesc[FaceDescId++] = (VertexIds[CurrentColId + 1][VertexRow]);
+						FaceDesc[FaceDescId++] = (VertexIds[CurrentColId + 1][VertexRow + 1]);
+						FaceDesc[FaceDescId++] = (VertexIds[CurrentColId][VertexRow + 1]);
+					}
+				}
+				break;
+			}
+			case ESide::Front:
+			{
+				FrontStartFaceId = FaceDescId;
+				for (int Step = 0; Step < NumSteps; Step++)
+				{
+					int LastStepRow = VertexIds[Step].Num() - 1;
+					FaceDesc[FaceDescId++] = (VertexIds[Step + LeftSideColumnId][LastStepRow - 1]);
+					FaceDesc[FaceDescId++] = (VertexIds[Step + LeftSideColumnId][LastStepRow]);
+					FaceDesc[FaceDescId++] = (VertexIds[Step][LastStepRow]);
+					FaceDesc[FaceDescId++] = (VertexIds[Step][LastStepRow - 1]);
+				}
+				break;
+			}
+			case ESide::Top:
+			{
+				TopStartFaceId = FaceDescId;
+				for (int Step = 0; Step < NumSteps; Step++)
+				{
+					int LastStepRow = VertexIds[Step].Num() - 1;
+					FaceDesc[FaceDescId++] = (VertexIds[Step + LeftSideColumnId][LastStepRow]);
+					FaceDesc[FaceDescId++] = (VertexIds[Step + LeftSideColumnId + 1][LastStepRow]);
+					FaceDesc[FaceDescId++] = (VertexIds[Step + 1][LastStepRow]);
+					FaceDesc[FaceDescId++] = (VertexIds[Step][LastStepRow]);
+				}
+				break;
+			}
+			case ESide::Back:
+			{
+				BackStartFaceId = FaceDescId;
+				for (int Step = 0; Step < NumSteps; Step++)
+				{
+					FaceDesc[FaceDescId++] = (VertexIds[NumSteps][Step]);
+					FaceDesc[FaceDescId++] = (VertexIds[NumSteps][Step + 1]);
+					FaceDesc[FaceDescId++] = (VertexIds[LeftSideColumnId + NumSteps][Step + 1]);
+					FaceDesc[FaceDescId++] = (VertexIds[LeftSideColumnId + NumSteps][Step]);
+				}
+				break;
+			}
+			case ESide::Bottom:
+			{
+				BottomStartFaceId = FaceDescId;
+				for (int Step = 0; Step < NumSteps; Step++)
+				{
+					FaceDesc[FaceDescId++] = (VertexIds[Step][0]);
+					FaceDesc[FaceDescId++] = (VertexIds[Step + 1][0]);
+					FaceDesc[FaceDescId++] = (VertexIds[Step + 1 + LeftSideColumnId][0]);
+					FaceDesc[FaceDescId++] = (VertexIds[Step + LeftSideColumnId][0]);
+				}
+				break;
+			}
+			}
+		}
+		LastFaceId = FaceDescId;
+
+		// Compute Normals
+		NormalDesc.SetNum(4 * NumQuads);
+		int NormalId = 0;
+		const int NumFrontFaceVertex = TopStartFaceId - FrontStartFaceId;
+		const int NumTopFaceVertex = BackStartFaceId - TopStartFaceId;
+		const int NumBackFaceVertex = BottomStartFaceId - BackStartFaceId;
+		const int NumBottomFaceVertex = LastFaceId - BottomStartFaceId;
+		for (const ESide& Side : AllSides)
+		{
+			switch (Side)
+			{
+			case ESide::Right:
+			case ESide::Left:
+			{
+				const int StartVertexId = (Side == ESide::Right ? 0 : NumVertsPerSide);
+				for (int SideVertex = 0; SideVertex < NumVertsPerSide; SideVertex++)
+				{
+					int VId = StartVertexId + SideVertex;
+					Normals[NormalId++] = GenerateNormal(Side, VId);
+				}
+				// For the Side Quads, the Normal descriptor is identical to the Face descriptor.
+				const int StartSideFaceVertex = (Side == ESide::Right ? 0 : 4 * NumQuadsPerSide);
+				const int NumSideFaceVertex = 4 * NumQuadsPerSide;
+				for (int SideFaceVertex = 0; SideFaceVertex < NumSideFaceVertex; SideFaceVertex++)
+				{
+					NormalDesc[StartSideFaceVertex + SideFaceVertex] = FaceDesc[StartSideFaceVertex + SideFaceVertex];
+				}
+				break;
+			}
+			case ESide::Front:
+			case ESide::Top:
+			{
+				const int NumFaceVertex = (Side == ESide::Front ? NumFrontFaceVertex : NumTopFaceVertex);
+				const int StartFaceId = (Side == ESide::Front ? FrontStartFaceId : TopStartFaceId);
+				for (int FaceVertexId = 0; FaceVertexId < NumFaceVertex; FaceVertexId++)
+				{
+					int VId = FaceDesc[StartFaceId + FaceVertexId];
+					Normals[NormalId] = GenerateNormal(Side, VId);
+					NormalDesc[StartFaceId + FaceVertexId] = NormalId;
+					NormalId++;
+				}
+				break;
+			}
+			case ESide::Back:
+			case ESide::Bottom:
+			{
+				// Generate the normals.
+				const int NormalStartId = NormalId;
+				const int NumFaceVertex = (Side == ESide::Back ? NumBackFaceVertex : NumBottomFaceVertex);
+				const int StartFaceId = (Side == ESide::Back ? BackStartFaceId : BottomStartFaceId);
+				for (int FaceVertexId = 0; FaceVertexId < NumFaceVertex; FaceVertexId += 4)
+				{
+					Normals[NormalId++] = GenerateNormal(Side, FaceDesc[StartFaceId + FaceVertexId]);
+					Normals[NormalId++] = GenerateNormal(Side, FaceDesc[StartFaceId + FaceVertexId + 3]);
+
+					// For the last face, add the last/cap vertices.
+					if (FaceVertexId + 4 >= NumFaceVertex)
+					{
+						Normals[NormalId++] = GenerateNormal(Side, FaceDesc[StartFaceId + FaceVertexId + 1]);
+						Normals[NormalId++] = GenerateNormal(Side, FaceDesc[StartFaceId + FaceVertexId + 2]);
+					}
+				}
+				// Assign the normals via the normal descriptor.
+				int FaceVertexId = StartFaceId;
+				int FaceNormalId = NormalStartId;
+				for (int FaceId = 0; FaceId < NumSteps; FaceId++)
+				{
+					// Face descriptor index order for back/bottom
+					//
+					//         1.___.2
+					//          |   |
+					//          |___|
+					//         0     3
+					//
+					// The normals are generated along each vertex row
+					// in order: (0, 3, 1, 2). Thus the respective normal
+					// IDs follow: (0, 2, 3, 1).
+					//
+					NormalDesc[FaceVertexId] = FaceNormalId;
+					NormalDesc[FaceVertexId + 1] = FaceNormalId + 2;
+					NormalDesc[FaceVertexId + 2] = FaceNormalId + 3;
+					NormalDesc[FaceVertexId + 3] = FaceNormalId + 1;
+					FaceVertexId += 4;
+					FaceNormalId += 2;
+				}
+				break;
+			}
+			}
+		}
+
+		// Compute UVs
+		UVDesc.SetNum(4 * NumQuads);
+		int UVId = 0;
+		float MaxDimension = FMathf::Max(FMathf::Max(NumSteps * StepDepth, NumSteps * StepHeight), StepWidth);
+		float UVScale = (bScaleUVByAspectRatio) ? (1.0f / (float)MaxDimension) : 1.0f;
+		for (const ESide& Side : AllSides)
+		{
+			switch (Side)
+			{
+			case ESide::Right:
+			case ESide::Left:
+			{
+				const int StartVertexId = (Side == ESide::Right ? 0 : NumVertsPerSide);
+				for (int SideVertex = 0; SideVertex < NumVertsPerSide; SideVertex++)
+				{
+					int VId = StartVertexId + SideVertex;
+					UVs[UVId++] = GenerateUV(Side, SideVertex, VId, UVScale);
+				}
+				// For the Side Quads, the UV descriptor is identical to the Face descriptor.
+				const int StartSideFaceVertex = (Side == ESide::Right ? 0 : 4 * NumQuadsPerSide);
+				const int NumSideFaceVertex = 4 * NumQuadsPerSide;
+				for (int SideFaceVertex = 0; SideFaceVertex < NumSideFaceVertex; SideFaceVertex++)
+				{
+					UVDesc[StartSideFaceVertex + SideFaceVertex] = FaceDesc[StartSideFaceVertex + SideFaceVertex];
+				}
+				break;
+			}
+			case ESide::Front:
+			case ESide::Top:
+			{
+				const int NumFaceVertex = (Side == ESide::Front ? NumFrontFaceVertex : NumTopFaceVertex);
+				const int StartFaceId = (Side == ESide::Front ? FrontStartFaceId : TopStartFaceId);
+				for (int FaceVertexId = 0; FaceVertexId < NumFaceVertex; FaceVertexId++)
+				{
+					int VId = FaceDesc[StartFaceId + FaceVertexId];
+					UVs[UVId] = GenerateUV(Side, FaceVertexId / 4, VId, UVScale);
+					UVDesc[StartFaceId + FaceVertexId] = UVId;
+					UVId++;
+				}
+				break;
+			}
+			case ESide::Back:
+			case ESide::Bottom:
+			{
+				// Generate the UVs.
+				const int UVStartId = UVId;
+				const int NumFaceVertex = (Side == ESide::Back ? NumBackFaceVertex : NumBottomFaceVertex);
+				const int StartFaceId = (Side == ESide::Back ? BackStartFaceId : BottomStartFaceId);
+				for (int FaceVertexId = 0; FaceVertexId < NumFaceVertex; FaceVertexId += 4)
+				{
+					const int SideFaceId = FaceVertexId / 4;
+					UVs[UVId++] = GenerateUV(Side, SideFaceId, FaceDesc[StartFaceId + FaceVertexId], UVScale);
+					UVs[UVId++] = GenerateUV(Side, SideFaceId, FaceDesc[StartFaceId + FaceVertexId + 3], UVScale);
+
+					// For the last face, add the last/cap vertices.
+					if (FaceVertexId + 4 >= NumFaceVertex)
+					{
+						UVs[UVId++] = GenerateUV(Side, SideFaceId, FaceDesc[StartFaceId + FaceVertexId + 1], UVScale);
+						UVs[UVId++] = GenerateUV(Side, SideFaceId, FaceDesc[StartFaceId + FaceVertexId + 2], UVScale);
+					}
+				}
+				// Assign the UVs via the normal descriptor.
+				int FaceVertexId = StartFaceId;
+				int FaceUVId = UVStartId;
+				for (int FaceId = 0; FaceId < NumSteps; FaceId++)
+				{
+					// Face descriptor index order for back/bottom
+					//
+					//         1.___.2
+					//          |   |
+					//          |___|
+					//         0     3
+					//
+					// The UVs are generated along each vertex row
+					// in order: (0, 3, 1, 2). Thus the respective UV
+					// IDs follow: (0, 2, 3, 1).
+					//
+					UVDesc[FaceVertexId] = FaceUVId;
+					UVDesc[FaceVertexId + 1] = FaceUVId + 2;
+					UVDesc[FaceVertexId + 2] = FaceUVId + 3;
+					UVDesc[FaceVertexId + 3] = FaceUVId + 1;
+					FaceVertexId += 4;
+					FaceUVId += 2;
+				}
+				break;
+			}
+			}
+		}
+
+		// Triangulate quad mesh into the output mesh.
+		auto FaceToSide = [this](int FaceId) -> ESide
+		{
+			ESide Side = ESide::Right;
+			const int FaceVertexId = FaceId * 4;
+			if (FaceVertexId >= RightStartFaceId && FaceVertexId < LeftStartFaceId)
+			{
+				Side = ESide::Right;
+			}
+			else if (FaceVertexId >= LeftStartFaceId && FaceVertexId < FrontStartFaceId)
+			{
+				Side = ESide::Left;
+			}
+			else if (FaceVertexId >= FrontStartFaceId && FaceVertexId < TopStartFaceId)
+			{
+				Side = ESide::Front;
+			}
+			else if (FaceVertexId >= TopStartFaceId && FaceVertexId < BackStartFaceId)
+			{
+				Side = ESide::Top;
+			}
+			else if (FaceVertexId >= BackStartFaceId && FaceVertexId < BottomStartFaceId)
+			{
+				Side = ESide::Back;
+			}
+			else if (FaceVertexId >= BottomStartFaceId && FaceVertexId < LastFaceId)
+			{
+				Side = ESide::Bottom;
+			}
+			else
+			{
+				check(false);
+			}
+			return Side;
+		};
+		constexpr int Tris[2][3] = { {0, 3, 2}, {2, 1, 0} };
+		int CurrentTriId = 0;
+		int GroupId = 0;
+		for (int FaceId = 0; FaceId < NumQuads; FaceId++)
+		{
+			ESide Side = FaceToSide(FaceId);
+			if (bPolygroupPerQuad)
+			{
+				GroupId = FaceId;
+			}
+			else
+			{
+				switch (Side)
+				{
+				case ESide::Right:
+					GroupId = 0;
+					break;
+				case ESide::Left:
+					GroupId = 1;
+					break;
+				case ESide::Front:
+				case ESide::Top:
+					// Each actual step (front/top) face are grouped
+					// as a polygroup.
+					check(NumFrontFaceVertex == NumTopFaceVertex);
+					check(NumFrontFaceVertex % 4 == 0);
+					GroupId = FaceId % (NumFrontFaceVertex / 4) + 2;
+					break;
+				case ESide::Back:
+					GroupId = 2 + NumFrontFaceVertex / 4 + 1;
+					break;
+				case ESide::Bottom:
+					GroupId = 2 + NumFrontFaceVertex / 4 + 2;
+					break;
+				}
+			}
+			for (int TriId = 0; TriId < 2; TriId++)
+			{
+				const int AId = 4 * FaceId + Tris[TriId][0];
+				const int BId = 4 * FaceId + Tris[TriId][1];
+				const int CId = 4 * FaceId + Tris[TriId][2];
+				const int FA = FaceDesc[AId];
+				const int FB = FaceDesc[BId];
+				const int FC = FaceDesc[CId];
+				SetTriangle(CurrentTriId, FA, FB, FC);
+				const int NA = NormalDesc[AId];
+				const int NB = NormalDesc[BId];
+				const int NC = NormalDesc[CId];
+				SetTriangleNormals(CurrentTriId, NA, NB, NC);
+				const int UVA = UVDesc[AId];
+				const int UVB = UVDesc[BId];
+				const int UVC = UVDesc[CId];
+				SetTriangleUVs(CurrentTriId, UVA, UVB, UVC);
+				SetTrianglePolygon(CurrentTriId, GroupId);
+				CurrentTriId++;
+			}
+		}
+
+		return *this;
+	}
+
+	/**
+	 * Generate a hollow base stair mesh.
+	 * 
+	 * Stair topology composition. (side view cross section)
+	 *
+	 * (EBaseStyle::HollowBase)
+	 *                  .___.
+	 *                  |   |          | z (height)
+	 *              .___|___|          |
+	 *              |   |   |          |_____ x (depth)
+	 *          .___|___|___|           \
+	 *          |   |   |                \ y (width)
+	 *  Row --> |___|___|
+	 *
+	 *          ^-- Column
+	 *
+	 */
+	virtual FMeshShapeGenerator& GenerateHollowBase()
+	{
+		// TODO: Not yet implemented.
+		check(false);
+		return *this;
+	}
+
+protected:
+	TArray<TArray<int>> VertexIds;
+	TArray<FIndex2i> VertexIdsToColumnRow;
+	TArray<int> FaceDesc;
+	TArray<int> NormalDesc;
+	TArray<int> UVDesc;
+
+	/** Mesh counts. */
+	int NumQuadsPerSide = 0;
+	int NumQuads = 0;
+	int NumVertsPerSide = 0;
+	int NumVerts = 0;
+	int NumAttrs = 0;
+	
+	/** Vertex column indices. */
+	int RightSideColumnId = 0;
+	int LeftSideColumnId = 0;
+
+	/** Face descriptor indices. */
+	int RightStartFaceId = 0;
+	int LeftStartFaceId = 0;
+	int FrontStartFaceId = 0;
+	int TopStartFaceId = 0;
+	int BackStartFaceId = 0;
+	int BottomStartFaceId = 0;
+	int LastFaceId = 0;
+
+protected:
+	virtual void Reset()
+	{
+		VertexIds.Reset();
+		VertexIdsToColumnRow.Reset();
+		FaceDesc.Reset();
+		NormalDesc.Reset();
+		UVDesc.Reset();
+
+		NumQuadsPerSide = 0;
+		NumQuads = 0;
+		NumVertsPerSide = 0;
+		NumVerts = 0;
+		NumAttrs = 0;
+
+		RightSideColumnId = 0;
+		LeftSideColumnId = 0;
+
+		FrontStartFaceId = 0;
+		TopStartFaceId = 0;
+		BackStartFaceId = 0;
+		BottomStartFaceId = 0;
+		LastFaceId = 0;
+	}
+
+	/**
+	 * Returns a vertex position.
+	 *
+	 * The method is provided the Right or Left side and the
+	 * corresponding vertex column/row index to compute. Column
+	 * and row indices refer to the matrix-like ordered vertex layout.
+	 * See diagram in GenerateSolidBase() / GenerateHollowBase().
+	 *
+	 * This generator only generates vertices for the Right & Left
+	 * Sides of the stairs.
+	 *
+	 * Mesh count protected variables are the only transient variables
+	 * guaranteed to be valid at the time GenerateVertex is invoked.
+	 *
+	 * @param Side The Right or Left side of the stairs.
+	 * @param VertexColumn The column index into the stair vertex layout.
+	 * @param VertexRow The row index into the stair vertex layout.
+	 */
+	virtual FVector3d GenerateVertex(ESide Side, int VertexColumn, int VertexRow) = 0;
+
+	/**
+	 * Returns a vertex normal vector.
+	 * 
+	 * All normals for a given side are shared except for Front & Top.
+	 * Border vertex normals per side are not shared.
+	 *
+	 * @param Side The side of the stairs to compute the normal.
+	 * @param VertexId The vertex index to compute the normal.
+	 */
+	virtual FVector3f GenerateNormal(ESide Side, int VertexId) = 0;
+	
+	/**
+	 * Returns a UV vector.
+	 * 
+	 * The Step parameter provides the Side-relative face. This
+	 * indicates which face for a given side is computing its
+	 * UV.
+	 *
+	 * All UVs for a given side are shared except for Front & Top.
+	 * Border edges along each side are UV island edges.
+	 *
+	 * @param Side The side of the stairs to compute the UV.
+	 * @param Step The Side-relative step face ID.
+	 * @param VertexId The vertex index to compute the UV.
+	 * @param UVScale The UV scale
+	 */
+	virtual FVector2f GenerateUV(ESide Side, int Step, int VertexId, float UVScale) = 0;
+};
+
+/**
+ * Generate an oriented Linear Stair mesh.
+ */
+class /*GEOMETRICOBJECTS_API*/ FLinearStairGenerator : public FStairGenerator
+{
+protected:
+	virtual FVector3d GenerateVertex(ESide Side, int VertexColumn, int VertexRow) override
+	{
+		if (Side != ESide::Right && Side != ESide::Left)
+		{
+			// Vertices are only generated for Right & Left sides.
+			check(false);
+		}
+		float X = VertexColumn * StepDepth;
+		float Y = (Side == ESide::Right ? 0.5 * StepWidth : -0.5 * StepWidth);
+		float Z = VertexRow * StepHeight;
+		return FVector3d(X, Y, Z);
+	}
+
+	virtual FVector3f GenerateNormal(ESide Side, int VertexId) override
+	{
+		FVector3f N(0.0f, 0.0f, 0.0f);
+		switch (Side)
+		{
+		case ESide::Right:
+			N = FVector3f(0.0f, 1.0f, 0.0f);
+			break;
+		case ESide::Left:
+			N = FVector3f(0.0f, -1.0f, 0.0f);
+			break;
+		case ESide::Front:
+			N = FVector3f(-1.0f, 0.0f, 0.0f);
+			break;
+		case ESide::Top:
+			N = FVector3f(0.0f, 0.0f, 1.0f);
+			break;
+		case ESide::Back:
+			N = FVector3f(1.0f, 0.0f, 0.0f);
+			break;
+		case ESide::Bottom:
+			N = FVector3f(0.0f, 0.0f, -1.0f);
+			break;
+		default:
+			check(false);
+			break;
+		}
+		return N;
+	}
+
+	virtual FVector2f GenerateUV(ESide Side, int Step, int VertexId, float UVScale) override
+	{
+		check(VertexId < VertexIdsToColumnRow.Num());
+		const int Col = VertexIdsToColumnRow[VertexId].A;
+		const int Row = VertexIdsToColumnRow[VertexId].B;
+		FVector2f UV(0.0f, 0.0f);
+		switch (Side)
+		{
+		case ESide::Right:
+		case ESide::Left:
+		{
+			const float UScale = NumSteps * StepDepth * UVScale;
+			const float VScale = NumSteps * StepHeight * UVScale;
+			UV.X = FMathf::Lerp(-0.5f, 0.5f, float(Col % LeftSideColumnId) / float(NumSteps + 1)) * UScale + 0.5f;
+			UV.Y = FMathf::Lerp(-0.5f, 0.5f, float(Row) / float(NumSteps + 1)) * VScale + 0.5f;;
+			break;
+		}
+		case ESide::Front:
+		{
+			const float UScale = StepWidth * UVScale;
+			const float VScale = StepHeight * UVScale;
+			UV.X = Col < LeftSideColumnId ? 0.5f : -0.5f;
+			UV.Y = Row > Step ? -0.5f : 0.5f;
+			UV.X = UV.X * UScale + 0.5f;
+			UV.Y = UV.Y * VScale + 0.5f;
+			break;
+		}
+		case ESide::Top:
+		{
+			const float UScale = StepWidth * UVScale;
+			const float VScale = StepDepth * UVScale;
+			UV.X = Col < LeftSideColumnId ? 0.5f : -0.5f;
+			UV.Y = (Col % LeftSideColumnId) > Step ? -0.5f : 0.5f;
+			UV.X = UV.X * UScale + 0.5f;
+			UV.Y = UV.Y * VScale + 0.5f;
+			break;
+		}
+		case ESide::Back:
+		{
+			const float UScale = StepWidth * UVScale;
+			const float VScale = NumSteps * StepHeight * UVScale;
+			UV.X = Col < LeftSideColumnId ? 0.5f : -0.5f;
+			UV.Y = FMathf::Lerp(-0.5f, 0.5f, float(Row) / float(NumSteps + 1));
+			UV.X = UV.X * UScale + 0.5f;
+			UV.Y = UV.Y * VScale + 0.5f;
+			break;
+		}
+		case ESide::Bottom:
+		{
+			const float UScale = StepWidth * UVScale;
+			const float VScale = NumSteps * StepDepth * UVScale;
+			UV.X = Col < LeftSideColumnId ? 0.5f : -0.5f;
+			UV.Y = FMathf::Lerp(-0.5f, 0.5f, float(Col % LeftSideColumnId) / float(NumSteps + 1));
+			UV.X = UV.X * UScale + 0.5f;
+			UV.Y = UV.Y * VScale + 0.5f;
+			break;
+		}
+		default:
+		{
+			check(false);
+			break;
+		}
+		}
+		return UV;
+	}
+};
\ No newline at end of file
diff --git a/Engine/Plugins/Experimental/MeshModelingToolset/Source/MeshModelingTools/Private/AddPrimitiveTool.cpp b/Engine/Plugins/Experimental/MeshModelingToolset/Source/MeshModelingTools/Private/AddPrimitiveTool.cpp
--- a/Engine/Plugins/Experimental/MeshModelingToolset/Source/MeshModelingTools/Private/AddPrimitiveTool.cpp
+++ b/Engine/Plugins/Experimental/MeshModelingToolset/Source/MeshModelingTools/Private/AddPrimitiveTool.cpp
@@ -18,6 +18,7 @@
 #include "Generators/SphereGenerator.h"
 #include "Generators/BoxSphereGenerator.h"
 #include "Generators/DiscMeshGenerator.h"
+#include "Generators/StairGenerator.h"
 #include "DynamicMeshAttributeSet.h"
 #include "FaceGroupUtil.h"
 
@@ -76,6 +77,9 @@ UInteractiveTool* UAddPrimitiveToolBuilder::BuildTool(const FToolBuilderState& S
 	case EMakeMeshShapeType::Sphere:
 		NewTool = NewObject<UAddSpherePrimitiveTool>(SceneState.ToolManager);
 		break;
+	case EMakeMeshShapeType::LinearStairs:
+		NewTool = NewObject<UAddLinearStairsPrimitiveTool>(SceneState.ToolManager);
+		break;
 	default:
 		break;
 	}
@@ -655,4 +659,24 @@ void UAddSphericalBoxPrimitiveTool::GenerateMesh(FDynamicMesh3* OutMesh) const
 	OutMesh->Copy(&SphereGen);
 }
 
+UAddLinearStairsPrimitiveTool::UAddLinearStairsPrimitiveTool(const FObjectInitializer& ObjectInitializer)
+	: Super(ObjectInitializer.SetDefaultSubobjectClass<UProceduralLinearStairsToolProperties>(TEXT("ShapeSettings")))
+{
+	AssetName = TEXT("LinearStairs");
+	SetToolDisplayName(LOCTEXT("LinearStairsToolName", "Create LinearStairs"));
+}
+
+void UAddLinearStairsPrimitiveTool::GenerateMesh(FDynamicMesh3* OutMesh) const
+{
+	FLinearStairGenerator StairGen;
+	auto* StairSettings = Cast<UProceduralLinearStairsToolProperties>(ShapeSettings);
+	StairGen.StepWidth = StairSettings->StepWidth;
+	StairGen.StepHeight = StairSettings->StepHeight;
+	StairGen.StepDepth = StairSettings->StepDepth;
+	StairGen.NumSteps = StairSettings->NumSteps;
+	StairGen.bPolygroupPerQuad = (ShapeSettings->PolygroupMode == EMakeMeshPolygroupMode::PerQuad);
+	StairGen.Generate();
+	OutMesh->Copy(&StairGen);
+}
+
 #undef LOCTEXT_NAMESPACE
diff --git a/Engine/Plugins/Experimental/MeshModelingToolset/Source/MeshModelingTools/Public/AddPrimitiveTool.h b/Engine/Plugins/Experimental/MeshModelingToolset/Source/MeshModelingTools/Public/AddPrimitiveTool.h
--- a/Engine/Plugins/Experimental/MeshModelingToolset/Source/MeshModelingTools/Public/AddPrimitiveTool.h
+++ b/Engine/Plugins/Experimental/MeshModelingToolset/Source/MeshModelingTools/Public/AddPrimitiveTool.h
@@ -37,7 +37,8 @@ public:
 		PuncturedDisc,
 		Torus,
 		SphericalBox,
-		Sphere
+		Sphere,
+		LinearStairs
 	};
 
 	IAssetGenerationAPI* AssetAPI;
@@ -339,6 +340,28 @@ public:
 	int Subdivisions = 16;
 };
 
+UCLASS()
+class MESHMODELINGTOOLS_API UProceduralLinearStairsToolProperties : public UProceduralShapeToolProperties
+{
+	GENERATED_BODY()
+public:
+	/** Width of Shape */
+	UPROPERTY(EditAnywhere, Category = "ShapeSettings|Dimensions", meta = (DisplayName = "Step Width", UIMin = "1.0", UIMax = "1000.0", ClampMin = "0.0001", ClampMax = "1000000.0", ProceduralShapeSetting))
+	float StepWidth = 150.f;
+
+	/** Depth of Shape */
+	UPROPERTY(EditAnywhere, Category = "ShapeSettings|Dimensions", meta = (DisplayName = "Step Depth", UIMin = "1.0", UIMax = "1000.0", ClampMin = "0.0001", ClampMax = "1000000.0", ProceduralShapeSetting))
+	float StepDepth = 30.f;
+
+	/** Height of Shape*/
+	UPROPERTY(EditAnywhere, Category = "ShapeSettings|Dimensions", meta = (DisplayName = "Step Height", UIMin = "1.0", UIMax = "1000.0", ClampMin = "0.0001", ClampMax = "1000000.0", ProceduralShapeSetting))
+	float StepHeight = 20.f;
+
+	/** Number of Steps of Shape */
+	UPROPERTY(EditAnywhere, Category = "ShapeSettings|Dimensions", meta = (DisplayName = "Number of Steps", UIMin = "1", UIMax = "100", ClampMin = "1", ClampMax = "1000", ProceduralShapeSetting))
+	int NumSteps = 5;
+};
+
 UCLASS(Transient)
 class MESHMODELINGTOOLS_API ULastActorInfo : public UObject
 {
@@ -556,3 +579,13 @@ protected:
 	void GenerateMesh(FDynamicMesh3* OutMesh) const override;
 };
 
+UCLASS()
+class UAddLinearStairsPrimitiveTool : public UAddPrimitiveTool
+{
+	GENERATED_BODY()
+public:
+	UAddLinearStairsPrimitiveTool(const FObjectInitializer& ObjectInitializer);
+protected:
+	void GenerateMesh(FDynamicMesh3* OutMesh) const override;
+};
+
diff --git a/Engine/Plugins/Experimental/ModelingToolsEditorMode/Source/ModelingToolsEditorMode/Private/ModelingToolsEditorMode.cpp b/Engine/Plugins/Experimental/ModelingToolsEditorMode/Source/ModelingToolsEditorMode/Private/ModelingToolsEditorMode.cpp
--- a/Engine/Plugins/Experimental/ModelingToolsEditorMode/Source/ModelingToolsEditorMode/Private/ModelingToolsEditorMode.cpp
+++ b/Engine/Plugins/Experimental/ModelingToolsEditorMode/Source/ModelingToolsEditorMode/Private/ModelingToolsEditorMode.cpp
@@ -346,6 +346,9 @@ void UModelingToolsEditorMode::Enter()
 	RegisterPrimitiveToolFunc(ToolManagerCommands.BeginAddSphericalBoxPrimitiveTool,
 							  TEXT("BeginAddSphericalBoxPrimitiveTool"),
 							  UAddPrimitiveToolBuilder::EMakeMeshShapeType::SphericalBox);
+	RegisterPrimitiveToolFunc(ToolManagerCommands.BeginAddLinearStairsPrimitiveTool,
+							  TEXT("BeginAddLinearStairsPrimitiveTool"),
+							  UAddPrimitiveToolBuilder::EMakeMeshShapeType::LinearStairs);
 
 	//
 	// make shape tools
diff --git a/Engine/Plugins/Experimental/ModelingToolsEditorMode/Source/ModelingToolsEditorMode/Private/ModelingToolsEditorModeToolkit.cpp b/Engine/Plugins/Experimental/ModelingToolsEditorMode/Source/ModelingToolsEditorMode/Private/ModelingToolsEditorModeToolkit.cpp
--- a/Engine/Plugins/Experimental/ModelingToolsEditorMode/Source/ModelingToolsEditorMode/Private/ModelingToolsEditorModeToolkit.cpp
+++ b/Engine/Plugins/Experimental/ModelingToolsEditorMode/Source/ModelingToolsEditorMode/Private/ModelingToolsEditorModeToolkit.cpp
@@ -468,6 +468,7 @@ void FModelingToolsEditorModeToolkit::BuildToolPalette_Experimental(FName Palett
 		ToolbarBuilder.AddToolBarButton(Commands.BeginAddRoundedRectanglePrimitiveTool);
 		ToolbarBuilder.AddToolBarButton(Commands.BeginAddDiscPrimitiveTool);
 		ToolbarBuilder.AddToolBarButton(Commands.BeginAddPuncturedDiscPrimitiveTool);
+		ToolbarBuilder.AddToolBarButton(Commands.BeginAddLinearStairsPrimitiveTool);
 	}
 	else if (PaletteIndex == CreateTabName)
 	{
diff --git a/Engine/Plugins/Experimental/ModelingToolsEditorMode/Source/ModelingToolsEditorMode/Private/ModelingToolsManagerActions.cpp b/Engine/Plugins/Experimental/ModelingToolsEditorMode/Source/ModelingToolsEditorMode/Private/ModelingToolsManagerActions.cpp
--- a/Engine/Plugins/Experimental/ModelingToolsEditorMode/Source/ModelingToolsEditorMode/Private/ModelingToolsManagerActions.cpp
+++ b/Engine/Plugins/Experimental/ModelingToolsEditorMode/Source/ModelingToolsEditorMode/Private/ModelingToolsManagerActions.cpp
@@ -33,6 +33,7 @@ void FModelingToolsManagerCommands::RegisterCommands()
 	UI_COMMAND(BeginAddTorusPrimitiveTool, "Torus", "Create new Torus StaticMesh Assets", EUserInterfaceActionType::ToggleButton, FInputChord());
 	UI_COMMAND(BeginAddSphericalBoxPrimitiveTool, "SphrB", "Create new Sphere StaticMesh Assets with Box Parameterization", EUserInterfaceActionType::ToggleButton, FInputChord());
 	UI_COMMAND(BeginAddSpherePrimitiveTool, "SphrA", "Create new Sphere StaticMesh Assets with Lat/Long Parameterization", EUserInterfaceActionType::ToggleButton, FInputChord());
+	UI_COMMAND(BeginAddLinearStairsPrimitiveTool, "StairL", "Create new Linear Stairs StaticMesh Assets", EUserInterfaceActionType::ToggleButton, FInputChord());
 
 	UI_COMMAND(BeginAddPatchTool, "Patch", "Add Patch", EUserInterfaceActionType::ToggleButton, FInputChord());
 	UI_COMMAND(BeginRevolveBoundaryTool, "BdryRev", "Revolve Mesh Boundary Loop into New Asset", EUserInterfaceActionType::ToggleButton, FInputChord());
diff --git a/Engine/Plugins/Experimental/ModelingToolsEditorMode/Source/ModelingToolsEditorMode/Public/ModelingToolsManagerActions.h b/Engine/Plugins/Experimental/ModelingToolsEditorMode/Source/ModelingToolsEditorMode/Public/ModelingToolsManagerActions.h
--- a/Engine/Plugins/Experimental/ModelingToolsEditorMode/Source/ModelingToolsEditorMode/Public/ModelingToolsManagerActions.h
+++ b/Engine/Plugins/Experimental/ModelingToolsEditorMode/Source/ModelingToolsEditorMode/Public/ModelingToolsManagerActions.h
@@ -26,6 +26,7 @@ public:
 	TSharedPtr<FUICommandInfo> BeginAddTorusPrimitiveTool;
 	TSharedPtr<FUICommandInfo> BeginAddSphericalBoxPrimitiveTool;
 	TSharedPtr<FUICommandInfo> BeginAddSpherePrimitiveTool;
+	TSharedPtr<FUICommandInfo> BeginAddLinearStairsPrimitiveTool;
 
 	TSharedPtr<FUICommandInfo> BeginAddPatchTool;
 	TSharedPtr<FUICommandInfo> BeginRevolveBoundaryTool;
