Deform mesh normals in FFDLattice. Computes the Jacobian of the trilinear or tricubic interpolation functions and rotates the original normals by the inverse Jacobian transpose times the determinant of the Jacobian.

#rnx

diff --git a/Engine/Plugins/Experimental/GeometryProcessing/Source/DynamicMesh/Private/Operations/FFDLattice.cpp b/Engine/Plugins/Experimental/GeometryProcessing/Source/DynamicMesh/Private/Operations/FFDLattice.cpp
--- a/Engine/Plugins/Experimental/GeometryProcessing/Source/DynamicMesh/Private/Operations/FFDLattice.cpp
+++ b/Engine/Plugins/Experimental/GeometryProcessing/Source/DynamicMesh/Private/Operations/FFDLattice.cpp
@@ -167,6 +167,96 @@ void FFFDLattice::GetDeformedMeshVertexPositions(const TArray<FVector3d>& Lattic
 
 }
 
+void FFFDLattice::GetRotatedOverlayNormals(const TArray<FVector3d>& LatticeControlPoints,
+										   const FDynamicMeshNormalOverlay* NormalOverlay,
+										   TArray<FVector3f>& OutNormals,
+										   ELatticeInterpolation Interpolation,
+										   FLatticeExecutionInfo ExecutionInfo,
+										   FProgressCancel* Progress) const
+{
+	int ElementCount = NormalOverlay->ElementCount();
+	OutNormals.SetNumZeroed(ElementCount);
+
+	EParallelForFlags ParallelForFlags = ExecutionInfo.bParallel ? EParallelForFlags::None : EParallelForFlags::ForceSingleThread;
+	ParallelForFlags |= EParallelForFlags::BackgroundPriority;
+
+	bool bCancelled = false;
+	auto InterpolationJob = [this, &OutNormals, &NormalOverlay, &LatticeControlPoints, &bCancelled, ExecutionInfo, Interpolation, Progress]
+	(int OverlayElementID)
+	{
+		// Every once in a while, check for cancellation
+		if ((OverlayElementID % ExecutionInfo.CancelCheckSize == 0) && Progress && Progress->Cancelled())
+		{
+			bCancelled = true;
+		}
+
+		int ParentVertexID = NormalOverlay->GetParentVertex(OverlayElementID);
+
+		if (bCancelled || VertexEmbeddings[ParentVertexID].LatticeCell[0] < 0)
+		{
+			return;
+		}
+
+		FMatrix3d Jacobian = (Interpolation == ELatticeInterpolation::Linear) ?
+			LinearInterpolationJacobian(VertexEmbeddings[ParentVertexID], LatticeControlPoints) :
+			CubicInterpolationJacobian(VertexEmbeddings[ParentVertexID], LatticeControlPoints);
+
+		// Typically we'd do transpose(inv(J)), however if a lattice cell inverts, the determinant is negative and the 
+		// resulting normal will be flipped. So we instead multiply by det(J)*transpose(inv(J)) to get the sign right.
+		FMatrix3d InvJacobian = Jacobian.DeterminantTimesInverseTranspose();
+		OutNormals[OverlayElementID] = FMatrix3f(InvJacobian) * NormalOverlay->GetElement(OverlayElementID);
+		OutNormals[OverlayElementID].Normalize();
+	};
+
+	ParallelFor(ElementCount, InterpolationJob, ParallelForFlags);
+}
+
+
+void FFFDLattice::GetRotatedMeshVertexNormals(const TArray<FVector3d>& LatticeControlPoints,
+											  const TArray<FVector3f>& OriginalNormals,
+											  TArray<FVector3f>& OutNormals,
+											  ELatticeInterpolation Interpolation,
+											  FLatticeExecutionInfo ExecutionInfo,
+											  FProgressCancel* Progress) const
+{
+	int MaxVertexID = OriginalNormals.Num();
+	OutNormals.SetNumZeroed(MaxVertexID);
+
+	EParallelForFlags ParallelForFlags = ExecutionInfo.bParallel ? EParallelForFlags::None : EParallelForFlags::ForceSingleThread;
+	ParallelForFlags |= EParallelForFlags::BackgroundPriority;
+
+	bool bCancelled = false;
+
+	check(VertexEmbeddings.Num() == OutNormals.Num());
+
+	auto InterpolationJob = [this, &OutNormals, &OriginalNormals, &LatticeControlPoints, &bCancelled, ExecutionInfo, MaxVertexID, Interpolation, Progress]
+	(int VertexID)
+	{
+		// Every once in a while, check for cancellation
+		if ((VertexID % ExecutionInfo.CancelCheckSize == 0) && Progress && Progress->Cancelled())
+		{
+			bCancelled = true;
+		}
+
+		if (bCancelled || VertexEmbeddings[VertexID].LatticeCell[0] < 0)
+		{
+			return;
+		}
+
+		FMatrix3d Jacobian = (Interpolation == ELatticeInterpolation::Linear) ?
+			LinearInterpolationJacobian(VertexEmbeddings[VertexID], LatticeControlPoints) :
+			CubicInterpolationJacobian(VertexEmbeddings[VertexID], LatticeControlPoints);
+
+		// Typically we'd do transpose(inv(J)), however if a lattice cell inverts, the determinant is negative and the 
+		// resulting normal will be flipped. So we instead multiply by det(J)*transpose(inv(J)) to get the sign right.
+		FMatrix3d InvJacobian = Jacobian.DeterminantTimesInverseTranspose();
+		OutNormals[VertexID] = FMatrix3f( InvJacobian ) * OriginalNormals[VertexID];
+		OutNormals[VertexID].Normalize();
+	};
+
+	ParallelFor(MaxVertexID, InterpolationJob, ParallelForFlags);
+}
+
 
 void FFFDLattice::GetValuePair(int I, int J, int K, FVector3d& A, FVector3d& B, 
 							   const TArray<FVector3d>& LatticeControlPoints) const
@@ -206,6 +296,29 @@ static float CubicBSplineKernel(float A)
 }
 
 
+static float CubicBSplineKernelDerivative(float A)
+{
+	// f'(a) =
+	//	a/2*(3|a| - 4)		            for 0 <= |a| < 1
+	//	-(a*(2-|a|)^2) / (2*|a|) 		for 1 <= |a| < 2
+	//	0								otherwise
+
+	float AbsA = FMath::Abs(A);
+	if (AbsA < 1.0f)
+	{
+		return A / 2.0f * (3.0f * AbsA - 4.0f);
+	}
+	else if (AbsA < 2.0f)
+	{
+		float TwoMinusAbsA = (2.0f - AbsA);
+		return -(A * TwoMinusAbsA * TwoMinusAbsA) / (2.0f * AbsA);
+	}
+
+	return 0.0f;
+}
+
+
+
 FVector3d FFFDLattice::InterpolatedPositionCubic(const FEmbedding& VertexEmbedding, 
 												 const TArray<FVector3d>& LatticeControlPoints) const
 {
@@ -307,6 +420,89 @@ FVector3d FFFDLattice::InterpolatedPosition(const FEmbedding& VertexEmbedding, c
 	return Sum;
 }
 
+FMatrix3d FFFDLattice::LinearInterpolationJacobian(const FEmbedding& VertexEmbedding, const TArray<FVector3d>& LatticeControlPoints) const
+{
+	// TODO: This was written for clarity and correctness. Could definitely be faster.
+
+	// Partial wrt x
+	FEmbedding XFloor = VertexEmbedding;
+	XFloor.CellWeighting.X = 0.0;
+	FEmbedding XCeil = VertexEmbedding;
+	XCeil.CellWeighting.X = 1.0;
+	FVector3d PartialX = InterpolatedPosition(XCeil, LatticeControlPoints) - InterpolatedPosition(XFloor, LatticeControlPoints);
+
+	// Partial wrt y
+	FEmbedding YFloor = VertexEmbedding;
+	YFloor.CellWeighting.Y = 0.0;
+	FEmbedding YCeil = VertexEmbedding;
+	YCeil.CellWeighting.Y = 1.0;
+	FVector3d PartialY = InterpolatedPosition(YCeil, LatticeControlPoints) - InterpolatedPosition(YFloor, LatticeControlPoints);
+
+	// Partial wrt z
+	FEmbedding ZFloor = VertexEmbedding;
+	ZFloor.CellWeighting.Z = 0.0;
+	FEmbedding ZCeil = VertexEmbedding;
+	ZCeil.CellWeighting.Z = 1.0;
+	FVector3d PartialZ = InterpolatedPosition(ZCeil, LatticeControlPoints) - InterpolatedPosition(ZFloor, LatticeControlPoints);
+
+	FMatrix3d Mat(PartialX, PartialY, PartialZ, false);
+
+	return Mat;
+}
+
+FMatrix3d FFFDLattice::CubicInterpolationJacobian(const FEmbedding& VertexEmbedding, const TArray<FVector3d>& LatticeControlPoints) const
+{
+	// TODO: This was written for clarity and correctness. Could definitely be faster.
+
+	double T = VertexEmbedding.CellWeighting.X;
+	double U = VertexEmbedding.CellWeighting.Y;
+	double V = VertexEmbedding.CellWeighting.Z;
+
+	FMatrix3d Sum = FMatrix3d::Zero();
+
+	for (int DI = -1; DI <= 2; ++DI)
+	{
+		double WeightX = CubicBSplineKernel(T - DI);
+		double DWeightXDX = CubicBSplineKernelDerivative(T - DI);
+
+		for (int DJ = -1; DJ <= 2; ++DJ)
+		{
+			double WeightY = CubicBSplineKernel(U - DJ);
+			double DWeightYDY = CubicBSplineKernelDerivative(U - DJ);
+
+			for (int DK = -1; DK <= 2; ++DK)
+			{
+				double WeightZ = CubicBSplineKernel(V - DK);
+				double DWeightZDZ = CubicBSplineKernelDerivative(V - DK);
+
+				double DWeightDX = DWeightXDX * WeightY * WeightZ;
+				double DWeightDY = WeightX * DWeightYDY * WeightZ;
+				double DWeightDZ = WeightX * WeightY * DWeightZDZ;
+
+				int i = VertexEmbedding.LatticeCell.X + DI;
+				int j = VertexEmbedding.LatticeCell.Y + DJ;
+				int k = VertexEmbedding.LatticeCell.Z + DK;
+
+				FVector3d LatticePoint;
+				if (i < 0 || i >= Dims.X || j < 0 || j >= Dims.Y || k < 0 || k >= Dims.Z)
+				{
+					// Get the extrapolated position for a "virtual" control point outside of the deformed lattice
+					LatticePoint = ExtrapolatedLatticePosition({ i,j,k }, LatticeControlPoints);
+				}
+				else
+				{
+					int PointIndex = ControlPointIndexFromCoordinates(i, j, k);
+					LatticePoint = LatticeControlPoints[PointIndex];
+				}
+
+				Sum += FMatrix3d(DWeightDX * LatticePoint, DWeightDY * LatticePoint, DWeightDZ * LatticePoint, false);
+			}
+		}
+	}
+
+	return Sum;
+}
+
 FVector3d FFFDLattice::ClosestLatticePosition(const FVector3i& VirtualControlPointIndex, const TArray<FVector3d>& LatticeControlPoints) const
 {
 	// Clamp to valid lattice index
diff --git a/Engine/Plugins/Experimental/GeometryProcessing/Source/DynamicMesh/Public/Operations/FFDLattice.h b/Engine/Plugins/Experimental/GeometryProcessing/Source/DynamicMesh/Public/Operations/FFDLattice.h
--- a/Engine/Plugins/Experimental/GeometryProcessing/Source/DynamicMesh/Public/Operations/FFDLattice.h
+++ b/Engine/Plugins/Experimental/GeometryProcessing/Source/DynamicMesh/Public/Operations/FFDLattice.h
@@ -5,6 +5,7 @@
 #include "VectorTypes.h"
 #include "BoxTypes.h"
 #include "Containers/StaticArray.h"
+#include "DynamicMeshAttributeSet.h"
 
 class FDynamicMesh3;
 class FProgressCancel;
@@ -49,6 +50,24 @@ public:
 										FLatticeExecutionInfo ExecutionInfo = FLatticeExecutionInfo(),
 										FProgressCancel* Progress = nullptr) const;
 
+	/// Using the Lattice's current control point positions, the original embedding information, and the original 
+	/// normals contained in an overlay, compute the new rotated normals
+	void GetRotatedOverlayNormals(const TArray<FVector3d>& LatticeControlPoints,
+								  const FDynamicMeshNormalOverlay* NormalOverlay,
+								  TArray<FVector3f>& OutNormals,
+								  ELatticeInterpolation Interpolation = ELatticeInterpolation::Linear,
+								  FLatticeExecutionInfo ExecutionInfo = FLatticeExecutionInfo(),
+								  FProgressCancel* Progress = nullptr) const;
+
+	/// Using the Lattice's current control point positions, the original embedding information, and the original 
+	/// normals, compute the new rotated normals
+	void GetRotatedMeshVertexNormals(const TArray<FVector3d>& LatticeControlPoints,
+									 const TArray<FVector3f>& OriginalNormals,
+									 TArray<FVector3f>& OutNormals,
+									 ELatticeInterpolation Interpolation,
+									 FLatticeExecutionInfo ExecutionInfo,
+									 FProgressCancel* Progress) const;
+
 	void GenerateInitialLatticePositions(TArray<FVector3d>& OutLatticePositions) const;
 	void GenerateLatticeEdges(TArray<FVector2i>& OutLatticeEdges) const;
 
@@ -94,6 +113,9 @@ protected:
 	FVector3d InterpolatedPosition(const FEmbedding& VertexEmbedding, const TArray<FVector3d>& LatticeControlPoints) const;
 	FVector3d InterpolatedPositionCubic(const FEmbedding& VertexEmbedding, const TArray<FVector3d>& LatticeControlPoints) const;
 
+	FMatrix3d LinearInterpolationJacobian(const FEmbedding& VertexEmbedding, const TArray<FVector3d>& LatticeControlPoints) const;
+	FMatrix3d CubicInterpolationJacobian(const FEmbedding& VertexEmbedding, const TArray<FVector3d>& LatticeControlPoints) const;
+
 	/// Clamp the given index to [0, Dims] and return the current lattice point position at the clamped index
 	FVector3d ClosestLatticePosition(const FVector3i& VirtualControlPointIndex, const TArray<FVector3d>& LatticeControlPoints) const;
 
diff --git a/Engine/Plugins/Experimental/GeometryProcessing/Source/GeometricObjects/Public/MatrixTypes.h b/Engine/Plugins/Experimental/GeometryProcessing/Source/GeometricObjects/Public/MatrixTypes.h
--- a/Engine/Plugins/Experimental/GeometryProcessing/Source/GeometricObjects/Public/MatrixTypes.h
+++ b/Engine/Plugins/Experimental/GeometryProcessing/Source/GeometricObjects/Public/MatrixTypes.h
@@ -64,6 +64,14 @@ struct TMatrix3
 		}
 	}
 
+	template<typename RealType2>
+	explicit constexpr TMatrix3(const TMatrix3<RealType2>& Mat) :
+		Row0(FVector3<RealType>{Mat.Row0}),
+		Row1(FVector3<RealType>{Mat.Row1}),
+		Row2(FVector3<RealType>{Mat.Row2})
+	{
+	}
+
 	static TMatrix3<RealType> Zero()
 	{
 		return TMatrix3<RealType>(0);
@@ -197,6 +205,30 @@ struct TMatrix3
 			Row0.Z, Row1.Z, Row2.Z);
 	}
 
+	// Computes |A|(A^-1)^T. 
+	// This value is sometimes useful for computing rotated surface normals from a continuous deformation.
+	// Since computing the 3x3 inverse using Cramer's rule divides by the 3x3 determinant in the final step, this 
+	// function avoids computing the determinant at all, since it would just cancel out.
+	TMatrix3<RealType> DeterminantTimesInverseTranspose() const
+	{
+		RealType a11 = Row0.X, a12 = Row0.Y, a13 = Row0.Z, a21 = Row1.X, a22 = Row1.Y, a23 = Row1.Z, a31 = Row2.X, a32 = Row2.Y, a33 = Row2.Z;
+		
+		RealType i00 = a33 * a22 - a32 * a23;
+		RealType i01 = -(a33 * a12 - a32 * a13);
+		RealType i02 = a23 * a12 - a22 * a13;
+
+		RealType i10 = -(a33 * a21 - a31 * a23);
+		RealType i11 = a33 * a11 - a31 * a13;
+		RealType i12 = -(a23 * a11 - a21 * a13);
+
+		RealType i20 = a32 * a21 - a31 * a22;
+		RealType i21 = -(a32 * a11 - a31 * a12);
+		RealType i22 = a22 * a11 - a21 * a12;
+
+		return TMatrix3<RealType>(i00, i10, i20, i01, i11, i21, i02, i12, i22);
+	}
+
+
 	bool EpsilonEqual(const TMatrix3<RealType>& Mat2, RealType Epsilon) const
 	{
 		return VectorUtil::EpsilonEqual(Row0, Mat2.Row0, Epsilon) &&
diff --git a/Engine/Plugins/Experimental/MeshModelingToolset/Source/MeshModelingTools/Private/LatticeDeformerTool.cpp b/Engine/Plugins/Experimental/MeshModelingToolset/Source/MeshModelingTools/Private/LatticeDeformerTool.cpp
--- a/Engine/Plugins/Experimental/MeshModelingToolset/Source/MeshModelingTools/Private/LatticeDeformerTool.cpp
+++ b/Engine/Plugins/Experimental/MeshModelingToolset/Source/MeshModelingTools/Private/LatticeDeformerTool.cpp
@@ -56,7 +56,8 @@ TUniquePtr<FDynamicMeshOperator> ULatticeDeformerOperatorFactory::MakeNewOperato
 		LatticeDeformerTool->OriginalMesh,
 		LatticeDeformerTool->Lattice,
 		LatticeDeformerTool->ControlPointsMechanic->GetControlPoints(),
-		OpInterpolationType);
+		OpInterpolationType,
+		LatticeDeformerTool->Settings->bDeformNormals);
 
 	return LatticeDeformOp;
 }
@@ -119,6 +120,10 @@ void ULatticeDeformerTool::Setup()
 	{
 		Preview->InvalidateResult();
 	});
+	Settings->WatchProperty(Settings->bDeformNormals, [this](bool)
+	{
+		Preview->InvalidateResult();
+	});
 	Settings->WatchProperty(Settings->GizmoCoordinateSystem, [this](EToolContextCoordinateSystem)
 	{
 		ControlPointsMechanic->SetCoordinateSystem(Settings->GizmoCoordinateSystem);
diff --git a/Engine/Plugins/Experimental/MeshModelingToolset/Source/MeshModelingTools/Public/LatticeDeformerTool.h b/Engine/Plugins/Experimental/MeshModelingToolset/Source/MeshModelingTools/Public/LatticeDeformerTool.h
--- a/Engine/Plugins/Experimental/MeshModelingToolset/Source/MeshModelingTools/Public/LatticeDeformerTool.h
+++ b/Engine/Plugins/Experimental/MeshModelingToolset/Source/MeshModelingTools/Public/LatticeDeformerTool.h
@@ -62,6 +62,10 @@ public:
 	UPROPERTY(EditAnywhere, Category = Interpolation )
 	ELatticeInterpolationType InterpolationType = ELatticeInterpolationType::Linear;
 
+	/** Whether to use approximate new vertex normals using the deformer */
+	UPROPERTY(EditAnywhere, Category = Interpolation)
+	bool bDeformNormals = false;
+
 	// Not user visible - used to disallow changing the lattice resolution after deformation
 	UPROPERTY(meta = (TransientToolProperty))
 	bool bCanChangeResolution = true;
diff --git a/Engine/Plugins/Experimental/MeshModelingToolset/Source/ModelingOperators/Private/DeformationOps/LatticeDeformerOp.cpp b/Engine/Plugins/Experimental/MeshModelingToolset/Source/ModelingOperators/Private/DeformationOps/LatticeDeformerOp.cpp
--- a/Engine/Plugins/Experimental/MeshModelingToolset/Source/ModelingOperators/Private/DeformationOps/LatticeDeformerOp.cpp
+++ b/Engine/Plugins/Experimental/MeshModelingToolset/Source/ModelingOperators/Private/DeformationOps/LatticeDeformerOp.cpp
@@ -21,21 +21,81 @@ void FLatticeDeformerOp::CalculateResult(FProgressCancel* Progress)
 	FLatticeExecutionInfo ExecutionInfo = FLatticeExecutionInfo();
 	ExecutionInfo.bParallel = true;
 	Lattice->GetDeformedMeshVertexPositions(LatticeControlPoints, DeformedPositions, InterpolationType, ExecutionInfo, Progress);
-
 	check(ResultMesh->VertexCount() == DeformedPositions.Num());
 
+	if (Progress && Progress->Cancelled())
+	{
+		return;
+	}
+
 	for (int vid : ResultMesh->VertexIndicesItr())
 	{
 		ResultMesh->SetVertex(vid, DeformedPositions[vid]);
 	}
+
+	if (bDeformNormals)
+	{
+		if (ResultMesh->HasAttributes())
+		{
+			FDynamicMeshNormalOverlay* NormalOverlay = ResultMesh->Attributes()->PrimaryNormals();
+			check(NormalOverlay != nullptr);
+
+			TArray<FVector3f> DeformedNormals;
+			Lattice->GetRotatedOverlayNormals(LatticeControlPoints,
+											  NormalOverlay,
+											  DeformedNormals,
+											  InterpolationType,
+											  ExecutionInfo,
+											  Progress);
+
+			if (Progress && Progress->Cancelled())
+			{
+				return;
+			}
+
+			for (int ElementID : NormalOverlay->ElementIndicesItr())
+			{
+				NormalOverlay->SetElement(ElementID, DeformedNormals[ElementID]);
+			}
+		}
+		else if (ResultMesh->HasVertexNormals())
+		{
+			TArray<FVector3f> OriginalNormals;
+			OriginalNormals.SetNum(ResultMesh->MaxVertexID());
+			for (int VertexID : ResultMesh->VertexIndicesItr())
+			{
+				OriginalNormals[VertexID] = ResultMesh->GetVertexNormal(VertexID);
+			}
+
+			TArray<FVector3f> RotatedNormals;
+			Lattice->GetRotatedMeshVertexNormals(LatticeControlPoints,
+												 OriginalNormals,
+												 RotatedNormals,
+												 InterpolationType,
+												 ExecutionInfo,
+												 Progress);
+
+			if (Progress && Progress->Cancelled())
+			{
+				return;
+			}
+
+			for (int vid : ResultMesh->VertexIndicesItr())
+			{
+				ResultMesh->SetVertexNormal(vid, RotatedNormals[vid]);
+			}
+		}
+	}
 }
 
 FLatticeDeformerOp::FLatticeDeformerOp(TSharedPtr<FDynamicMesh3, ESPMode::ThreadSafe> InOriginalMesh,
 									   TSharedPtr<FFFDLattice, ESPMode::ThreadSafe> InLattice,
 									   const TArray<FVector3d>& InLatticeControlPoints,
-									   ELatticeInterpolation InInterpolationType) :
+									   ELatticeInterpolation InInterpolationType,
+									   bool bInDeformNormals) :
 	Lattice(InLattice),
 	OriginalMesh(InOriginalMesh),
 	LatticeControlPoints(InLatticeControlPoints),
-	InterpolationType(InInterpolationType)
+	InterpolationType(InInterpolationType),
+	bDeformNormals(bInDeformNormals)
 {}
diff --git a/Engine/Plugins/Experimental/MeshModelingToolset/Source/ModelingOperators/Public/DeformationOps/LatticeDeformerOp.h b/Engine/Plugins/Experimental/MeshModelingToolset/Source/ModelingOperators/Public/DeformationOps/LatticeDeformerOp.h
--- a/Engine/Plugins/Experimental/MeshModelingToolset/Source/ModelingOperators/Public/DeformationOps/LatticeDeformerOp.h
+++ b/Engine/Plugins/Experimental/MeshModelingToolset/Source/ModelingOperators/Public/DeformationOps/LatticeDeformerOp.h
@@ -19,7 +19,8 @@ public:
 	FLatticeDeformerOp(TSharedPtr<FDynamicMesh3, ESPMode::ThreadSafe> InOriginalMesh,
 					   TSharedPtr<FFFDLattice, ESPMode::ThreadSafe> InLattice,
 					   const TArray<FVector3d>& InLatticeControlPoints,
-					   ELatticeInterpolation InInterpolationType);
+					   ELatticeInterpolation InInterpolationType,
+					   bool bInDeformNormals);
 
 	// FDynamicMeshOperator implementation
 	void CalculateResult(FProgressCancel* Progress) override;
@@ -31,5 +32,6 @@ protected:
 	const TSharedPtr<const FDynamicMesh3, ESPMode::ThreadSafe> OriginalMesh;
 	const TArray<FVector3d> LatticeControlPoints;
 	ELatticeInterpolation InterpolationType;
+	bool bDeformNormals = false;
 };
 
