GeometryProcessing: convert debugging sanity check to checkSlow to avoid runtime crashes
#rnx

diff --git a/Engine/Plugins/Experimental/GeometryProcessing/Source/DynamicMesh/Private/Operations/GroupTopologyDeformer.cpp b/Engine/Plugins/Experimental/GeometryProcessing/Source/DynamicMesh/Private/Operations/GroupTopologyDeformer.cpp
--- a/Engine/Plugins/Experimental/GeometryProcessing/Source/DynamicMesh/Private/Operations/GroupTopologyDeformer.cpp
+++ b/Engine/Plugins/Experimental/GeometryProcessing/Source/DynamicMesh/Private/Operations/GroupTopologyDeformer.cpp
@@ -278,7 +278,7 @@ void FGroupTopologyDeformer::ComputeEncoding()
 				VtxEncoding.Weights[k] /= WeightSum;
 				Reconstruct += VtxEncoding.Weights[k] * (Mesh->GetVertex(Face.BoundaryVerts[k]) + VtxEncoding.Deltas[k]);
 			}
-			check(Reconstruct.Distance(Pos) < 0.0001);
+			checkSlow(Reconstruct.Distance(Pos) < 0.0001);   // sanity check
 		}
 	}
 }
