skip seeds that are outside the marching cubes grid

diff --git a/Engine/Plugins/Experimental/GeometryProcessing/Source/GeometricObjects/Public/Generators/MarchingCubes.h b/Engine/Plugins/Experimental/GeometryProcessing/Source/GeometricObjects/Public/Generators/MarchingCubes.h
--- a/Engine/Plugins/Experimental/GeometryProcessing/Source/GeometricObjects/Public/Generators/MarchingCubes.h
+++ b/Engine/Plugins/Experimental/GeometryProcessing/Source/GeometricObjects/Public/Generators/MarchingCubes.h
@@ -571,7 +571,7 @@ protected:
 		for (FVector3<double> seed : Seeds)
 		{
 			FVector3i seed_idx = cell_index(seed);
-			if (done_cells[seed_idx] == 1)
+			if (!done_cells.IsValidIndex(seed_idx) || done_cells[seed_idx] == 1)
 			{
 				continue;
 			}
@@ -614,11 +614,11 @@ protected:
 	{
 		parallel_mesh_access = true;
 
-		ParallelFor(Seeds.Num(), [&](int32 Index)
+		ParallelFor(Seeds.Num(), [this, &Seeds](int32 Index)
 		{
 			FVector3<double> Seed = Seeds[Index];
 			FVector3i seed_idx = cell_index(Seed);
-			if (set_cell_if_not_done(seed_idx) == false)
+			if (!done_cells.IsValidIndex(seed_idx) || set_cell_if_not_done(seed_idx) == false)
 			{
 				return;
 			}
diff --git a/Engine/Plugins/Experimental/GeometryProcessing/Source/GeometricObjects/Public/Spatial/DenseGrid3.h b/Engine/Plugins/Experimental/GeometryProcessing/Source/GeometricObjects/Public/Spatial/DenseGrid3.h
--- a/Engine/Plugins/Experimental/GeometryProcessing/Source/GeometricObjects/Public/Spatial/DenseGrid3.h
+++ b/Engine/Plugins/Experimental/GeometryProcessing/Source/GeometricObjects/Public/Spatial/DenseGrid3.h
@@ -42,6 +42,12 @@ public:
 		return Dimensions.X * Dimensions.Y * Dimensions.Z;
 	}
 
+	bool IsValidIndex(const FVector3i& Index) const
+	{
+		return Index.X >= 0 && Index.Y >= 0 && Index.Z >= 0
+			&& Index.X < Dimensions.X && Index.Y < Dimensions.Y && Index.Z < Dimensions.Z;
+	}
+
 	const FVector3i& GetDimensions() const
 	{
 		return Dimensions;
