Geometry Tools: Mesh Simplification - Fixes to Mode that retains vertex locations.
                           This included two fixes.
                           1) a bug that would choose the average position when collapsing an edge on a mesh boundary ( so not retaining a vertex location )
                           2) eliminates the chance that an edge adjacent to a UV seam (or other attribute seam) will collapse to the vertex not on the seam (this would have moved the seam and resulted in new attribute values).

diff --git a/Engine/Plugins/Experimental/GeometryProcessing/Source/DynamicMesh/Private/DynamicMeshOverlay.cpp b/Engine/Plugins/Experimental/GeometryProcessing/Source/DynamicMesh/Private/DynamicMeshOverlay.cpp
--- a/Engine/Plugins/Experimental/GeometryProcessing/Source/DynamicMesh/Private/DynamicMeshOverlay.cpp
+++ b/Engine/Plugins/Experimental/GeometryProcessing/Source/DynamicMesh/Private/DynamicMeshOverlay.cpp
@@ -443,11 +443,26 @@ void TDynamicMeshOverlay<RealType, ElementSize>::InitializeNewTriangle(int tid)
 
 
 template<typename RealType, int ElementSize>
-bool TDynamicMeshOverlay<RealType,ElementSize>::IsSeamEdge(int eid) const
+bool TDynamicMeshOverlay<RealType,ElementSize>::IsSeamEdge(int eid, bool* bIsNonIntersecting) const
 {
+	if (bIsNonIntersecting != nullptr)
+	{
+		*bIsNonIntersecting = false;
+	}
+
 	FIndex2i et = ParentMesh->GetEdgeT(eid);
 	if (et.B == FDynamicMesh3::InvalidID)
 	{
+		if (bIsNonIntersecting != nullptr)
+		{
+			FIndex2i ev = ParentMesh->GetEdgeV(eid);
+			int CountA = CountVertexElements(ev.A);
+			int CountB = CountVertexElements(ev.B);
+			
+			// will be false if another seam intersects is adjacent to either end of the seam edge
+			*bIsNonIntersecting = (CountA == 1) && (CountB == 1);
+		}
+
 		return true;
 	}
 
@@ -462,15 +477,41 @@ bool TDynamicMeshOverlay<RealType,ElementSize>::IsSeamEdge(int eid) const
 
 	FIndex3i Triangle0 = GetTriangle(et.A);
 	FIndex3i BaseTriangle0(ParentVertices[Triangle0.A], ParentVertices[Triangle0.B], ParentVertices[Triangle0.C]);
-	int idx_base_a1 = BaseTriangle0.IndexOf(base_a);
-	int idx_base_b1 = BaseTriangle0.IndexOf(base_b);
+	int idx_base_a0 = BaseTriangle0.IndexOf(base_a);
+	int idx_base_b0 = BaseTriangle0.IndexOf(base_b);
 
 	FIndex3i Triangle1 = GetTriangle(et.B);
 	FIndex3i BaseTriangle1(ParentVertices[Triangle1.A], ParentVertices[Triangle1.B], ParentVertices[Triangle1.C]);
-	int idx_base_a2 = BaseTriangle1.IndexOf(base_a);
-	int idx_base_b2 = BaseTriangle1.IndexOf(base_b);
+	int idx_base_a1 = BaseTriangle1.IndexOf(base_a);
+	int idx_base_b1 = BaseTriangle1.IndexOf(base_b);
+
+	int el_a_tri0 = Triangle0[idx_base_a0];
+	int el_b_tri0 = Triangle0[idx_base_b0];
+	int el_a_tri1 = Triangle1[idx_base_a1];
+	int el_b_tri1 = Triangle1[idx_base_b1];
 
-	return !IndexUtil::SamePairUnordered(Triangle0[idx_base_a1], Triangle0[idx_base_b1], Triangle1[idx_base_a2], Triangle1[idx_base_b2]);
+	bool bIsSeam = !IndexUtil::SamePairUnordered(el_a_tri0, el_b_tri0, el_a_tri1, el_b_tri1);
+
+	if (bIsNonIntersecting != nullptr)
+	{
+		if ((el_a_tri0 == el_a_tri1 || el_a_tri0 == el_b_tri1) || (el_b_tri0 == el_b_tri1 || el_b_tri0 == el_a_tri1))
+		{
+			// seam edge "intersects" with end of the seam
+			*bIsNonIntersecting = false;
+		}
+		else
+		{
+			// check that exactly two elements are associated with the vertices at each end of the edge
+			int CountA = CountVertexElements(base_a);
+			int CountB = CountVertexElements(base_b);
+
+			// will be false if another seam intersects is adjacent to either end of the seam edge
+			*bIsNonIntersecting = (CountA == 2) && (CountB == 2);
+		}
+	}
+	
+	return bIsSeam;
+	
 
 
 	// TODO: this doesn't seem to work but it should, and would be more efficient:
diff --git a/Engine/Plugins/Experimental/GeometryProcessing/Source/DynamicMesh/Private/MeshSimplification.cpp b/Engine/Plugins/Experimental/GeometryProcessing/Source/DynamicMesh/Private/MeshSimplification.cpp
--- a/Engine/Plugins/Experimental/GeometryProcessing/Source/DynamicMesh/Private/MeshSimplification.cpp
+++ b/Engine/Plugins/Experimental/GeometryProcessing/Source/DynamicMesh/Private/MeshSimplification.cpp
@@ -269,7 +269,11 @@ FVector3d TMeshSimplification<QuadricErrorType>::OptimalPoint(int eid, const FQu
 	{
 		if (Mesh->IsBoundaryEdge(eid))
 		{
-			return (Mesh->GetVertex(ea) + Mesh->GetVertex(eb)) * 0.5;
+			const bool bModeAllowsVertMovement = (CollapseMode != ESimplificationCollapseModes::MinimalExistingVertexError);
+			if (bModeAllowsVertMovement)
+			{ 
+				return (Mesh->GetVertex(ea) + Mesh->GetVertex(eb)) * 0.5;
+			} // else MinimalExistingVertexError case below will choose one of the vertex locations
 		}
 		else
 		{
@@ -715,7 +719,7 @@ void TMeshSimplification<QuadricErrorType>::DoSimplify()
 	{
 		return;
 	}
-
+	// [TODO] - consider, skip this when CollapseMode == ESimplificationCollapseModes::MinimalExistingVertexError ? 
 	Reproject();
 
 	ProfileEndPass();
@@ -1117,11 +1121,96 @@ void TMeshSimplification<QuadricErrorType>::FastCollapsePass(double fMinEdgeLeng
 
 
 
+template <typename QuadricErrorType>
+bool TMeshSimplification<QuadricErrorType>::CanCollapseEdge(int edgeID, int a, int b, int c, int d, int t0, int t1, int& collapse_to) const
+{
+	const bool bPreserveSeamTopology = (CollapseMode == ESimplificationCollapseModes::MinimalExistingVertexError);
 
+	if (bAllowSeamCollapse && !bPreserveSeamTopology)
+	{
+		return CanCollapseVertex(edgeID, a, b, collapse_to);
+	}
 
+	// make sure that the retained edges from the collapsed triangles don't merge seams
+	bool bCanCollapse = FMeshRefinerBase::CanCollapseEdge(edgeID, a, b, c, d, t0, t1, collapse_to);
 
+	// make sure more general seam topology is preserved
+	if (bCanCollapse && bAllowSeamCollapse && bPreserveSeamTopology)
+	{
+		if (!Constraints)
+		{
+			return bCanCollapse;
+		}
+
+		// NB: We have to be more restrictive with the MinimalQuadricPositionError mode
+		// in order to preclude the possibility of a seam moving during collapse.
 
+		// check if this edge is a seam
+		if (Constraints->HasEdgeConstraint(edgeID))
+		{
+			// examine local topology
+			bool bCanCollapseSeam = true;
+			if (const FDynamicMeshAttributeSet* Attributes = Mesh->Attributes())
+			{
+
+				for (int i = 0; bCanCollapseSeam && i < Attributes->NumUVLayers(); ++i)
+				{
+					auto* Overlay = Attributes->GetUVLayer(i);
+					bool bIsNonIntersecting;
+					if (Overlay->IsSeamEdge(edgeID, &bIsNonIntersecting))
+					{
+						bCanCollapseSeam = bCanCollapseSeam && bIsNonIntersecting;
+					}
+				}
+				for (int i = 0; bCanCollapseSeam && i < Attributes->NumNormalLayers(); ++i)
+				{
+					auto* Overlay = Attributes->GetNormalLayer(i);
+					bool bIsNonIntersecting;
+					if (Overlay->IsSeamEdge(edgeID, &bIsNonIntersecting))
+					{
+						bCanCollapseSeam = bCanCollapseSeam && bIsNonIntersecting;
+					}
+				}
+			}
+			bCanCollapse = bCanCollapseSeam;
+		}
+		else
+		{
+			// this edge was not a seam, but need to check if one or both ends are part of other seams 
+			// - this is done by checking for vertex constraint
+			bool bVertexAOnSeam = Constraints->HasVertexConstraint(a);
+			bool bVertexBOnSeam = Constraints->HasVertexConstraint(b);
+			if (bVertexAOnSeam && bVertexBOnSeam)
+			{
+				bCanCollapse = false;
+			}
+			else if (bVertexAOnSeam)
+			{
+				if (collapse_to == -1)
+				{
+					collapse_to = a;
+				}
+				else if (collapse_to != a)
+				{
+					bCanCollapse = false;
+				}
+			}
+			else if (bVertexBOnSeam)
+			{
+				if (collapse_to == -1)
+				{
+					collapse_to = b;
+				}
+				else if (collapse_to != b)
+				{
+					bCanCollapse = false;
+				}
+			}
+		}
+	}
 
+	return bCanCollapse;
+}
 
 
 
@@ -1171,21 +1260,14 @@ ESimplificationResult TMeshSimplification<QuadricErrorType>::CollapseEdge(int ed
 	// check if we should collapse, and also find which vertex we should retain
 	// in cases where we have constraints/etc
 	int collapse_to = -1;
-	bool bCanCollapse = false;
-	if (!bAllowSeamCollapse)
-	{
-		bCanCollapse = CanCollapseEdge(edgeID, a, b, c, d, t0, t1, collapse_to);
-	}
-	else
-	{
-		bCanCollapse = CanCollapseVertex(edgeID, a, b, collapse_to);
-	}
+	bool bCanCollapse = CanCollapseEdge(edgeID, a, b, c, d, t0, t1, collapse_to);
 
 	if (bCanCollapse == false)
 	{
 		return ESimplificationResult::Ignored_Constrained;
 	}
 
+	const bool bMinimalVertexMode = (CollapseMode == ESimplificationCollapseModes::MinimalExistingVertexError);
 	// if we have a boundary, we want to collapse to boundary
 	if (bPreserveBoundaryShape && bHaveBoundary)
 	{
@@ -1197,19 +1279,23 @@ ESimplificationResult TMeshSimplification<QuadricErrorType>::CollapseEdge(int ed
 				return ESimplificationResult::Ignored_Constrained;
 			}
 		}
-		if (IsBoundaryVertex(b))
-		{
-			collapse_to = b;
-		}
-		else if (IsBoundaryVertex(a))
+
+		if (!bMinimalVertexMode) // the minimal existing vertex error has already resolved this with more complicated logic
 		{
-			collapse_to = a;
+			if (IsBoundaryVertex(b))
+			{
+				collapse_to = b;
+			}
+			else if (IsBoundaryVertex(a))
+			{
+				collapse_to = a;
+			}
 		}
 	}
 
 	if (RequireKeepVert == a || RequireKeepVert == b)
 	{
-		if (collapse_to >= 0 && collapse_to != RequireKeepVert)
+		if (collapse_to != -1 && collapse_to != RequireKeepVert)
 		{
 			return ESimplificationResult::Ignored_Constrained;
 		}
@@ -1222,26 +1308,35 @@ ESimplificationResult TMeshSimplification<QuadricErrorType>::CollapseEdge(int ed
 	ESimplificationResult retVal = ESimplificationResult::Failed_OpNotSuccessful;
 
 	int iKeep = b, iCollapse = a;
-	bool bCanMove = true;
-	
+	bool bConstraintsSpecifyPosition = false;
 	if (collapse_to != -1)
 	{
 		iKeep = collapse_to;
 		iCollapse = (iKeep == a) ? b : a;
 
-		// if constraints require a fixed position
+		// if constraints or collapse mode require a fixed position
 		if (Constraints)
 		{
-			bCanMove = Constraints->GetVertexConstraint(collapse_to).bCanMove;
+			bConstraintsSpecifyPosition = bMinimalVertexMode || (!Constraints->GetVertexConstraint(collapse_to).bCanMove );
 		}
 	}
 	double collapse_t = 0;
-	if (bCanMove)
+	if (!bConstraintsSpecifyPosition)
 	{
+		checkSlow(!bMinimalVertexMode || (vNewPos == vA || vNewPos == vB));
+		
+		// [TODO] maybe skip Projection call when !bModeAllowsVertMovment.
 		vNewPos = GetProjectedCollapsePosition(iKeep, vNewPos);
 		double div = vA.Distance(vB);
 		collapse_t = (div < FMathd::ZeroTolerance) ? 0.5 : (vNewPos.Distance(Mesh->GetVertex(iKeep))) / div;
 		collapse_t = VectorUtil::Clamp(collapse_t, 0.0, 1.0);
+
+		if (bMinimalVertexMode)
+		{
+			// this _should_ already be 0 or 1 with perfect precision. 
+			// round here to make sure later attribute lerps don't change values
+			collapse_t = FMath::RoundToDouble(collapse_t);
+		}
 	}
 	else
 	{
diff --git a/Engine/Plugins/Experimental/GeometryProcessing/Source/DynamicMesh/Public/DynamicMeshOverlay.h b/Engine/Plugins/Experimental/GeometryProcessing/Source/DynamicMesh/Public/DynamicMeshOverlay.h
--- a/Engine/Plugins/Experimental/GeometryProcessing/Source/DynamicMesh/Public/DynamicMeshOverlay.h
+++ b/Engine/Plugins/Experimental/GeometryProcessing/Source/DynamicMesh/Public/DynamicMeshOverlay.h
@@ -385,8 +385,11 @@ public:
 	}
 
 
-	/** Returns true if the parent-mesh edge is a "Seam" in this overlay */
-	bool IsSeamEdge(int EdgeID) const;
+	/** Returns true if the parent-mesh edge is a "Seam" in this overlay.  
+	*   If present, bIsNonIntersectingOut will be true only if this is a seam edge 
+	*   that does not intersect with another seam or the end of the seam. 
+	*/
+	bool IsSeamEdge(int EdgeID, bool* bIsNonIntersectingOut = nullptr) const;
 	/** Returns true if the parent-mesh edge is a "Seam End" in this overlay, meaning the adjacent element triangles share one element, not two */
 	bool IsSeamEndEdge(int EdgeID) const;
 	/** Returns true if the parent-mesh vertex is connected to any seam edges */
diff --git a/Engine/Plugins/Experimental/GeometryProcessing/Source/DynamicMesh/Public/MeshSimplification.h b/Engine/Plugins/Experimental/GeometryProcessing/Source/DynamicMesh/Public/MeshSimplification.h
--- a/Engine/Plugins/Experimental/GeometryProcessing/Source/DynamicMesh/Public/MeshSimplification.h
+++ b/Engine/Plugins/Experimental/GeometryProcessing/Source/DynamicMesh/Public/MeshSimplification.h
@@ -62,7 +62,9 @@ public:
 	/** if false, face and vertex quadrics are recomputed in the neighborhood of each collapse, definitely slower but maybe higher quality*/
 	bool bRetainQuadricMemory = false;
 
-	/** if true, we try to keep boundary vertices on boundary. You probably want this. */
+	/** if true, we try to keep boundary vertices on boundary. 
+	* You probably want this if you aren't using seam quadrics.
+	*/
 	bool bPreserveBoundaryShape = true;
 
 	/** if true, we allow UV and Normal seams to collapse during simplification.*/
@@ -272,7 +274,13 @@ protected:
 		return IsBoundaryVtxCache[vid];
 	}
 
-
+	/**
+	* Figure out if we can collapse edge eid=[a,b] under current constraint set
+	* This depends on the base class function FMeshRefinerBase::CanCollapseEdge, 
+	* but includes additional restrictions specified by the CollapseMode of this simplifier.
+	* See the base class for description of parameters.
+	*/
+	bool CanCollapseEdge(int edgeID, int a, int b, int c, int d, int t0, int t1, int& collapse_to) const;
 
 
 	/**
