GeometryProcessing: Add RenderBuffersToDynamicMesh() conversion function to build a FDynamicMesh3 from StaticMesh LOD data. Add functions to explicitly enable/disable Tangents overlays in FDynamicMesh3AttributeSet.
#rnx

diff --git a/Engine/Plugins/Experimental/GeometryProcessing/Source/DynamicMesh/Private/DynamicMeshAttributeSet.cpp b/Engine/Plugins/Experimental/GeometryProcessing/Source/DynamicMesh/Private/DynamicMeshAttributeSet.cpp
--- a/Engine/Plugins/Experimental/GeometryProcessing/Source/DynamicMesh/Private/DynamicMeshAttributeSet.cpp
+++ b/Engine/Plugins/Experimental/GeometryProcessing/Source/DynamicMesh/Private/DynamicMeshAttributeSet.cpp
@@ -252,6 +252,17 @@ void FDynamicMeshAttributeSet::SetNumUVLayers(int Num)
 
 
 
+void FDynamicMeshAttributeSet::EnableTangents()
+{
+	SetNumNormalLayers(3);
+}
+
+void FDynamicMeshAttributeSet::DisableTangents()
+{
+	SetNumNormalLayers(1);
+}
+
+
 void FDynamicMeshAttributeSet::SetNumNormalLayers(int Num)
 {
 	if (NormalLayers.Num() == Num)
diff --git a/Engine/Plugins/Experimental/GeometryProcessing/Source/DynamicMesh/Public/DynamicMeshAttributeSet.h b/Engine/Plugins/Experimental/GeometryProcessing/Source/DynamicMesh/Public/DynamicMeshAttributeSet.h
--- a/Engine/Plugins/Experimental/GeometryProcessing/Source/DynamicMesh/Public/DynamicMeshAttributeSet.h
+++ b/Engine/Plugins/Experimental/GeometryProcessing/Source/DynamicMesh/Public/DynamicMeshAttributeSet.h
@@ -97,21 +97,6 @@ private:
 	void Reparent(FDynamicMesh3* NewParent);
 
 public:
-	/** @return number of UV layers */
-	virtual int NumUVLayers() const
-	{
-		return UVLayers.Num();
-	}
-
-	virtual void SetNumUVLayers(int Num);
-
-	/** @return number of Normals layers */
-	virtual int NumNormalLayers() const
-	{
-		return NormalLayers.Num();
-	}
-
-	virtual void SetNumNormalLayers(int Num);
 
 	/** @return true if the given edge is a seam edge in any overlay */
 	virtual bool IsSeamEdge(int EdgeID) const;
@@ -132,6 +117,15 @@ public:
 	// UV Layers 
 	//
 
+	/** @return number of UV layers */
+	virtual int NumUVLayers() const
+	{
+		return UVLayers.Num();
+	}
+
+	/** Set number of UV (2-vector float overlay) layers */
+	virtual void SetNumUVLayers(int Num);
+
 	/** @return the UV layer at the given Index  if exists, else nullptr*/
 	FDynamicMeshUVOverlay* GetUVLayer(int Index)
 	{
@@ -168,6 +162,21 @@ public:
 	// Normal Layers 
 	//
 
+	/** @return number of Normals layers */
+	virtual int NumNormalLayers() const
+	{
+		return NormalLayers.Num();
+	}
+
+	/** Set number of Normals (3-vector float overlay) layers */
+	virtual void SetNumNormalLayers(int Num);
+
+	/** Enable Tangents overlays (Tangent = Normal layer 1, Bitangent = Normal layer 2) */
+	void EnableTangents();
+
+	/** Disable Tangents overlays */
+	void DisableTangents();
+
 	/** @return the Normal layer at the given Index if exists, else nullptr */
 	FDynamicMeshNormalOverlay* GetNormalLayer(int Index)
 	{
diff --git a/Engine/Plugins/Experimental/GeometryProcessing/Source/MeshConversion/MeshConversion.Build.cs b/Engine/Plugins/Experimental/GeometryProcessing/Source/MeshConversion/MeshConversion.Build.cs
--- a/Engine/Plugins/Experimental/GeometryProcessing/Source/MeshConversion/MeshConversion.Build.cs
+++ b/Engine/Plugins/Experimental/GeometryProcessing/Source/MeshConversion/MeshConversion.Build.cs
@@ -11,6 +11,7 @@ public class MeshConversion : ModuleRules
         PublicDependencyModuleNames.AddRange(
 			new string[] {
                 "Core",
+				"Engine",
                 "MeshDescription",
 				"StaticMeshDescription",
 				"GeometricObjects",
diff --git a/Engine/Plugins/Experimental/GeometryProcessing/Source/MeshConversion/Private/IndexedMeshToDynamicMesh.cpp b/Engine/Plugins/Experimental/GeometryProcessing/Source/MeshConversion/Private/IndexedMeshToDynamicMesh.cpp
new file mode 100644
--- /dev/null
+++ b/Engine/Plugins/Experimental/GeometryProcessing/Source/MeshConversion/Private/IndexedMeshToDynamicMesh.cpp
@@ -0,0 +1,101 @@
+// Copyright Epic Games, Inc. All Rights Reserved. 
+
+#include "IndexedMeshToDynamicMesh.h"
+#include "DynamicMeshAttributeSet.h"
+#include "DynamicMeshOverlay.h"
+#include "MeshNormals.h"
+#include "Operations/MergeCoincidentMeshEdges.h"
+
+
+void UE::Conversion::RenderBuffersToDynamicMesh(
+	const FStaticMeshVertexBuffers& VertexData,
+	const FRawStaticIndexBuffer& IndexData,
+	const FStaticMeshSectionArray& SectionData,
+	FDynamicMesh3& MeshOut,
+	bool bAttemptToWeldSeams)
+{
+	const FStaticMeshVertexBuffer& VertexBuffer = VertexData.StaticMeshVertexBuffer;
+	const FPositionVertexBuffer& PositionBuffer = VertexData.PositionVertexBuffer;
+	const FColorVertexBuffer& ColorBuffer = VertexData.ColorVertexBuffer;
+
+	int32 NumUVChannels = VertexBuffer.GetNumTexCoords();
+	bool bHasUVs = (NumUVChannels > 0);
+	NumUVChannels = FMath::Max(1, NumUVChannels);
+
+	MeshOut.Clear();
+
+	MeshOut.EnableAttributes();
+	FDynamicMeshNormalOverlay* Normals = MeshOut.Attributes()->PrimaryNormals();
+	MeshOut.Attributes()->EnableTangents();
+	FDynamicMeshNormalOverlay* Tangents = MeshOut.Attributes()->PrimaryTangents();
+	FDynamicMeshNormalOverlay* BiTangents = MeshOut.Attributes()->PrimaryBiTangents();
+
+	MeshOut.Attributes()->SetNumUVLayers(NumUVChannels);
+	TArray<FDynamicMeshUVOverlay*> UVs;
+	for (int32 ui = 0; ui < NumUVChannels; ++ui)
+	{
+		UVs.Add(MeshOut.Attributes()->GetUVLayer(ui));
+	}
+	MeshOut.Attributes()->EnableMaterialID();
+	FDynamicMeshMaterialAttribute* MaterialID = MeshOut.Attributes()->GetMaterialID();
+
+	int32 NumVertices = PositionBuffer.GetNumVertices();
+	for (int32 vi = 0; vi < NumVertices; ++vi)
+	{
+		FVector Position = PositionBuffer.VertexPosition(vi);
+		int32 vid = MeshOut.AppendVertex(Position);
+		ensure(vid == vi);
+
+		FVector Normal = VertexBuffer.VertexTangentZ(vi);
+		int32 nid = Normals->AppendElement((FVector3f)Normal);
+		ensure(nid == vi);
+
+		FVector Tangent = VertexBuffer.VertexTangentX(vi);
+		int32 tid = Tangents->AppendElement((FVector3f)Tangent);
+		ensure(tid == vi);
+
+		FVector BiTangent = VertexBuffer.VertexTangentY(vi);
+		int32 bid = BiTangents->AppendElement((FVector3f)BiTangent);
+		ensure(bid == vi);
+
+		for (int32 ui = 0; ui < NumUVChannels; ++ui)
+		{
+			FVector2D UV = (bHasUVs) ? VertexBuffer.GetVertexUV(vi, ui) : FVector2D::ZeroVector;
+			int32 uvid = UVs[ui]->AppendElement((FVector2f)UV);
+			ensure(uvid == vi);
+		}
+	}
+
+
+	int32 NumTriangles = IndexData.GetNumIndices() / 3;
+	for (int32 ti = 0; ti < NumTriangles; ++ti)
+	{
+		int32 a = IndexData.GetIndex(3 * ti);
+		int32 b = IndexData.GetIndex(3 * ti + 1);
+		int32 c = IndexData.GetIndex(3 * ti + 2);
+		int32 tid = MeshOut.AppendTriangle(a, b, c);
+		ensure(tid == ti);
+		Normals->SetTriangle(tid, FIndex3i(a, b, c));
+		for (int32 ui = 0; ui < NumUVChannels; ++ui)
+		{
+			UVs[ui]->SetTriangle(tid, FIndex3i(a, b, c));
+		}
+	}
+
+	for (const FStaticMeshSection& Section : SectionData)
+	{
+		int32 MatIndex = Section.MaterialIndex;
+		for (uint32 k = 0; k < Section.NumTriangles; ++k)
+		{
+			uint32 TriIdx = (Section.FirstIndex/3) + k;
+			MaterialID->SetValue((int32)TriIdx, MatIndex);
+		}
+	}
+
+
+	if (bAttemptToWeldSeams)
+	{
+		FMergeCoincidentMeshEdges Merge(&MeshOut);
+		Merge.Apply();
+	}
+}
\ No newline at end of file
diff --git a/Engine/Plugins/Experimental/GeometryProcessing/Source/MeshConversion/Public/IndexedMeshToDynamicMesh.h b/Engine/Plugins/Experimental/GeometryProcessing/Source/MeshConversion/Public/IndexedMeshToDynamicMesh.h
new file mode 100644
--- /dev/null
+++ b/Engine/Plugins/Experimental/GeometryProcessing/Source/MeshConversion/Public/IndexedMeshToDynamicMesh.h
@@ -0,0 +1,32 @@
+// Copyright Epic Games, Inc. All Rights Reserved.
+
+#pragma once
+
+#include "CoreMinimal.h"
+#include "StaticMeshResources.h"
+#include "DynamicMesh3.h"
+
+
+namespace UE {
+namespace Conversion {
+
+/**
+ * Convert rendering buffer data (from FStaticMeshLODResources) to a DynamicMesh3.
+ * The output FDynamicMesh3 attribute set will have as many UV channels as found in VertexData,
+ * as well as overlay Tangents enabled/copied, and MaterialID set.
+ * @param bAttemptToWeldSeams if true, try to weld seam edges together to create a more closed mesh using FMergeCoincidentMeshEdges
+ */
+void MESHCONVERSION_API RenderBuffersToDynamicMesh(
+	const FStaticMeshVertexBuffers& VertexData,
+	const FRawStaticIndexBuffer& IndexData,
+	const FStaticMeshSectionArray& SectionData,
+	FDynamicMesh3& MeshOut,
+	bool bAttemptToWeldSeams = false );
+
+
+
+
+}  // end namespace Conversion
+}  // end namespace UE
+
+
