Fix instances of TFunctionRef Var = [](){}
Storing a lambda in a TFunctionRef without storing it in any other variable makes it go out of scope and leaves a dangling pointer to the out-of-scope lambda in the TFunctionRef.
#rnx

diff --git a/Engine/Plugins/Experimental/GeometryProcessing/Source/DynamicMesh/Private/MeshNormals.cpp b/Engine/Plugins/Experimental/GeometryProcessing/Source/DynamicMesh/Private/MeshNormals.cpp
--- a/Engine/Plugins/Experimental/GeometryProcessing/Source/DynamicMesh/Private/MeshNormals.cpp
+++ b/Engine/Plugins/Experimental/GeometryProcessing/Source/DynamicMesh/Private/MeshNormals.cpp
@@ -373,7 +373,7 @@ void FMeshNormals::InitializeOverlayRegionToPerVertexNormals(FDynamicMeshNormalO
 	TSet<int32> TriangleSet(Triangles);
 	TArray<int32> Vertices;
 	MeshIndexUtil::TriangleToVertexIDs(Mesh, Triangles, Vertices);
-	TFunctionRef<bool(int32)> TriangleSetFunc = [&](int32 tid) { return TriangleSet.Contains(tid); };
+	auto TriangleSetFunc = [&](int32 tid) { return TriangleSet.Contains(tid); };
 	int32 NumVertices = Vertices.Num();
 	TMap<int32, int32> TriangleMap;
 	TriangleMap.Reserve(NumVertices);
@@ -383,7 +383,7 @@ void FMeshNormals::InitializeOverlayRegionToPerVertexNormals(FDynamicMeshNormalO
 	for ( int32 i = 0; i < NumVertices; ++i)
 	{
 		int32 vid = Vertices[i];
-		FVector3d Normal = FMeshNormals::ComputeVertexNormal(*Mesh, vid, TriangleSetFunc, true, true);
+		FVector3d Normal = FMeshNormals::ComputeVertexNormal(*Mesh, vid, TFunctionRef<bool(int32)>(TriangleSetFunc), true, true);
 		int32 nid = NormalOverlay->AppendElement((FVector3f)Normal);
 		VertNormals[i] = nid;
 
diff --git a/Engine/Source/Editor/SceneOutliner/Private/SSceneOutliner.cpp b/Engine/Source/Editor/SceneOutliner/Private/SSceneOutliner.cpp
--- a/Engine/Source/Editor/SceneOutliner/Private/SSceneOutliner.cpp
+++ b/Engine/Source/Editor/SceneOutliner/Private/SSceneOutliner.cpp
@@ -1309,7 +1309,7 @@ void SSceneOutliner::PasteFoldersEnd()
 
 void SSceneOutliner::DuplicateFoldersHierarchy()
 {
-	TFunctionRef<void(const FSceneOutlinerTreeItemPtr&)> RecursiveFolderSelect = [&](const FSceneOutlinerTreeItemPtr& Item)
+	TFunction<void(const FSceneOutlinerTreeItemPtr&)> RecursiveFolderSelect = [&](const FSceneOutlinerTreeItemPtr& Item)
 	{
 		if (Item->IsA<FFolderTreeItem>())
 		{
diff --git a/Engine/Source/Runtime/Engine/Private/WorldPartition/RuntimeSpatialHash/RuntimeSpatialHashHLOD.cpp b/Engine/Source/Runtime/Engine/Private/WorldPartition/RuntimeSpatialHash/RuntimeSpatialHashHLOD.cpp
--- a/Engine/Source/Runtime/Engine/Private/WorldPartition/RuntimeSpatialHash/RuntimeSpatialHashHLOD.cpp
+++ b/Engine/Source/Runtime/Engine/Private/WorldPartition/RuntimeSpatialHash/RuntimeSpatialHashHLOD.cpp
@@ -298,7 +298,7 @@ static bool SortHLODGrids(const TMap<FName, TSet<FName>>& InHLODGridsGraph, TArr
 	TSet<FName>		ProcessedGridsSet;	// Processed grids
 	TSet<FName>		VisitedGridsSet;	// Visited grids
 
-	TFunctionRef<bool(const FName, uint32)> VisitGraph = [&](const FName GridName, uint32 CurrentDepth) -> bool
+	TFunction<bool(const FName, uint32)> VisitGraph = [&](const FName GridName, uint32 CurrentDepth) -> bool
 	{
 		uint32& GridDepth = OutGridsDepth.FindOrAdd(GridName);
 		GridDepth = FMath::Max(GridDepth, CurrentDepth);
