GeometryProcessing: add various Append and Transform functions to FSimpleShapeSet3d.
ModelingComponents: add CollisionTransformUtil with various utility functions for working w/ Components and collision geometry. FPhysicsDataCollection now can initialize FSimpleShapeSet3d from existing FKAggregateGeom.
ModelingTools: Update BakeTransformTool, EditPivotTool to properly update existing Simple Collision geometry. Update CombineMeshesTool to properly transfer/combine/append Simple Collision geometry and settings.
#rnx

diff --git a/Engine/Plugins/Experimental/GeometryProcessing/Source/DynamicMesh/Private/ShapeApproximation/SimpleShapeSet3.cpp b/Engine/Plugins/Experimental/GeometryProcessing/Source/DynamicMesh/Private/ShapeApproximation/SimpleShapeSet3.cpp
--- a/Engine/Plugins/Experimental/GeometryProcessing/Source/DynamicMesh/Private/ShapeApproximation/SimpleShapeSet3.cpp
+++ b/Engine/Plugins/Experimental/GeometryProcessing/Source/DynamicMesh/Private/ShapeApproximation/SimpleShapeSet3.cpp
@@ -3,6 +3,7 @@
 #include "ShapeApproximation/SimpleShapeSet3.h"
 
 #include "MeshQueries.h"
+#include "MeshTransforms.h"
 #include "Intersection/ContainmentQueries3.h"
 
 
@@ -304,6 +305,174 @@ void FSimpleShapeSet3d::RemoveContainedGeometry()
 }
 
 
+static void TransformSphereShape(FSphereShape3d& SphereShape, const FTransform3d& Transform)
+{
+	double RadiusScale = Transform.GetScale().Length() / FVector3d::One().Length();
+	SphereShape.Sphere.Center = Transform.TransformPosition(SphereShape.Sphere.Center);
+	SphereShape.Sphere.Radius *= RadiusScale;
+}
+static void TransformBoxShape(FBoxShape3d& BoxShape, const FTransform3d& Transform)
+{
+	FVector3d CornerVec = BoxShape.Box.Frame.PointAt(BoxShape.Box.Extents) - BoxShape.Box.Frame.Origin;
+	BoxShape.Box.Frame.Transform(Transform);
+	CornerVec = Transform.TransformVector(CornerVec);
+	BoxShape.Box.Extents.X = CornerVec.Dot(BoxShape.Box.AxisX());
+	BoxShape.Box.Extents.Y = CornerVec.Dot(BoxShape.Box.AxisY());
+	BoxShape.Box.Extents.Z = CornerVec.Dot(BoxShape.Box.AxisZ());
+}
+static void TransformCapsuleShape(FCapsuleShape3d& CapsuleShape, const FTransform3d& Transform)
+{
+	FVector3d P0 = Transform.TransformPosition(CapsuleShape.Capsule.Segment.StartPoint());
+	FVector3d P1 = Transform.TransformPosition(CapsuleShape.Capsule.Segment.EndPoint());
+
+	CapsuleShape.Capsule.Segment.Center = 0.5 * (P0 + P1);
+	CapsuleShape.Capsule.Segment.Direction = (P1 - P0);
+	CapsuleShape.Capsule.Segment.Extent = CapsuleShape.Capsule.Segment.Direction.Normalize() * 0.5;
+
+	double CurRadius = CapsuleShape.Capsule.Radius;
+	FFrame3d CapsuleFrame(CapsuleShape.Capsule.Segment.Center, CapsuleShape.Capsule.Segment.Direction);
+	FVector3d SideVec = CapsuleFrame.PointAt(FVector3d(CurRadius, CurRadius, 0)) - CapsuleFrame.Origin;
+	FVector3d NewSideVec = Transform.TransformVector(SideVec);
+	double RadiusScale = NewSideVec.Length() / SideVec.Length();
+	CapsuleShape.Capsule.Radius *= RadiusScale;
+}
+
+
+
+
+static void TransformSphereShape(FSphereShape3d& SphereShape, const TArray<FTransform3d>& TransformSequence)
+{
+	for (const FTransform3d& XForm : TransformSequence)
+	{
+		SphereShape.Sphere.Center = XForm.TransformPosition(SphereShape.Sphere.Center);
+		double RadiusScale = XForm.GetScale().Length() / FVector3d::One().Length();
+		SphereShape.Sphere.Radius *= RadiusScale;
+	}
+}
+static void TransformBoxShape(FBoxShape3d& BoxShape, const TArray<FTransform3d>& TransformSequence)
+{
+	FVector3d CornerVec = BoxShape.Box.Frame.PointAt(BoxShape.Box.Extents) - BoxShape.Box.Frame.Origin;
+	for (const FTransform3d& XForm : TransformSequence)
+	{
+		BoxShape.Box.Frame.Transform(XForm);
+		CornerVec = XForm.TransformVector(CornerVec);
+	}
+	BoxShape.Box.Extents.X = CornerVec.Dot(BoxShape.Box.AxisX());
+	BoxShape.Box.Extents.Y = CornerVec.Dot(BoxShape.Box.AxisY());
+	BoxShape.Box.Extents.Z = CornerVec.Dot(BoxShape.Box.AxisZ());
+}
+static void TransformCapsuleShape(FCapsuleShape3d& CapsuleShape, const TArray<FTransform3d>& TransformSequence)
+{
+	FVector3d P0 = CapsuleShape.Capsule.Segment.StartPoint();
+	FVector3d P1 = CapsuleShape.Capsule.Segment.EndPoint();
+
+	double CurRadius = CapsuleShape.Capsule.Radius;
+	FFrame3d CapsuleFrame(CapsuleShape.Capsule.Segment.Center, CapsuleShape.Capsule.Segment.Direction);
+	FVector3d InitialSideVec = CapsuleFrame.PointAt(FVector3d(CurRadius, CurRadius, 0)) - CapsuleFrame.Origin;
+	FVector3d NewSideVec = InitialSideVec;
+
+	for (const FTransform3d& XForm : TransformSequence)
+	{
+		P0 = XForm.TransformPosition(P0);
+		P1 = XForm.TransformPosition(P1);
+		NewSideVec = XForm.TransformVector(NewSideVec);
+	}
+
+	CapsuleShape.Capsule.Segment.Center = 0.5 * (P0 + P1);
+	CapsuleShape.Capsule.Segment.Direction = (P1 - P0);
+	CapsuleShape.Capsule.Segment.Extent = CapsuleShape.Capsule.Segment.Direction.Normalize() * 0.5;
+	double RadiusScale = NewSideVec.Length() / InitialSideVec.Length();
+	CapsuleShape.Capsule.Radius *= RadiusScale;
+}
+
+
+
+
+void FSimpleShapeSet3d::Append(const FSimpleShapeSet3d& OtherShapeSet)
+{
+	for (FSphereShape3d SphereShape : OtherShapeSet.Spheres)
+	{
+		Spheres.Add(SphereShape);
+	}
+
+	for (FBoxShape3d BoxShape : OtherShapeSet.Boxes)
+	{
+		Boxes.Add(BoxShape);
+	}
+
+	for (FCapsuleShape3d CapsuleShape : OtherShapeSet.Capsules)
+	{
+		Capsules.Add(CapsuleShape);
+	}
+
+	Convexes.Reserve(Convexes.Num() + OtherShapeSet.Convexes.Num());
+	for (const FConvexShape3d& ConvexShape : OtherShapeSet.Convexes)
+	{
+		Convexes.Add(ConvexShape);
+	}
+}
+
+
+void FSimpleShapeSet3d::Append(const FSimpleShapeSet3d& OtherShapeSet, const FTransform3d& Transform)
+{
+	for (FSphereShape3d SphereShape : OtherShapeSet.Spheres)
+	{
+		TransformSphereShape(SphereShape, Transform);
+		Spheres.Add(SphereShape);
+	}
+
+	for (FBoxShape3d BoxShape : OtherShapeSet.Boxes)
+	{
+		TransformBoxShape(BoxShape, Transform);
+		Boxes.Add(BoxShape);
+	}
+
+	for (FCapsuleShape3d CapsuleShape : OtherShapeSet.Capsules)
+	{
+		TransformCapsuleShape(CapsuleShape, Transform);
+		Capsules.Add(CapsuleShape);
+	}
+
+	Convexes.Reserve(Convexes.Num() + OtherShapeSet.Convexes.Num());
+	for (const FConvexShape3d& ConvexShape : OtherShapeSet.Convexes)
+	{
+		Convexes.Add(ConvexShape);
+		MeshTransforms::ApplyTransform(Convexes.Last().Mesh, Transform);
+	}
+}
+
+
+
+void FSimpleShapeSet3d::Append(const FSimpleShapeSet3d& OtherShapeSet, const TArray<FTransform3d>& TransformSequence)
+{
+	for (FSphereShape3d SphereShape : OtherShapeSet.Spheres)
+	{
+		TransformSphereShape(SphereShape, TransformSequence);
+		Spheres.Add(SphereShape);
+	}
+
+	for (FBoxShape3d BoxShape : OtherShapeSet.Boxes)
+	{
+		TransformBoxShape(BoxShape, TransformSequence);
+		Boxes.Add(BoxShape);
+	}
+
+	for (FCapsuleShape3d CapsuleShape : OtherShapeSet.Capsules)
+	{
+		TransformCapsuleShape(CapsuleShape, TransformSequence);
+		Capsules.Add(CapsuleShape);
+	}
+
+	Convexes.Reserve(Convexes.Num() + OtherShapeSet.Convexes.Num());
+	for (const FConvexShape3d& ConvexShape : OtherShapeSet.Convexes)
+	{
+		Convexes.Add(ConvexShape);
+		for (const FTransform3d& XForm : TransformSequence)
+		{
+			MeshTransforms::ApplyTransform(Convexes.Last().Mesh, XForm);
+		}
+	}
+}
 
 
 
@@ -343,4 +512,29 @@ void FSimpleShapeSet3d::FilterByVolume(int32 MaximumCount)
 	Boxes = MoveTemp(NewSet.Boxes);
 	Capsules = MoveTemp(NewSet.Capsules);
 	Convexes = MoveTemp(NewSet.Convexes);
-}
\ No newline at end of file
+}
+
+
+
+void FSimpleShapeSet3d::ApplyTransform(const FTransform3d& Transform)
+{
+	for (FSphereShape3d& SphereShape : Spheres)
+	{
+		TransformSphereShape(SphereShape, Transform);
+	}
+
+	for (FBoxShape3d& BoxShape : Boxes)
+	{
+		TransformBoxShape(BoxShape, Transform);
+	}
+
+	for (FCapsuleShape3d& CapsuleShape : Capsules)
+	{
+		TransformCapsuleShape(CapsuleShape, Transform);
+	}
+
+	for (FConvexShape3d& ConvexShape : Convexes)
+	{
+		MeshTransforms::ApplyTransform(ConvexShape.Mesh, Transform);
+	}
+}
diff --git a/Engine/Plugins/Experimental/GeometryProcessing/Source/DynamicMesh/Public/ShapeApproximation/SimpleShapeSet3.h b/Engine/Plugins/Experimental/GeometryProcessing/Source/DynamicMesh/Public/ShapeApproximation/SimpleShapeSet3.h
--- a/Engine/Plugins/Experimental/GeometryProcessing/Source/DynamicMesh/Public/ShapeApproximation/SimpleShapeSet3.h
+++ b/Engine/Plugins/Experimental/GeometryProcessing/Source/DynamicMesh/Public/ShapeApproximation/SimpleShapeSet3.h
@@ -9,6 +9,7 @@
 #include "CapsuleTypes.h"
 #include "SphereTypes.h"
 #include "DynamicMesh3.h"
+#include "TransformTypes.h"
 
 
 /**
@@ -116,6 +117,24 @@ struct DYNAMICMESH_API FSimpleShapeSet3d
 	/** @return total number of elements in all sets */
 	int32 TotalElementsNum() const { return Spheres.Num() + Boxes.Num() + Capsules.Num() + Convexes.Num(); }
 
+	/**
+	 * Append elements of another shape set
+	 */
+	void Append(const FSimpleShapeSet3d& OtherShapeSet);
+
+
+	/**
+	 * Append elements of another shape set with given transform applied
+	 */
+	void Append(const FSimpleShapeSet3d& OtherShapeSet, const FTransform3d& Transform);
+
+
+	/**
+	 * Append elements of another shape set with given transforms applied
+	 */
+	void Append(const FSimpleShapeSet3d& OtherShapeSet, const TArray<FTransform3d>& TransformSequence);
+
+
 	/**
 	 * Remove any of the elements that are fully contained in larger elements
 	 */
@@ -125,5 +144,11 @@ struct DYNAMICMESH_API FSimpleShapeSet3d
 	 * Sort the elements by volume and then discard all but the largest MaximumCount elements
 	 */
 	void FilterByVolume(int32 MaximumCount);
+
+
+	/**
+	 * Transform shape elements. This will be a best-effort as if there is non-uniform scaling only Convexes can be transformed correctly
+	 */
+	void ApplyTransform(const FTransform3d& Transform);
 };
 
diff --git a/Engine/Plugins/Experimental/MeshModelingToolset/Source/MeshModelingTools/Private/EditPivotTool.cpp b/Engine/Plugins/Experimental/MeshModelingToolset/Source/MeshModelingTools/Private/EditPivotTool.cpp
--- a/Engine/Plugins/Experimental/MeshModelingToolset/Source/MeshModelingTools/Private/EditPivotTool.cpp
+++ b/Engine/Plugins/Experimental/MeshModelingToolset/Source/MeshModelingTools/Private/EditPivotTool.cpp
@@ -13,6 +13,7 @@
 #include "MeshTransforms.h"
 #include "BaseBehaviors/ClickDragBehavior.h"
 #include "ToolSceneQueriesUtil.h"
+#include "Physics/ComponentCollisionUtil.h"
 
 #include "BaseGizmos/GizmoComponents.h"
 #include "BaseGizmos/TransformGizmo.h"
@@ -445,14 +446,18 @@ void UEditPivotTool::UpdateAssets(const FFrame3d& NewPivotWorldFrame)
 		{
 			FTransform3d ToBake(OriginalTransforms[ComponentIdx] * NewWorldInverse);
 
+			UPrimitiveComponent* Component = ComponentTargets[ComponentIdx]->GetOwnerComponent();
+			Component->Modify();
+
+			// transform simple collision geometry
+			UE::Geometry::TransformSimpleCollision(Component, ToBake);
+
 			ComponentTargets[ComponentIdx]->CommitMesh([&ToBake](const FPrimitiveComponentTarget::FCommitParams& CommitParams)
 			{
 				FMeshDescriptionEditableTriangleMeshAdapter EditableMeshDescAdapter(CommitParams.MeshDescription);
 				MeshAdapterTransforms::ApplyTransform(EditableMeshDescAdapter, ToBake);
 			});
 
-			UPrimitiveComponent* Component = ComponentTargets[ComponentIdx]->GetOwnerComponent();
-			Component->Modify();
 			Component->SetWorldTransform(NewWorldTransform);
 		}
 		else
diff --git a/Engine/Plugins/Experimental/MeshModelingToolset/Source/ModelingComponents/Private/Physics/ComponentCollisionUtil.cpp b/Engine/Plugins/Experimental/MeshModelingToolset/Source/ModelingComponents/Private/Physics/ComponentCollisionUtil.cpp
new file mode 100644
--- /dev/null
+++ b/Engine/Plugins/Experimental/MeshModelingToolset/Source/ModelingComponents/Private/Physics/ComponentCollisionUtil.cpp
@@ -0,0 +1,174 @@
+// Copyright Epic Games, Inc. All Rights Reserved. 
+
+#include "Physics/ComponentCollisionUtil.h"
+
+#include "OrientedBoxTypes.h"
+#include "SegmentTypes.h"
+#include "CapsuleTypes.h"
+#include "SphereTypes.h"
+#include "DynamicMesh3.h"
+#include "ShapeApproximation/SimpleShapeSet3.h"
+
+#include "Physics/PhysicsDataCollection.h"
+
+
+#include "Components/StaticMeshComponent.h"
+#include "Engine/StaticMesh.h"
+#include "Engine/Classes/PhysicsEngine/BodySetup.h"
+#include "PhysicsEngine/AggregateGeom.h"
+
+
+UE::Geometry::FComponentCollisionSettings UE::Geometry::GetCollisionSettings(const UPrimitiveComponent* Component)
+{
+	FComponentCollisionSettings Settings;
+
+	const UStaticMeshComponent* StaticMeshComponent = CastChecked<UStaticMeshComponent>(Component);
+	if (ensure(StaticMeshComponent))
+	{
+		const UStaticMesh* StaticMesh = StaticMeshComponent->GetStaticMesh();
+		if (ensure(StaticMesh))
+		{
+			UBodySetup* BodySetup = StaticMesh->GetBodySetup();
+			if (ensure(BodySetup))
+			{
+				Settings.CollisionTypeFlag = (int32)BodySetup->CollisionTraceFlag;
+			}
+		}
+	}
+	return Settings;
+}
+
+
+void UE::Geometry::UpdateSimpleCollision(
+	UBodySetup* BodySetup, 
+	const FKAggregateGeom* NewGeometry,
+	UStaticMesh* StaticMesh,
+	FComponentCollisionSettings CollisionSettings)
+{
+	BodySetup->Modify();
+	BodySetup->RemoveSimpleCollision();
+
+	// set new collision geometry
+	BodySetup->AggGeom = *NewGeometry;
+
+	// update collision type
+	BodySetup->CollisionTraceFlag = (ECollisionTraceFlag)CollisionSettings.CollisionTypeFlag;
+
+	// rebuild physics meshes
+	BodySetup->CreatePhysicsMeshes();
+
+	// rebuild nav collision (? StaticMeshEditor does this)
+	StaticMesh->CreateNavCollision(/*bIsUpdate=*/true);
+
+	// update physics state on all components using this StaticMesh
+	for (FThreadSafeObjectIterator Iter(UStaticMeshComponent::StaticClass()); Iter; ++Iter)
+	{
+		UStaticMeshComponent* SMComponent = Cast<UStaticMeshComponent>(*Iter);
+		if (SMComponent->GetStaticMesh() == StaticMesh)
+		{
+			if (SMComponent->IsPhysicsStateCreated())
+			{
+				SMComponent->RecreatePhysicsState();
+			}
+		}
+	}
+
+	// mark static mesh as dirty so it gets resaved?
+	StaticMesh->MarkPackageDirty();
+
+#if WITH_EDITORONLY_DATA
+	// mark the static mesh as having customized collision so it is not regenerated on reimport
+	StaticMesh->bCustomizedCollision = CollisionSettings.bIsGeneratedCollision;
+#endif // WITH_EDITORONLY_DATA
+}
+
+
+bool UE::Geometry::SetSimpleCollision(
+	UPrimitiveComponent* Component,
+	const FSimpleShapeSet3d* ShapeSet,
+	FComponentCollisionSettings CollisionSettings)
+{
+	FPhysicsDataCollection PhysicsData;
+	PhysicsData.InitializeFromComponent(Component, false);
+	if (ensure(PhysicsData.SourceComponent.IsValid()) == false || ensure(ShapeSet != nullptr) == false )
+	{
+		return false;
+	}
+
+	PhysicsData.Geometry = *ShapeSet;
+	PhysicsData.CopyGeometryToAggregate();
+
+	// FPhysicsDataCollection stores its references as const, but the input Component was non-const so this is ok to do
+	UStaticMesh* StaticMesh = const_cast<UStaticMesh*>(PhysicsData.SourceStaticMesh.Get());
+	UBodySetup* BodySetup = StaticMesh->GetBodySetup();
+	UpdateSimpleCollision(BodySetup, &PhysicsData.AggGeom, StaticMesh, CollisionSettings);
+
+	return true;
+}
+
+
+
+bool UE::Geometry::TransformSimpleCollision(
+	UPrimitiveComponent* Component,
+	const FTransform3d& Transform)
+{
+	FPhysicsDataCollection PhysicsData;
+	PhysicsData.InitializeFromComponent(Component, true);
+	if ( ensure(PhysicsData.SourceComponent.IsValid()) == false)
+	{
+		return false;
+	}
+
+	FComponentCollisionSettings Settings = GetCollisionSettings(Component);
+	Settings.bIsGeneratedCollision = false;
+
+	PhysicsData.Geometry.ApplyTransform(Transform);
+	PhysicsData.ClearAggregate();
+	PhysicsData.CopyGeometryToAggregate();
+
+	// FPhysicsDataCollection stores its references as const, but the input Component was non-const so this is ok to do
+	UStaticMesh* StaticMesh = const_cast<UStaticMesh*>(PhysicsData.SourceStaticMesh.Get());
+	UBodySetup* BodySetup = StaticMesh->GetBodySetup();
+	UpdateSimpleCollision(BodySetup, &PhysicsData.AggGeom, StaticMesh, Settings);
+
+	return true;
+}
+
+
+
+
+bool UE::Geometry::AppendSimpleCollision(
+	const UPrimitiveComponent* Component,
+	FSimpleShapeSet3d* ShapeSetOut,
+	const FTransform3d& Transform)
+{
+	FPhysicsDataCollection PhysicsData;
+	PhysicsData.InitializeFromComponent(Component, true);
+	if (ensure(PhysicsData.SourceComponent.IsValid()) == false || ensure(ShapeSetOut != nullptr) == false)
+	{
+		return false;
+	}
+
+	ShapeSetOut->Append(PhysicsData.Geometry, Transform);
+	return true;
+}
+
+
+bool UE::Geometry::AppendSimpleCollision(
+	const UPrimitiveComponent* Component,
+	FSimpleShapeSet3d* ShapeSetOut,
+	const TArray<FTransform3d>& TransformSeqeuence)
+{
+	FPhysicsDataCollection PhysicsData;
+	PhysicsData.InitializeFromComponent(Component, true);
+	if (ensure(PhysicsData.SourceComponent.IsValid()) == false || ensure(ShapeSetOut != nullptr) == false)
+	{
+		return false;
+	}
+
+	ShapeSetOut->Append(PhysicsData.Geometry, TransformSeqeuence);
+
+	return true;
+}
+
+
diff --git a/Engine/Plugins/Experimental/MeshModelingToolset/Source/ModelingComponents/Private/Physics/PhysicsDataCollection.cpp b/Engine/Plugins/Experimental/MeshModelingToolset/Source/ModelingComponents/Private/Physics/PhysicsDataCollection.cpp
--- a/Engine/Plugins/Experimental/MeshModelingToolset/Source/ModelingComponents/Private/Physics/PhysicsDataCollection.cpp
+++ b/Engine/Plugins/Experimental/MeshModelingToolset/Source/ModelingComponents/Private/Physics/PhysicsDataCollection.cpp
@@ -3,42 +3,50 @@
 #include "Physics/PhysicsDataCollection.h"
 #include "Physics/CollisionGeometryConversion.h"
 
-#include "Engine/Classes/Engine/StaticMesh.h"
-#include "Engine/Classes/Components/StaticMeshComponent.h"
-#include "Engine/Classes/PhysicsEngine/BodySetup.h"
+#include "Engine/StaticMesh.h"
+#include "Components/StaticMeshComponent.h"
+#include "PhysicsEngine/BodySetup.h"
 
 
 
 void FPhysicsDataCollection::InitializeFromComponent(const UActorComponent* Component, bool bInitializeAggGeom)
 {
 	const UStaticMeshComponent* StaticMeshComponent = CastChecked<UStaticMeshComponent>(Component);
-	const UStaticMesh* StaticMesh = StaticMeshComponent->GetStaticMesh();
-	
-	SourceComponent = StaticMeshComponent;
-	SourceStaticMesh = StaticMesh;
-	BodySetup = StaticMesh->GetBodySetup();
+	if (ensure(StaticMeshComponent))
+	{
+		const UStaticMesh* StaticMesh = StaticMeshComponent->GetStaticMesh();
+		if (ensure(StaticMesh))
+		{
+			SourceComponent = StaticMeshComponent;
+			SourceStaticMesh = StaticMesh;
+			BodySetup = StaticMesh->GetBodySetup();
 
-	ExternalScale3D = FVector(1.f, 1.f, 1.f);
+			ExternalScale3D = FVector(1.f, 1.f, 1.f);
 
-	if (bInitializeAggGeom)
-	{
-		AggGeom = BodySetup->AggGeom;
-		// transfer AggGeom to FSimpleShapeSet3d...
+			if (bInitializeAggGeom)
+			{
+				AggGeom = BodySetup->AggGeom;
+				UE::Geometry::GetShapeSet(AggGeom, Geometry);
+			}
+		}
 	}
 }
 
 
 void FPhysicsDataCollection::InitializeFromStaticMesh(const UStaticMesh* StaticMesh, bool bInitializeAggGeom)
 {
-	SourceStaticMesh = StaticMesh;
-	BodySetup = StaticMesh->GetBodySetup();
+	if (ensure(StaticMesh))
+	{
+		SourceStaticMesh = StaticMesh;
+		BodySetup = StaticMesh->GetBodySetup();
 
-	ExternalScale3D = FVector(1.f, 1.f, 1.f);
+		ExternalScale3D = FVector(1.f, 1.f, 1.f);
 
-	if (bInitializeAggGeom)
-	{
-		AggGeom = BodySetup->AggGeom;
-		// transfer AggGeom to FSimpleShapeSet3d...
+		if (bInitializeAggGeom)
+		{
+			AggGeom = BodySetup->AggGeom;
+			UE::Geometry::GetShapeSet(AggGeom, Geometry);
+		}
 	}
 }
 
diff --git a/Engine/Plugins/Experimental/MeshModelingToolset/Source/ModelingComponents/Public/Physics/CollisionGeometryConversion.h b/Engine/Plugins/Experimental/MeshModelingToolset/Source/ModelingComponents/Public/Physics/CollisionGeometryConversion.h
--- a/Engine/Plugins/Experimental/MeshModelingToolset/Source/ModelingComponents/Public/Physics/CollisionGeometryConversion.h
+++ b/Engine/Plugins/Experimental/MeshModelingToolset/Source/ModelingComponents/Public/Physics/CollisionGeometryConversion.h
@@ -9,9 +9,9 @@
 #include "CapsuleTypes.h"
 #include "SphereTypes.h"
 #include "DynamicMesh3.h"
+#include "ShapeApproximation/SimpleShapeSet3.h"
 
-
-#include "Engine/Classes/PhysicsEngine/AggregateGeom.h"
+#include "PhysicsEngine/AggregateGeom.h"
 
 
 namespace UE
@@ -71,5 +71,100 @@ namespace UE
 		}
 
 
+
+
+		/**
+		 * Convert FKSphereElem to FSphereShape3d
+		 */
+		void GetShape(const FKSphereElem& SphereElem, FSphereShape3d& SphereShapeOut)
+		{
+			SphereShapeOut.Sphere.Center = (FVector3d)SphereElem.Center;
+			SphereShapeOut.Sphere.Radius = (double)SphereElem.Radius;
+		}
+
+		/**
+		 * Convert FKBoxElem to FBoxShape3d
+		 */
+		void GetShape(const FKBoxElem& BoxElem, FBoxShape3d& BoxShapeOut)
+		{
+			BoxShapeOut.Box.Frame.Origin = (FVector3d)BoxElem.Center;
+			BoxShapeOut.Box.Frame.Rotation = (FQuaterniond)BoxElem.Rotation.Quaternion();
+			BoxShapeOut.Box.Extents = FVector3d(BoxElem.X, BoxElem.Y, BoxElem.Z) / 2.0;
+		}
+
+		/**
+		 * Convert FKSphylElem to FCapsuleShape3d
+		 */
+		void GetShape(const FKSphylElem& CapsuleElem, FCapsuleShape3d& CapsuleShapeOut)
+		{
+			FQuaterniond CapsuleRotation(CapsuleElem.Rotation.Quaternion());
+			CapsuleShapeOut.Capsule.Segment.Center = (FVector3d)CapsuleElem.Center;
+			CapsuleShapeOut.Capsule.Segment.Direction = CapsuleRotation.AxisZ();
+			CapsuleShapeOut.Capsule.Segment.Extent = (double)CapsuleElem.Length / 2.0;
+			CapsuleShapeOut.Capsule.Radius = (double)CapsuleElem.Radius;
+		}
+
+
+		/**
+		 * Convert FKConvexElem to FDynamicMesh3 
+		 */
+		void GetShape(const FKConvexElem& ConvexElem, FDynamicMesh3& MeshOut)
+		{
+			MeshOut.Clear();
+			for (int32 k = 0; k < ConvexElem.VertexData.Num(); ++k)
+			{
+				MeshOut.AppendVertex((FVector3d)ConvexElem.VertexData[k]);
+			}
+			int32 NumTriangles = ConvexElem.IndexData.Num() / 3;
+			for (int32 k = 0; k < NumTriangles; ++k)
+			{
+				MeshOut.AppendTriangle( ConvexElem.IndexData[3*k], ConvexElem.IndexData[3*k+1], ConvexElem.IndexData[3*k+2]);
+			}
+		}
+
+		/**
+		 * Convert FKConvexElem to FConvexShape3d
+		 */
+		void GetShape(const FKConvexElem& ConvexElem, FConvexShape3d& ConvexShapeOut)
+		{
+			GetShape(ConvexElem, ConvexShapeOut.Mesh);
+		}
+
+
+		/**
+		 * Convert FKAggregateGeom to FSimpleShapeSet3d
+		 */
+		void GetShapeSet(const FKAggregateGeom& AggGeom, FSimpleShapeSet3d& ShapeSetOut)
+		{
+			for (const FKBoxElem& BoxElem : AggGeom.BoxElems)
+			{
+				FBoxShape3d BoxShape;
+				GetShape(BoxElem, BoxShape);
+				ShapeSetOut.Boxes.Add(BoxShape);
+			}
+
+			for (const FKSphereElem& SphereElem : AggGeom.SphereElems)
+			{
+				FSphereShape3d SphereShape;
+				GetShape(SphereElem, SphereShape);
+				ShapeSetOut.Spheres.Add(SphereShape);
+			}
+
+			for (const FKSphylElem& CapsuleElem : AggGeom.SphylElems)
+			{
+				FCapsuleShape3d CapsuleShape;
+				GetShape(CapsuleElem, CapsuleShape);
+				ShapeSetOut.Capsules.Add(CapsuleShape);
+			}
+
+			for (const FKConvexElem& ConvexElem : AggGeom.ConvexElems)
+			{
+				FConvexShape3d ConvexShape;
+				GetShape(ConvexElem, ConvexShape);
+				ShapeSetOut.Convexes.Add(ConvexShape);
+			}
+		}
+
+
 	}
 }
\ No newline at end of file
diff --git a/Engine/Plugins/Experimental/MeshModelingToolset/Source/ModelingComponents/Public/Physics/ComponentCollisionUtil.h b/Engine/Plugins/Experimental/MeshModelingToolset/Source/ModelingComponents/Public/Physics/ComponentCollisionUtil.h
new file mode 100644
--- /dev/null
+++ b/Engine/Plugins/Experimental/MeshModelingToolset/Source/ModelingComponents/Public/Physics/ComponentCollisionUtil.h
@@ -0,0 +1,84 @@
+// Copyright Epic Games, Inc. All Rights Reserved.
+
+#pragma once
+
+#include "CoreMinimal.h"
+#include "Components/PrimitiveComponent.h"
+#include "TransformTypes.h"
+
+class UBodySetup;
+class UStaticMesh;
+struct FKAggregateGeom;
+struct FSimpleShapeSet3d;
+
+namespace UE
+{
+namespace Geometry
+{
+
+/**
+ * Component/BodySetup collision settings (eg StaticMeshComponent) we might need to pass through the functions below
+ */
+MODELINGCOMPONENTS_API struct FComponentCollisionSettings
+{
+	int32 CollisionTypeFlag = 0;		// this is ECollisionTraceFlag
+	bool bIsGeneratedCollision = true;
+};
+
+
+/**
+ * @return current Component collision settings
+ */
+MODELINGCOMPONENTS_API FComponentCollisionSettings GetCollisionSettings(
+	const UPrimitiveComponent* Component);
+
+
+/**
+ * Apply Transform to any Simple Collision geometry owned by Component.
+ * Note that Nonuniform scaling support is very limited and will generally enlarge collision volumes.
+ */
+MODELINGCOMPONENTS_API bool TransformSimpleCollision(
+	UPrimitiveComponent* Component,
+	const FTransform3d& Transform);
+
+/**
+ * Replace existing Simple Collision geometry in Component with that defined by ShapeSet,
+ * and update the Component/BodySetup collision settings
+ */
+MODELINGCOMPONENTS_API bool SetSimpleCollision(
+	UPrimitiveComponent* Component,
+	const FSimpleShapeSet3d* ShapeSet,
+	FComponentCollisionSettings CollisionSettings = FComponentCollisionSettings() );
+
+
+/**
+ * Apply Transform to the existing Simple Collision geometry in Component and then append to to ShapeSetOut
+ */
+MODELINGCOMPONENTS_API bool AppendSimpleCollision(
+	const UPrimitiveComponent* SourceComponent,
+	FSimpleShapeSet3d* ShapeSetOut,
+	const FTransform3d& Transform);
+
+/**
+ * Apply TransformSequence (in-order) to the existing Simple Collision geometry in Component and then append to to ShapeSetOut
+ */
+MODELINGCOMPONENTS_API bool AppendSimpleCollision(
+	const UPrimitiveComponent* SourceComponent,
+	FSimpleShapeSet3d* ShapeSetOut,
+	const TArray<FTransform3d>& TransformSeqeuence);
+
+/**
+ * Replace existing Simple Collision geometry in BodySetup with that defined by NewGeometry,
+ * and update the Component/BodySetup collision settings. Optional StaticMesh argument allows
+ * for necessary updates to it and any active UStaticMeshComponent that reference it
+ */
+MODELINGCOMPONENTS_API void UpdateSimpleCollision(
+	UBodySetup* BodySetup,
+	const FKAggregateGeom* NewGeometry,
+	UStaticMesh* StaticMesh = nullptr,
+	FComponentCollisionSettings CollisionSettings = FComponentCollisionSettings());
+
+
+
+}  // end namespace Geometry
+}  // end namespace UE
\ No newline at end of file
diff --git a/Engine/Plugins/Experimental/MeshModelingToolset/Source/ModelingComponents/Public/Physics/PhysicsDataCollection.h b/Engine/Plugins/Experimental/MeshModelingToolset/Source/ModelingComponents/Public/Physics/PhysicsDataCollection.h
--- a/Engine/Plugins/Experimental/MeshModelingToolset/Source/ModelingComponents/Public/Physics/PhysicsDataCollection.h
+++ b/Engine/Plugins/Experimental/MeshModelingToolset/Source/ModelingComponents/Public/Physics/PhysicsDataCollection.h
@@ -4,8 +4,8 @@
 
 #include "CoreMinimal.h"
 #include "ShapeApproximation/SimpleShapeSet3.h"
-#include "Engine/Classes/Engine/StaticMesh.h"
-#include "Engine/Classes/PhysicsEngine/AggregateGeom.h"
+#include "Engine/StaticMesh.h"
+#include "PhysicsEngine/AggregateGeom.h"
 
 class UActorComponent;
 class UBodySetup;
