MeshModelingTools: Add curved stair generator

#rnx

diff --git a/Engine/Plugins/Experimental/GeometryProcessing/Source/GeometricObjects/Public/Generators/StairGenerator.h b/Engine/Plugins/Experimental/GeometryProcessing/Source/GeometricObjects/Public/Generators/StairGenerator.h
--- a/Engine/Plugins/Experimental/GeometryProcessing/Source/GeometricObjects/Public/Generators/StairGenerator.h
+++ b/Engine/Plugins/Experimental/GeometryProcessing/Source/GeometricObjects/Public/Generators/StairGenerator.h
@@ -20,9 +20,6 @@ public:
 	/** The width of each step. */
 	float StepWidth = 150.0f;
 
-	/** The depth of each step. */
-	float StepDepth = 30.0f;
-
 	/** The height of each step. */
 	float StepHeight = 20.0f;
 
@@ -48,6 +45,8 @@ public:
 	EBaseStyle BaseStyle = EBaseStyle::SolidBase;
 
 public:
+	virtual ~FStairGenerator() {}
+
 	/** Generate the mesh */
 	virtual FMeshShapeGenerator& Generate() override
 	{
@@ -318,7 +317,7 @@ public:
 		// Compute UVs
 		UVDesc.SetNum(4 * NumQuads);
 		int UVId = 0;
-		float MaxDimension = FMathf::Max(FMathf::Max(NumSteps * StepDepth, NumSteps * StepHeight), StepWidth);
+		float MaxDimension = GetMaxDimension();
 		float UVScale = (bScaleUVByAspectRatio) ? (1.0f / (float)MaxDimension) : 1.0f;
 		for (const ESide& Side : AllSides)
 		{
@@ -552,6 +551,11 @@ protected:
 	int LastFaceId = 0;
 
 protected:
+	/**
+	 * Reset state data on the generator.
+	 * 
+	 * This is invoked at the head of the Generate() method.
+	 */
 	virtual void Reset()
 	{
 		VertexIds.Reset();
@@ -623,6 +627,12 @@ protected:
 	 * @param UVScale The UV scale
 	 */
 	virtual FVector2f GenerateUV(ESide Side, int Step, int VertexId, float UVScale) = 0;
+
+	/**
+	 * Returns the max dimension of the staircase for the purposes
+	 * of computing the world UV scale.
+	 */
+	virtual float GetMaxDimension() = 0;
 };
 
 /**
@@ -630,6 +640,13 @@ protected:
  */
 class /*GEOMETRICOBJECTS_API*/ FLinearStairGenerator : public FStairGenerator
 {
+public:
+	/** The depth of each step. */
+	float StepDepth = 30.0f;
+
+public:
+	virtual ~FLinearStairGenerator() {}
+
 protected:
 	virtual FVector3d GenerateVertex(ESide Side, int VertexColumn, int VertexRow) override
 	{
@@ -646,26 +663,26 @@ protected:
 
 	virtual FVector3f GenerateNormal(ESide Side, int VertexId) override
 	{
-		FVector3f N(0.0f, 0.0f, 0.0f);
+		FVector3f N;
 		switch (Side)
 		{
 		case ESide::Right:
-			N = FVector3f(0.0f, 1.0f, 0.0f);
+			N = FVector3f::UnitY();
 			break;
 		case ESide::Left:
-			N = FVector3f(0.0f, -1.0f, 0.0f);
+			N = -FVector3f::UnitY();
 			break;
 		case ESide::Front:
-			N = FVector3f(-1.0f, 0.0f, 0.0f);
+			N = -FVector3f::UnitX();
 			break;
 		case ESide::Top:
-			N = FVector3f(0.0f, 0.0f, 1.0f);
+			N = FVector3f::UnitZ();
 			break;
 		case ESide::Back:
-			N = FVector3f(1.0f, 0.0f, 0.0f);
+			N = FVector3f::UnitX();
 			break;
 		case ESide::Bottom:
-			N = FVector3f(0.0f, 0.0f, -1.0f);
+			N = -FVector3f::UnitZ();
 			break;
 		default:
 			check(false);
@@ -676,7 +693,6 @@ protected:
 
 	virtual FVector2f GenerateUV(ESide Side, int Step, int VertexId, float UVScale) override
 	{
-		check(VertexId < VertexIdsToColumnRow.Num());
 		const int Col = VertexIdsToColumnRow[VertexId].A;
 		const int Row = VertexIdsToColumnRow[VertexId].B;
 		FVector2f UV(0.0f, 0.0f);
@@ -688,7 +704,7 @@ protected:
 			const float UScale = NumSteps * StepDepth * UVScale;
 			const float VScale = NumSteps * StepHeight * UVScale;
 			UV.X = FMathf::Lerp(-0.5f, 0.5f, float(Col % LeftSideColumnId) / float(NumSteps + 1)) * UScale + 0.5f;
-			UV.Y = FMathf::Lerp(-0.5f, 0.5f, float(Row) / float(NumSteps + 1)) * VScale + 0.5f;;
+			UV.Y = FMathf::Lerp(-0.5f, 0.5f, float(Row) / float(NumSteps + 1)) * VScale + 0.5f;
 			break;
 		}
 		case ESide::Front:
@@ -739,4 +755,237 @@ protected:
 		}
 		return UV;
 	}
-};
\ No newline at end of file
+
+	virtual float GetMaxDimension()
+	{
+		return FMathf::Max(FMathf::Max(NumSteps * StepDepth, NumSteps * StepHeight), StepWidth);
+	}
+};
+
+/**
+ * Generate an oriented Curved Stair mesh.
+ */
+class /*GEOMETRICOBJECTS_API*/ FCurvedStairGenerator : public FStairGenerator
+{
+public:
+	/** Inner radius of the curved staircase */
+	float InnerRadius = 150.0f;
+
+	/** Curve angle of the staircase (in degrees) */
+	float CurveAngle = 90.0f;
+
+public:
+	virtual ~FCurvedStairGenerator() {}
+
+protected:
+	typedef FStairGenerator ParentClass;
+
+	/** Precompute/cached data */
+	bool bIsClockwise = true;
+	float CurveRadians = 0.0f;
+	float CurveRadiansPerStep = 0.0f;
+	float OuterRadius = 0.0f;
+	float RadiusRatio = 1.0f;
+	FVector3f BackNormal = FVector3f::Zero();
+
+protected:
+	virtual void Reset()
+	{
+		ParentClass::Reset();
+
+		bIsClockwise = CurveAngle > 0.0f;
+		CurveRadians = CurveAngle * TMathUtilConstants<float>::DegToRad;
+		CurveRadiansPerStep = CurveRadians / NumSteps;
+		OuterRadius = InnerRadius + StepWidth;
+		RadiusRatio = OuterRadius / InnerRadius;
+		BackNormal = FVector3f::Zero();
+	}
+
+	virtual FVector3d GenerateVertex(ESide Side, int VertexColumn, int VertexRow) override
+	{
+		if (Side != ESide::Right && Side != ESide::Left)
+		{
+			// Vertices are only generated for Right & Left sides.
+			check(false);
+		}
+		float X = 0.0f;
+		float Y = 0.0f;
+		float Z = VertexRow * StepHeight;
+
+		float XCoeff = FMathf::Cos(VertexColumn * CurveRadiansPerStep);
+		float YCoeff = FMathf::Sin(VertexColumn * CurveRadiansPerStep);
+		if (bIsClockwise)
+		{
+			X = (Side == ESide::Right ? XCoeff * InnerRadius : XCoeff * OuterRadius);
+			Y = (Side == ESide::Right ? YCoeff * InnerRadius : YCoeff * OuterRadius);
+		}
+		else
+		{
+			X = (Side == ESide::Right ? XCoeff * -OuterRadius : XCoeff * -InnerRadius);
+			Y = (Side == ESide::Right ? YCoeff * -OuterRadius : YCoeff * -InnerRadius);
+		}
+		return FVector3d(X, Y, Z);
+	}
+
+	virtual FVector3f GenerateNormal(ESide Side, int VertexId) override
+	{
+		const int Col = VertexIdsToColumnRow[VertexId].A;
+		const int Row = VertexIdsToColumnRow[VertexId].B;
+
+		FVector3f N;
+		switch (Side)
+		{
+		case ESide::Right:
+		case ESide::Left:
+		{
+			float X = FMathf::Cos((Col % (NumSteps + 1)) * CurveRadiansPerStep);
+			float Y = FMathf::Sin((Col % (NumSteps + 1)) * CurveRadiansPerStep);
+			N = (Side == ESide::Right ? FVector3f(-X, -Y, 0.0f) : FVector3f(X, Y, 0.0f));
+			N.Normalize();
+			break;
+		}
+		case ESide::Front:
+		{
+			float X = FMathf::Cos((Col % (NumSteps + 1)) * CurveRadiansPerStep);
+			float Y = FMathf::Sin((Col % (NumSteps + 1)) * CurveRadiansPerStep);
+			N = FVector3f(Y, -X, 0.0f);
+			N.Normalize();
+			break;
+		}
+		case ESide::Top:
+		{
+			N = FVector3f::UnitZ();
+			break;
+		}
+		case ESide::Back:
+		{
+			if (BackNormal == FVector3f::Zero())
+			{
+				float X = FMathf::Cos(NumSteps * CurveRadiansPerStep);
+				float Y = FMathf::Sin(NumSteps * CurveRadiansPerStep);
+				BackNormal = FVector3f(-Y, X, 0.0f);
+			}
+			N = BackNormal;
+			break;
+		}
+		case ESide::Bottom:
+		{
+			N = -FVector3f::UnitZ();
+			break;
+		}
+		default:
+		{
+			check(false);
+			break;
+		}
+		}
+		return N;
+	}
+
+	virtual FVector2f GenerateUV(ESide Side, int Step, int VertexId, float UVScale) override
+	{
+		const int Col = VertexIdsToColumnRow[VertexId].A;
+		const int Row = VertexIdsToColumnRow[VertexId].B;
+		FVector2f UV;
+		switch (Side)
+		{
+		case ESide::Right:
+		case ESide::Left:
+		{
+			const float UScale = OuterRadius * UVScale;
+			const float VScale = NumSteps * StepHeight * UVScale;
+			UV.X = FMathf::Lerp(-0.5f, 0.5f, float(Col % LeftSideColumnId) / float(NumSteps + 1));
+			UV.Y = FMathf::Lerp(-0.5f, 0.5f, float(Row) / float(NumSteps + 1));
+
+			// Proportionally scale the UVs of each side according
+			// to the RadiusRatio and ensure [0,1].
+			if (RadiusRatio * UScale > 1.0f)
+			{
+				UV.X /= (Side == ESide::Left ? UScale : UScale * RadiusRatio);
+			}
+			else if (Side == ESide::Left)
+			{
+				UV.X *= RadiusRatio;
+			}
+			UV.X = UV.X * UScale + 0.5f;
+			UV.Y = UV.Y * VScale + 0.5f;
+			break;
+		}
+		case ESide::Front:
+		{
+			const float UScale = StepWidth * UVScale;
+			const float VScale = StepHeight * UVScale;
+			UV.X = Col < LeftSideColumnId ? 0.5f : -0.5f;
+			UV.Y = Row > Step ? -0.5f : 0.5f;
+			UV.X = UV.X * UScale + 0.5f;
+			UV.Y = UV.Y * VScale + 0.5f;
+			break;
+		}
+		case ESide::Top:
+		{
+			const float UScale = StepWidth * UVScale;
+			const float VScale = OuterRadius / NumSteps * UVScale;
+			UV.X = Col < LeftSideColumnId ? 0.5f : -0.5f;
+			UV.Y = (Col % LeftSideColumnId) > Step ? -0.5f : 0.5f;
+
+			// Proportionally scale the outer edge of top faces
+			// to avoid stretching while ensuring [0,1].
+			if (RadiusRatio * VScale > 1.0f)
+			{
+				UV.Y /= (Col >= LeftSideColumnId ? VScale : VScale * RadiusRatio);
+			}
+			else if (Col >= LeftSideColumnId)
+			{
+				UV.Y *= RadiusRatio;
+			}
+			
+			UV.X = UV.X * UScale + 0.5f;
+			UV.Y = UV.Y * VScale + 0.5f;
+			break;
+		}
+		case ESide::Back:
+		{
+			const float UScale = StepWidth * UVScale;
+			const float VScale = NumSteps * StepHeight * UVScale;
+			UV.X = Col < LeftSideColumnId ? 0.5f : -0.5f;
+			UV.Y = FMathf::Lerp(-0.5f, 0.5f, float(Row) / float(NumSteps + 1));
+			UV.X = UV.X * UScale + 0.5f;
+			UV.Y = UV.Y * VScale + 0.5f;
+			break;
+		}
+		case ESide::Bottom:
+		{
+			const float UScale = StepWidth * UVScale;
+			const float VScale = OuterRadius * UVScale;
+			UV.X = Col < LeftSideColumnId ? 0.5f : -0.5f;
+			UV.Y = FMathf::Lerp(-0.5f, 0.5f, float(Col % LeftSideColumnId) / float(NumSteps + 1));
+
+			// Proportionally scale outer edge of bottom faces
+			// to avoid stretching while ensuring [0,1].
+			if (RadiusRatio * VScale > 1.0f)
+			{
+				UV.Y /= (Col >= LeftSideColumnId ? VScale : VScale * RadiusRatio);
+			}
+			else if (Col >= LeftSideColumnId)
+			{
+				UV.Y *= RadiusRatio;
+			}
+			UV.X = UV.X * UScale + 0.5f;
+			UV.Y = UV.Y * VScale + 0.5f;
+			break;
+		}
+		default:
+		{
+			check(false);
+			break;
+		}
+		}
+		return UV;
+	}
+
+	virtual float GetMaxDimension()
+	{
+		float MaxDepth = FMathf::Abs(CurveRadians) * OuterRadius;
+		return FMathf::Max(FMathf::Max(MaxDepth, NumSteps * StepHeight), StepWidth);
+	}
+};
diff --git a/Engine/Plugins/Experimental/MeshModelingToolset/Source/MeshModelingTools/Private/AddPrimitiveTool.cpp b/Engine/Plugins/Experimental/MeshModelingToolset/Source/MeshModelingTools/Private/AddPrimitiveTool.cpp
--- a/Engine/Plugins/Experimental/MeshModelingToolset/Source/MeshModelingTools/Private/AddPrimitiveTool.cpp
+++ b/Engine/Plugins/Experimental/MeshModelingToolset/Source/MeshModelingTools/Private/AddPrimitiveTool.cpp
@@ -77,8 +77,8 @@ UInteractiveTool* UAddPrimitiveToolBuilder::BuildTool(const FToolBuilderState& S
 	case EMakeMeshShapeType::Sphere:
 		NewTool = NewObject<UAddSpherePrimitiveTool>(SceneState.ToolManager);
 		break;
-	case EMakeMeshShapeType::LinearStairs:
-		NewTool = NewObject<UAddLinearStairsPrimitiveTool>(SceneState.ToolManager);
+	case EMakeMeshShapeType::Stairs:
+		NewTool = NewObject<UAddStairsPrimitiveTool>(SceneState.ToolManager);
 		break;
 	default:
 		break;
@@ -659,24 +659,45 @@ void UAddSphericalBoxPrimitiveTool::GenerateMesh(FDynamicMesh3* OutMesh) const
 	OutMesh->Copy(&SphereGen);
 }
 
-UAddLinearStairsPrimitiveTool::UAddLinearStairsPrimitiveTool(const FObjectInitializer& ObjectInitializer)
-	: Super(ObjectInitializer.SetDefaultSubobjectClass<UProceduralLinearStairsToolProperties>(TEXT("ShapeSettings")))
+UAddStairsPrimitiveTool::UAddStairsPrimitiveTool(const FObjectInitializer& ObjectInitializer)
+	: Super(ObjectInitializer.SetDefaultSubobjectClass<UProceduralStairsToolProperties>(TEXT("ShapeSettings")))
 {
-	AssetName = TEXT("LinearStairs");
-	SetToolDisplayName(LOCTEXT("LinearStairsToolName", "Create LinearStairs"));
+	AssetName = TEXT("Stairs");
+	SetToolDisplayName(LOCTEXT("StairsToolName", "Create Stairs"));
 }
 
-void UAddLinearStairsPrimitiveTool::GenerateMesh(FDynamicMesh3* OutMesh) const
+void UAddStairsPrimitiveTool::GenerateMesh(FDynamicMesh3* OutMesh) const
 {
-	FLinearStairGenerator StairGen;
-	auto* StairSettings = Cast<UProceduralLinearStairsToolProperties>(ShapeSettings);
-	StairGen.StepWidth = StairSettings->StepWidth;
-	StairGen.StepHeight = StairSettings->StepHeight;
-	StairGen.StepDepth = StairSettings->StepDepth;
-	StairGen.NumSteps = StairSettings->NumSteps;
-	StairGen.bPolygroupPerQuad = (ShapeSettings->PolygroupMode == EMakeMeshPolygroupMode::PerQuad);
-	StairGen.Generate();
-	OutMesh->Copy(&StairGen);
+	auto* StairSettings = Cast<UProceduralStairsToolProperties>(ShapeSettings);
+	switch (StairSettings->StairsType)
+	{
+	case EProceduralStairsType::Linear:
+	{
+		FLinearStairGenerator StairGen;
+		StairGen.StepWidth = StairSettings->StepWidth;
+		StairGen.StepHeight = StairSettings->StepHeight;
+		StairGen.StepDepth = StairSettings->StepDepth;
+		StairGen.NumSteps = StairSettings->NumSteps;
+		StairGen.bPolygroupPerQuad = (ShapeSettings->PolygroupMode == EMakeMeshPolygroupMode::PerQuad);
+		StairGen.Generate();
+		OutMesh->Copy(&StairGen);
+		break;
+	}
+	case EProceduralStairsType::Curved:
+	{
+		FCurvedStairGenerator StairGen;
+		StairGen.StepWidth = StairSettings->StepWidth;
+		StairGen.StepHeight = StairSettings->StepHeight;
+		StairGen.NumSteps = StairSettings->NumSteps;
+		StairGen.InnerRadius = StairSettings->InnerRadius;
+		StairGen.CurveAngle = StairSettings->CurveAngle;
+		StairGen.bPolygroupPerQuad = (ShapeSettings->PolygroupMode == EMakeMeshPolygroupMode::PerQuad);
+		StairGen.Generate();
+		OutMesh->Copy(&StairGen);
+		break;
+	}
+	}
 }
 
+
 #undef LOCTEXT_NAMESPACE
diff --git a/Engine/Plugins/Experimental/MeshModelingToolset/Source/MeshModelingTools/Public/AddPrimitiveTool.h b/Engine/Plugins/Experimental/MeshModelingToolset/Source/MeshModelingTools/Public/AddPrimitiveTool.h
--- a/Engine/Plugins/Experimental/MeshModelingToolset/Source/MeshModelingTools/Public/AddPrimitiveTool.h
+++ b/Engine/Plugins/Experimental/MeshModelingToolset/Source/MeshModelingTools/Public/AddPrimitiveTool.h
@@ -38,7 +38,7 @@ public:
 		Torus,
 		SphericalBox,
 		Sphere,
-		LinearStairs
+		Stairs
 	};
 
 	IAssetGenerationAPI* AssetAPI;
@@ -340,28 +340,50 @@ public:
 	int Subdivisions = 16;
 };
 
+UENUM()
+enum class EProceduralStairsType
+{
+	/** Create a linear staircase */
+	Linear UMETA(DisplayName = "Linear Stairs"),
+	/** Create a curved staircase */
+	Curved UMETA(DisplayName = "Curved Stairs")
+};
+
 UCLASS()
-class MESHMODELINGTOOLS_API UProceduralLinearStairsToolProperties : public UProceduralShapeToolProperties
+class MESHMODELINGTOOLS_API UProceduralStairsToolProperties : public UProceduralShapeToolProperties
 {
 	GENERATED_BODY()
 public:
-	/** Width of Shape */
-	UPROPERTY(EditAnywhere, Category = "ShapeSettings|Dimensions", meta = (DisplayName = "Step Width", UIMin = "1.0", UIMax = "1000.0", ClampMin = "0.0001", ClampMax = "1000000.0", ProceduralShapeSetting))
-	float StepWidth = 150.f;
+	/** Type of staircase to create */
+	UPROPERTY(EditAnywhere, Category = "ShapeSettings")
+		EProceduralStairsType StairsType = EProceduralStairsType::Linear;
 
-	/** Depth of Shape */
-	UPROPERTY(EditAnywhere, Category = "ShapeSettings|Dimensions", meta = (DisplayName = "Step Depth", UIMin = "1.0", UIMax = "1000.0", ClampMin = "0.0001", ClampMax = "1000000.0", ProceduralShapeSetting))
-	float StepDepth = 30.f;
+	/** Number of Steps of Shape */
+	UPROPERTY(EditAnywhere, Category = "ShapeSettings|Dimensions", meta = (DisplayName = "Number of Steps", UIMin = "1", UIMax = "100", ClampMin = "1", ClampMax = "1000", ProceduralShapeSetting))
+	int NumSteps = 8;
 
-	/** Height of Shape*/
+	/** Width of each step */
+	UPROPERTY(EditAnywhere, Category = "ShapeSettings|Dimensions", meta = (DisplayName = "Step Width", UIMin = "1.0", UIMax = "1000.0", ClampMin = "0.0001", ClampMax = "1000000.0", ProceduralShapeSetting))
+	float StepWidth = 150.0f;
+
+	/** Height of each step */
 	UPROPERTY(EditAnywhere, Category = "ShapeSettings|Dimensions", meta = (DisplayName = "Step Height", UIMin = "1.0", UIMax = "1000.0", ClampMin = "0.0001", ClampMax = "1000000.0", ProceduralShapeSetting))
-	float StepHeight = 20.f;
+	float StepHeight = 20.0f;
 
-	/** Number of Steps of Shape */
-	UPROPERTY(EditAnywhere, Category = "ShapeSettings|Dimensions", meta = (DisplayName = "Number of Steps", UIMin = "1", UIMax = "100", ClampMin = "1", ClampMax = "1000", ProceduralShapeSetting))
-	int NumSteps = 5;
+	/** Depth of each step for linear stair shapes */
+	UPROPERTY(EditAnywhere, Category = "ShapeSettings|Dimensions", meta = (DisplayName = "Step Depth", UIMin = "1.0", UIMax = "1000.0", ClampMin = "0.0001", ClampMax = "1000000.0", ProceduralShapeSetting, EditCondition = "StairsType == EProceduralStairsType::Linear", EditConditionHides))
+	float StepDepth = 30.0f;
+
+	/** Curve angle for curved/spiral stair shapes */
+	UPROPERTY(EditAnywhere, Category = "ShapeSettings|Dimensions", meta = (DisplayName = "Curve Angle", UIMin = "-360.0", UIMax = "360.0", ClampMin = "-360.0", ClampMax = "360.0", ProceduralShapeSetting, EditCondition = "StairsType == EProceduralStairsType::Curved", EditConditionHides))
+	float CurveAngle = 90.0f;
+
+	/** Inner radius for curved/spiral stair shapes */
+	UPROPERTY(EditAnywhere, Category = "ShapeSettings|Dimensions", meta = (DisplayName = "Inner Radius", UIMin = "1.0", UIMax = "1000.0", ClampMin = "0.0001", ClampMax = "1000000.0", ProceduralShapeSetting, EditCondition = "StairsType == EProceduralStairsType::Curved", EditConditionHides))
+	float InnerRadius = 150.0f;
 };
 
+
 UCLASS(Transient)
 class MESHMODELINGTOOLS_API ULastActorInfo : public UObject
 {
@@ -580,12 +602,13 @@ protected:
 };
 
 UCLASS()
-class UAddLinearStairsPrimitiveTool : public UAddPrimitiveTool
+class UAddStairsPrimitiveTool : public UAddPrimitiveTool
 {
 	GENERATED_BODY()
 public:
-	UAddLinearStairsPrimitiveTool(const FObjectInitializer& ObjectInitializer);
+	UAddStairsPrimitiveTool(const FObjectInitializer& ObjectInitializer);
 protected:
 	void GenerateMesh(FDynamicMesh3* OutMesh) const override;
 };
 
+
diff --git a/Engine/Plugins/Experimental/ModelingToolsEditorMode/Source/ModelingToolsEditorMode/Private/ModelingToolsEditorMode.cpp b/Engine/Plugins/Experimental/ModelingToolsEditorMode/Source/ModelingToolsEditorMode/Private/ModelingToolsEditorMode.cpp
--- a/Engine/Plugins/Experimental/ModelingToolsEditorMode/Source/ModelingToolsEditorMode/Private/ModelingToolsEditorMode.cpp
+++ b/Engine/Plugins/Experimental/ModelingToolsEditorMode/Source/ModelingToolsEditorMode/Private/ModelingToolsEditorMode.cpp
@@ -347,9 +347,9 @@ void UModelingToolsEditorMode::Enter()
 	RegisterPrimitiveToolFunc(ToolManagerCommands.BeginAddSphericalBoxPrimitiveTool,
 							  TEXT("BeginAddSphericalBoxPrimitiveTool"),
 							  UAddPrimitiveToolBuilder::EMakeMeshShapeType::SphericalBox);
-	RegisterPrimitiveToolFunc(ToolManagerCommands.BeginAddLinearStairsPrimitiveTool,
-							  TEXT("BeginAddLinearStairsPrimitiveTool"),
-							  UAddPrimitiveToolBuilder::EMakeMeshShapeType::LinearStairs);
+	RegisterPrimitiveToolFunc(ToolManagerCommands.BeginAddStairsPrimitiveTool,
+							  TEXT("BeginAddStairsPrimitiveTool"),
+							  UAddPrimitiveToolBuilder::EMakeMeshShapeType::Stairs);
 
 	//
 	// make shape tools
diff --git a/Engine/Plugins/Experimental/ModelingToolsEditorMode/Source/ModelingToolsEditorMode/Private/ModelingToolsEditorModeToolkit.cpp b/Engine/Plugins/Experimental/ModelingToolsEditorMode/Source/ModelingToolsEditorMode/Private/ModelingToolsEditorModeToolkit.cpp
--- a/Engine/Plugins/Experimental/ModelingToolsEditorMode/Source/ModelingToolsEditorMode/Private/ModelingToolsEditorModeToolkit.cpp
+++ b/Engine/Plugins/Experimental/ModelingToolsEditorMode/Source/ModelingToolsEditorMode/Private/ModelingToolsEditorModeToolkit.cpp
@@ -469,7 +469,7 @@ void FModelingToolsEditorModeToolkit::BuildToolPalette_Experimental(FName Palett
 		ToolbarBuilder.AddToolBarButton(Commands.BeginAddRoundedRectanglePrimitiveTool);
 		ToolbarBuilder.AddToolBarButton(Commands.BeginAddDiscPrimitiveTool);
 		ToolbarBuilder.AddToolBarButton(Commands.BeginAddPuncturedDiscPrimitiveTool);
-		ToolbarBuilder.AddToolBarButton(Commands.BeginAddLinearStairsPrimitiveTool);
+		ToolbarBuilder.AddToolBarButton(Commands.BeginAddStairsPrimitiveTool);
 	}
 	else if (PaletteIndex == CreateTabName)
 	{
diff --git a/Engine/Plugins/Experimental/ModelingToolsEditorMode/Source/ModelingToolsEditorMode/Private/ModelingToolsManagerActions.cpp b/Engine/Plugins/Experimental/ModelingToolsEditorMode/Source/ModelingToolsEditorMode/Private/ModelingToolsManagerActions.cpp
--- a/Engine/Plugins/Experimental/ModelingToolsEditorMode/Source/ModelingToolsEditorMode/Private/ModelingToolsManagerActions.cpp
+++ b/Engine/Plugins/Experimental/ModelingToolsEditorMode/Source/ModelingToolsEditorMode/Private/ModelingToolsManagerActions.cpp
@@ -33,7 +33,7 @@ void FModelingToolsManagerCommands::RegisterCommands()
 	UI_COMMAND(BeginAddTorusPrimitiveTool, "Torus", "Create new Torus StaticMesh Assets", EUserInterfaceActionType::ToggleButton, FInputChord());
 	UI_COMMAND(BeginAddSphericalBoxPrimitiveTool, "SphrB", "Create new Sphere StaticMesh Assets with Box Parameterization", EUserInterfaceActionType::ToggleButton, FInputChord());
 	UI_COMMAND(BeginAddSpherePrimitiveTool, "SphrA", "Create new Sphere StaticMesh Assets with Lat/Long Parameterization", EUserInterfaceActionType::ToggleButton, FInputChord());
-	UI_COMMAND(BeginAddLinearStairsPrimitiveTool, "StairL", "Create new Linear Stairs StaticMesh Assets", EUserInterfaceActionType::ToggleButton, FInputChord());
+	UI_COMMAND(BeginAddStairsPrimitiveTool, "Stair", "Create new Stairs StaticMesh Assets", EUserInterfaceActionType::ToggleButton, FInputChord());
 
 	UI_COMMAND(BeginAddPatchTool, "Patch", "Add Patch", EUserInterfaceActionType::ToggleButton, FInputChord());
 	UI_COMMAND(BeginRevolveBoundaryTool, "BdryRev", "Revolve Mesh Boundary Loop into New Asset", EUserInterfaceActionType::ToggleButton, FInputChord());
diff --git a/Engine/Plugins/Experimental/ModelingToolsEditorMode/Source/ModelingToolsEditorMode/Public/ModelingToolsManagerActions.h b/Engine/Plugins/Experimental/ModelingToolsEditorMode/Source/ModelingToolsEditorMode/Public/ModelingToolsManagerActions.h
--- a/Engine/Plugins/Experimental/ModelingToolsEditorMode/Source/ModelingToolsEditorMode/Public/ModelingToolsManagerActions.h
+++ b/Engine/Plugins/Experimental/ModelingToolsEditorMode/Source/ModelingToolsEditorMode/Public/ModelingToolsManagerActions.h
@@ -26,7 +26,7 @@ public:
 	TSharedPtr<FUICommandInfo> BeginAddTorusPrimitiveTool;
 	TSharedPtr<FUICommandInfo> BeginAddSphericalBoxPrimitiveTool;
 	TSharedPtr<FUICommandInfo> BeginAddSpherePrimitiveTool;
-	TSharedPtr<FUICommandInfo> BeginAddLinearStairsPrimitiveTool;
+	TSharedPtr<FUICommandInfo> BeginAddStairsPrimitiveTool;
 
 	TSharedPtr<FUICommandInfo> BeginAddPatchTool;
 	TSharedPtr<FUICommandInfo> BeginRevolveBoundaryTool;
