refactor jacketing to use one AABB tree per mesh, and transform samples into the local space of each tree, rather than a giant combined AABB tree for all meshes

diff --git a/Engine/Plugins/Experimental/GeometryProcessing/Source/DynamicMesh/Private/Operations/RemoveOccludedTriangles.cpp b/Engine/Plugins/Experimental/GeometryProcessing/Source/DynamicMesh/Private/Operations/RemoveOccludedTriangles.cpp
--- a/Engine/Plugins/Experimental/GeometryProcessing/Source/DynamicMesh/Private/Operations/RemoveOccludedTriangles.cpp
+++ b/Engine/Plugins/Experimental/GeometryProcessing/Source/DynamicMesh/Private/Operations/RemoveOccludedTriangles.cpp
@@ -73,7 +73,7 @@ namespace UE
 		bool DYNAMICMESH_API RemoveInternalTriangles(FDynamicMesh3& Mesh, bool bTestPerCompoment,
 			EOcclusionTriangleSampling SamplingMethod,
 			EOcclusionCalculationMode OcclusionMode,
-			int RandomSamplesPerTri, double WindingNumberThreshold)
+			int RandomSamplesPerTri, double WindingNumberThreshold, bool bTestOccludedByAny)
 		{
 			if (bTestPerCompoment)
 			{
@@ -96,7 +96,7 @@ namespace UE
 				TArrayView<FTransform3d> TransformsView(&Id, 1);
 				TArrayView<TMeshAABBTree3<FComponentMesh>*> SpatialsView(Spatials.GetData(), Spatials.Num());
 				TArrayView<TFastWindingTree<FComponentMesh>*> WindingsView(Windings.GetData(), Windings.Num());
-				return Remover.Apply(TransformsView, SpatialsView, WindingsView);
+				return Remover.Apply(TransformsView, SpatialsView, WindingsView, TArrayView<const FTransform3d>(), bTestOccludedByAny);
 			}
 			else
 			{
diff --git a/Engine/Plugins/Experimental/GeometryProcessing/Source/DynamicMesh/Public/Operations/RemoveOccludedTriangles.h b/Engine/Plugins/Experimental/GeometryProcessing/Source/DynamicMesh/Public/Operations/RemoveOccludedTriangles.h
--- a/Engine/Plugins/Experimental/GeometryProcessing/Source/DynamicMesh/Public/Operations/RemoveOccludedTriangles.h
+++ b/Engine/Plugins/Experimental/GeometryProcessing/Source/DynamicMesh/Public/Operations/RemoveOccludedTriangles.h
@@ -51,7 +51,7 @@ namespace UE
 		bool DYNAMICMESH_API RemoveInternalTriangles(FDynamicMesh3& Mesh, bool bTestPerComponent = false,
 			EOcclusionTriangleSampling SamplingMethod = EOcclusionTriangleSampling::Centroids, 
 			EOcclusionCalculationMode OcclusionMode = EOcclusionCalculationMode::FastWindingNumber,
-			int RandomSamplesPerTri = 1, double WindingNumberThreshold = .5);
+			int RandomSamplesPerTri = 0, double WindingNumberThreshold = .5, bool bTestOccludedByAny = false);
 	}
 }
 
@@ -84,10 +84,14 @@ public:
 	 * @param MeshLocalToOccluderSpaces Transforms to take instances of the local mesh into the space of the occluders
 	 * @param Spatials AABB trees for all occluders
 	 * @param FastWindingTrees Precomputed fast winding trees for occluders
+	 * @param SpatialTransforms Transforms AABB/winding tree to shared occluder space (if empty, Identity is used)
+	 * @param bTestOccludedByAny If true, a triangle is occluded if it is fully occluded by *any* occluder.
+	 *							 Otherwise, we test if it's occluded by the combination of all occluders.
 	 * @return true on success
 	 */
 	virtual bool Select(const TArrayView<const FTransform3d> MeshLocalToOccluderSpaces,
-		const TArrayView<TMeshAABBTree3<OccluderTriangleMeshType>*> Spatials, const TArrayView<TFastWindingTree<OccluderTriangleMeshType>*> FastWindingTrees)
+		const TArrayView<TMeshAABBTree3<OccluderTriangleMeshType>*> Spatials, const TArrayView<TFastWindingTree<OccluderTriangleMeshType>*> FastWindingTrees,
+		const TArrayView<const FTransform3d> SpatialTransforms = TArrayView<const FTransform3d>(), bool bTestOccludedByAny = false)
 	{
 		if (Cancelled())
 		{
@@ -125,36 +129,189 @@ public:
 			TriangleBaryCoordSamples.Add(FVector3d(1 / 3.0, 1 / 3.0, 1 / 3.0));
 		}
 
-		auto IsOccludedFWN = [this]
-			(TMeshAABBTree3<OccluderTriangleMeshType>* Spatial, TFastWindingTree<OccluderTriangleMeshType>* FastWindingTree, const FVector3d& Pt)
+		// Helper struct to track both the transform and its inverse
+		struct FTransformWithInv
 		{
-			return FastWindingTree->FastWindingNumber(Pt) > WindingIsoValue;
+			FTransform3d Transform;
+			FQuaterniond InvRot;
+			FVector3d InvScale;
+
+			FTransformWithInv() : 
+				Transform(FTransform3d::Identity()), InvRot(FQuaterniond::Identity()), InvScale(FVector3d::One())
+			{}
+			FTransformWithInv(const FTransform3d& Transform) : Transform(Transform)
+			{
+				ComputeInv();
+			}
+
+			void ComputeInv()
+			{
+				InvRot = Transform.GetRotation().Inverse();
+				InvScale = Transform.GetScale();
+				const double ScaleTolerance = FMathd::ZeroTolerance * 100;
+				for (int i = 0; i < 3; i++)
+				{
+					// for near-zero scale, pretend scale was a not-quite-as-small number instead; TODO handle these cases more correctly
+					if (FMath::Abs(InvScale[i]) < ScaleTolerance)
+					{
+						InvScale[i] = FMathd::SignNonZero(InvScale[i]) / ScaleTolerance;
+					}
+					else
+					{
+						InvScale[i] = 1.0 / InvScale[i];
+					}
+				}
+			}
+
+			inline FVector3d InverseTransformPosition(const FVector3d& P) const
+			{
+				return InvScale * (InvRot * (P - Transform.GetTranslation()));
+			}
+
+			inline FVector3d InverseTransformVector(const FVector3d& V) const
+			{
+				return InvScale * (InvRot * V);
+			}
+
+			// checks if the transform is an exact match
+			inline bool IsSameTransform(const FTransform3d& Other) const
+			{
+				return Other.GetTranslation() == Transform.GetTranslation()
+					&& Other.GetScale() == Transform.GetScale()
+					&& Other.GetRotation().EpsilonEqual(Transform.GetRotation(), 0);
+			}
+
+			// apply the inverse transform unless the point was originally transformed by this same transform; then just return the original point
+			// use to avoid floating point error shoving a sample into the surface that it came from
+			inline FVector3d InverseTransformUnlessMatch(const FVector3d& Pt, const FVector3d& OriginalPt, const FTransform3d& OriginalXF) const
+			{
+				if (IsSameTransform(OriginalXF))
+				{
+					return OriginalPt;
+				}
+				else
+				{
+					return InverseTransformPosition(Pt);
+				}
+			}
+		};
+
+		typedef TFunctionRef<bool(const TArrayView<TMeshAABBTree3<OccluderTriangleMeshType>*> Spatials, const TArrayView<TFastWindingTree<OccluderTriangleMeshType>*> FastWindingTrees,
+			const FVector3d& Pt, const FVector3d& OriginalPt, const FTransform3d& OriginalXF, const TArray<FTransformWithInv>& SpatialTransforms)> FIsOccludedFn;
+
+		// test if point is occluded by the combination of all spatials, according to winding number
+		FIsOccludedFn IsOccludedFWNTotal =
+			[this](const TArrayView<TMeshAABBTree3<OccluderTriangleMeshType>*> Spatials, const TArrayView<TFastWindingTree<OccluderTriangleMeshType>*> FastWindingTrees,
+			const FVector3d& Pt, const FVector3d& OriginalPt, const FTransform3d& OriginalXF, const TArray<FTransformWithInv>& SpatialTransforms) -> bool
+		{
+			double WindingSum = 0;
+			for (int Idx = 0, NumPts = FastWindingTrees.Num(); Idx < NumPts; ++Idx)
+			{
+				FVector3d XFPt = SpatialTransforms[Idx].InverseTransformUnlessMatch(Pt, OriginalPt, OriginalXF);
+				WindingSum += FastWindingTrees[Idx]->FastWindingNumber(XFPt);
+			}
+			return WindingSum > WindingIsoValue;
 		};
-		auto IsOccludedSimple = [this, &RayDirs, &NR]
-			(TMeshAABBTree3<OccluderTriangleMeshType>* Spatial, TFastWindingTree<OccluderTriangleMeshType>* FastWindingTree, const FVector3d& Pt)
+		// test if point is occluded by the any of the spatials, according to winding number
+		FIsOccludedFn IsOccludedFWNAny =
+			[this](const TArrayView<TMeshAABBTree3<OccluderTriangleMeshType>*> Spatials, const TArrayView<TFastWindingTree<OccluderTriangleMeshType>*> FastWindingTrees,
+				const FVector3d& Pt, const FVector3d& OriginalPt, const FTransform3d& OriginalXF, const TArray<FTransformWithInv>& SpatialTransforms) -> bool
+		{
+			for (int Idx = 0, NumPts = FastWindingTrees.Num(); Idx < NumPts; ++Idx)
+			{
+				FVector3d XFPt = SpatialTransforms[Idx].InverseTransformUnlessMatch(Pt, OriginalPt, OriginalXF);
+				if (FastWindingTrees[Idx]->FastWindingNumber(XFPt) > WindingIsoValue)
+				{
+					return true;
+				}
+			}
+			return false;
+		};
+		// test if point is occluded by the combination of all spatials, according to raycasts
+		FIsOccludedFn IsOccludedSimpleTotal = [this, &RayDirs, &NR]
+		(const TArrayView<TMeshAABBTree3<OccluderTriangleMeshType>*> Spatials, const TArrayView<TFastWindingTree<OccluderTriangleMeshType>*> FastWindingTrees,
+			const FVector3d& Pt, const FVector3d& OriginalPt, const FTransform3d& OriginalXF, const TArray<FTransformWithInv>& SpatialTransforms) -> bool
 		{
 			FRay3d Ray;
-			Ray.Origin = Pt;
-
+			
 			for (int RayIdx = 0; RayIdx < NR; ++RayIdx)
 			{
-				Ray.Direction = RayDirs[RayIdx];
-				int HitTID = Spatial->FindNearestHitTriangle(Ray);
-				if (HitTID == IndexConstants::InvalidID)
+				bool bAnyHit = false;
+				for (int SpatialIdx = 0, NumPts = Spatials.Num(); SpatialIdx < NumPts; ++SpatialIdx)
+				{
+					FVector3d XFPt = SpatialTransforms[SpatialIdx].InverseTransformUnlessMatch(Pt, OriginalPt, OriginalXF);
+					Ray.Direction = SpatialTransforms[SpatialIdx].InverseTransformVector(RayDirs[RayIdx]);
+					Ray.Origin = XFPt;
+					int HitTID = Spatials[SpatialIdx]->FindNearestHitTriangle(Ray);
+					if (HitTID != IndexConstants::InvalidID)
+					{
+						bAnyHit = true;
+						break;
+					}
+				}
+				if (!bAnyHit) // found a ray direction that hits none of the spatials
 				{
 					return false;
 				}
 			}
 			return true;
 		};
+		// test if point is occluded by the any of the spatials, according to raycasts
+		FIsOccludedFn IsOccludedSimpleAny = [this, &RayDirs, &NR]
+		(const TArrayView<TMeshAABBTree3<OccluderTriangleMeshType>*> Spatials, const TArrayView<TFastWindingTree<OccluderTriangleMeshType>*> FastWindingTrees,
+			const FVector3d& Pt, const FVector3d& OriginalPt, const FTransform3d& OriginalXF, const TArray<FTransformWithInv>& SpatialTransforms) -> bool
+		{
+			FRay3d Ray;
 
-		TFunctionRef<bool(TMeshAABBTree3<OccluderTriangleMeshType> * Spatial, TFastWindingTree<OccluderTriangleMeshType> * FastWindingTree, const FVector3d& Pt)> IsOccludedF =
-			(InsideMode == EOcclusionCalculationMode::FastWindingNumber) ? 
-			(TFunctionRef<bool(TMeshAABBTree3<OccluderTriangleMeshType> * Spatial, TFastWindingTree<OccluderTriangleMeshType> * FastWindingTree, const FVector3d& Pt)>)IsOccludedFWN :
-			(TFunctionRef<bool(TMeshAABBTree3<OccluderTriangleMeshType> * Spatial, TFastWindingTree<OccluderTriangleMeshType> * FastWindingTree, const FVector3d& Pt)>)IsOccludedSimple;
+			for (int SpatialIdx = 0, NumPts = Spatials.Num(); SpatialIdx < NumPts; ++SpatialIdx)
+			{
+				FVector3d XFPt = SpatialTransforms[SpatialIdx].InverseTransformUnlessMatch(Pt, OriginalPt, OriginalXF);
+				Ray.Origin = XFPt;
+				bool bIsOccluded = true;
+				for (int RayIdx = 0; RayIdx < NR; ++RayIdx)
+				{
+					Ray.Direction = SpatialTransforms[SpatialIdx].InverseTransformVector(RayDirs[RayIdx]);
+					int HitTID = Spatials[SpatialIdx]->FindNearestHitTriangle(Ray);
+					if (HitTID == IndexConstants::InvalidID)
+					{
+						bIsOccluded = false;
+						break;
+					}
+				}
+				if (bIsOccluded) // found a spatial that every ray direction hit
+				{
+					return true;
+				}
+			}
+			return false;
+		};
+
+		FIsOccludedFn IsOccludedF = InsideMode == EOcclusionCalculationMode::FastWindingNumber ?
+			(bTestOccludedByAny ? IsOccludedFWNAny : IsOccludedFWNTotal)
+				:
+			(bTestOccludedByAny ? IsOccludedSimpleAny : IsOccludedSimpleTotal);
 
 		bool bForceSingleThread = false;
 
+		bool bHasSpatialTransforms = Spatials.Num() == SpatialTransforms.Num();
+
+		TArray<FTransformWithInv> SpatialTransformsWithInv;
+		if (bHasSpatialTransforms)
+		{
+			for (int32 Idx = 0; Idx < Spatials.Num(); Idx++)
+			{
+				SpatialTransformsWithInv.Emplace(SpatialTransforms[Idx]);
+			}
+		}
+		else
+		{
+			// fill with identity transforms
+			for (int32 Idx = 0; Idx < Spatials.Num(); Idx++)
+			{
+				SpatialTransformsWithInv.Emplace();
+			}
+		}
+
 		TArray<bool> VertexOccluded;
 		if ( TriangleSamplingMethod == EOcclusionTriangleSampling::Vertices || TriangleSamplingMethod == EOcclusionTriangleSampling::VerticesAndCentroids )
 		{
@@ -164,28 +321,29 @@ public:
 			FMeshNormals Normals(Mesh);
 			Normals.ComputeVertexNormals();
 
-			for (int TreeIdx = 0; TreeIdx < Spatials.Num(); TreeIdx++)
+			ParallelFor(Mesh->MaxVertexID(), 
+			[this, &Normals, &VertexOccluded, &IsOccludedF, &MeshLocalToOccluderSpaces,
+				&Spatials, &FastWindingTrees, &SpatialTransformsWithInv](int32 VID)
 			{
-				TMeshAABBTree3<OccluderTriangleMeshType>* Spatial = Spatials[TreeIdx];
-				TFastWindingTree<OccluderTriangleMeshType>* FastWindingTree = FastWindingTrees.Num() > 0 ? FastWindingTrees[TreeIdx] : nullptr;
-				ParallelFor(Mesh->MaxVertexID(), [this, &Normals, &VertexOccluded, &IsOccludedF, &MeshLocalToOccluderSpaces, Spatial, FastWindingTree](int32 VID)
-					{
-						if (!Mesh->IsVertex(VID) || VertexOccluded[VID])
-						{
-							return;
-						}
-						FVector3d SamplePos = Mesh->GetVertex(VID);
-						FVector3d Normal = Normals[VID];
-						SamplePos += Normal * NormalOffset;
-						bool bAllOccluded = true;
-						for (const FTransform3d& XF : MeshLocalToOccluderSpaces)
-						{
-							bAllOccluded = bAllOccluded && IsOccludedF(Spatial, FastWindingTree, XF.TransformPosition(SamplePos));
-						}
-						checkSlow(VertexOccluded[VID] == false); // should have skipped the vertex if we already knew it was occluded (e.g. from another occluder)
-						VertexOccluded[VID] = bAllOccluded;
-					}, bForceSingleThread);
-			}
+				if (!Mesh->IsVertex(VID) || VertexOccluded[VID])
+				{
+					return;
+				}
+				FVector3d SamplePos = Mesh->GetVertex(VID);
+				FVector3d Normal = Normals[VID];
+				SamplePos += Normal * NormalOffset;
+				bool bAllOccluded = true;
+				for (int32 TransformIdx = 0, TransformNum = MeshLocalToOccluderSpaces.Num(); TransformIdx < TransformNum; TransformIdx++)
+				{
+					const FTransform3d& OriginalTransform = MeshLocalToOccluderSpaces[TransformIdx];
+					FVector3d XFPos = OriginalTransform.TransformPosition(SamplePos);
+
+					bAllOccluded = bAllOccluded && IsOccludedF(Spatials, FastWindingTrees, 
+						XFPos, SamplePos, OriginalTransform, SpatialTransformsWithInv);
+				}
+				checkSlow(VertexOccluded[VID] == false); // should have skipped the vertex if we already knew it was occluded (e.g. from another occluder)
+				VertexOccluded[VID] = bAllOccluded;
+			}, bForceSingleThread);
 		}
 		if (Cancelled())
 		{
@@ -194,44 +352,47 @@ public:
 
 		TArray<FThreadSafeBool> TriOccluded;
 		TriOccluded.Init(false, Mesh->MaxTriangleID());
-		for (int TreeIdx = 0; TreeIdx < Spatials.Num(); TreeIdx++)
+
+		ParallelFor(Mesh->MaxTriangleID(), 
+		[this, &VertexOccluded, &IsOccludedF, &MeshLocalToOccluderSpaces,
+			&TriangleBaryCoordSamples, &TriOccluded,
+			&Spatials, &FastWindingTrees, &SpatialTransformsWithInv](int32 TID)
 		{
-			TMeshAABBTree3<OccluderTriangleMeshType>* Spatial = Spatials[TreeIdx];
-			TFastWindingTree<OccluderTriangleMeshType>* FastWindingTree = FastWindingTrees.Num() > 0 ? FastWindingTrees[TreeIdx] : nullptr;
-			ParallelFor(Mesh->MaxTriangleID(), [this, &VertexOccluded, &IsOccludedF, &MeshLocalToOccluderSpaces, &TriangleBaryCoordSamples, &TriOccluded, Spatial, FastWindingTree](int32 TID)
+			if (!Mesh->IsTriangle(TID))
+			{
+				return;
+			}
+
+			bool bInside = true;
+			if (TriangleSamplingMethod == EOcclusionTriangleSampling::Vertices || TriangleSamplingMethod == EOcclusionTriangleSampling::VerticesAndCentroids)
+			{
+				FIndex3i Tri = Mesh->GetTriangle(TID);
+				bInside = VertexOccluded[Tri.A] && VertexOccluded[Tri.B] && VertexOccluded[Tri.C];
+			}
+			if (bInside && TriangleBaryCoordSamples.Num() > 0)
+			{
+				FVector3d Normal = Mesh->GetTriNormal(TID);
+				FVector3d V0, V1, V2;
+				Mesh->GetTriVertices(TID, V0, V1, V2);
+				for (int32 SampleIdx = 0, NumSamples = TriangleBaryCoordSamples.Num(); bInside && SampleIdx < NumSamples; SampleIdx++)
 				{
-					if (!Mesh->IsTriangle(TID))
+					FVector3d BaryCoords = TriangleBaryCoordSamples[SampleIdx];
+					FVector3d SamplePos = V0 * BaryCoords.X + V1 * BaryCoords.Y + V2 * BaryCoords.Z + Normal * NormalOffset;
+					for (int32 TransformIdx = 0, TransformNum = MeshLocalToOccluderSpaces.Num(); TransformIdx < TransformNum; TransformIdx++)
 					{
-						return;
-					}
+						const FTransform3d& OriginalTransform = MeshLocalToOccluderSpaces[TransformIdx];
+						FVector3d XFPos = OriginalTransform.TransformPosition(SamplePos);
 
-					bool bInside = true;
-					if (TriangleSamplingMethod == EOcclusionTriangleSampling::Vertices || TriangleSamplingMethod == EOcclusionTriangleSampling::VerticesAndCentroids)
-					{
-						FIndex3i Tri = Mesh->GetTriangle(TID);
-						bInside = VertexOccluded[Tri.A] && VertexOccluded[Tri.B] && VertexOccluded[Tri.C];
-					}
-					if (bInside && TriangleBaryCoordSamples.Num() > 0)
-					{
-						FVector3d Normal = Mesh->GetTriNormal(TID);
-						FVector3d V0, V1, V2;
-						Mesh->GetTriVertices(TID, V0, V1, V2);
-						for (int32 SampleIdx = 0, NumSamples = TriangleBaryCoordSamples.Num(); bInside && SampleIdx < NumSamples; SampleIdx++)
-						{
-							FVector3d BaryCoords = TriangleBaryCoordSamples[SampleIdx];
-							FVector3d SamplePos = V0 * BaryCoords.X + V1 * BaryCoords.Y + V2 * BaryCoords.Z + Normal * NormalOffset;
-							for (const FTransform3d& XF : MeshLocalToOccluderSpaces)
-							{
-								bInside = bInside && IsOccludedF(Spatial, FastWindingTree, XF.TransformPosition(SamplePos));
-							}
-						}
-					}
-					if (bInside)
-					{
-						TriOccluded[TID] = true;
+						bInside = bInside && IsOccludedF(Spatials, FastWindingTrees,
+							XFPos, SamplePos, OriginalTransform, SpatialTransformsWithInv);
 					}
-				}, bForceSingleThread);
-		}
+				}
+			}
+			if (bInside)
+			{
+				TriOccluded[TID] = true;
+			}
+		}, bForceSingleThread);
 
 
 		if (Cancelled())
@@ -278,12 +439,16 @@ public:
 	 * @param MeshLocalToOccluderSpaces Transforms to take instances of the local mesh into the space of the occluders
 	 * @param Spatials AABB trees for all occluders
 	 * @param FastWindingTrees Precomputed fast winding trees for occluders
+	 * @param SpatialTransforms Transforms AABB/winding tree to shared occluder space (if empty, Identity is used)
+	 * @param bTestOccludedByAny If true, a triangle is occluded if it is fully occluded by *any* occluder.
+	 *							 Otherwise, we test if it's occluded by the combination of all occluders.
 	 * @return true on success
 	 */
 	virtual bool Apply(const TArrayView<const FTransform3d> MeshLocalToOccluderSpaces, 
-		const TArrayView<TMeshAABBTree3<OccluderTriangleMeshType>*> Spatials, const TArrayView<TFastWindingTree<OccluderTriangleMeshType>*> FastWindingTrees)
+		const TArrayView<TMeshAABBTree3<OccluderTriangleMeshType>*> Spatials, const TArrayView<TFastWindingTree<OccluderTriangleMeshType>*> FastWindingTrees,
+		const TArrayView<const FTransform3d> SpatialTransforms = TArrayView<const FTransform3d>(), bool bTestOccludedByAny = false)
 	{
-		if (!Select(MeshLocalToOccluderSpaces, Spatials, FastWindingTrees))
+		if (!Select(MeshLocalToOccluderSpaces, Spatials, FastWindingTrees, SpatialTransforms, bTestOccludedByAny))
 		{
 			return false;
 		}
@@ -306,7 +471,7 @@ public:
 	{
 		TArrayView<TMeshAABBTree3<OccluderTriangleMeshType>*> Spatials(&Spatial, 1);
 		TArrayView<TFastWindingTree<OccluderTriangleMeshType>*> FastWindingTrees(&FastWindingTree, 1);
-		return Select(MeshLocalToOccluderSpaces, Spatials, FastWindingTrees);
+		return Select(MeshLocalToOccluderSpaces, Spatials, FastWindingTrees, TArrayView<const FTransform3d>());
 	}
 
 
@@ -324,7 +489,7 @@ public:
 	{
 		TArrayView<TMeshAABBTree3<OccluderTriangleMeshType>*> Spatials(&Spatial, 1);
 		TArrayView<TFastWindingTree<OccluderTriangleMeshType>*> FastWindingTrees(&FastWindingTree, 1);
-		return Apply(MeshLocalToOccluderSpaces, Spatials, FastWindingTrees);
+		return Apply(MeshLocalToOccluderSpaces, Spatials, FastWindingTrees, TArrayView<const FTransform3d>());
 	}
 
 	/**
diff --git a/Engine/Plugins/Experimental/MeshModelingToolset/Source/MeshModelingTools/Private/RemoveOccludedTrianglesTool.cpp b/Engine/Plugins/Experimental/MeshModelingToolset/Source/MeshModelingTools/Private/RemoveOccludedTrianglesTool.cpp
--- a/Engine/Plugins/Experimental/MeshModelingToolset/Source/MeshModelingTools/Private/RemoveOccludedTrianglesTool.cpp
+++ b/Engine/Plugins/Experimental/MeshModelingToolset/Source/MeshModelingTools/Private/RemoveOccludedTrianglesTool.cpp
@@ -199,8 +199,6 @@ void URemoveOccludedTrianglesTool::SetupPreviews()
 	int32 NumTargets = ComponentTargets.Num();
 	int32 NumPreviews = PreviewToTargetIdx.Num();
 
-	CombinedMeshTrees = MakeShared<IndexMeshWithAcceleration, ESPMode::ThreadSafe>();
-
 #if WITH_EDITOR
 	static const FText SlowTaskText = LOCTEXT("RemoveOccludedTrianglesInit", "Building mesh occlusion data...");
 
@@ -222,6 +220,10 @@ void URemoveOccludedTrianglesTool::SetupPreviews()
 	OccludedGroupIDs.Init(-1, NumPreviews);
 	OccludedGroupLayers.Init(-1, NumPreviews);
 
+	OccluderTrees.SetNum(NumTargets);
+	OccluderWindings.SetNum(NumTargets);
+	OccluderTransforms.SetNum(NumTargets);
+
 	OriginalDynamicMeshes.SetNum(NumPreviews);
 	PreviewToCopyIdx.Reset(); PreviewToCopyIdx.SetNum(NumPreviews);
 	for (int32 TargetIdx = 0; TargetIdx < NumTargets; TargetIdx++)
@@ -275,6 +277,10 @@ void URemoveOccludedTrianglesTool::SetupPreviews()
 			Preview->PreviewMesh->UpdatePreview(OriginalDynamicMeshes[PreviewIdx].Get());
 			Preview->SetVisibility(true);
 
+			OccluderTrees[TargetIdx] = MakeShared<FDynamicMeshAABBTree3, ESPMode::ThreadSafe>(OriginalDynamicMeshes[PreviewIdx].Get());
+			OccluderWindings[TargetIdx] = MakeShared<TFastWindingTree<FDynamicMesh3>, ESPMode::ThreadSafe>(OccluderTrees[TargetIdx].Get());
+			OccluderTransforms[TargetIdx] = (FTransform3d)ComponentTargets[TargetIdx]->GetWorldTransform();
+
 			// configure secondary render material
 			Preview->SecondaryMaterial = OccludedMaterial;
 
@@ -305,6 +311,10 @@ void URemoveOccludedTrianglesTool::SetupPreviews()
 			
 			PreviewMesh->SetVisible(true);
 
+			OccluderTrees[TargetIdx] = OccluderTrees[PreviewToTargetIdx[PreviewIdx]];
+			OccluderWindings[TargetIdx] = OccluderWindings[PreviewToTargetIdx[PreviewIdx]];
+			OccluderTransforms[TargetIdx] = (FTransform3d)ComponentTargets[TargetIdx]->GetWorldTransform();
+
 			PreviewMesh->SetSecondaryRenderMaterial(OccludedMaterial);
 			PreviewMesh->EnableSecondaryTriangleBuffers(MoveTemp(IsOccludedGroupFn));
 
@@ -312,11 +322,7 @@ void URemoveOccludedTrianglesTool::SetupPreviews()
 				PreviewCopies[CopyIdx]->UpdatePreview(Compute->PreviewMesh->GetPreviewDynamicMesh());
 			});
 		}
-
-		CombinedMeshTrees->AddMesh(*OriginalDynamicMeshes[PreviewIdx], FTransform3d(ComponentTargets[TargetIdx]->GetWorldTransform()));
 	}
-
-	CombinedMeshTrees->BuildAcceleration();
 }
 
 
@@ -400,9 +406,19 @@ TUniquePtr<FDynamicMeshOperator> URemoveOccludedTrianglesOperatorFactory::MakeNe
 	FTransform LocalToWorld = Tool->ComponentTargets[ComponentIndex]->GetWorldTransform();
 	Op->OriginalMesh = Tool->OriginalDynamicMeshes[PreviewIdx];
 
-	Op->CombinedMeshTrees = Tool->CombinedMeshTrees;
-
-	Op->bOnlySelfOcclude = Tool->BasicProperties->bOnlySelfOcclude;
+	if (Tool->BasicProperties->bOnlySelfOcclude)
+	{
+		int32 TargetIdx = Tool->PreviewToTargetIdx[PreviewIdx];
+		Op->OccluderTrees.Add(Tool->OccluderTrees[TargetIdx]);
+		Op->OccluderWindings.Add(Tool->OccluderWindings[TargetIdx]);
+		Op->OccluderTransforms.Add(FTransform3d::Identity());
+	}
+	else
+	{
+		Op->OccluderTrees = Tool->OccluderTrees;
+		Op->OccluderWindings = Tool->OccluderWindings;
+		Op->OccluderTransforms = Tool->OccluderTransforms;
+	}
 
 	Op->AddRandomRays = Tool->BasicProperties->AddRandomRays;
 
diff --git a/Engine/Plugins/Experimental/MeshModelingToolset/Source/MeshModelingTools/Public/RemoveOccludedTrianglesTool.h b/Engine/Plugins/Experimental/MeshModelingToolset/Source/MeshModelingTools/Public/RemoveOccludedTrianglesTool.h
--- a/Engine/Plugins/Experimental/MeshModelingToolset/Source/MeshModelingTools/Public/RemoveOccludedTrianglesTool.h
+++ b/Engine/Plugins/Experimental/MeshModelingToolset/Source/MeshModelingTools/Public/RemoveOccludedTrianglesTool.h
@@ -203,8 +203,10 @@ protected:
 protected:
 	TArray<TSharedPtr<FDynamicMesh3, ESPMode::ThreadSafe>> OriginalDynamicMeshes;
 
-	// pre-transformed combined all-mesh data (with mesh geometry data in just raw index buffers as that's all we need)
-	TSharedPtr<IndexMeshWithAcceleration, ESPMode::ThreadSafe> CombinedMeshTrees;
+	// AABB trees and winding trees for every mesh target, with repeated instances as pointers to the same data
+	TArray<TSharedPtr<FDynamicMeshAABBTree3, ESPMode::ThreadSafe>> OccluderTrees;
+	TArray<TSharedPtr<TFastWindingTree<FDynamicMesh3>, ESPMode::ThreadSafe>> OccluderWindings;
+	TArray<FTransform3d> OccluderTransforms;
 
 	TArray<TArray<int32>> PreviewToCopyIdx;
 	TArray<int32> PreviewToTargetIdx;
diff --git a/Engine/Plugins/Experimental/MeshModelingToolset/Source/ModelingOperators/Private/BaseOps/VoxelBaseOp.cpp b/Engine/Plugins/Experimental/MeshModelingToolset/Source/ModelingOperators/Private/BaseOps/VoxelBaseOp.cpp
--- a/Engine/Plugins/Experimental/MeshModelingToolset/Source/ModelingOperators/Private/BaseOps/VoxelBaseOp.cpp
+++ b/Engine/Plugins/Experimental/MeshModelingToolset/Source/ModelingOperators/Private/BaseOps/VoxelBaseOp.cpp
@@ -18,7 +18,7 @@ void FVoxelBaseOp::PostProcessResult(FProgressCancel* Progress, double MeshCellS
 
 	if (bRemoveInternalSurfaces)
 	{
-		UE::MeshAutoRepair::RemoveInternalTriangles(*ResultMesh.Get(), true, EOcclusionTriangleSampling::Centroids, EOcclusionCalculationMode::FastWindingNumber);
+		UE::MeshAutoRepair::RemoveInternalTriangles(*ResultMesh.Get(), true, EOcclusionTriangleSampling::Centroids, EOcclusionCalculationMode::FastWindingNumber, 0, .5, true);
 	}
 
 	if (Progress && Progress->Cancelled())
diff --git a/Engine/Plugins/Experimental/MeshModelingToolset/Source/ModelingOperators/Private/CleaningOps/RemoveOccludedTrianglesOp.cpp b/Engine/Plugins/Experimental/MeshModelingToolset/Source/ModelingOperators/Private/CleaningOps/RemoveOccludedTrianglesOp.cpp
--- a/Engine/Plugins/Experimental/MeshModelingToolset/Source/ModelingOperators/Private/CleaningOps/RemoveOccludedTrianglesOp.cpp
+++ b/Engine/Plugins/Experimental/MeshModelingToolset/Source/ModelingOperators/Private/CleaningOps/RemoveOccludedTrianglesOp.cpp
@@ -15,31 +15,6 @@
 #include "Async/ParallelFor.h"
 
 
-void IndexMeshWithAcceleration::AddMesh(const FDynamicMesh3& MeshIn, FTransform3d Transform)
-{
-	int VertexIndexStart = Vertices.Num();
-	Vertices.SetNum(VertexIndexStart + MeshIn.MaxVertexID());
-	ParallelFor(MeshIn.MaxVertexID(), [this, &MeshIn, &Transform, &VertexIndexStart](int VID)
-	{
-		if (MeshIn.IsVertex(VID))
-		{
-			Vertices[VID + VertexIndexStart] = Transform.TransformPosition(MeshIn.GetVertex(VID));
-		}
-	}
-	, false);
-
-	// don't parallelize triangles b/c we want need them compact
-	for (int TID = 0; TID < MeshIn.MaxTriangleID(); TID++)
-	{
-		if (MeshIn.IsTriangle(TID))
-		{
-			FIndex3i Triangle = MeshIn.GetTriangle(TID);
-			Triangles.Add(Triangle.A + VertexIndexStart);
-			Triangles.Add(Triangle.B + VertexIndexStart);
-			Triangles.Add(Triangle.C + VertexIndexStart);
-		}
-	}
-}
 
 
 void FRemoveOccludedTrianglesOp::SetTransform(const FTransform& Transform)
@@ -97,69 +72,42 @@ void FRemoveOccludedTrianglesOp::CalculateResult(FProgressCancel* Progress)
 		}
 	};
 
-	if (bOnlySelfOcclude)
+	TRemoveOccludedTriangles<FDynamicMesh3> Jacket(ResultMesh.Get());
+
+	if (Progress && Progress->Cancelled())
 	{
-		TRemoveOccludedTriangles<FDynamicMesh3> SelfJacket(ResultMesh.Get());
-		FDynamicMeshAABBTree3 SelfAABB(ResultMesh.Get());
-		TFastWindingTree<FDynamicMesh3> SelfFWTree(&SelfAABB);
-		if (Progress && Progress->Cancelled())
-		{
-			return;
-		}
-		SelfJacket.InsideMode = InsideMode;
-		SelfJacket.TriangleSamplingMethod = TriangleSamplingMethod;
-		SelfJacket.WindingIsoValue = WindingIsoValue;
-		SelfJacket.NormalOffset = NormalOffset;
-		SelfJacket.AddRandomRays = AddRandomRays;
-		SelfJacket.AddTriangleSamples = AddTriangleSamples;
-		SelfJacket.Select(FTransform3d::Identity(), &SelfAABB, &SelfFWTree);
-		if (ShrinkRemoval > 0)
-		{
-			ShrinkSelection(*ResultMesh.Get(), SelfJacket.RemovedT, ShrinkRemoval);
-		}
+		return;
+	}
 
-		if (bSetTriangleGroupInsteadOfRemoving)
-		{
-			FIndex2i GroupIDAndLayerIndex = SetNewGroupSelection(*ResultMesh.Get(), SelfJacket.RemovedT, ActiveGroupLayer, bActiveGroupLayerIsDefault);
-			CreatedGroupID = GroupIDAndLayerIndex.A;
-			CreatedGroupLayerIndex = GroupIDAndLayerIndex.B;
-		}
-		else
-		{
-			SelfJacket.RemoveSelected();
-		}
+	Jacket.InsideMode = InsideMode;
+	Jacket.TriangleSamplingMethod = TriangleSamplingMethod;
+	Jacket.WindingIsoValue = WindingIsoValue;
+	Jacket.NormalOffset = NormalOffset;
+	Jacket.AddRandomRays = AddRandomRays;
+	Jacket.AddTriangleSamples = AddTriangleSamples;
+	// copy shared pointers to ray pointers for the jacketing algorithm (they w/ shorter lifespan so it's ok)
+	TArray<FDynamicMeshAABBTree3*> RawOccluderTrees; RawOccluderTrees.SetNum(OccluderTrees.Num());
+	TArray<TFastWindingTree<FDynamicMesh3>*> RawOccluderWindings; RawOccluderWindings.SetNum(OccluderWindings.Num());
+	for (int32 TreeIdx = 0; TreeIdx < OccluderTrees.Num(); TreeIdx++)
+	{
+		RawOccluderTrees[TreeIdx] = OccluderTrees[TreeIdx].Get();
+		RawOccluderWindings[TreeIdx] = OccluderWindings[TreeIdx].Get();
 	}
-	else
+	Jacket.Select(MeshTransforms, RawOccluderTrees, RawOccluderWindings, OccluderTransforms);
+	if (ShrinkRemoval > 0)
 	{
-		TRemoveOccludedTriangles<TIndexMeshArrayAdapter<int, double, FVector3d>> Jacket(ResultMesh.Get());
-
-		if (Progress && Progress->Cancelled())
-		{
-			return;
-		}
-
-		Jacket.InsideMode = InsideMode;
-		Jacket.TriangleSamplingMethod = TriangleSamplingMethod;
-		Jacket.WindingIsoValue = WindingIsoValue;
-		Jacket.NormalOffset = NormalOffset;
-		Jacket.AddRandomRays = AddRandomRays;
-		Jacket.AddTriangleSamples = AddTriangleSamples;
-		Jacket.Select(MeshTransforms, &CombinedMeshTrees->AABB, &CombinedMeshTrees->FastWinding);
-		if (ShrinkRemoval > 0)
-		{
-			ShrinkSelection(*ResultMesh.Get(), Jacket.RemovedT, ShrinkRemoval);
-		}
+		ShrinkSelection(*ResultMesh.Get(), Jacket.RemovedT, ShrinkRemoval);
+	}
 
-		if (bSetTriangleGroupInsteadOfRemoving)
-		{
-			FIndex2i GroupIDAndLayerIndex = SetNewGroupSelection(*ResultMesh.Get(), Jacket.RemovedT, ActiveGroupLayer, bActiveGroupLayerIsDefault);
-			CreatedGroupID = GroupIDAndLayerIndex.A;
-			CreatedGroupLayerIndex = GroupIDAndLayerIndex.B;
-		}
-		else
-		{
-			Jacket.RemoveSelected();
-		}
+	if (bSetTriangleGroupInsteadOfRemoving)
+	{
+		FIndex2i GroupIDAndLayerIndex = SetNewGroupSelection(*ResultMesh.Get(), Jacket.RemovedT, ActiveGroupLayer, bActiveGroupLayerIsDefault);
+		CreatedGroupID = GroupIDAndLayerIndex.A;
+		CreatedGroupLayerIndex = GroupIDAndLayerIndex.B;
+	}
+	else
+	{
+		Jacket.RemoveSelected();
 	}
 
 	if (MinTriCountConnectedComponent > 0 || MinAreaConnectedComponent > 0)
diff --git a/Engine/Plugins/Experimental/MeshModelingToolset/Source/ModelingOperators/Private/CompositionOps/VoxelBlendMeshesOp.cpp b/Engine/Plugins/Experimental/MeshModelingToolset/Source/ModelingOperators/Private/CompositionOps/VoxelBlendMeshesOp.cpp
--- a/Engine/Plugins/Experimental/MeshModelingToolset/Source/ModelingOperators/Private/CompositionOps/VoxelBlendMeshesOp.cpp
+++ b/Engine/Plugins/Experimental/MeshModelingToolset/Source/ModelingOperators/Private/CompositionOps/VoxelBlendMeshesOp.cpp
@@ -75,7 +75,7 @@ void FVoxelBlendMeshesOp::CalculateResult(FProgressCancel* Progress)
 
 			if (bRemoveInternalsAfterVoxWrap)
 			{
-				UE::MeshAutoRepair::RemoveInternalTriangles(TransformedMeshes[MeshIdx], true, EOcclusionTriangleSampling::Centroids, EOcclusionCalculationMode::FastWindingNumber);
+				UE::MeshAutoRepair::RemoveInternalTriangles(TransformedMeshes[MeshIdx], true, EOcclusionTriangleSampling::Centroids, EOcclusionCalculationMode::FastWindingNumber, 0, .5, true);
 			}
 		}
 
diff --git a/Engine/Plugins/Experimental/MeshModelingToolset/Source/ModelingOperators/Private/CompositionOps/VoxelMorphologyMeshesOp.cpp b/Engine/Plugins/Experimental/MeshModelingToolset/Source/ModelingOperators/Private/CompositionOps/VoxelMorphologyMeshesOp.cpp
--- a/Engine/Plugins/Experimental/MeshModelingToolset/Source/ModelingOperators/Private/CompositionOps/VoxelMorphologyMeshesOp.cpp
+++ b/Engine/Plugins/Experimental/MeshModelingToolset/Source/ModelingOperators/Private/CompositionOps/VoxelMorphologyMeshesOp.cpp
@@ -111,7 +111,7 @@ void FVoxelMorphologyMeshesOp::CalculateResult(FProgressCancel* Progress)
 
 		if (bRemoveInternalsAfterVoxWrap)
 		{
-			UE::MeshAutoRepair::RemoveInternalTriangles(CombinedMesh, true, EOcclusionTriangleSampling::Centroids, EOcclusionCalculationMode::FastWindingNumber);
+			UE::MeshAutoRepair::RemoveInternalTriangles(CombinedMesh, true, EOcclusionTriangleSampling::Centroids, EOcclusionCalculationMode::FastWindingNumber, 0, .5, true);
 		}
 
 		Spatial.Build(); // rebuild w/ updated mesh
diff --git a/Engine/Plugins/Experimental/MeshModelingToolset/Source/ModelingOperators/Public/CleaningOps/RemoveOccludedTrianglesOp.h b/Engine/Plugins/Experimental/MeshModelingToolset/Source/ModelingOperators/Public/CleaningOps/RemoveOccludedTrianglesOp.h
--- a/Engine/Plugins/Experimental/MeshModelingToolset/Source/ModelingOperators/Public/CleaningOps/RemoveOccludedTrianglesOp.h
+++ b/Engine/Plugins/Experimental/MeshModelingToolset/Source/ModelingOperators/Public/CleaningOps/RemoveOccludedTrianglesOp.h
@@ -6,31 +6,10 @@
 #include "Util/ProgressCancel.h"
 #include "ModelingOperators.h"
 
+#include "DynamicMeshAABBTree3.h"
 #include "MeshAdapter.h"
 #include "Operations/RemoveOccludedTriangles.h"
 
-// simple struct to help construct big trees of all the occluders
-struct MODELINGOPERATORS_API IndexMeshWithAcceleration
-{
-	TArray<int> Triangles;
-	TArray<FVector3d> Vertices;
-	TIndexMeshArrayAdapter<int, double, FVector3d> Mesh;
-	TMeshAABBTree3<TIndexMeshArrayAdapter<int, double, FVector3d>> AABB;
-	TFastWindingTree<TIndexMeshArrayAdapter<int, double, FVector3d>> FastWinding;
-
-	IndexMeshWithAcceleration() : AABB(&Mesh, false), FastWinding(&AABB, false)
-	{
-		Mesh.SetSources(&Vertices, &Triangles);
-	}
-
-	void AddMesh(const FDynamicMesh3& MeshIn, FTransform3d Transform);
-
-	void BuildAcceleration()
-	{
-		AABB.Build();
-		FastWinding.Build();
-	}
-};
 
 
 class MODELINGOPERATORS_API FRemoveOccludedTrianglesOp : public FDynamicMeshOperator
@@ -40,7 +19,10 @@ public:
 
 	// inputs
 	TSharedPtr<FDynamicMesh3, ESPMode::ThreadSafe> OriginalMesh;
-	TSharedPtr<IndexMeshWithAcceleration, ESPMode::ThreadSafe> CombinedMeshTrees;
+	
+	TArray<TSharedPtr<FDynamicMeshAABBTree3, ESPMode::ThreadSafe>> OccluderTrees;
+	TArray<TSharedPtr<TFastWindingTree<FDynamicMesh3>, ESPMode::ThreadSafe>> OccluderWindings;
+	TArray<FTransform3d> OccluderTransforms;
 
 	TArray<FTransform3d> MeshTransforms;
 
@@ -59,9 +41,6 @@ public:
 	// add triangle samples per triangle (in addition to TriangleSamplingMethod)
 	int AddTriangleSamples = 0;
 
-	// if true, *ignore* the combined selected occluder data in CombinedMeshTrees, and only consider self-occlusion of this mesh
-	bool bOnlySelfOcclude = false;
-
 	int ShrinkRemoval = 0;
 
 	double MinAreaConnectedComponent = 0;
