ModelingTools: Add support for FlatMidpointFan cap style to SweepGenerator

#rnx

diff --git a/Engine/Plugins/Experimental/GeometryProcessing/Source/GeometricObjects/Public/Generators/SweepGenerator.h b/Engine/Plugins/Experimental/GeometryProcessing/Source/GeometricObjects/Public/Generators/SweepGenerator.h
--- a/Engine/Plugins/Experimental/GeometryProcessing/Source/GeometricObjects/Public/Generators/SweepGenerator.h
+++ b/Engine/Plugins/Experimental/GeometryProcessing/Source/GeometricObjects/Public/Generators/SweepGenerator.h
@@ -24,7 +24,8 @@
 enum class /*GEOMETRICOBJECTS_API*/ ECapType
 {
 	None = 0,
-	FlatTriangulation = 1
+	FlatTriangulation = 1,
+	FlatMidpointFan = 2
 	// TODO: Cone, other caps ...
 };
 
@@ -121,16 +122,15 @@ protected:
 					NumUVs += XVerts;
 					NumNormals += XVerts;
 				}
-
+				else if (Caps[CapIdx] == ECapType::FlatMidpointFan)
+				{
+					NumTriangles += XVerts;
+					NumPolygons++;
+					NumUVs += XVerts + 1;
+					NumNormals += XVerts + 1;
+					NumVerts += 1;
+				}
 				// TODO: support more cap type; e.g.:
-				//else if (Caps[CapIdx] == ECapType::FlatMidpointFan)
-				//{
-				//	NumTriangles += XVerts;
-				//	NumPolygons += XVerts;
-				//	NumUVs += XVerts + 1;
-				//	NumNormals += XVerts + 1;
-				//	NumVerts += 1;
-				//}
 				//else if (Caps[CapIdx] == ECapType::Cone)
 				//{
 				//	NumTriangles += XVerts;
@@ -148,7 +148,6 @@ protected:
 		{
 			for (int32 CapIdx = 0; CapIdx < 2; CapIdx++)
 			{
-
 				if (Caps[CapIdx] == ECapType::FlatTriangulation)
 				{
 					int32 VertOffset = CapIdx * (XVerts * (NumCrossSections - 1));
@@ -181,6 +180,49 @@ protected:
 					{
 						FVector2f CenteredVert = (FVector2f)CrossSection.GetVertices()[Idx] * CapUVScale + CapUVOffset;
 						SetUV(CapUVStart[CapIdx] + Idx, FVector2f(CenteredVert.X * SideScale, CenteredVert.Y), VertOffset + Idx);
+
+						// correct normal to be filled by subclass
+						SetNormal(CapNormalStart[CapIdx] + Idx, FVector3f::Zero(), VertOffset + Idx);
+					}
+				}
+				else if (Caps[CapIdx] == ECapType::FlatMidpointFan)
+				{
+					int32 VertOffset = CapIdx * (XVerts * (NumCrossSections - 1));
+					int32 CapVertStartIdx = CapVertStart[CapIdx];
+					int32 TriIdx = CapTriangleStart[CapIdx];
+					int32 PolyIdx = CapPolygonStart[CapIdx];
+					for (int32 VertIdx = 0; VertIdx < XVerts; VertIdx++)
+					{
+						bool Flipped = CapIdx == 0;
+						SetTriangle(TriIdx,
+							VertOffset + VertIdx,
+							CapVertStartIdx,
+							VertOffset + (VertIdx + 1) % XVerts,
+							Flipped);
+						SetTriangleUVs(TriIdx,
+							CapUVStart[CapIdx] + VertIdx,
+							CapUVStart[CapIdx] + XVerts,
+							CapUVStart[CapIdx] + (VertIdx + 1) % XVerts,
+							Flipped);
+						SetTriangleNormals(TriIdx,
+							CapNormalStart[CapIdx] + VertIdx,
+							CapNormalStart[CapIdx] + XVerts,
+							CapNormalStart[CapIdx] + (VertIdx + 1) % XVerts,
+							Flipped);
+						SetTrianglePolygon(TriIdx, PolyIdx);
+						TriIdx++;
+					}
+
+					// Set cap midpoint UV & Normal
+					// (correct normal to be filled by subclass)
+					SetUV(CapUVStart[CapIdx] + XVerts, FVector2f::Zero() + CapUVOffset, CapVertStartIdx);
+					SetNormal(CapNormalStart[CapIdx] + XVerts, FVector3f::Zero(), CapVertStartIdx);
+
+					// Set cap profile UVs & Normal
+					for (int32 Idx = 0; Idx < XVerts; Idx++)
+					{
+						FVector2f CenteredVert = (FVector2f)CrossSection.GetVertices()[Idx] * CapUVScale + CapUVOffset;
+						SetUV(CapUVStart[CapIdx] + Idx, FVector2f(CenteredVert.X, CenteredVert.Y), VertOffset + Idx);
 						SetNormal(CapNormalStart[CapIdx] + Idx, FVector3f::Zero(), VertOffset + Idx);
 					}
 				}
@@ -351,8 +393,8 @@ public:
 
 		if (bCapped)
 		{
-			Caps[0] = ECapType::FlatTriangulation;
-			Caps[1] = ECapType::FlatTriangulation;
+			Caps[0] = ECapType::FlatMidpointFan;
+			Caps[1] = ECapType::FlatMidpointFan;
 		}
 
 		int NumX = Radii.Num();
@@ -417,9 +459,19 @@ public:
 				}
 			}
 		}
-		// if capped, set top/bottom normals
+		
 		if (bCapped)
 		{
+			// if capped, set vertices.
+			for (int CapIdx = 0; CapIdx < 2; CapIdx++)
+			{
+				if (Caps[CapIdx] == ECapType::FlatMidpointFan)
+				{
+					Vertices[CapVertStart[CapIdx]] = FVector3d::UnitZ() * Heights[CapIdx];
+				}
+			}
+
+			// if capped, set top/bottom normals
 			for (int SubIdx = 0; SubIdx < X.VertexCount(); SubIdx++)
 			{
 				for (int XBotTop = 0; XBotTop < 2; ++XBotTop)
@@ -427,6 +479,13 @@ public:
 					Normals[CapNormalStart[XBotTop] + SubIdx] = FVector3f(0, 0, 2 * XBotTop - 1);
 				}
 			}
+			for (int CapIdx = 0; CapIdx < 2; CapIdx++)
+			{
+				if (Caps[CapIdx] == ECapType::FlatMidpointFan)
+				{
+					Normals[CapNormalStart[CapIdx] + X.VertexCount()] = FVector3f(0, 0, 2 * CapIdx - 1);
+				}
+			}
 		}
 
 		if (bUVScaleMatchSidesAndCaps)
