Safer handling of islands for generic interface for uv packing

diff --git a/Engine/Plugins/Experimental/GeometryProcessing/Source/DynamicMesh/Private/Parameterization/MeshUVPacking.cpp b/Engine/Plugins/Experimental/GeometryProcessing/Source/DynamicMesh/Private/Parameterization/MeshUVPacking.cpp
--- a/Engine/Plugins/Experimental/GeometryProcessing/Source/DynamicMesh/Private/Parameterization/MeshUVPacking.cpp
+++ b/Engine/Plugins/Experimental/GeometryProcessing/Source/DynamicMesh/Private/Parameterization/MeshUVPacking.cpp
@@ -65,9 +65,9 @@ bool FDynamicMeshUVPacker::StandardPack()
 		return MeshView.UVOverlay->AreTrianglesConnected(Triangle0, Triangle1);
 		});
 
-	return Packer.StandardPack(&MeshView, UVComponents.Num(), [&UVComponents](int Idx)
+	return Packer.StandardPack(&MeshView, UVComponents.Num(), [&UVComponents](int Idx, TArray<int32>& Island)
 		{
-			return TArrayView<const int>(UVComponents.GetComponent(Idx).Indices);
+			Island = UVComponents.GetComponent(Idx).Indices;
 		});
 }
 
@@ -87,8 +87,8 @@ bool FDynamicMeshUVPacker::StackPack()
 		return MeshView.UVOverlay->AreTrianglesConnected(Triangle0, Triangle1);
 		});
 
-	return Packer.StackPack(&MeshView, UVComponents.Num(), [&UVComponents](int Idx)
+	return Packer.StackPack(&MeshView, UVComponents.Num(), [&UVComponents](int Idx, TArray<int32>& Island)
 		{
-			return TArrayView<const int>(UVComponents.GetComponent(Idx).Indices);
+			Island = UVComponents.GetComponent(Idx).Indices;
 		});
 }
\ No newline at end of file
diff --git a/Engine/Plugins/Experimental/GeometryProcessing/Source/GeometricObjects/Private/Parameterization/UVPacking.cpp b/Engine/Plugins/Experimental/GeometryProcessing/Source/GeometricObjects/Private/Parameterization/UVPacking.cpp
--- a/Engine/Plugins/Experimental/GeometryProcessing/Source/GeometricObjects/Private/Parameterization/UVPacking.cpp
+++ b/Engine/Plugins/Experimental/GeometryProcessing/Source/GeometricObjects/Private/Parameterization/UVPacking.cpp
@@ -799,7 +799,7 @@ void FStandardChartPacker::RasterizeChart(const FUVIsland& Chart, uint32 RectW,
 }} // namespace UE::InternalUVPacking
 
 
-bool FUVPacker::StandardPack(IUVMeshView* Mesh, int NumIslands, TFunctionRef<const TArrayView<const int>(int)> GetIsland)
+bool FUVPacker::StandardPack(IUVMeshView* Mesh, int NumIslands, TFunctionRef<void(int, TArray<int32>&)> CopyIsland)
 {
 	using namespace UE::InternalUVPacking;
 	FStandardChartPacker Packer;
@@ -818,7 +818,7 @@ bool FUVPacker::StandardPack(IUVMeshView* Mesh, int NumIslands, TFunctionRef<con
 
 		Chart.Id = ci + 1;
 
-		Chart.Triangles = GetIsland(ci);
+		CopyIsland(ci, Chart.Triangles);
 
 		Chart.MinUV = FVector2d(FLT_MAX, FLT_MAX);
 		Chart.MaxUV = FVector2d(-FLT_MAX, -FLT_MAX);
@@ -904,7 +904,7 @@ bool FUVPacker::StandardPack(IUVMeshView* Mesh, int NumIslands, TFunctionRef<con
 
 
 
-bool FUVPacker::StackPack(IUVMeshView* Mesh, int NumIslands, TFunctionRef<const TArrayView<const int>(int)> GetIsland)
+bool FUVPacker::StackPack(IUVMeshView* Mesh, int NumIslands, TFunctionRef<void(int, TArray<int32>&)> CopyIsland)
 {
 	double GutterWidth = (double)GutterSize / (double)TextureResolution;
 
@@ -917,7 +917,8 @@ bool FUVPacker::StackPack(IUVMeshView* Mesh, int NumIslands, TFunctionRef<const
 	for (int32 ci = 0; ci < NumCharts; ++ci)
 	{
 		FAxisAlignedBox2d IslandBounds = FAxisAlignedBox2d::Empty();
-		const TArrayView<const int> Island = GetIsland(ci);
+		TArray<int32> Island;
+		CopyIsland(ci, Island);
 		for (int32 tid : Island)
 		{
 			FVector2f UVTri[3];
@@ -949,7 +950,8 @@ bool FUVPacker::StackPack(IUVMeshView* Mesh, int NumIslands, TFunctionRef<const
 	{
 		IslandElements.Reset();
 		FAxisAlignedBox2d IslandBounds = AllIslandBounds[ci];
-		const TArrayView<const int> Island = GetIsland(ci);
+		TArray<int32> Island;
+		CopyIsland(ci, Island);
 		for (int32 tid : Island)
 		{
 			FIndex3i UVTri = Mesh->GetUVTriangle(tid);
diff --git a/Engine/Plugins/Experimental/GeometryProcessing/Source/GeometricObjects/Public/Parameterization/UVPacking.h b/Engine/Plugins/Experimental/GeometryProcessing/Source/GeometricObjects/Public/Parameterization/UVPacking.h
--- a/Engine/Plugins/Experimental/GeometryProcessing/Source/GeometricObjects/Public/Parameterization/UVPacking.h
+++ b/Engine/Plugins/Experimental/GeometryProcessing/Source/GeometricObjects/Public/Parameterization/UVPacking.h
@@ -47,14 +47,14 @@ public:
 	 * All UV islands are packed into standard positive-unit-square.
 	 * Only supports single-pixel border size.
 	 */
-	bool StandardPack(IUVMeshView* Mesh, int NumIslands, TFunctionRef<const TArrayView<const int>(int)> GetIsland);
+	bool StandardPack(IUVMeshView* Mesh, int NumIslands, TFunctionRef<void(int, TArray<int32>&)> CopyIsland);
 
 	/// Version of StandardPack that takes an array of arrays instead of a TFunctionRef, for convenience
 	bool StandardPack(IUVMeshView* Mesh, const TArray<TArray<int>>& UVIslands)
 	{
-		return StandardPack(Mesh, UVIslands.Num(), [UVIslands](int Idx)
+		return StandardPack(Mesh, UVIslands.Num(), [&UVIslands](int Idx, TArray<int32>& IslandOut)
 			{
-				return UVIslands[Idx];
+				IslandOut = UVIslands[Idx];
 			});
 	}
 
@@ -63,14 +63,14 @@ public:
 	 * and translate each islands separately so that it's bbox-min is at the origin.
 	 * So the islands are "stacked" and all fit in the unit box.
 	 */
-	bool StackPack(IUVMeshView* Mesh, int NumIslands, TFunctionRef<const TArrayView<const int>(int)> GetIsland);
+	bool StackPack(IUVMeshView* Mesh, int NumIslands, TFunctionRef<void(int, TArray<int32>&)> CopyIsland);
 
 	/// Version of StackPack that takes an array of arrays instead of a TFunctionRef, for convenience
 	bool StackPack(IUVMeshView* Mesh, const TArray<TArray<int>>& UVIslands)
 	{
-		return StackPack(Mesh, UVIslands.Num(), [UVIslands](int Idx)
+		return StackPack(Mesh, UVIslands.Num(), [&UVIslands](int Idx, TArray<int32>& IslandOut)
 			{
-				return UVIslands[Idx];
+				IslandOut = UVIslands[Idx];
 			});
 	}
 
@@ -127,11 +127,16 @@ inline void CreateUVIslandsFromMeshTopology(TriangleMeshType& Mesh,
 		FIndex3i Tri = Mesh.GetTriangle(TID);
 		uint32 IslandID = VertComponents.Find(Tri.A);
 		int32 Idx = -1;
-		if (!IslandIDToIdx.Contains(IslandID))
+		int32* FoundIdx = IslandIDToIdx.Find(IslandID);
+		if (!FoundIdx)
 		{
 			Idx = IslandsOut.Emplace();
 			IslandIDToIdx.Add(IslandID, Idx);
 		}
+		else
+		{
+			Idx = *FoundIdx;
+		}
 		IslandsOut[Idx].Add(TID);
 	}
 }
