Fix geometry processing unit tests and bug in AABB tree's find-nearest-triangle

#rnx

diff --git a/Engine/Plugins/Experimental/GeometryProcessing/Source/GeometricObjects/Public/Distance/DistTriangle3Triangle3.h b/Engine/Plugins/Experimental/GeometryProcessing/Source/GeometricObjects/Public/Distance/DistTriangle3Triangle3.h
--- a/Engine/Plugins/Experimental/GeometryProcessing/Source/GeometricObjects/Public/Distance/DistTriangle3Triangle3.h
+++ b/Engine/Plugins/Experimental/GeometryProcessing/Source/GeometricObjects/Public/Distance/DistTriangle3Triangle3.h
@@ -19,9 +19,6 @@ template <typename Real>
 class TDistTriangle3Triangle3
 {
 public:
-	// Input
-	TTriangle3<Real> Triangle[2];
-
 	// Output
 	Real DistanceSquared = -1.0;
 	FVector3<Real> TriangleClosest[2];
@@ -37,6 +34,13 @@ public:
 		Triangle[1] = TriangleB;
 	}
 
+	void SetTriangle(int WhichTriangle, const TTriangle3<Real>& TriangleIn)
+	{
+		check(WhichTriangle >= 0 && WhichTriangle < 2);
+		Triangle[WhichTriangle] = TriangleIn;
+		DistanceSquared = -1.0; // clear cached result
+	}
+
 	Real Get()
 	{
 		return TMathUtil<Real>::Sqrt(ComputeResult());
@@ -115,6 +119,11 @@ public:
 		DistanceSquared = sqrDist;
 		return DistanceSquared;
 	}
+	
+private:
+
+	// Input triangles; must be set via SetTriangle to clear cached result
+	TTriangle3<Real> Triangle[2];
 };
 
 typedef TDistTriangle3Triangle3<float> FDistTriangle3Triangle3f;
diff --git a/Engine/Plugins/Experimental/GeometryProcessing/Source/GeometricObjects/Public/Generators/SweepGenerator.h b/Engine/Plugins/Experimental/GeometryProcessing/Source/GeometricObjects/Public/Generators/SweepGenerator.h
--- a/Engine/Plugins/Experimental/GeometryProcessing/Source/GeometricObjects/Public/Generators/SweepGenerator.h
+++ b/Engine/Plugins/Experimental/GeometryProcessing/Source/GeometricObjects/Public/Generators/SweepGenerator.h
@@ -365,6 +365,7 @@ public:
 	int AngleSamples = 16;
 	bool bCapped = false;
 	bool bUVScaleMatchSidesAndCaps = true;
+	ECapType CapType = ECapType::FlatMidpointFan;
 
 	static float ComputeSegLengths(const TArrayView<float>& Radii, const TArrayView<float>& Heights, TArray<float>& AlongPercents)
 	{
@@ -393,8 +394,8 @@ public:
 
 		if (bCapped)
 		{
-			Caps[0] = ECapType::FlatMidpointFan;
-			Caps[1] = ECapType::FlatMidpointFan;
+			Caps[0] = CapType;
+			Caps[1] = CapType;
 		}
 
 		int NumX = Radii.Num();
@@ -636,6 +637,7 @@ public:
 
 	bool bCapped = false;
 	bool bLoop = false;
+	ECapType CapType = ECapType::FlatTriangulation;
 
 	// When true, the generator attempts to scale UV's in a way that preserves scaling across different mesh
 	// results, aiming for 1.0 in UV space to be equal to UnitUVInWorldCoordinates in world space. This in
@@ -655,8 +657,8 @@ public:
 
 		if (bCapped && !bLoop)
 		{
-			Caps[0] = ECapType::FlatTriangulation;
-			Caps[1] = ECapType::FlatTriangulation;
+			Caps[0] = CapType;
+			Caps[1] = CapType;
 		}
 		int PathNum = Path.Num();
 		
diff --git a/Engine/Plugins/Experimental/GeometryProcessing/Source/GeometricObjects/Public/Spatial/MeshAABBTree3.h b/Engine/Plugins/Experimental/GeometryProcessing/Source/GeometricObjects/Public/Spatial/MeshAABBTree3.h
--- a/Engine/Plugins/Experimental/GeometryProcessing/Source/GeometricObjects/Public/Spatial/MeshAABBTree3.h
+++ b/Engine/Plugins/Experimental/GeometryProcessing/Source/GeometricObjects/Public/Spatial/MeshAABBTree3.h
@@ -1257,7 +1257,7 @@ public:
 					otri.V[1] = TransformF(otri.V[1]);
 					otri.V[2] = TransformF(otri.V[2]);
 				}
-				dist.Triangle[0] = otri;
+				dist.SetTriangle(0, otri);
 
 				// inner iteration over "our" triangles
 				for (int i = 1; i <= num_tris; ++i)
@@ -1268,7 +1268,7 @@ public:
 						continue;
 					}
 					Mesh->GetTriVertices(ti, Tri.V[0], Tri.V[1], Tri.V[2]);
-					dist.Triangle[1] = Tri;
+					dist.SetTriangle(1, Tri);
 					double dist_sqr = dist.GetSquared();
 					if (dist_sqr < nearest_sqr)
 					{
