another mac CIS fix for modeling toolset-related code that was previously windows only

diff --git a/Engine/Plugins/Experimental/GeometryProcessing/Source/GeometricObjects/Private/Generators/MarchingCubes.cpp b/Engine/Plugins/Experimental/GeometryProcessing/Source/GeometricObjects/Private/Generators/MarchingCubes.cpp
new file mode 100644
--- /dev/null
+++ b/Engine/Plugins/Experimental/GeometryProcessing/Source/GeometricObjects/Private/Generators/MarchingCubes.cpp
@@ -0,0 +1,8 @@
+// Copyright Epic Games, Inc. All Rights Reserved.
+
+#include "Generators/MarchingCubes.h"
+
+
+constexpr int FMarchingCubes::EdgeIndices[12][2];
+constexpr int FMarchingCubes::EdgeTable[256];
+constexpr int FMarchingCubes::TriTable[256][16];
\ No newline at end of file
diff --git a/Engine/Plugins/Experimental/GeometryProcessing/Source/GeometricObjects/Public/Generators/MarchingCubes.h b/Engine/Plugins/Experimental/GeometryProcessing/Source/GeometricObjects/Public/Generators/MarchingCubes.h
--- a/Engine/Plugins/Experimental/GeometryProcessing/Source/GeometricObjects/Public/Generators/MarchingCubes.h
+++ b/Engine/Plugins/Experimental/GeometryProcessing/Source/GeometricObjects/Public/Generators/MarchingCubes.h
@@ -886,11 +886,11 @@ protected:
 	*/
 
 
-	constexpr static int EdgeIndices[12][2] = {
+	GEOMETRICOBJECTS_API constexpr static int EdgeIndices[12][2] = {
 		{0,1}, {1,2}, {2,3}, {3,0}, {4,5}, {5,6}, {6,7}, {7,4}, {0,4}, {1,5}, {2,6}, {3,7}
 	};
 
-	constexpr static int EdgeTable[256] = {
+	GEOMETRICOBJECTS_API constexpr static int EdgeTable[256] = {
 		0x0  , 0x109, 0x203, 0x30a, 0x406, 0x50f, 0x605, 0x70c,
 		0x80c, 0x905, 0xa0f, 0xb06, 0xc0a, 0xd03, 0xe09, 0xf00,
 		0x190, 0x99 , 0x393, 0x29a, 0x596, 0x49f, 0x795, 0x69c,
@@ -925,7 +925,7 @@ protected:
 		0x70c, 0x605, 0x50f, 0x406, 0x30a, 0x203, 0x109, 0x0   };
 
 
-	constexpr static int TriTable[256][16] =
+	GEOMETRICOBJECTS_API constexpr static int TriTable[256][16] =
 	{
 		{-1, -1, -1, -1, -1, -1, -1, -1, -1, -1, -1, -1, -1, -1, -1, -1},
 		{0, 8, 3, -1, -1, -1, -1, -1, -1, -1, -1, -1, -1, -1, -1, -1},
