Fixup Eigen symbol visibility.

#rnx

diff --git a/Engine/Plugins/Experimental/AlembicImporter/Source/AlembicLibrary/Private/EigenHelper.h b/Engine/Plugins/Experimental/AlembicImporter/Source/AlembicLibrary/Private/EigenHelper.h
--- a/Engine/Plugins/Experimental/AlembicImporter/Source/AlembicLibrary/Private/EigenHelper.h
+++ b/Engine/Plugins/Experimental/AlembicImporter/Source/AlembicLibrary/Private/EigenHelper.h
@@ -13,10 +13,12 @@
 	#pragma warning(push)
 	#pragma warning(disable:6294) // Ill-defined for-loop:  initial condition does not satisfy test.  Loop body not executed.
 #endif
+PRAGMA_DEFAULT_VISIBILITY_START
 THIRD_PARTY_INCLUDES_START
 	#include <Eigen/Dense>
 	#include <Eigen/SVD>
 THIRD_PARTY_INCLUDES_END
+PRAGMA_DEFAULT_VISIBILITY_END
 #if defined(_MSC_VER) && USING_CODE_ANALYSIS
 	#pragma warning(pop)
 #endif
diff --git a/Engine/Plugins/Experimental/GeometryProcessing/Source/DynamicMesh/Private/Solvers/Internal/FSparseMatrixD.h b/Engine/Plugins/Experimental/GeometryProcessing/Source/DynamicMesh/Private/Solvers/Internal/FSparseMatrixD.h
--- a/Engine/Plugins/Experimental/GeometryProcessing/Source/DynamicMesh/Private/Solvers/Internal/FSparseMatrixD.h
+++ b/Engine/Plugins/Experimental/GeometryProcessing/Source/DynamicMesh/Private/Solvers/Internal/FSparseMatrixD.h
@@ -22,9 +22,11 @@
 #pragma warning(disable : 6313)
 #pragma warning(disable : 6294)
 #endif
+PRAGMA_DEFAULT_VISIBILITY_START
 THIRD_PARTY_INCLUDES_START
 #include <Eigen/Sparse>
 THIRD_PARTY_INCLUDES_END
+PRAGMA_DEFAULT_VISIBILITY_END
 #if defined(_MSC_VER) && USING_CODE_ANALYSIS
 #pragma warning(pop)
 #endif
@@ -105,4 +107,4 @@ public:
 		return (Column.Num() > 0) ? ConstVectorType(&Column[0], Column.Num()) : ConstVectorType(nullptr, 0);
 	}
 
-};
\ No newline at end of file
+};
diff --git a/Engine/Plugins/Experimental/GeometryProcessing/Source/DynamicMesh/Private/Solvers/Internal/LaplacianOperators.cpp b/Engine/Plugins/Experimental/GeometryProcessing/Source/DynamicMesh/Private/Solvers/Internal/LaplacianOperators.cpp
--- a/Engine/Plugins/Experimental/GeometryProcessing/Source/DynamicMesh/Private/Solvers/Internal/LaplacianOperators.cpp
+++ b/Engine/Plugins/Experimental/GeometryProcessing/Source/DynamicMesh/Private/Solvers/Internal/LaplacianOperators.cpp
@@ -15,9 +15,11 @@
 #pragma warning(disable : 6313)
 #pragma warning(disable : 6294)
 #endif
+PRAGMA_DEFAULT_VISIBILITY_START
 THIRD_PARTY_INCLUDES_START
 #include <Eigen/Dense> // for Matrix4d in testing
 THIRD_PARTY_INCLUDES_END
+PRAGMA_DEFAULT_VISIBILITY_END
 #if defined(_MSC_VER) && USING_CODE_ANALYSIS
 #pragma warning(pop)
 #endif
diff --git a/Engine/Plugins/Experimental/GeometryProcessing/Source/DynamicMesh/Private/Solvers/Internal/MatrixSolver.h b/Engine/Plugins/Experimental/GeometryProcessing/Source/DynamicMesh/Private/Solvers/Internal/MatrixSolver.h
--- a/Engine/Plugins/Experimental/GeometryProcessing/Source/DynamicMesh/Private/Solvers/Internal/MatrixSolver.h
+++ b/Engine/Plugins/Experimental/GeometryProcessing/Source/DynamicMesh/Private/Solvers/Internal/MatrixSolver.h
@@ -14,6 +14,7 @@
 #pragma warning(disable : 6313)
 #pragma warning(disable : 6294)
 #endif
+PRAGMA_DEFAULT_VISIBILITY_START
 THIRD_PARTY_INCLUDES_START
 #include <Eigen/Sparse>
 #include <Eigen/Core>
@@ -26,6 +27,7 @@ THIRD_PARTY_INCLUDES_START
 #endif
 #include <Eigen/Dense>
 THIRD_PARTY_INCLUDES_END
+PRAGMA_DEFAULT_VISIBILITY_END
 #if defined(_MSC_VER) && USING_CODE_ANALYSIS
 #pragma warning(pop)
 #endif
diff --git a/Engine/Plugins/Runtime/HairStrands/Source/HairStrandsCore/Private/GroomBindingBuilder.cpp b/Engine/Plugins/Runtime/HairStrands/Source/HairStrandsCore/Private/GroomBindingBuilder.cpp
--- a/Engine/Plugins/Runtime/HairStrands/Source/HairStrandsCore/Private/GroomBindingBuilder.cpp
+++ b/Engine/Plugins/Runtime/HairStrands/Source/HairStrandsCore/Private/GroomBindingBuilder.cpp
@@ -23,6 +23,7 @@
 #pragma warning(push)
 #pragma warning(disable:6294) // Ill-defined for-loop:  initial condition does not satisfy test.  Loop body not executed.
 #endif
+PRAGMA_DEFAULT_VISIBILITY_START
 THIRD_PARTY_INCLUDES_START
 #include <Eigen/Core>
 #include <Eigen/Dense>
@@ -30,6 +31,7 @@ THIRD_PARTY_INCLUDES_START
 #include <Eigen/SparseLU>
 #include <Eigen/SVD>
 THIRD_PARTY_INCLUDES_END
+PRAGMA_DEFAULT_VISIBILITY_END
 #if defined(_MSC_VER) && USING_CODE_ANALYSIS
 #pragma warning(pop)
 #endif
diff --git a/Engine/Plugins/Runtime/HairStrands/Source/HairStrandsCore/Private/HairCardsBuilder.cpp b/Engine/Plugins/Runtime/HairStrands/Source/HairStrandsCore/Private/HairCardsBuilder.cpp
--- a/Engine/Plugins/Runtime/HairStrands/Source/HairStrandsCore/Private/HairCardsBuilder.cpp
+++ b/Engine/Plugins/Runtime/HairStrands/Source/HairStrandsCore/Private/HairCardsBuilder.cpp
@@ -32,12 +32,14 @@
 #pragma warning(push)
 #pragma warning(disable:6294) // Ill-defined for-loop:  initial condition does not satisfy test.  Loop body not executed.
 #endif
+PRAGMA_DEFAULT_VISIBILITY_START
 THIRD_PARTY_INCLUDES_START
 #include <Eigen/Core>
 #include <Eigen/Dense>
 #include <Eigen/Sparse>
 #include <Eigen/SparseLU>
 THIRD_PARTY_INCLUDES_END
+PRAGMA_DEFAULT_VISIBILITY_END
 #if defined(_MSC_VER) && USING_CODE_ANALYSIS
 #pragma warning(pop)
 #endif
