GeometryProcessing: move all public types in GeometryProcessing plugin into UE::Geometry:: namespace, as well as some code in MeshModelingToolset. Fix fallout in Water, Paint Mode, ChaosEditor, DataPr

diff --git a/Engine/Plugins/Experimental/GeometryProcessing/Source/GeometricObjects/Public/ExplicitUseGeometryMathTypes.h b/Engine/Plugins/Experimental/GeometryProcessing/Source/GeometricObjects/Public/ExplicitUseGeometryMathTypes.h
new file mode 100644
--- /dev/null
+++ b/Engine/Plugins/Experimental/GeometryProcessing/Source/GeometricObjects/Public/ExplicitUseGeometryMathTypes.h
@@ -0,0 +1,123 @@
+// Copyright Epic Games, Inc. All Rights Reserved.
+
+#pragma once
+
+//
+// This header exists to simplify usage of the core UE::Geometry:: math types, such as FVector3d/f, FTransform3d/f, etc. 
+// The GeometryProcessing plugin used these type names for several years before the UE Core math library was extended with
+// double-precision support, at which point it was decided to rename FVector as FVector3f/FVector3d, etc, creating a conflict. 
+// Unfortunately the UE Core versions are not a direct replacement for the GeometryProcessing variants, because:
+//   1) they have templated base classes, and significant other template code takes advantage of this. The UE Core Math types currently do not.
+//   2) they use more reliable (but potentially slower) standard-library math functions
+//   3) they have function naming schemes compatible with other standard vector math libraries (which simplifies porting/etc)
+//   4) they have internal member differences in some case (eg UE::Geometry::FTransform3d does not use SSE, Matrix2/3/4 stores as Vectors, etc)
+//
+// As a result the global ::FVector3d and UE::Geometry::FVector3d names are ambiguous. To avoid requiring explicit UE::Geometry:: disambiguation 
+// in all cases, which would make code harder to read and vastly complicate UE4/UE5 merging, in code that primarily uses GeometryProcessing
+// types, we prefer to explicitly specify - via C++ 'using' declarations - that we are using the UE::Geometry:: types at cpp file scope.
+// Rather than add a standard block of declarations into each file, this header can be included **IN THE CPP ONLY** to do so.
+//
+// Note that the header does not contain 'using namespace UE::Geometry;' - we leave that to the cpp file, so the standard boilerplate is:
+//
+//		#include "ExplicitUseGeometryMathTypes.h"		// using UE::Geometry::(math types)
+//		using namespace UE::Geometry;
+//
+// This header **MUST NOT BE INCLUDED IN OTHER HEADERS** as this could result in ambiguous definitions for downstream includers.
+// In headers, the GeometryProcessing core math types can be explicitly 'used' in other namespaces, or in class definitions, but not in global scope.
+//
+
+
+#include "IndexTypes.h"
+#include "VectorTypes.h"
+#include "MatrixTypes.h"
+#include "Quaternion.h"
+#include "BoxTypes.h"
+#include "PlaneTypes.h"
+#include "HalfspaceTypes.h"
+#include "LineTypes.h"
+#include "SegmentTypes.h"
+#include "RayTypes.h"
+#include "SphereTypes.h"
+#include "CapsuleTypes.h"
+#include "OrientedBoxTypes.h"
+#include "TriangleTypes.h"
+#include "CircleTypes.h"
+#include "FrameTypes.h"
+#include "TransformTypes.h"
+
+using UE::Geometry::FIndex2i;
+using UE::Geometry::FIndex3i;
+using UE::Geometry::FIndex4i;
+
+using UE::Geometry::FVector2f;
+using UE::Geometry::FVector2d;
+using UE::Geometry::FVector2i;
+using UE::Geometry::FVector3f;
+using UE::Geometry::FVector3d;
+using UE::Geometry::FVector3i;
+using UE::Geometry::FVector4f;
+using UE::Geometry::FVector4d;
+using UE::Geometry::FVector4i;
+
+using UE::Geometry::FMatrix3f;
+using UE::Geometry::FMatrix3d;
+using UE::Geometry::FMatrix2f;
+using UE::Geometry::FMatrix2d;
+
+using UE::Geometry::FQuaternionf;
+using UE::Geometry::FQuaterniond;
+
+using UE::Geometry::FInterval1f;
+using UE::Geometry::FInterval1d;
+using UE::Geometry::FInterval1i;
+using UE::Geometry::FAxisAlignedBox2f;
+using UE::Geometry::FAxisAlignedBox2d;
+using UE::Geometry::FAxisAlignedBox2i;
+using UE::Geometry::FAxisAlignedBox3f;
+using UE::Geometry::FAxisAlignedBox3d;
+using UE::Geometry::FAxisAlignedBox3i;
+
+using UE::Geometry::FHalfspace3d;
+using UE::Geometry::FHalfspace3f;
+
+using UE::Geometry::FPlane3d;
+using UE::Geometry::FPlane3f;
+
+using UE::Geometry::FLine2f;
+using UE::Geometry::FLine2d;
+using UE::Geometry::FLine3f;
+using UE::Geometry::FLine3d;
+
+using UE::Geometry::FSegment2f;
+using UE::Geometry::FSegment2d;
+using UE::Geometry::FSegment3f;
+using UE::Geometry::FSegment3d;
+
+using UE::Geometry::FRay3f;
+using UE::Geometry::FRay3d;
+
+using UE::Geometry::FSphere3d;
+using UE::Geometry::FSphere3f;
+
+using UE::Geometry::FCapsule3d;
+using UE::Geometry::FCapsule3f;
+
+using UE::Geometry::FOrientedBox3d;
+using UE::Geometry::FOrientedBox3f;
+
+using UE::Geometry::FTriangle2d;
+using UE::Geometry::FTriangle2f;
+using UE::Geometry::FTriangle3d;
+using UE::Geometry::FTriangle3f;
+
+using UE::Geometry::FCircle2d;
+using UE::Geometry::FCircle2f;
+using UE::Geometry::FCircle3d;
+using UE::Geometry::FCircle3f;
+
+using UE::Geometry::FTransform3d;
+using UE::Geometry::FTransform3f;
+
+using UE::Geometry::FFrame3d;
+using UE::Geometry::FFrame3f;
+
diff --git a/Engine/Plugins/Experimental/GeometryProcessing/Source/GeometricObjects/Public/Generators/CapsuleGenerator.h b/Engine/Plugins/Experimental/GeometryProcessing/Source/GeometricObjects/Public/Generators/CapsuleGenerator.h
--- a/Engine/Plugins/Experimental/GeometryProcessing/Source/GeometricObjects/Public/Generators/CapsuleGenerator.h
+++ b/Engine/Plugins/Experimental/GeometryProcessing/Source/GeometricObjects/Public/Generators/CapsuleGenerator.h
@@ -6,6 +6,12 @@
 #include "OrientedBoxTypes.h"
 #include "Util/IndexUtil.h"
 
+namespace UE
+{
+namespace Geometry
+{
+
+
 /**
  * Generate a Capsule mesh, with UVs wrapped cylindrically.
  * This is basically a "stretched" standard sphere triangulation, where we have
@@ -235,6 +241,8 @@ public:
 		OutputPolarTriangles();
 		return *this;
 	}
+};
 
 
-};
+} // end namespace UE::Geometry
+} // end namespace UE
diff --git a/Engine/Plugins/Experimental/GeometryProcessing/Source/GeometricObjects/Public/Generators/MeshShapeGenerator.h b/Engine/Plugins/Experimental/GeometryProcessing/Source/GeometricObjects/Public/Generators/MeshShapeGenerator.h
--- a/Engine/Plugins/Experimental/GeometryProcessing/Source/GeometricObjects/Public/Generators/MeshShapeGenerator.h
+++ b/Engine/Plugins/Experimental/GeometryProcessing/Source/GeometricObjects/Public/Generators/MeshShapeGenerator.h
@@ -5,6 +5,11 @@
 #include "VectorTypes.h"
 #include "IndexTypes.h"
 
+namespace UE
+{
+namespace Geometry
+{
+
 /**
  * Base class for triangle mesh generators (eg like to generate sphere, cylinder, etc)
  * Subclasses must implement ::Generate() 
@@ -261,6 +266,7 @@ public:
 		return FVector3i((int)round(c.X), (int)round(c.Y), (int)round(c.Z));
 	}
 
+};
 
-
-};
\ No newline at end of file
+} // end namespace UE::Geometry
+} // end namespace UE
\ No newline at end of file
diff --git a/Engine/Plugins/Experimental/GeometryProcessing/Source/GeometricObjects/Public/Generators/MinimalBoxMeshGenerator.h b/Engine/Plugins/Experimental/GeometryProcessing/Source/GeometricObjects/Public/Generators/MinimalBoxMeshGenerator.h
--- a/Engine/Plugins/Experimental/GeometryProcessing/Source/GeometricObjects/Public/Generators/MinimalBoxMeshGenerator.h
+++ b/Engine/Plugins/Experimental/GeometryProcessing/Source/GeometricObjects/Public/Generators/MinimalBoxMeshGenerator.h
@@ -6,6 +6,11 @@
 #include "OrientedBoxTypes.h"
 #include "Util/IndexUtil.h"
 
+namespace UE
+{
+namespace Geometry
+{
+
 /**
  * Generate an oriented Box mesh with the smallest number of triangles possible (6 vertices, 12 triangles)
  */
@@ -79,8 +84,8 @@ public:
 		return *this;
 	}
 
+};
 
 
-
-
-};
+} // end namespace UE::Geometry
+} // end namespace UE
diff --git a/Engine/Plugins/Experimental/GeometryProcessing/Source/GeometricObjects/Public/Generators/RectangleMeshGenerator.h b/Engine/Plugins/Experimental/GeometryProcessing/Source/GeometricObjects/Public/Generators/RectangleMeshGenerator.h
--- a/Engine/Plugins/Experimental/GeometryProcessing/Source/GeometricObjects/Public/Generators/RectangleMeshGenerator.h
+++ b/Engine/Plugins/Experimental/GeometryProcessing/Source/GeometricObjects/Public/Generators/RectangleMeshGenerator.h
@@ -5,6 +5,11 @@
 #include "MeshShapeGenerator.h"
 #include "Misc/EnumClassFlags.h"
 
+namespace UE
+{
+namespace Geometry
+{
+
 /**
  * Generate planar rectangular mesh with variable number of subdivisions along width and height.
  * By default, center of rectangle is centered at (0,0,0) origin
@@ -102,7 +107,8 @@ public:
 
 	/** Generate the mesh */
 	virtual FMeshShapeGenerator& Generate() override;
+};
 
 
-
-};
+} // end namespace UE::Geometry
+} // end namespace UE
\ No newline at end of file
diff --git a/Engine/Plugins/Experimental/GeometryProcessing/Source/GeometricObjects/Public/Generators/StairGenerator.h b/Engine/Plugins/Experimental/GeometryProcessing/Source/GeometricObjects/Public/Generators/StairGenerator.h
--- a/Engine/Plugins/Experimental/GeometryProcessing/Source/GeometricObjects/Public/Generators/StairGenerator.h
+++ b/Engine/Plugins/Experimental/GeometryProcessing/Source/GeometricObjects/Public/Generators/StairGenerator.h
@@ -5,6 +5,13 @@
 #include "IndexTypes.h"
 #include "MathUtil.h"
 
+
+namespace UE
+{
+namespace Geometry
+{
+
+
 /**
  * Base Stair mesh generator class.
  */
@@ -989,3 +996,8 @@ protected:
 		return FMathf::Max(FMathf::Max(MaxDepth, NumSteps * StepHeight), StepWidth);
 	}
 };
+
+
+
+} // end namespace UE::Geometry
+} // end namespace UE
\ No newline at end of file
diff --git a/Engine/Plugins/Experimental/GeometryProcessing/Source/GeometricObjects/Public/HalfspaceTypes.h b/Engine/Plugins/Experimental/GeometryProcessing/Source/GeometricObjects/Public/HalfspaceTypes.h
--- a/Engine/Plugins/Experimental/GeometryProcessing/Source/GeometricObjects/Public/HalfspaceTypes.h
+++ b/Engine/Plugins/Experimental/GeometryProcessing/Source/GeometricObjects/Public/HalfspaceTypes.h
@@ -5,6 +5,11 @@
 #include "Math/UnrealMath.h"
 #include "VectorTypes.h"
 
+namespace UE
+{
+namespace Geometry
+{
+
 /*
  * 3D Halfspace stored as parameters to Plane Equation (Normal, Normal.Dot(PointOnPlane))
  * The Normal points "into" the halfspace, ie X is inside if (Normal.Dot(X) - Constant) >= 0
@@ -49,4 +54,8 @@ public:
 
 
 typedef THalfspace3<float> FHalfspace3f;
-typedef THalfspace3<double> FHalfspace3d;
\ No newline at end of file
+typedef THalfspace3<double> FHalfspace3d;
+
+
+} // end namespace UE::Geometry
+} // end namespace UE
\ No newline at end of file
diff --git a/Engine/Plugins/Experimental/GeometryProcessing/Source/GeometricObjects/Public/Image/ImageDimensions.h b/Engine/Plugins/Experimental/GeometryProcessing/Source/GeometricObjects/Public/Image/ImageDimensions.h
--- a/Engine/Plugins/Experimental/GeometryProcessing/Source/GeometricObjects/Public/Image/ImageDimensions.h
+++ b/Engine/Plugins/Experimental/GeometryProcessing/Source/GeometricObjects/Public/Image/ImageDimensions.h
@@ -4,6 +4,11 @@
 
 #include "VectorTypes.h"
 
+namespace UE
+{
+namespace Geometry
+{
+
 /**
  * FImageDimensions provides various functions for working with size/indices/coordinates of a
  * 2D image, as well as standard UV spaces
@@ -129,6 +134,7 @@ public:
 };
 
 
-
+} // end namespace UE::Geometry
+} // end namespace UE
 
 
diff --git a/Engine/Plugins/Experimental/GeometryProcessing/Source/GeometricObjects/Public/Implicit/GridInterpolant.h b/Engine/Plugins/Experimental/GeometryProcessing/Source/GeometricObjects/Public/Implicit/GridInterpolant.h
--- a/Engine/Plugins/Experimental/GeometryProcessing/Source/GeometricObjects/Public/Implicit/GridInterpolant.h
+++ b/Engine/Plugins/Experimental/GeometryProcessing/Source/GeometricObjects/Public/Implicit/GridInterpolant.h
@@ -5,6 +5,12 @@
 #pragma once
 
 #include "MathUtil.h"
+#include "VectorTypes.h"
+
+namespace UE
+{
+namespace Geometry
+{
 
 
 /**
@@ -184,3 +190,7 @@ public:
 		return FVector3d(gradX, gradY, gradZ);
 	}
 };
+
+
+} // end namespace UE::Geometry
+} // end namespace UE
\ No newline at end of file
diff --git a/Engine/Plugins/Experimental/GeometryProcessing/Source/GeometricObjects/Public/Implicit/Morphology.h b/Engine/Plugins/Experimental/GeometryProcessing/Source/GeometricObjects/Public/Implicit/Morphology.h
--- a/Engine/Plugins/Experimental/GeometryProcessing/Source/GeometricObjects/Public/Implicit/Morphology.h
+++ b/Engine/Plugins/Experimental/GeometryProcessing/Source/GeometricObjects/Public/Implicit/Morphology.h
@@ -15,8 +15,10 @@
 #include "Generators/MarchingCubes.h"
 #include "MeshQueries.h"
 
-
-
+namespace UE
+{
+namespace Geometry
+{
 
 template<typename TriangleMeshType>
 class TImplicitMorphology
@@ -100,15 +102,15 @@ public:
 		double SignedOffset = UnsignedOffset;
 		switch (MorphologyOp)
 		{
-		case TImplicitMorphology<FDynamicMesh3>::EMorphologyOp::Dilate:
-		case TImplicitMorphology<FDynamicMesh3>::EMorphologyOp::Close:
+		case TImplicitMorphology<TriangleMeshType>::EMorphologyOp::Dilate:
+		case TImplicitMorphology<TriangleMeshType>::EMorphologyOp::Close:
 			SignedOffset = -SignedOffset;
 			break;
 		}
 
 		ComputeFirstPass(UnsignedOffset, SignedOffset);
 
-		if (MorphologyOp == TImplicitMorphology<FDynamicMesh3>::EMorphologyOp::Close || MorphologyOp == TImplicitMorphology<FDynamicMesh3>::EMorphologyOp::Open)
+		if (MorphologyOp == TImplicitMorphology<TriangleMeshType>::EMorphologyOp::Close || MorphologyOp == TImplicitMorphology<TriangleMeshType>::EMorphologyOp::Open)
 		{
 			ComputeSecondPass(UnsignedOffset, -SignedOffset);
 		}
@@ -239,3 +241,7 @@ protected:
 		MarchingCubes.Implicit = nullptr;
 	}
 };
+
+
+} // end namespace UE::Geometry
+} // end namespace UE
\ No newline at end of file
diff --git a/Engine/Plugins/Experimental/GeometryProcessing/Source/GeometricObjects/Public/IndexTypes.h b/Engine/Plugins/Experimental/GeometryProcessing/Source/GeometricObjects/Public/IndexTypes.h
--- a/Engine/Plugins/Experimental/GeometryProcessing/Source/GeometricObjects/Public/IndexTypes.h
+++ b/Engine/Plugins/Experimental/GeometryProcessing/Source/GeometricObjects/Public/IndexTypes.h
@@ -2,6 +2,7 @@
 
 #pragma once
 
+#include "GeometryBase.h"
 #include "Math/IntVector.h"
 #include <limits>
 
@@ -11,6 +12,11 @@ namespace IndexConstants
 	constexpr int InvalidID = -1;
 }
 
+namespace UE
+{
+namespace Geometry
+{
+
 /**
  * 2-index tuple. Ported from g3Sharp library, with the intention of
  * maintaining compatibility with existing g3Sharp code. Has an API
@@ -231,3 +237,7 @@ FORCEINLINE uint32 GetTypeHash(const FIndex4i& Index)
 	// Note: this assumes there's no padding that could contain uncompared data.
 	return FCrc::MemCrc_DEPRECATED(&Index, sizeof(FIndex4i));
 }
+
+
+} // end namespace UE::Geometry
+} // end namespace UE
diff --git a/Engine/Plugins/Experimental/GeometryProcessing/Source/GeometricObjects/Public/Intersection/IntersectionUtil.h b/Engine/Plugins/Experimental/GeometryProcessing/Source/GeometricObjects/Public/Intersection/IntersectionUtil.h
--- a/Engine/Plugins/Experimental/GeometryProcessing/Source/GeometricObjects/Public/Intersection/IntersectionUtil.h
+++ b/Engine/Plugins/Experimental/GeometryProcessing/Source/GeometricObjects/Public/Intersection/IntersectionUtil.h
@@ -6,6 +6,10 @@
 #include "BoxTypes.h"
 
 
+namespace UE
+{
+namespace Geometry
+{
 
 /**
  * FLinearIntersection contains intersection information returned by linear/primitive intersection functions
@@ -18,11 +22,16 @@ struct FLinearIntersection
 };
 
 
+} // end namespace UE::Geometry
+} // end namespace UE
+
+
 /**
  * IntersectionUtil contains functions to compute intersections between geometric objects
  */
 namespace IntersectionUtil
 {
+	using namespace UE::Geometry;
 
 	template<typename RealType>
 	bool RayTriangleTest(const FVector3<RealType>& RayOrigin, const FVector3<RealType>& RayDirection, const FVector3<RealType>& V0, const FVector3<RealType>& V1, const FVector3<RealType>& V2)
diff --git a/Engine/Plugins/Experimental/GeometryProcessing/Source/GeometricObjects/Public/Intersection/IntrLine2Triangle2.h b/Engine/Plugins/Experimental/GeometryProcessing/Source/GeometricObjects/Public/Intersection/IntrLine2Triangle2.h
--- a/Engine/Plugins/Experimental/GeometryProcessing/Source/GeometricObjects/Public/Intersection/IntrLine2Triangle2.h
+++ b/Engine/Plugins/Experimental/GeometryProcessing/Source/GeometricObjects/Public/Intersection/IntrLine2Triangle2.h
@@ -11,6 +11,11 @@
 
 #include "Intersection/Intersector1.h"
 
+namespace UE
+{
+namespace Geometry
+{
+
 /**
  * Compute intersection between 2D Line and 2D Triangle
  */
@@ -258,4 +263,7 @@ public:
 };
 
 typedef TIntrLine2Triangle2<float> FIntrLine2Triangle2f;
-typedef TIntrLine2Triangle2<double> FIntrLine2Triangle2d;
\ No newline at end of file
+typedef TIntrLine2Triangle2<double> FIntrLine2Triangle2d;
+
+} // end namespace UE::Geometry
+} // end namespace UE
diff --git a/Engine/Plugins/Experimental/GeometryProcessing/Source/GeometricObjects/Public/Intersection/IntrRay3AxisAlignedBox3.h b/Engine/Plugins/Experimental/GeometryProcessing/Source/GeometricObjects/Public/Intersection/IntrRay3AxisAlignedBox3.h
--- a/Engine/Plugins/Experimental/GeometryProcessing/Source/GeometricObjects/Public/Intersection/IntrRay3AxisAlignedBox3.h
+++ b/Engine/Plugins/Experimental/GeometryProcessing/Source/GeometricObjects/Public/Intersection/IntrRay3AxisAlignedBox3.h
@@ -8,6 +8,10 @@
 #include "BoxTypes.h"
 #include "VectorUtil.h"
 
+namespace UE
+{
+namespace Geometry
+{
 
 /**
  * Compute intersection between 3D ray and 3D axis-aligned box
@@ -217,4 +221,7 @@ protected:
 };
 
 typedef TIntrRay3AxisAlignedBox3<float> FIntrRay3AxisAlignedBox3f;
-typedef TIntrRay3AxisAlignedBox3<double> FIntrRay3AxisAlignedBox3d;
\ No newline at end of file
+typedef TIntrRay3AxisAlignedBox3<double> FIntrRay3AxisAlignedBox3d;
+
+} // end namespace UE::Geometry
+} // end namespace UE
diff --git a/Engine/Plugins/Experimental/GeometryProcessing/Source/GeometricObjects/Public/Intersection/IntrRay3Triangle3.h b/Engine/Plugins/Experimental/GeometryProcessing/Source/GeometricObjects/Public/Intersection/IntrRay3Triangle3.h
--- a/Engine/Plugins/Experimental/GeometryProcessing/Source/GeometricObjects/Public/Intersection/IntrRay3Triangle3.h
+++ b/Engine/Plugins/Experimental/GeometryProcessing/Source/GeometricObjects/Public/Intersection/IntrRay3Triangle3.h
@@ -9,6 +9,11 @@
 #include "TriangleTypes.h"
 #include "VectorUtil.h"
 
+namespace UE
+{
+namespace Geometry
+{
+
 /**
  * Compute intersection between 3D ray and 3D triangle
  */
@@ -169,4 +174,7 @@ public:
 };
 
 typedef TIntrRay3Triangle3<float> FIntrRay3Triangle3f;
-typedef TIntrRay3Triangle3<double> FIntrRay3Triangle3d;
\ No newline at end of file
+typedef TIntrRay3Triangle3<double> FIntrRay3Triangle3d;
+
+} // end namespace UE::Geometry
+} // end namespace UE
diff --git a/Engine/Plugins/Experimental/GeometryProcessing/Source/GeometricObjects/Public/Intersection/IntrSegment2Segment2.h b/Engine/Plugins/Experimental/GeometryProcessing/Source/GeometricObjects/Public/Intersection/IntrSegment2Segment2.h
--- a/Engine/Plugins/Experimental/GeometryProcessing/Source/GeometricObjects/Public/Intersection/IntrSegment2Segment2.h
+++ b/Engine/Plugins/Experimental/GeometryProcessing/Source/GeometricObjects/Public/Intersection/IntrSegment2Segment2.h
@@ -12,6 +12,11 @@
 #include "VectorTypes.h"
 #include "VectorUtil.h"  // for EIntersectionType
 
+namespace UE
+{
+namespace Geometry
+{
+
 // ported from WildMagic5
 //
 //
@@ -243,4 +248,7 @@ protected:
 };
 
 typedef TIntrSegment2Segment2<double> FIntrSegment2Segment2d;
-typedef TIntrSegment2Segment2<float> FIntrSegment2Segment2f;
\ No newline at end of file
+typedef TIntrSegment2Segment2<float> FIntrSegment2Segment2f;
+
+} // end namespace UE::Geometry
+} // end namespace UE
diff --git a/Engine/Plugins/Experimental/GeometryProcessing/Source/GeometricObjects/Public/Intersection/IntrSegment2Triangle2.h b/Engine/Plugins/Experimental/GeometryProcessing/Source/GeometricObjects/Public/Intersection/IntrSegment2Triangle2.h
--- a/Engine/Plugins/Experimental/GeometryProcessing/Source/GeometricObjects/Public/Intersection/IntrSegment2Triangle2.h
+++ b/Engine/Plugins/Experimental/GeometryProcessing/Source/GeometricObjects/Public/Intersection/IntrSegment2Triangle2.h
@@ -11,6 +11,11 @@
 
 #include "Intersection/IntrLine2Triangle2.h"
 
+namespace UE
+{
+namespace Geometry
+{
+
 /**
  * Compute intersection between 2D segment and 2D triangle
  */
@@ -142,4 +147,7 @@ public:
 };
 
 typedef TIntrSegment2Triangle2<float> FIntrSegment2Triangle2f;
-typedef TIntrSegment2Triangle2<double> FIntrSegment2Triangle2d;
\ No newline at end of file
+typedef TIntrSegment2Triangle2<double> FIntrSegment2Triangle2d;
+
+} // end namespace UE::Geometry
+} // end namespace UE
diff --git a/Engine/Plugins/Experimental/GeometryProcessing/Source/GeometricObjects/Public/Intersection/IntrTriangle2Triangle2.h b/Engine/Plugins/Experimental/GeometryProcessing/Source/GeometricObjects/Public/Intersection/IntrTriangle2Triangle2.h
--- a/Engine/Plugins/Experimental/GeometryProcessing/Source/GeometricObjects/Public/Intersection/IntrTriangle2Triangle2.h
+++ b/Engine/Plugins/Experimental/GeometryProcessing/Source/GeometricObjects/Public/Intersection/IntrTriangle2Triangle2.h
@@ -8,6 +8,11 @@
 #include "TriangleTypes.h"
 #include "VectorUtil.h"
 
+namespace UE
+{
+namespace Geometry
+{
+
 /**
  * Compute intersection between 2D triangles
  */
@@ -311,4 +316,7 @@ private:
 };
 
 typedef TIntrTriangle2Triangle2<float> FIntrTriangle2Triangle2f;
-typedef TIntrTriangle2Triangle2<double> FIntrTriangle2Triangle2d;
\ No newline at end of file
+typedef TIntrTriangle2Triangle2<double> FIntrTriangle2Triangle2d;
+
+} // end namespace UE::Geometry
+} // end namespace UE
diff --git a/Engine/Plugins/Experimental/GeometryProcessing/Source/GeometricObjects/Public/Intersection/IntrTriangle3Triangle3.h b/Engine/Plugins/Experimental/GeometryProcessing/Source/GeometricObjects/Public/Intersection/IntrTriangle3Triangle3.h
--- a/Engine/Plugins/Experimental/GeometryProcessing/Source/GeometricObjects/Public/Intersection/IntrTriangle3Triangle3.h
+++ b/Engine/Plugins/Experimental/GeometryProcessing/Source/GeometricObjects/Public/Intersection/IntrTriangle3Triangle3.h
@@ -14,6 +14,11 @@
 #include "Intersection/IntrSegment2Triangle2.h"
 #include "Intersection/IntrTriangle2Triangle2.h"
 
+namespace UE
+{
+namespace Geometry
+{
+
 /**
  * Compute intersection between 3D triangles
  * use Test() for fast boolean query, does not compute intersection info
@@ -824,4 +829,7 @@ protected:
 };
 
 typedef TIntrTriangle3Triangle3<float> FIntrTriangle3Triangle3f;
-typedef TIntrTriangle3Triangle3<double> FIntrTriangle3Triangle3d;
\ No newline at end of file
+typedef TIntrTriangle3Triangle3<double> FIntrTriangle3Triangle3d;
+
+} // end namespace UE::Geometry
+} // end namespace UE
diff --git a/Engine/Plugins/Experimental/GeometryProcessing/Source/GeometricObjects/Public/LineTypes.h b/Engine/Plugins/Experimental/GeometryProcessing/Source/GeometricObjects/Public/LineTypes.h
--- a/Engine/Plugins/Experimental/GeometryProcessing/Source/GeometricObjects/Public/LineTypes.h
+++ b/Engine/Plugins/Experimental/GeometryProcessing/Source/GeometricObjects/Public/LineTypes.h
@@ -7,6 +7,11 @@
 #include "Math/UnrealMath.h"
 #include "VectorTypes.h"
 
+namespace UE
+{
+namespace Geometry
+{
+
 /**
  * TLine2 is a two-dimensional infinite line.
  * The line is stored in (Center,Direction) form.
@@ -212,4 +217,8 @@ struct TLine3
 };
 
 typedef TLine3<double> FLine3d;
-typedef TLine3<float> FLine3f;
\ No newline at end of file
+typedef TLine3<float> FLine3f;
+
+} // end namespace UE::Geometry
+} // end namespace UE
+
diff --git a/Engine/Plugins/Experimental/GeometryProcessing/Source/GeometricObjects/Public/MathUtil.h b/Engine/Plugins/Experimental/GeometryProcessing/Source/GeometricObjects/Public/MathUtil.h
--- a/Engine/Plugins/Experimental/GeometryProcessing/Source/GeometricObjects/Public/MathUtil.h
+++ b/Engine/Plugins/Experimental/GeometryProcessing/Source/GeometricObjects/Public/MathUtil.h
@@ -2,7 +2,7 @@
 
 #pragma once
 
-#include "CoreTypes.h"   // required for GEOMETRICOBJECTS_API macro
+#include "GeometryBase.h"
 #include <cmath>
 #include <cfloat>
 
diff --git a/Engine/Plugins/Experimental/GeometryProcessing/Source/GeometricObjects/Public/MatrixTypes.h b/Engine/Plugins/Experimental/GeometryProcessing/Source/GeometricObjects/Public/MatrixTypes.h
--- a/Engine/Plugins/Experimental/GeometryProcessing/Source/GeometricObjects/Public/MatrixTypes.h
+++ b/Engine/Plugins/Experimental/GeometryProcessing/Source/GeometricObjects/Public/MatrixTypes.h
@@ -5,6 +5,11 @@
 #include "VectorTypes.h"
 #include "VectorUtil.h"
 
+namespace UE
+{
+namespace Geometry
+{
+
 template <typename RealType>
 struct TMatrix3
 {
@@ -476,4 +481,7 @@ inline static TMatrix3<RealType> CrossProductMatrix(const FVector3<RealType>& A)
 typedef TMatrix3<float> FMatrix3f;
 typedef TMatrix3<double> FMatrix3d;
 typedef TMatrix2<float> FMatrix2f;
-typedef TMatrix2<double> FMatrix2d;
\ No newline at end of file
+typedef TMatrix2<double> FMatrix2d;
+
+} // end namespace UE::Geometry
+} // end namespace UE
\ No newline at end of file
diff --git a/Engine/Plugins/Experimental/GeometryProcessing/Source/GeometricObjects/Public/MeshAdapterTransforms.h b/Engine/Plugins/Experimental/GeometryProcessing/Source/GeometricObjects/Public/MeshAdapterTransforms.h
--- a/Engine/Plugins/Experimental/GeometryProcessing/Source/GeometricObjects/Public/MeshAdapterTransforms.h
+++ b/Engine/Plugins/Experimental/GeometryProcessing/Source/GeometricObjects/Public/MeshAdapterTransforms.h
@@ -21,6 +21,12 @@
  */
 namespace MeshAdapterTransforms
 {
+	using namespace UE::Geometry;
+	using FVector3d = UE::Geometry::FVector3d;
+	using FVector3f = UE::Geometry::FVector3f;
+	using FTransform3d = UE::Geometry::FTransform3d;
+	using FFrame3d = UE::Geometry::FFrame3d;
+
 	/**
 	* Apply Translation to vertex positions of Mesh. Does not modify any other attributes.
 	*/
diff --git a/Engine/Plugins/Experimental/GeometryProcessing/Source/GeometricObjects/Public/OrientedBoxTypes.h b/Engine/Plugins/Experimental/GeometryProcessing/Source/GeometricObjects/Public/OrientedBoxTypes.h
--- a/Engine/Plugins/Experimental/GeometryProcessing/Source/GeometricObjects/Public/OrientedBoxTypes.h
+++ b/Engine/Plugins/Experimental/GeometryProcessing/Source/GeometricObjects/Public/OrientedBoxTypes.h
@@ -8,6 +8,12 @@
 #include "BoxTypes.h"
 #include "FrameTypes.h"
 
+namespace UE
+{
+namespace Geometry
+{
+
+
 /**
  * TOrientedBox3 is a non-axis-aligned 3D box defined by a 3D frame and extents along the axes of that frame
  * The frame is at the center of the box.
@@ -269,8 +275,8 @@ struct TOrientedBox3
 
 };
 
-
-
-
 typedef TOrientedBox3<float> FOrientedBox3f;
 typedef TOrientedBox3<double> FOrientedBox3d;
+
+} // end namespace UE::Geometry
+} // end namespace UE
diff --git a/Engine/Plugins/Experimental/GeometryProcessing/Source/GeometricObjects/Public/Parameterization/UVPacking.h b/Engine/Plugins/Experimental/GeometryProcessing/Source/GeometricObjects/Public/Parameterization/UVPacking.h
--- a/Engine/Plugins/Experimental/GeometryProcessing/Source/GeometricObjects/Public/Parameterization/UVPacking.h
+++ b/Engine/Plugins/Experimental/GeometryProcessing/Source/GeometricObjects/Public/Parameterization/UVPacking.h
@@ -4,12 +4,12 @@
 
 #include "VectorTypes.h"
 #include "IndexTypes.h"
-
-
-#include "CoreMinimal.h"
 #include "DisjointSet.h"
 
-
+namespace UE
+{
+namespace Geometry
+{
 
 /**
  * FUVPacker implements various strategies for packing UV islands in a
@@ -77,11 +77,14 @@ public:
 };
 
 
+} // end namespace UE::Geometry
+} // end namespace UE
+
 namespace UE
 {
 namespace UVPacking
 {
-
+	using namespace UE::Geometry;
 /**
  * Create UV islands from a triangle mesh connectivity.
  * Assumes the triangles are already split at UV seams, but topologically connected otherwise.
diff --git a/Engine/Plugins/Experimental/GeometryProcessing/Source/GeometricObjects/Public/PlaneTypes.h b/Engine/Plugins/Experimental/GeometryProcessing/Source/GeometricObjects/Public/PlaneTypes.h
--- a/Engine/Plugins/Experimental/GeometryProcessing/Source/GeometricObjects/Public/PlaneTypes.h
+++ b/Engine/Plugins/Experimental/GeometryProcessing/Source/GeometricObjects/Public/PlaneTypes.h
@@ -7,6 +7,11 @@
 #include "VectorTypes.h"
 #include "VectorUtil.h"
 
+namespace UE
+{
+namespace Geometry
+{
+
 template<typename RealType>
 struct TPlane3
 {
@@ -149,5 +154,6 @@ struct TPlane3
 typedef TPlane3<float> FPlane3f;
 typedef TPlane3<double> FPlane3d;
 
-
+} // end namespace UE::Geometry
+} // end namespace UE
 
diff --git a/Engine/Plugins/Experimental/GeometryProcessing/Source/GeometricObjects/Public/Quaternion.h b/Engine/Plugins/Experimental/GeometryProcessing/Source/GeometricObjects/Public/Quaternion.h
--- a/Engine/Plugins/Experimental/GeometryProcessing/Source/GeometricObjects/Public/Quaternion.h
+++ b/Engine/Plugins/Experimental/GeometryProcessing/Source/GeometricObjects/Public/Quaternion.h
@@ -2,12 +2,19 @@
 
 #pragma once
 
+#include "Math/Quat.h"		// Engine types
 #include "VectorTypes.h"
 #include "MatrixTypes.h"
 #include "IndexTypes.h"
 
 // ported from geometry3Sharp Quaternion
 
+namespace UE
+{
+namespace Geometry
+{
+
+
 template<typename RealType>
 struct TQuaternion
 {
@@ -498,3 +505,6 @@ bool TQuaternion<RealType>::EpsilonEqual(const TQuaternion<RealType>& Other, Rea
 		(RealType)fabs(W - Other.W) <= Epsilon;
 }
 
+
+} // end namespace UE::Geometry
+} // end namespace UE
\ No newline at end of file
diff --git a/Engine/Plugins/Experimental/GeometryProcessing/Source/GeometricObjects/Public/RayTypes.h b/Engine/Plugins/Experimental/GeometryProcessing/Source/GeometricObjects/Public/RayTypes.h
--- a/Engine/Plugins/Experimental/GeometryProcessing/Source/GeometricObjects/Public/RayTypes.h
+++ b/Engine/Plugins/Experimental/GeometryProcessing/Source/GeometricObjects/Public/RayTypes.h
@@ -8,6 +8,11 @@
 #include "VectorTypes.h"
 #include "Math/Ray.h"
 
+namespace UE
+{
+namespace Geometry
+{
+
 /*
  * 3D Ray stored as Origin point and normalized Direction vector
  */
@@ -121,3 +126,6 @@ public:
 typedef TRay3<float> FRay3f;
 typedef TRay3<double> FRay3d;
 
+} // end namespace UE::Geometry
+} // end namespace UE
+
diff --git a/Engine/Plugins/Experimental/GeometryProcessing/Source/GeometricObjects/Public/Sampling/Gaussians.h b/Engine/Plugins/Experimental/GeometryProcessing/Source/GeometricObjects/Public/Sampling/Gaussians.h
--- a/Engine/Plugins/Experimental/GeometryProcessing/Source/GeometricObjects/Public/Sampling/Gaussians.h
+++ b/Engine/Plugins/Experimental/GeometryProcessing/Source/GeometricObjects/Public/Sampling/Gaussians.h
@@ -5,6 +5,12 @@
 #include "MathUtil.h"
 #include "VectorTypes.h"
 
+namespace UE
+{
+namespace Geometry
+{
+
+
 /**
  * 1D Gaussian
  */
@@ -175,4 +181,8 @@ public:
 	}
 };
 typedef TGaussian3<float> TGaussian3f;
-typedef TGaussian3<double> TGaussian3d;
\ No newline at end of file
+typedef TGaussian3<double> TGaussian3d;
+
+
+} // end namespace UE::Geometry
+} // end namespace UE
\ No newline at end of file
diff --git a/Engine/Plugins/Experimental/GeometryProcessing/Source/GeometricObjects/Public/Sampling/SphericalFibonacci.h b/Engine/Plugins/Experimental/GeometryProcessing/Source/GeometricObjects/Public/Sampling/SphericalFibonacci.h
--- a/Engine/Plugins/Experimental/GeometryProcessing/Source/GeometricObjects/Public/Sampling/SphericalFibonacci.h
+++ b/Engine/Plugins/Experimental/GeometryProcessing/Source/GeometricObjects/Public/Sampling/SphericalFibonacci.h
@@ -6,6 +6,10 @@
 #include "VectorTypes.h"
 #include "MatrixTypes.h"
 
+namespace UE
+{
+namespace Geometry
+{
 
 /**
  * A Spherical Fibonacci (SF) Point Set is a set of points that are roughly evenly distributed on a sphere.
@@ -125,6 +129,7 @@ protected:
 		return a * b + -TMathUtil<RealType>::Floor(a * b);
 	}
 
-
-
 };
+
+} // end namespace UE::Geometry
+} // end namespace UE
\ No newline at end of file
diff --git a/Engine/Plugins/Experimental/GeometryProcessing/Source/GeometricObjects/Public/Sampling/VectorSetAnalysis.h b/Engine/Plugins/Experimental/GeometryProcessing/Source/GeometricObjects/Public/Sampling/VectorSetAnalysis.h
--- a/Engine/Plugins/Experimental/GeometryProcessing/Source/GeometricObjects/Public/Sampling/VectorSetAnalysis.h
+++ b/Engine/Plugins/Experimental/GeometryProcessing/Source/GeometricObjects/Public/Sampling/VectorSetAnalysis.h
@@ -5,6 +5,10 @@
 #include "CoreMinimal.h"
 #include "VectorTypes.h"
 
+namespace UE
+{
+namespace Geometry
+{
 
 /**
  * TVectorSetAnalysis3 computes various analyses of a set of input Vectors (currently mainly clustering.
@@ -100,4 +104,7 @@ public:
 };
 
 typedef TVectorSetAnalysis3<float> FVectorSetAnalysis3f;
-typedef TVectorSetAnalysis3<double> FVectorSetAnalysis3d;
\ No newline at end of file
+typedef TVectorSetAnalysis3<double> FVectorSetAnalysis3d;
+
+} // end namespace UE::Geometry
+} // end namespace UE
diff --git a/Engine/Plugins/Experimental/GeometryProcessing/Source/GeometricObjects/Public/SegmentTypes.h b/Engine/Plugins/Experimental/GeometryProcessing/Source/GeometricObjects/Public/SegmentTypes.h
--- a/Engine/Plugins/Experimental/GeometryProcessing/Source/GeometricObjects/Public/SegmentTypes.h
+++ b/Engine/Plugins/Experimental/GeometryProcessing/Source/GeometricObjects/Public/SegmentTypes.h
@@ -8,6 +8,10 @@
 #include "VectorTypes.h"
 #include "BoxTypes.h"
 
+namespace UE
+{
+namespace Geometry
+{
 
 /*
  * 2D Line Segment stored as Center point, normalized Direction vector, and scalar Extent
@@ -535,3 +539,7 @@ protected:
 };
 typedef TSegment3<float> FSegment3f;
 typedef TSegment3<double> FSegment3d;
+
+} // end namespace UE::Geometry
+} // end namespace UE
+
diff --git a/Engine/Plugins/Experimental/GeometryProcessing/Source/GeometricObjects/Public/Spatial/FastWinding.h b/Engine/Plugins/Experimental/GeometryProcessing/Source/GeometricObjects/Public/Spatial/FastWinding.h
--- a/Engine/Plugins/Experimental/GeometryProcessing/Source/GeometricObjects/Public/Spatial/FastWinding.h
+++ b/Engine/Plugins/Experimental/GeometryProcessing/Source/GeometricObjects/Public/Spatial/FastWinding.h
@@ -13,6 +13,11 @@
  */
 namespace FastTriWinding
 {
+	using namespace UE::Geometry;
+	using FTriangle3d = UE::Geometry::FTriangle3d;
+	using FVector3d = UE::Geometry::FVector3d;
+	using FMatrix3d = UE::Geometry::FMatrix3d;
+
 	/**
 	 *  precompute constant coefficients of triangle winding number approximation
 	 *  P: 'Center' of expansion for Triangles (area-weighted centroid avg)
@@ -168,6 +173,11 @@ namespace FastTriWinding
 
 
 
+namespace UE
+{
+namespace Geometry
+{
+
 /**
  * Fast Mesh Winding Number extension to a TMeshAABBTree3.
  * This class is an "add-on" to the AABBTree, that can compute the Fast Mesh Winding Number.
@@ -496,3 +506,7 @@ private:
 		}
 	}
 };
+
+
+} // end namespace UE::Geometry
+} // end namespace UE
\ No newline at end of file
diff --git a/Engine/Plugins/Experimental/GeometryProcessing/Source/GeometricObjects/Public/Spatial/MeshAABBTree3.h b/Engine/Plugins/Experimental/GeometryProcessing/Source/GeometricObjects/Public/Spatial/MeshAABBTree3.h
--- a/Engine/Plugins/Experimental/GeometryProcessing/Source/GeometricObjects/Public/Spatial/MeshAABBTree3.h
+++ b/Engine/Plugins/Experimental/GeometryProcessing/Source/GeometricObjects/Public/Spatial/MeshAABBTree3.h
@@ -12,11 +12,12 @@
 #include "Spatial/SpatialInterfaces.h"
 #include "Distance/DistTriangle3Triangle3.h"
 
-template <class TriangleMeshType>
-class TFastWindingTree;
 
 namespace MeshIntersection
 {
+	using namespace UE::Geometry;
+	using FVector3d = UE::Geometry::FVector3d;
+
 	/**
 	 * Intersection query result types for triangle mesh intersections
 	 */
@@ -44,6 +45,16 @@ namespace MeshIntersection
 	};
 }
 
+
+namespace UE
+{
+namespace Geometry
+{
+
+template <class TriangleMeshType>
+class TFastWindingTree;
+
+
 template <class TriangleMeshType>
 class TMeshAABBTree3 : public IMeshSpatial
 {
@@ -2158,3 +2169,7 @@ private:
 		TreeTraversalImpl(IBox, 0, t, FQueryOptions());
 	}
 };
+
+
+} // end namespace UE::Geometry
+} // end namespace UE
\ No newline at end of file
diff --git a/Engine/Plugins/Experimental/GeometryProcessing/Source/GeometricObjects/Public/Spatial/MeshWindingNumberGrid.h b/Engine/Plugins/Experimental/GeometryProcessing/Source/GeometricObjects/Public/Spatial/MeshWindingNumberGrid.h
--- a/Engine/Plugins/Experimental/GeometryProcessing/Source/GeometricObjects/Public/Spatial/MeshWindingNumberGrid.h
+++ b/Engine/Plugins/Experimental/GeometryProcessing/Source/GeometricObjects/Public/Spatial/MeshWindingNumberGrid.h
@@ -13,6 +13,11 @@
 #include "MatrixTypes.h"
 
 
+namespace UE
+{
+namespace Geometry
+{
+
 /**
  * Sample mesh winding number (MWN) on a discrete grid. Can sample full grid, or
  * compute MWN values along a specific iso-contour and then fill in rest of grid
@@ -343,11 +348,7 @@ private:
 		});
 	}
 
-
-
-
-
-
-
-
 };
+
+} // end namespace UE::Geometry
+} // end namespace UE
\ No newline at end of file
diff --git a/Engine/Plugins/Experimental/GeometryProcessing/Source/GeometricObjects/Public/Spatial/SampledScalarField2.h b/Engine/Plugins/Experimental/GeometryProcessing/Source/GeometricObjects/Public/Spatial/SampledScalarField2.h
--- a/Engine/Plugins/Experimental/GeometryProcessing/Source/GeometricObjects/Public/Spatial/SampledScalarField2.h
+++ b/Engine/Plugins/Experimental/GeometryProcessing/Source/GeometricObjects/Public/Spatial/SampledScalarField2.h
@@ -9,7 +9,10 @@
 #include "VectorTypes.h"
 #include "DenseGrid2.h"
 
-
+namespace UE
+{
+namespace Geometry
+{
 
 /**
  * TSampledScalarField2 implements a generic 2D grid of values that can be interpolated in various ways.
@@ -185,4 +188,7 @@ public:
 };
 
 typedef TSampledScalarField2<double, double> FSampledScalarField2d;
-typedef TSampledScalarField2<float, float> FSampledScalarField2f;
\ No newline at end of file
+typedef TSampledScalarField2<float, float> FSampledScalarField2f;
+
+} // end namespace UE::Geometry
+} // end namespace UE
diff --git a/Engine/Plugins/Experimental/GeometryProcessing/Source/GeometricObjects/Public/Spatial/SparseDynamicPointOctree3.h b/Engine/Plugins/Experimental/GeometryProcessing/Source/GeometricObjects/Public/Spatial/SparseDynamicPointOctree3.h
--- a/Engine/Plugins/Experimental/GeometryProcessing/Source/GeometricObjects/Public/Spatial/SparseDynamicPointOctree3.h
+++ b/Engine/Plugins/Experimental/GeometryProcessing/Source/GeometricObjects/Public/Spatial/SparseDynamicPointOctree3.h
@@ -14,7 +14,10 @@
 #include "Async/ParallelFor.h"
 #include "Spatial/SparseDynamicOctree3.h"    // for FDynamicFlagArray
 
-
+namespace UE
+{
+namespace Geometry
+{
 
 /**
  * FSparsePointOctreeCell is a Node in a SparseDynamicOctree3. 
@@ -967,4 +970,7 @@ FString FSparseDynamicPointOctree3::FStatistics::ToString() const
 			k, LevelBoxCounts[k], LevelObjCounts[k], ((float)LevelObjCounts[k] / (float)LevelBoxCounts[k]), LevelMaxObjCounts[k], LevelCellSizes[k] );
 	}
 	return Result;
-}
\ No newline at end of file
+}
+
+} // end namespace UE::Geometry
+} // end namespace UE
diff --git a/Engine/Plugins/Experimental/GeometryProcessing/Source/GeometricObjects/Public/SphereTypes.h b/Engine/Plugins/Experimental/GeometryProcessing/Source/GeometricObjects/Public/SphereTypes.h
--- a/Engine/Plugins/Experimental/GeometryProcessing/Source/GeometricObjects/Public/SphereTypes.h
+++ b/Engine/Plugins/Experimental/GeometryProcessing/Source/GeometricObjects/Public/SphereTypes.h
@@ -5,6 +5,11 @@
 #include "Math/UnrealMath.h"
 #include "VectorTypes.h"
 
+namespace UE
+{
+namespace Geometry
+{
+
 /*
  * 3D Sphere stored as Center point and Radius
  */
@@ -98,4 +103,7 @@ public:
 
 
 typedef TSphere3<float> FSphere3f;
-typedef TSphere3<double> FSphere3d;
\ No newline at end of file
+typedef TSphere3<double> FSphere3d;
+
+} // end namespace UE::Geometry
+} // end namespace UE
\ No newline at end of file
diff --git a/Engine/Plugins/Experimental/GeometryProcessing/Source/GeometricObjects/Public/TriangleTypes.h b/Engine/Plugins/Experimental/GeometryProcessing/Source/GeometricObjects/Public/TriangleTypes.h
--- a/Engine/Plugins/Experimental/GeometryProcessing/Source/GeometricObjects/Public/TriangleTypes.h
+++ b/Engine/Plugins/Experimental/GeometryProcessing/Source/GeometricObjects/Public/TriangleTypes.h
@@ -1,12 +1,16 @@
 // Copyright Epic Games, Inc. All Rights Reserved.
 
-// Port of WildMagic DistPoint3TTriangle3
-
 #pragma once
 
 #include "VectorTypes.h"
 #include "VectorUtil.h"
 
+namespace UE
+{
+namespace Geometry
+{
+
+
 /**
  * Triangle utility functions
  */
@@ -237,3 +241,5 @@ typedef TTriangle3<float> FTriangle3f;
 typedef TTriangle3<double> FTriangle3d;
 typedef TTriangle3<int> FTriangle3i;
 
+} // end namespace UE::Geometry
+} // end namespace UE
\ No newline at end of file
diff --git a/Engine/Plugins/Experimental/GeometryProcessing/Source/GeometricObjects/Public/Util/ColorConstants.h b/Engine/Plugins/Experimental/GeometryProcessing/Source/GeometricObjects/Public/Util/ColorConstants.h
--- a/Engine/Plugins/Experimental/GeometryProcessing/Source/GeometricObjects/Public/Util/ColorConstants.h
+++ b/Engine/Plugins/Experimental/GeometryProcessing/Source/GeometricObjects/Public/Util/ColorConstants.h
@@ -4,6 +4,11 @@
 
 #include "Math/Color.h"
 
+namespace UE
+{
+namespace Geometry
+{
+
 
 /**
  * Color constants in Linear color space.
@@ -223,4 +228,8 @@ namespace LinearColors
 	 */
 	GEOMETRICOBJECTS_API FColor SelectFColor(int32 Index);
 
-}
\ No newline at end of file
+}
+
+
+} // end namespace UE::Geometry
+} // end namespace UE
\ No newline at end of file
diff --git a/Engine/Plugins/Experimental/GeometryProcessing/Source/GeometricObjects/Public/Util/GridIndexing2.h b/Engine/Plugins/Experimental/GeometryProcessing/Source/GeometricObjects/Public/Util/GridIndexing2.h
--- a/Engine/Plugins/Experimental/GeometryProcessing/Source/GeometricObjects/Public/Util/GridIndexing2.h
+++ b/Engine/Plugins/Experimental/GeometryProcessing/Source/GeometricObjects/Public/Util/GridIndexing2.h
@@ -6,6 +6,10 @@
 
 #include "VectorTypes.h"
 
+namespace UE
+{
+namespace Geometry
+{
 
 /**
  * Convert between integer grid coordinates and scaled real-valued coordinates (ie assumes integer grid origin == real origin)
@@ -85,4 +89,5 @@ typedef TShiftGridIndexer2<float> FShiftGridIndexer2f;
 typedef TShiftGridIndexer2<double> FShiftGridIndexer2d;
 
 
-
+} // end namespace UE::Geometry
+} // end namespace UE
diff --git a/Engine/Plugins/Experimental/GeometryProcessing/Source/GeometricObjects/Public/Util/IndexSetDecompositions.h b/Engine/Plugins/Experimental/GeometryProcessing/Source/GeometricObjects/Public/Util/IndexSetDecompositions.h
--- a/Engine/Plugins/Experimental/GeometryProcessing/Source/GeometricObjects/Public/Util/IndexSetDecompositions.h
+++ b/Engine/Plugins/Experimental/GeometryProcessing/Source/GeometricObjects/Public/Util/IndexSetDecompositions.h
@@ -5,10 +5,10 @@
 #include "VectorTypes.h"
 
 
-//
-// @todo: base interface that allows for iterating over
-// different kinds of index decompositions in a generic way
-// 
+namespace UE
+{
+namespace Geometry
+{
 
 
 /**
@@ -77,6 +77,7 @@ public:
 		return (*Found)->IndexArray;
 	}
 
+};
 
-
-};
\ No newline at end of file
+} // end namespace UE::Geometry
+} // end namespace UE
diff --git a/Engine/Plugins/Experimental/GeometryProcessing/Source/GeometricObjects/Public/Util/IndexUtil.h b/Engine/Plugins/Experimental/GeometryProcessing/Source/GeometricObjects/Public/Util/IndexUtil.h
--- a/Engine/Plugins/Experimental/GeometryProcessing/Source/GeometricObjects/Public/Util/IndexUtil.h
+++ b/Engine/Plugins/Experimental/GeometryProcessing/Source/GeometricObjects/Public/Util/IndexUtil.h
@@ -11,7 +11,8 @@
 
 namespace IndexUtil
 {
-
+	using namespace UE::Geometry;
+	using UE::Geometry::FIndex2i;
 
 	/**
 	 * @return true if  [a0,a1] and [b0,b1] are the same pair, ignoring order
diff --git a/Engine/Plugins/Experimental/GeometryProcessing/Source/GeometricObjects/Public/Util/RevolveUtil.h b/Engine/Plugins/Experimental/GeometryProcessing/Source/GeometricObjects/Public/Util/RevolveUtil.h
--- a/Engine/Plugins/Experimental/GeometryProcessing/Source/GeometricObjects/Public/Util/RevolveUtil.h
+++ b/Engine/Plugins/Experimental/GeometryProcessing/Source/GeometricObjects/Public/Util/RevolveUtil.h
@@ -7,12 +7,14 @@
 #include "FrameTypes.h"
 #include "VectorTypes.h"
 
-class UNewMeshMaterialProperties;
-class URevolveProperties;
-class FCurveSweepOp;
+PREDECLARE_GEOMETRY(class FCurveSweepOp);
 
 namespace RevolveUtil
 {
+	using namespace UE::Geometry;
+	using UE::Geometry::FVector3d;
+	using UE::Geometry::FFrame3d;
+
 	/** 
 	 * Generates a sweep curve with the initial frame at the origin and the rest rotated around the given axis.
 	 *
diff --git a/Engine/Plugins/Experimental/GeometryProcessing/Source/GeometricObjects/Public/VectorTypes.h b/Engine/Plugins/Experimental/GeometryProcessing/Source/GeometricObjects/Public/VectorTypes.h
--- a/Engine/Plugins/Experimental/GeometryProcessing/Source/GeometricObjects/Public/VectorTypes.h
+++ b/Engine/Plugins/Experimental/GeometryProcessing/Source/GeometricObjects/Public/VectorTypes.h
@@ -7,7 +7,8 @@
 #include <sstream>
 
 
-
+namespace UE {
+namespace Geometry {
 
 
 /**
@@ -1085,4 +1086,8 @@ FORCEINLINE uint32 GetTypeHash(const TVector4<T>& Vector)
 	// (this is how FIntVector and all the other FVectors do their hash functions)
 	// Note: this assumes there's no padding that could contain uncompared data.
 	return FCrc::MemCrc_DEPRECATED(&Vector, sizeof(TVector4<T>));
-}
\ No newline at end of file
+}
+
+
+} // end namespace UE::Geometry
+} // end namespace UE
\ No newline at end of file
diff --git a/Engine/Plugins/Experimental/GeometryProcessing/Source/GeometricObjects/Public/VectorUtil.h b/Engine/Plugins/Experimental/GeometryProcessing/Source/GeometricObjects/Public/VectorUtil.h
--- a/Engine/Plugins/Experimental/GeometryProcessing/Source/GeometricObjects/Public/VectorUtil.h
+++ b/Engine/Plugins/Experimental/GeometryProcessing/Source/GeometricObjects/Public/VectorUtil.h
@@ -28,6 +28,10 @@ enum class EIntersectionType
 
 namespace VectorUtil
 {
+	using namespace UE::Geometry;
+	using FVector2d = UE::Geometry::FVector2d;
+	using FVector3d = UE::Geometry::FVector3d;
+
 	/**
 	 * @return true if all components of V are finite
 	 */
diff --git a/Engine/Plugins/Experimental/GeometryProcessing/Source/GeometryAlgorithms/Private/ConstrainedDelaunay2.cpp b/Engine/Plugins/Experimental/GeometryProcessing/Source/GeometryAlgorithms/Private/ConstrainedDelaunay2.cpp
--- a/Engine/Plugins/Experimental/GeometryProcessing/Source/GeometryAlgorithms/Private/ConstrainedDelaunay2.cpp
+++ b/Engine/Plugins/Experimental/GeometryProcessing/Source/GeometryAlgorithms/Private/ConstrainedDelaunay2.cpp
@@ -7,6 +7,9 @@
 #include "Async/ParallelFor.h"
 #include <vector>
 
+#include "ExplicitUseGeometryMathTypes.h"		// using UE::Geometry::(math types)
+using namespace UE::Geometry;
+
 //namespace
 //{
 //#define DEBUG_FILE_DUMPING 1
@@ -468,7 +471,7 @@ bool TConstrainedDelaunay2<RealType>::Triangulate()
 }
 //
 template<typename RealType>
-TArray<FIndex3i> GEOMETRYALGORITHMS_API ConstrainedDelaunayTriangulate(const TGeneralPolygon2<RealType>& GeneralPolygon)
+TArray<FIndex3i> GEOMETRYALGORITHMS_API UE::Geometry::ConstrainedDelaunayTriangulate(const TGeneralPolygon2<RealType>& GeneralPolygon)
 {
 	TConstrainedDelaunay2<RealType> Triangulation;
 	Triangulation.FillRule = TConstrainedDelaunay2<RealType>::EFillRule::Positive;
@@ -477,11 +480,18 @@ TArray<FIndex3i> GEOMETRYALGORITHMS_API ConstrainedDelaunayTriangulate(const TGe
 	return Triangulation.Triangles;
 }
 
+namespace UE
+{
+namespace Geometry
+{
 
-template TArray<FIndex3i> GEOMETRYALGORITHMS_API ConstrainedDelaunayTriangulate(const TGeneralPolygon2<double>& GeneralPolygon);
-template TArray<FIndex3i> GEOMETRYALGORITHMS_API ConstrainedDelaunayTriangulate(const TGeneralPolygon2<float>& GeneralPolygon);
+template TArray<FIndex3i> GEOMETRYALGORITHMS_API UE::Geometry::ConstrainedDelaunayTriangulate(const TGeneralPolygon2<double>& GeneralPolygon);
+template TArray<FIndex3i> GEOMETRYALGORITHMS_API UE::Geometry::ConstrainedDelaunayTriangulate(const TGeneralPolygon2<float>& GeneralPolygon);
 
 
 
 template struct TConstrainedDelaunay2<float>;
 template struct TConstrainedDelaunay2<double>;
+
+} // end namespace UE::Geometry
+} // end namespace UE
\ No newline at end of file
diff --git a/Engine/Plugins/Experimental/GeometryProcessing/Source/GeometryAlgorithms/Private/ConvexHull2.cpp b/Engine/Plugins/Experimental/GeometryProcessing/Source/GeometryAlgorithms/Private/ConvexHull2.cpp
--- a/Engine/Plugins/Experimental/GeometryProcessing/Source/GeometryAlgorithms/Private/ConvexHull2.cpp
+++ b/Engine/Plugins/Experimental/GeometryProcessing/Source/GeometryAlgorithms/Private/ConvexHull2.cpp
@@ -7,6 +7,9 @@
 #include "ExactPredicates.h"
 #include "Algo/Unique.h"
 
+#include "ExplicitUseGeometryMathTypes.h"		// using UE::Geometry::(math types)
+using namespace UE::Geometry;
+
 template<class RealType>
 bool TConvexHull2<RealType>::Solve(int32 NumPoints, TFunctionRef<FVector2<RealType>(int32)> GetPointFunc, TFunctionRef<bool(int32)> FilterFunc)
 {
@@ -286,5 +289,13 @@ void TConvexHull2<RealType>::GetTangent(TFunctionRef<FVector2<RealType>(int32)>
 }
 
 
+namespace UE
+{
+namespace Geometry
+{
+
 template class TConvexHull2<float>;
 template class TConvexHull2<double>;
+
+} // end namespace UE::Geometry
+} // end namespace UE
\ No newline at end of file
diff --git a/Engine/Plugins/Experimental/GeometryProcessing/Source/GeometryAlgorithms/Private/ConvexHull3.cpp b/Engine/Plugins/Experimental/GeometryProcessing/Source/GeometryAlgorithms/Private/ConvexHull3.cpp
--- a/Engine/Plugins/Experimental/GeometryProcessing/Source/GeometryAlgorithms/Private/ConvexHull3.cpp
+++ b/Engine/Plugins/Experimental/GeometryProcessing/Source/GeometryAlgorithms/Private/ConvexHull3.cpp
@@ -7,7 +7,8 @@
 
 #include "Async/ParallelFor.h"
 
-
+#include "ExplicitUseGeometryMathTypes.h"		// using UE::Geometry::(math types)
+using namespace UE::Geometry;
 
 
 /**
@@ -654,7 +655,13 @@ bool TConvexHull3<RealType>::Solve(int32 NumPoints, TFunctionRef<FVector3<RealTy
 }
 
 
-
+namespace UE
+{
+namespace Geometry
+{
 
 template class TConvexHull3<float>;
 template class TConvexHull3<double>;
+
+} // end namespace UE::Geometry
+} // end namespace UE
\ No newline at end of file
diff --git a/Engine/Plugins/Experimental/GeometryProcessing/Source/GeometryAlgorithms/Private/Distance/DistLine3Circle3.cpp b/Engine/Plugins/Experimental/GeometryProcessing/Source/GeometryAlgorithms/Private/Distance/DistLine3Circle3.cpp
--- a/Engine/Plugins/Experimental/GeometryProcessing/Source/GeometryAlgorithms/Private/Distance/DistLine3Circle3.cpp
+++ b/Engine/Plugins/Experimental/GeometryProcessing/Source/GeometryAlgorithms/Private/Distance/DistLine3Circle3.cpp
@@ -5,6 +5,8 @@
 #include "GteUtil.h"
 
 
+using namespace UE::Geometry;
+
 template <typename Real>
 Real TDistLine3Circle3<Real>::ComputeResult()
 {
@@ -32,7 +34,13 @@ Real TDistLine3Circle3<Real>::ComputeResult()
 
 
 
-
+namespace UE
+{
+namespace Geometry
+{
 
 template class GEOMETRYALGORITHMS_API TDistLine3Circle3<float>;
 template class GEOMETRYALGORITHMS_API TDistLine3Circle3<double>;
+
+} // end namespace UE::Geometry
+} // end namespace UE
diff --git a/Engine/Plugins/Experimental/GeometryProcessing/Source/GeometryAlgorithms/Private/FitCapsule3.cpp b/Engine/Plugins/Experimental/GeometryProcessing/Source/GeometryAlgorithms/Private/FitCapsule3.cpp
--- a/Engine/Plugins/Experimental/GeometryProcessing/Source/GeometryAlgorithms/Private/FitCapsule3.cpp
+++ b/Engine/Plugins/Experimental/GeometryProcessing/Source/GeometryAlgorithms/Private/FitCapsule3.cpp
@@ -4,6 +4,8 @@
 #include "ThirdParty/GTEngine/Mathematics/GteCapsule.h"
 #include "ThirdParty/GTEngine/Mathematics/GteContCapsule3.h"
 
+#include "ExplicitUseGeometryMathTypes.h"		// using UE::Geometry::(math types)
+using namespace UE::Geometry;
 
 template<typename RealType>
 bool TFitCapsule3<RealType>::Solve(int32 NumPoints, TFunctionRef<FVector3<RealType>(int32)> GetPointFunc)
@@ -38,5 +40,13 @@ bool TFitCapsule3<RealType>::Solve(int32 NumPoints, TFunctionRef<FVector3<RealTy
 }
 
 
+namespace UE
+{
+namespace Geometry
+{
+
 template class GEOMETRYALGORITHMS_API TFitCapsule3<float>;
-template class GEOMETRYALGORITHMS_API TFitCapsule3<double>;
\ No newline at end of file
+template class GEOMETRYALGORITHMS_API TFitCapsule3<double>;
+
+} // end namespace UE::Geometry
+} // end namespace UE
diff --git a/Engine/Plugins/Experimental/GeometryProcessing/Source/GeometryAlgorithms/Private/MinVolumeBox3.cpp b/Engine/Plugins/Experimental/GeometryProcessing/Source/GeometryAlgorithms/Private/MinVolumeBox3.cpp
--- a/Engine/Plugins/Experimental/GeometryProcessing/Source/GeometryAlgorithms/Private/MinVolumeBox3.cpp
+++ b/Engine/Plugins/Experimental/GeometryProcessing/Source/GeometryAlgorithms/Private/MinVolumeBox3.cpp
@@ -23,6 +23,12 @@ THIRD_PARTY_INCLUDES_END
 #include "GteUtil.h"
 #include "Util/ProgressCancel.h"
 
+#include "ExplicitUseGeometryMathTypes.h"		// using UE::Geometry::(math types)
+using namespace UE::Geometry;
+
+namespace UE {
+namespace Geometry {
+
 template <typename RealType>
 struct TMinVolumeBox3Internal
 {
@@ -103,7 +109,8 @@ struct TMinVolumeBox3Internal
 
 };
 
-
+} // end namespace UE::Geometry
+} // end namespace UE
 
 template<typename RealType>
 bool TMinVolumeBox3<RealType>::Solve(int32 NumPoints, TFunctionRef<FVector3<RealType>(int32)> GetPointFunc, bool bUseExactBox, FProgressCancel* Progress )
@@ -143,6 +150,13 @@ void TMinVolumeBox3<RealType>::Initialize(int32 NumPoints, bool bUseExactBox)
 }
 
 
+namespace UE
+{
+namespace Geometry
+{
 
 template class GEOMETRYALGORITHMS_API TMinVolumeBox3<float>;
-template class GEOMETRYALGORITHMS_API TMinVolumeBox3<double>;
\ No newline at end of file
+template class GEOMETRYALGORITHMS_API TMinVolumeBox3<double>;
+
+} // end namespace UE::Geometry
+} // end namespace UE
diff --git a/Engine/Plugins/Experimental/GeometryProcessing/Source/GeometryAlgorithms/Private/MinVolumeSphere3.cpp b/Engine/Plugins/Experimental/GeometryProcessing/Source/GeometryAlgorithms/Private/MinVolumeSphere3.cpp
--- a/Engine/Plugins/Experimental/GeometryProcessing/Source/GeometryAlgorithms/Private/MinVolumeSphere3.cpp
+++ b/Engine/Plugins/Experimental/GeometryProcessing/Source/GeometryAlgorithms/Private/MinVolumeSphere3.cpp
@@ -8,8 +8,12 @@
 #include "ThirdParty/GTEngine/Mathematics/GteUIntegerAP32.h"
 #include "ThirdParty/GTEngine/Mathematics/GteMinimumVolumeSphere3.h"
 
+#include "ExplicitUseGeometryMathTypes.h"		// using UE::Geometry::(math types)
+using namespace UE::Geometry;
 
 
+namespace UE {
+namespace Geometry {
 
 template <typename RealType>
 struct TMinVolumeSphere3Internal
@@ -57,6 +61,9 @@ struct TMinVolumeSphere3Internal
 
 };
 
+} // end namespace UE::Geometry
+} // end namespace UE
+
 
 
 template<typename RealType>
@@ -100,7 +107,13 @@ void TMinVolumeSphere3<RealType>::Initialize(int32 NumPoints, bool bUseExactComp
 
 
 
-
+namespace UE
+{
+namespace Geometry
+{
 
 template class GEOMETRYALGORITHMS_API TMinVolumeSphere3<float>;
-template class GEOMETRYALGORITHMS_API TMinVolumeSphere3<double>;
\ No newline at end of file
+template class GEOMETRYALGORITHMS_API TMinVolumeSphere3<double>;
+
+} // end namespace UE::Geometry
+} // end namespace UE
diff --git a/Engine/Plugins/Experimental/GeometryProcessing/Source/GeometryAlgorithms/Public/Arrangement2d.h b/Engine/Plugins/Experimental/GeometryProcessing/Source/GeometryAlgorithms/Public/Arrangement2d.h
--- a/Engine/Plugins/Experimental/GeometryProcessing/Source/GeometryAlgorithms/Public/Arrangement2d.h
+++ b/Engine/Plugins/Experimental/GeometryProcessing/Source/GeometryAlgorithms/Public/Arrangement2d.h
@@ -13,6 +13,10 @@
 
 #include "CoreMinimal.h"
 
+namespace UE {
+namespace Geometry {
+
+
 /**
  * Arrangement2d constructs a planar arrangement of a set of 2D line segments.
  * When a segment is inserted, existing edges are split, and the inserted
@@ -534,3 +538,8 @@ protected:
 		return num_hits > 0;
 	}
 };
+
+
+
+} // end namespace UE::Geometry
+} // end namespace UE
diff --git a/Engine/Plugins/Experimental/GeometryProcessing/Source/GeometryAlgorithms/Public/ConstrainedDelaunay2.h b/Engine/Plugins/Experimental/GeometryProcessing/Source/GeometryAlgorithms/Public/ConstrainedDelaunay2.h
--- a/Engine/Plugins/Experimental/GeometryProcessing/Source/GeometryAlgorithms/Public/ConstrainedDelaunay2.h
+++ b/Engine/Plugins/Experimental/GeometryProcessing/Source/GeometryAlgorithms/Public/ConstrainedDelaunay2.h
@@ -11,6 +11,11 @@
 #include "CoreMinimal.h"
 
 
+namespace UE {
+namespace Geometry {
+
+
+
 template<typename RealType>
 struct TConstrainedDelaunay2
 {
@@ -146,3 +151,7 @@ TArray<FIndex3i> GEOMETRYALGORITHMS_API ConstrainedDelaunayTriangulate(const TGe
 typedef TConstrainedDelaunay2<float> FConstrainedDelaunay2f;
 typedef TConstrainedDelaunay2<double> FConstrainedDelaunay2d;
 
+
+
+} // end namespace UE::Geometry
+} // end namespace UE
\ No newline at end of file
diff --git a/Engine/Plugins/Experimental/GeometryProcessing/Source/GeometryAlgorithms/Public/ConvexHull3.h b/Engine/Plugins/Experimental/GeometryProcessing/Source/GeometryAlgorithms/Public/ConvexHull3.h
--- a/Engine/Plugins/Experimental/GeometryProcessing/Source/GeometryAlgorithms/Public/ConvexHull3.h
+++ b/Engine/Plugins/Experimental/GeometryProcessing/Source/GeometryAlgorithms/Public/ConvexHull3.h
@@ -9,6 +9,11 @@
 #include "PlaneTypes.h"
 #include "HalfspaceTypes.h"
 
+namespace UE {
+namespace Geometry {
+
+
+
 template <typename RealType> struct TConvexHull3Internal;
 
 /**
@@ -173,4 +178,9 @@ protected:
 };
 
 typedef TConvexHull3<float> FConvexHull3f;
-typedef TConvexHull3<double> FConvexHull3d;
\ No newline at end of file
+typedef TConvexHull3<double> FConvexHull3d;
+
+
+} // end namespace UE::Geometry
+} // end namespace UE
+
diff --git a/Engine/Plugins/Experimental/GeometryProcessing/Source/GeometryAlgorithms/Public/Distance/DistLine3Circle3.h b/Engine/Plugins/Experimental/GeometryProcessing/Source/GeometryAlgorithms/Public/Distance/DistLine3Circle3.h
--- a/Engine/Plugins/Experimental/GeometryProcessing/Source/GeometryAlgorithms/Public/Distance/DistLine3Circle3.h
+++ b/Engine/Plugins/Experimental/GeometryProcessing/Source/GeometryAlgorithms/Public/Distance/DistLine3Circle3.h
@@ -9,6 +9,9 @@
 #include "CircleTypes.h"
 #include "LineTypes.h"
 
+namespace UE {
+namespace Geometry {
+
 
 /**
 * Compute unsigned distance and closest-points between 3D line and 3D circle
@@ -48,3 +51,7 @@ public:
 typedef TDistLine3Circle3<float> FDistLine3Circle3f;
 typedef TDistLine3Circle3<double> FDistLine3Circle3d;
 
+
+
+} // end namespace UE::Geometry
+} // end namespace UE
diff --git a/Engine/Plugins/Experimental/GeometryProcessing/Source/GeometryAlgorithms/Public/FitCapsule3.h b/Engine/Plugins/Experimental/GeometryProcessing/Source/GeometryAlgorithms/Public/FitCapsule3.h
--- a/Engine/Plugins/Experimental/GeometryProcessing/Source/GeometryAlgorithms/Public/FitCapsule3.h
+++ b/Engine/Plugins/Experimental/GeometryProcessing/Source/GeometryAlgorithms/Public/FitCapsule3.h
@@ -6,6 +6,11 @@
 #include "VectorTypes.h"
 #include "CapsuleTypes.h"
 
+namespace UE {
+namespace Geometry {
+
+
+
 /**
  * Fit a Capsule to a set of 3D points.
  * Currently solved by using least-squares to fit a line to the point set.
@@ -49,4 +54,8 @@ public:
 };
 
 typedef TFitCapsule3<float> FFitCapsule3f;
-typedef TFitCapsule3<double> FFitCapsule3d;
\ No newline at end of file
+typedef TFitCapsule3<double> FFitCapsule3d;
+
+
+} // end namespace UE::Geometry
+} // end namespace UE
diff --git a/Engine/Plugins/Experimental/GeometryProcessing/Source/GeometryAlgorithms/Public/MinVolumeBox3.h b/Engine/Plugins/Experimental/GeometryProcessing/Source/GeometryAlgorithms/Public/MinVolumeBox3.h
--- a/Engine/Plugins/Experimental/GeometryProcessing/Source/GeometryAlgorithms/Public/MinVolumeBox3.h
+++ b/Engine/Plugins/Experimental/GeometryProcessing/Source/GeometryAlgorithms/Public/MinVolumeBox3.h
@@ -7,10 +7,13 @@
 #include "OrientedBoxTypes.h"
 #include "Templates/PimplPtr.h"
 
-template <typename RealType> struct TMinVolumeBox3Internal;
-
 class FProgressCancel;
 
+namespace UE {
+namespace Geometry {
+
+template <typename RealType> struct TMinVolumeBox3Internal;
+
 /**
  * Calculate a Minimal-Volume Oriented Box for a set of 3D points.
  * This internally first computes the Convex Hull of the point set. 
@@ -41,4 +44,8 @@ protected:
 };
 
 typedef TMinVolumeBox3<float> FMinVolumeBox3f;
-typedef TMinVolumeBox3<double> FMinVolumeBox3d;
\ No newline at end of file
+typedef TMinVolumeBox3<double> FMinVolumeBox3d;
+
+
+} // end namespace UE::Geometry
+} // end namespace UE
diff --git a/Engine/Plugins/Experimental/GeometryProcessing/Source/GeometryAlgorithms/Public/MinVolumeSphere3.h b/Engine/Plugins/Experimental/GeometryProcessing/Source/GeometryAlgorithms/Public/MinVolumeSphere3.h
--- a/Engine/Plugins/Experimental/GeometryProcessing/Source/GeometryAlgorithms/Public/MinVolumeSphere3.h
+++ b/Engine/Plugins/Experimental/GeometryProcessing/Source/GeometryAlgorithms/Public/MinVolumeSphere3.h
@@ -7,6 +7,8 @@
 #include "SphereTypes.h"
 #include "Templates/PimplPtr.h"
 
+namespace UE {
+namespace Geometry {
 
 template <typename RealType> struct TMinVolumeSphere3Internal;
 
@@ -39,4 +41,7 @@ protected:
 };
 
 typedef TMinVolumeSphere3<float> FMinVolumeSphere3f;
-typedef TMinVolumeSphere3<double> FMinVolumeSphere3d;
\ No newline at end of file
+typedef TMinVolumeSphere3<double> FMinVolumeSphere3d;
+
+} // end namespace UE::Geometry
+} // end namespace UE
\ No newline at end of file
diff --git a/Engine/Plugins/Experimental/GeometryProcessing/Source/MeshConversion/Public/DynamicMeshToMeshDescription.h b/Engine/Plugins/Experimental/GeometryProcessing/Source/MeshConversion/Public/DynamicMeshToMeshDescription.h
--- a/Engine/Plugins/Experimental/GeometryProcessing/Source/MeshConversion/Public/DynamicMeshToMeshDescription.h
+++ b/Engine/Plugins/Experimental/GeometryProcessing/Source/MeshConversion/Public/DynamicMeshToMeshDescription.h
@@ -7,8 +7,8 @@
 #include "MeshDescription.h"
 #include "MeshConversionOptions.h"
 
-// predeclare tangents template
-template<typename RealType> class TMeshTangents;
+PREDECLARE_GEOMETRY(template<typename RealType> class TMeshTangents);
+using UE::Geometry::FDynamicMesh3;
 
 /**
  * Convert FDynamicMesh3 to FMeshDescription
@@ -105,7 +105,7 @@ public:
 	 *	NOTE: assumes the order of triangles in the MeshIn correspond to the ordering you'd get by iterating over triangles, on MeshOut
 	 *		  This matches conversion currently used in MeshDescriptionToDynamicMesh.cpp, but if that changes we will need to change this function to match!
 	 */
-	void UpdateTangents(const FDynamicMesh3* MeshIn, FMeshDescription& MeshOut, const TMeshTangents<double>* SrcTangents);
+	void UpdateTangents(const FDynamicMesh3* MeshIn, FMeshDescription& MeshOut, const UE::Geometry::TMeshTangents<double>* SrcTangents);
 
 
 	/**
diff --git a/Engine/Plugins/Experimental/GeometryProcessing/Source/MeshConversion/Public/MeshDescriptionAdapter.h b/Engine/Plugins/Experimental/GeometryProcessing/Source/MeshConversion/Public/MeshDescriptionAdapter.h
--- a/Engine/Plugins/Experimental/GeometryProcessing/Source/MeshConversion/Public/MeshDescriptionAdapter.h
+++ b/Engine/Plugins/Experimental/GeometryProcessing/Source/MeshConversion/Public/MeshDescriptionAdapter.h
@@ -19,6 +19,9 @@
  */
 struct /*MESHCONVERSION_API*/ FMeshDescriptionTriangleMeshAdapter
 {
+	using FVector3d = UE::Geometry::FVector3d;
+	using FVector3f = UE::Geometry::FVector3f;
+	using FIndex3i = UE::Geometry::FIndex3i;
 protected:
 	const FMeshDescription* Mesh;
 	TVertexAttributesConstRef<FVector> VertexPositions;
@@ -112,6 +115,9 @@ public:
  */
 struct /*MESHCONVERSION_API*/ FMeshDescriptionEditableTriangleMeshAdapter
 {
+	using FVector3d = UE::Geometry::FVector3d;
+	using FVector3f = UE::Geometry::FVector3f;
+	using FIndex3i = UE::Geometry::FIndex3i;
 protected:
 	FMeshDescription* Mesh;
 	TVertexAttributesRef<FVector> VertexPositions;
diff --git a/Engine/Plugins/Experimental/GeometryProcessing/Source/MeshConversion/Public/MeshDescriptionBuilder.h b/Engine/Plugins/Experimental/GeometryProcessing/Source/MeshConversion/Public/MeshDescriptionBuilder.h
--- a/Engine/Plugins/Experimental/GeometryProcessing/Source/MeshConversion/Public/MeshDescriptionBuilder.h
+++ b/Engine/Plugins/Experimental/GeometryProcessing/Source/MeshConversion/Public/MeshDescriptionBuilder.h
@@ -4,8 +4,9 @@
 
 #include "CoreMinimal.h"
 #include "MeshDescription.h"
+#include "GeometryBase.h"
 
-class FDynamicMesh3;
+PREDECLARE_USE_GEOMETRY_CLASS(FDynamicMesh3);
 
 
 /**
diff --git a/Engine/Plugins/Experimental/GeometryProcessing/Source/MeshConversion/Public/MeshDescriptionToDynamicMesh.h b/Engine/Plugins/Experimental/GeometryProcessing/Source/MeshConversion/Public/MeshDescriptionToDynamicMesh.h
--- a/Engine/Plugins/Experimental/GeometryProcessing/Source/MeshConversion/Public/MeshDescriptionToDynamicMesh.h
+++ b/Engine/Plugins/Experimental/GeometryProcessing/Source/MeshConversion/Public/MeshDescriptionToDynamicMesh.h
@@ -7,7 +7,8 @@
 #include "MeshDescription.h"
 
 // predeclare tangents template
-template<typename RealType> class TMeshTangents;
+PREDECLARE_GEOMETRY(template<typename RealType> class TMeshTangents);
+using UE::Geometry::FDynamicMesh3;
 
 /**
  * Convert FMeshDescription to FDynamicMesh3
@@ -74,13 +75,13 @@ public:
 	 * Copy tangents from MeshDescription to a FMeshTangents instance.
 	 * @warning Convert() must have been used to create the TargetMesh before calling this function
 	 */
-	void CopyTangents(const FMeshDescription* SourceMesh, const FDynamicMesh3* TargetMesh, TMeshTangents<float>* TangentsOut);
+	void CopyTangents(const FMeshDescription* SourceMesh, const FDynamicMesh3* TargetMesh, UE::Geometry::TMeshTangents<float>* TangentsOut);
 
 	/**
 	 * Copy tangents from MeshDescription to a FMeshTangents instance.
 	 * @warning Convert() must have been used to create the TargetMesh before calling this function
 	 */
-	void CopyTangents(const FMeshDescription* SourceMesh, const FDynamicMesh3* TargetMesh, TMeshTangents<double>* TangentsOut);
+	void CopyTangents(const FMeshDescription* SourceMesh, const FDynamicMesh3* TargetMesh, UE::Geometry::TMeshTangents<double>* TangentsOut);
 
 
 public:
diff --git a/Engine/Plugins/Experimental/MeshLODToolset/Source/MeshLODToolset/Private/Graphs/GenerateMeshLODGraph.cpp b/Engine/Plugins/Experimental/MeshLODToolset/Source/MeshLODToolset/Private/Graphs/GenerateMeshLODGraph.cpp
--- a/Engine/Plugins/Experimental/MeshLODToolset/Source/MeshLODToolset/Private/Graphs/GenerateMeshLODGraph.cpp
+++ b/Engine/Plugins/Experimental/MeshLODToolset/Source/MeshLODToolset/Private/Graphs/GenerateMeshLODGraph.cpp
@@ -30,10 +30,11 @@
 #include "PhysicsNodes/GenerateSimpleCollisionNode.h"
 #include "GeometryFlowExecutor.h"
 
+#include "ExplicitUseGeometryMathTypes.h"		// using UE::Geometry::(math types)
+using namespace UE::Geometry;
 using namespace UE::GeometryFlow;
 
 
-
 void FGenerateMeshLODGraph::SetSourceMesh(const FDynamicMesh3& SourceMeshIn)
 {
 	UpdateSourceNodeValue<FDynamicMeshSourceNode>(*Graph, MeshSourceNode, SourceMeshIn);
@@ -635,7 +636,7 @@ void FGenerateMeshLODGraph::BuildGraph()
 
 
 
-int32 FGenerateMeshLODGraph::AppendTextureBakeNode(const TImageBuilder<FVector4f>& SourceImage, const FString& Identifier)
+int32 FGenerateMeshLODGraph::AppendTextureBakeNode(const TImageBuilder<UE::Geometry::FVector4f>& SourceImage, const FString& Identifier)
 {
 	FBakeTextureGraphInfo NewNode;
 	NewNode.Index = BakeTextureNodes.Num();
diff --git a/Engine/Plugins/Experimental/MeshLODToolset/Source/MeshLODToolset/Private/Tools/GenerateStaticMeshLODAssetTool.cpp b/Engine/Plugins/Experimental/MeshLODToolset/Source/MeshLODToolset/Private/Tools/GenerateStaticMeshLODAssetTool.cpp
--- a/Engine/Plugins/Experimental/MeshLODToolset/Source/MeshLODToolset/Private/Tools/GenerateStaticMeshLODAssetTool.cpp
+++ b/Engine/Plugins/Experimental/MeshLODToolset/Source/MeshLODToolset/Private/Tools/GenerateStaticMeshLODAssetTool.cpp
@@ -48,10 +48,11 @@
 #include "Misc/ScopedSlowTask.h"
 #endif
 
+#include "ExplicitUseGeometryMathTypes.h"		// using UE::Geometry::(math types)
+using namespace UE::Geometry;
 
 #define LOCTEXT_NAMESPACE "UGenerateStaticMeshLODAssetTool"
 
-
 //
 // Local Op Stuff
 //
diff --git a/Engine/Plugins/Experimental/MeshLODToolset/Source/MeshLODToolset/Private/Tools/LODManagerTool.cpp b/Engine/Plugins/Experimental/MeshLODToolset/Source/MeshLODToolset/Private/Tools/LODManagerTool.cpp
--- a/Engine/Plugins/Experimental/MeshLODToolset/Source/MeshLODToolset/Private/Tools/LODManagerTool.cpp
+++ b/Engine/Plugins/Experimental/MeshLODToolset/Source/MeshLODToolset/Private/Tools/LODManagerTool.cpp
@@ -20,10 +20,11 @@
 
 #include "Components/PrimitiveComponent.h"
 
+#include "ExplicitUseGeometryMathTypes.h"		// using UE::Geometry::(math types)
+using namespace UE::Geometry;
 
 #define LOCTEXT_NAMESPACE "ULODManagerTool"
 
-
 /*
  * ToolBuilder
  */
diff --git a/Engine/Plugins/Experimental/MeshLODToolset/Source/MeshLODToolset/Public/Graphs/GenerateMeshLODGraph.h b/Engine/Plugins/Experimental/MeshLODToolset/Source/MeshLODToolset/Public/Graphs/GenerateMeshLODGraph.h
--- a/Engine/Plugins/Experimental/MeshLODToolset/Source/MeshLODToolset/Public/Graphs/GenerateMeshLODGraph.h
+++ b/Engine/Plugins/Experimental/MeshLODToolset/Source/MeshLODToolset/Public/Graphs/GenerateMeshLODGraph.h
@@ -18,7 +18,7 @@
 #include "DataTypes/MeshImageBakingData.h"
 #include "PhysicsNodes/GenerateSimpleCollisionNode.h"
 
-
+using UE::Geometry::FDynamicMesh3;
 
 struct FMeshLODGraphPreFilterSettings
 {
@@ -28,10 +28,11 @@ struct FMeshLODGraphPreFilterSettings
 
 class FGenerateMeshLODGraph
 {
+	using FVector4f = UE::Geometry::FVector4f;
 public:
 	void BuildGraph();
 
-	int32 AppendTextureBakeNode(const TImageBuilder<FVector4f>& SourceImage, const FString& Identifier);
+	int32 AppendTextureBakeNode(const UE::Geometry::TImageBuilder<FVector4f>& SourceImage, const FString& Identifier);
 
 	void SetSourceMesh(const FDynamicMesh3& SourceMesh);
 
@@ -65,15 +66,15 @@ public:
 
 	void EvaluateResult(
 		FDynamicMesh3& ResultMesh,
-		FMeshTangentsd& ResultTangents,
-		FSimpleShapeSet3d& ResultCollision,
+		UE::Geometry::FMeshTangentsd& ResultTangents,
+		UE::Geometry::FSimpleShapeSet3d& ResultCollision,
 		UE::GeometryFlow::FNormalMapImage& NormalMap,
 		TArray<TUniquePtr<UE::GeometryFlow::FTextureImage>>& TextureImages,
 		FProgressCancel* Progress);
 
 	void EvaluateResultParallel(FDynamicMesh3& ResultMesh,
-								FMeshTangentsd& ResultTangents,
-								FSimpleShapeSet3d& ResultCollision,
+								UE::Geometry::FMeshTangentsd& ResultTangents,
+								UE::Geometry::FSimpleShapeSet3d& ResultCollision,
 								UE::GeometryFlow::FNormalMapImage& NormalMap,
 								TArray<TUniquePtr<UE::GeometryFlow::FTextureImage>>& TextureImages,
 								FProgressCancel* Progress);
diff --git a/Engine/Plugins/Experimental/MeshLODToolset/Source/MeshLODToolset/Public/Graphs/GenerateStaticMeshLODProcess.h b/Engine/Plugins/Experimental/MeshLODToolset/Source/MeshLODToolset/Public/Graphs/GenerateStaticMeshLODProcess.h
--- a/Engine/Plugins/Experimental/MeshLODToolset/Source/MeshLODToolset/Public/Graphs/GenerateStaticMeshLODProcess.h
+++ b/Engine/Plugins/Experimental/MeshLODToolset/Source/MeshLODToolset/Public/Graphs/GenerateStaticMeshLODProcess.h
@@ -16,7 +16,7 @@ class UTexture2D;
 class UMaterialInstanceConstant;
 class UMaterialInstanceDynamic;
 struct FMeshDescription;
-
+using UE::Geometry::FDynamicMesh3;
 
 UENUM()
 enum class EGenerateStaticMeshLODBakeResolution
@@ -160,7 +160,7 @@ UCLASS(Transient)
 class UGenerateStaticMeshLODProcess : public UObject
 {
 	GENERATED_BODY()
-
+	using FVector4f = UE::Geometry::FVector4f;
 public:
 
 	bool Initialize(UStaticMesh* SourceMesh);
@@ -181,8 +181,8 @@ public:
 
 	bool ComputeDerivedSourceData(FProgressCancel* Progress);
 	const FDynamicMesh3& GetDerivedLOD0Mesh() const { return DerivedLODMesh; }
-	const FMeshTangentsd& GetDerivedLOD0MeshTangents() const { return DerivedLODMeshTangents; }
-	const FSimpleShapeSet3d& GetDerivedCollision() const { return DerivedCollision; }
+	const UE::Geometry::FMeshTangentsd& GetDerivedLOD0MeshTangents() const { return DerivedLODMeshTangents; }
+	const UE::Geometry::FSimpleShapeSet3d& GetDerivedCollision() const { return DerivedCollision; }
 
 	/**
 	 * Creates new Asset
@@ -229,8 +229,8 @@ protected:
 	{
 		UTexture2D* Texture = nullptr;
 		FName ParameterName;
-		FImageDimensions Dimensions;
-		TImageBuilder<FVector4f> Image;
+		UE::Geometry::FImageDimensions Dimensions;
+		UE::Geometry::TImageBuilder<FVector4f> Image;
 		bool bIsNormalMap = false;
 		bool bIsDefaultTexture = false;
 		bool bShouldBakeTexture = false;
@@ -257,8 +257,8 @@ protected:
 	FString DerivedAssetPath;
 
 	FDynamicMesh3 DerivedLODMesh;				// the new generated LOD0 mesh
-	FMeshTangentsd DerivedLODMeshTangents;		// Tangents for DerivedLODMesh
-	FSimpleShapeSet3d DerivedCollision;			// Simple Collision for DerivedLODMesh
+	UE::Geometry::FMeshTangentsd DerivedLODMeshTangents;		// Tangents for DerivedLODMesh
+	UE::Geometry::FSimpleShapeSet3d DerivedCollision;			// Simple Collision for DerivedLODMesh
 
 	// Texture set potentially required by output Material set
 	UE::GeometryFlow::FNormalMapImage DerivedNormalMapImage;	// Normal Map
diff --git a/Engine/Plugins/Experimental/MeshLODToolset/Source/MeshLODToolset/Public/Tools/GenerateStaticMeshLODAssetTool.h b/Engine/Plugins/Experimental/MeshLODToolset/Source/MeshLODToolset/Public/Tools/GenerateStaticMeshLODAssetTool.h
--- a/Engine/Plugins/Experimental/MeshLODToolset/Source/MeshLODToolset/Public/Tools/GenerateStaticMeshLODAssetTool.h
+++ b/Engine/Plugins/Experimental/MeshLODToolset/Source/MeshLODToolset/Public/Tools/GenerateStaticMeshLODAssetTool.h
@@ -187,7 +187,7 @@ protected:
 	UPROPERTY()
 	UGenerateStaticMeshLODProcess* GenerateProcess;
 
-	TUniquePtr<IDynamicMeshOperatorFactory> OpFactory;
+	TUniquePtr<UE::Geometry::IDynamicMeshOperatorFactory> OpFactory;
 
 	void OnSettingsModified();
 
diff --git a/Engine/Plugins/Experimental/MeshLODToolset/Source/MeshLODToolset/Public/Tools/LODManagerTool.h b/Engine/Plugins/Experimental/MeshLODToolset/Source/MeshLODToolset/Public/Tools/LODManagerTool.h
--- a/Engine/Plugins/Experimental/MeshLODToolset/Source/MeshLODToolset/Public/Tools/LODManagerTool.h
+++ b/Engine/Plugins/Experimental/MeshLODToolset/Source/MeshLODToolset/Public/Tools/LODManagerTool.h
@@ -230,7 +230,7 @@ protected:
 
 	struct FLODMeshInfo
 	{
-		FDynamicMesh3 Mesh;
+		UE::Geometry::FDynamicMesh3 Mesh;
 
 		bool bInfoCached = false;
 		TArray<int> BoundaryEdges;
diff --git a/Engine/Plugins/Experimental/MeshModelingToolset/Source/MeshModelingTools/Private/AddPatchTool.cpp b/Engine/Plugins/Experimental/MeshModelingToolset/Source/MeshModelingTools/Private/AddPatchTool.cpp
--- a/Engine/Plugins/Experimental/MeshModelingToolset/Source/MeshModelingTools/Private/AddPatchTool.cpp
+++ b/Engine/Plugins/Experimental/MeshModelingToolset/Source/MeshModelingTools/Private/AddPatchTool.cpp
@@ -18,8 +18,10 @@
 
 #include "DynamicMeshEditor.h"
 
-#define LOCTEXT_NAMESPACE "UAddPatchTool"
+#include "ExplicitUseGeometryMathTypes.h"		// using UE::Geometry::(math types)
+using namespace UE::Geometry;
 
+#define LOCTEXT_NAMESPACE "UAddPatchTool"
 
 /*
  * ToolBuilder
diff --git a/Engine/Plugins/Experimental/MeshModelingToolset/Source/MeshModelingTools/Private/AddPrimitiveTool.cpp b/Engine/Plugins/Experimental/MeshModelingToolset/Source/MeshModelingTools/Private/AddPrimitiveTool.cpp
--- a/Engine/Plugins/Experimental/MeshModelingToolset/Source/MeshModelingTools/Private/AddPrimitiveTool.cpp
+++ b/Engine/Plugins/Experimental/MeshModelingToolset/Source/MeshModelingTools/Private/AddPrimitiveTool.cpp
@@ -28,8 +28,10 @@
 #include "UObject/PropertyIterator.h"
 #include "UObject/UnrealType.h"
 
-#define LOCTEXT_NAMESPACE "UAddPrimitiveTool"
+#include "ExplicitUseGeometryMathTypes.h"		// using UE::Geometry::(math types)
+using namespace UE::Geometry;
 
+#define LOCTEXT_NAMESPACE "UAddPrimitiveTool"
 
 /*
  * ToolBuilder
diff --git a/Engine/Plugins/Experimental/MeshModelingToolset/Source/MeshModelingTools/Private/BakeMeshAttributeMapsTool.cpp b/Engine/Plugins/Experimental/MeshModelingToolset/Source/MeshModelingTools/Private/BakeMeshAttributeMapsTool.cpp
--- a/Engine/Plugins/Experimental/MeshModelingToolset/Source/MeshModelingTools/Private/BakeMeshAttributeMapsTool.cpp
+++ b/Engine/Plugins/Experimental/MeshModelingToolset/Source/MeshModelingTools/Private/BakeMeshAttributeMapsTool.cpp
@@ -36,8 +36,10 @@
 #include "Engine/Classes/Components/StaticMeshComponent.h"
 #include "Engine/Classes/Engine/StaticMesh.h"
 
-#define LOCTEXT_NAMESPACE "UBakeMeshAttributeMapsTool"
+#include "ExplicitUseGeometryMathTypes.h"		// using UE::Geometry::(math types)
+using namespace UE::Geometry;
 
+#define LOCTEXT_NAMESPACE "UBakeMeshAttributeMapsTool"
 
 /*
  * ToolBuilder
diff --git a/Engine/Plugins/Experimental/MeshModelingToolset/Source/MeshModelingTools/Private/ConvertToPolygonsTool.cpp b/Engine/Plugins/Experimental/MeshModelingToolset/Source/MeshModelingTools/Private/ConvertToPolygonsTool.cpp
--- a/Engine/Plugins/Experimental/MeshModelingToolset/Source/MeshModelingTools/Private/ConvertToPolygonsTool.cpp
+++ b/Engine/Plugins/Experimental/MeshModelingToolset/Source/MeshModelingTools/Private/ConvertToPolygonsTool.cpp
@@ -20,8 +20,10 @@
 #include "MeshOpPreviewHelpers.h"
 #include "ModelingOperators.h"
 
-#define LOCTEXT_NAMESPACE "UConvertToPolygonsTool"
+#include "ExplicitUseGeometryMathTypes.h"		// using UE::Geometry::(math types)
+using namespace UE::Geometry;
 
+#define LOCTEXT_NAMESPACE "UConvertToPolygonsTool"
 
 /*
  * ToolBuilder
diff --git a/Engine/Plugins/Experimental/MeshModelingToolset/Source/MeshModelingTools/Private/DrawAndRevolveTool.cpp b/Engine/Plugins/Experimental/MeshModelingToolset/Source/MeshModelingTools/Private/DrawAndRevolveTool.cpp
--- a/Engine/Plugins/Experimental/MeshModelingToolset/Source/MeshModelingTools/Private/DrawAndRevolveTool.cpp
+++ b/Engine/Plugins/Experimental/MeshModelingToolset/Source/MeshModelingTools/Private/DrawAndRevolveTool.cpp
@@ -16,8 +16,10 @@
 #include "ToolSceneQueriesUtil.h"
 #include "ToolSetupUtil.h"
 
-#define LOCTEXT_NAMESPACE "UDrawAndRevolveTool"
+#include "ExplicitUseGeometryMathTypes.h"		// using UE::Geometry::(math types)
+using namespace UE::Geometry;
 
+#define LOCTEXT_NAMESPACE "UDrawAndRevolveTool"
 
 const FText InitializationModeMessage = LOCTEXT("CurveInitialization", "Draw a profile curve and revolve it around purple axis. Ctrl+click repositions draw plane, Ctrl+Shift+click repositions without aligning axis. End curve by clicking the end again or connecting to start. Backspace deletes points, shift inverts snapping behavior.");
 const FText EditModeMessage = LOCTEXT("CurveEditing", "Click points to select them, Shift+click to add/remove points to selection. Ctrl+click a segment to add a point, or select an endpoint and Ctrl+click somewhere on the plane to add to the ends. Backspace deletes selected points. Holding Shift toggles snapping to be opposite the EnableSnapping setting.");
diff --git a/Engine/Plugins/Experimental/MeshModelingToolset/Source/MeshModelingTools/Private/DrawPolyPathTool.cpp b/Engine/Plugins/Experimental/MeshModelingToolset/Source/MeshModelingTools/Private/DrawPolyPathTool.cpp
--- a/Engine/Plugins/Experimental/MeshModelingToolset/Source/MeshModelingTools/Private/DrawPolyPathTool.cpp
+++ b/Engine/Plugins/Experimental/MeshModelingToolset/Source/MeshModelingTools/Private/DrawPolyPathTool.cpp
@@ -20,8 +20,10 @@
 #include "Operations/ExtrudeMesh.h"
 #include "MeshNormals.h"
 
-#define LOCTEXT_NAMESPACE "UDrawPolyPathTool"
+#include "ExplicitUseGeometryMathTypes.h"		// using UE::Geometry::(math types)
+using namespace UE::Geometry;
 
+#define LOCTEXT_NAMESPACE "UDrawPolyPathTool"
 
 namespace
 {
diff --git a/Engine/Plugins/Experimental/MeshModelingToolset/Source/MeshModelingTools/Private/DynamicMeshBrushTool.cpp b/Engine/Plugins/Experimental/MeshModelingToolset/Source/MeshModelingTools/Private/DynamicMeshBrushTool.cpp
--- a/Engine/Plugins/Experimental/MeshModelingToolset/Source/MeshModelingTools/Private/DynamicMeshBrushTool.cpp
+++ b/Engine/Plugins/Experimental/MeshModelingToolset/Source/MeshModelingTools/Private/DynamicMeshBrushTool.cpp
@@ -4,10 +4,12 @@
 #include "InteractiveToolManager.h"
 #include "ToolBuilderUtil.h"
 
+#include "ExplicitUseGeometryMathTypes.h"		// using UE::Geometry::(math types)
+using namespace UE::Geometry;
+
 // localization namespace
 #define LOCTEXT_NAMESPACE "UDynamicMeshBrushTool"
 
-
 /*
  * Tool
  */
diff --git a/Engine/Plugins/Experimental/MeshModelingToolset/Source/MeshModelingTools/Private/DynamicMeshSculptTool.cpp b/Engine/Plugins/Experimental/MeshModelingToolset/Source/MeshModelingTools/Private/DynamicMeshSculptTool.cpp
--- a/Engine/Plugins/Experimental/MeshModelingToolset/Source/MeshModelingTools/Private/DynamicMeshSculptTool.cpp
+++ b/Engine/Plugins/Experimental/MeshModelingToolset/Source/MeshModelingTools/Private/DynamicMeshSculptTool.cpp
@@ -36,8 +36,10 @@
 #include "Materials/Material.h"
 #include "Materials/MaterialInstanceDynamic.h"
 
-#define LOCTEXT_NAMESPACE "UDynamicMeshSculptTool"
+#include "ExplicitUseGeometryMathTypes.h"		// using UE::Geometry::(math types)
+using namespace UE::Geometry;
 
+#define LOCTEXT_NAMESPACE "UDynamicMeshSculptTool"
 
 namespace
 {
diff --git a/Engine/Plugins/Experimental/MeshModelingToolset/Source/MeshModelingTools/Private/EditMeshPolygonsTool.cpp b/Engine/Plugins/Experimental/MeshModelingToolset/Source/MeshModelingTools/Private/EditMeshPolygonsTool.cpp
--- a/Engine/Plugins/Experimental/MeshModelingToolset/Source/MeshModelingTools/Private/EditMeshPolygonsTool.cpp
+++ b/Engine/Plugins/Experimental/MeshModelingToolset/Source/MeshModelingTools/Private/EditMeshPolygonsTool.cpp
@@ -40,9 +40,10 @@
 #include "Materials/MaterialInstanceDynamic.h"
 #include "Components/BrushComponent.h"
 
-#define LOCTEXT_NAMESPACE "UEditMeshPolygonsTool"
-
+#include "ExplicitUseGeometryMathTypes.h"		// using UE::Geometry::(math types)
+using namespace UE::Geometry;
 
+#define LOCTEXT_NAMESPACE "UEditMeshPolygonsTool"
 
 /*
  * ToolBuilder
diff --git a/Engine/Plugins/Experimental/MeshModelingToolset/Source/MeshModelingTools/Private/EditUVIslandsTool.cpp b/Engine/Plugins/Experimental/MeshModelingToolset/Source/MeshModelingTools/Private/EditUVIslandsTool.cpp
--- a/Engine/Plugins/Experimental/MeshModelingToolset/Source/MeshModelingTools/Private/EditUVIslandsTool.cpp
+++ b/Engine/Plugins/Experimental/MeshModelingToolset/Source/MeshModelingTools/Private/EditUVIslandsTool.cpp
@@ -14,9 +14,10 @@
 
 #include "Materials/MaterialInstanceDynamic.h"
 
-#define LOCTEXT_NAMESPACE "UEditUVIslandsTool"
-
+#include "ExplicitUseGeometryMathTypes.h"		// using UE::Geometry::(math types)
+using namespace UE::Geometry;
 
+#define LOCTEXT_NAMESPACE "UEditUVIslandsTool"
 
 /*
  * ToolBuilder
diff --git a/Engine/Plugins/Experimental/MeshModelingToolset/Source/MeshModelingTools/Private/LatticeDeformerTool.cpp b/Engine/Plugins/Experimental/MeshModelingToolset/Source/MeshModelingTools/Private/LatticeDeformerTool.cpp
--- a/Engine/Plugins/Experimental/MeshModelingToolset/Source/MeshModelingTools/Private/LatticeDeformerTool.cpp
+++ b/Engine/Plugins/Experimental/MeshModelingToolset/Source/MeshModelingTools/Private/LatticeDeformerTool.cpp
@@ -19,8 +19,10 @@
 #include "Algo/ForEach.h"
 #include "Operations/FFDLattice.h"
 
-#define LOCTEXT_NAMESPACE "ULatticeDeformerTool"
+#include "ExplicitUseGeometryMathTypes.h"		// using UE::Geometry::(math types)
+using namespace UE::Geometry;
 
+#define LOCTEXT_NAMESPACE "ULatticeDeformerTool"
 
 // Tool builder
 
diff --git a/Engine/Plugins/Experimental/MeshModelingToolset/Source/MeshModelingTools/Private/MeshAttributePaintTool.cpp b/Engine/Plugins/Experimental/MeshModelingToolset/Source/MeshModelingTools/Private/MeshAttributePaintTool.cpp
--- a/Engine/Plugins/Experimental/MeshModelingToolset/Source/MeshModelingTools/Private/MeshAttributePaintTool.cpp
+++ b/Engine/Plugins/Experimental/MeshModelingToolset/Source/MeshModelingTools/Private/MeshAttributePaintTool.cpp
@@ -7,14 +7,13 @@
 #include "ToolSetupUtil.h"
 #include "Selections/MeshConnectedComponents.h"
 
-
 #include "MeshDescription.h"
 
+#include "ExplicitUseGeometryMathTypes.h"		// using UE::Geometry::(math types)
+using namespace UE::Geometry;
 
 #define LOCTEXT_NAMESPACE "UMeshAttributePaintTool"
 
-
-
 class FMeshDescriptionVertexAttributeAdapter : public IMeshVertexAttributeAdapter
 {
 public:
diff --git a/Engine/Plugins/Experimental/MeshModelingToolset/Source/MeshModelingTools/Private/MeshInspectorTool.cpp b/Engine/Plugins/Experimental/MeshModelingToolset/Source/MeshModelingTools/Private/MeshInspectorTool.cpp
--- a/Engine/Plugins/Experimental/MeshModelingToolset/Source/MeshModelingTools/Private/MeshInspectorTool.cpp
+++ b/Engine/Plugins/Experimental/MeshModelingToolset/Source/MeshModelingTools/Private/MeshInspectorTool.cpp
@@ -22,8 +22,10 @@
 #include "ToolSetupUtil.h"
 #include "AssetUtils/MeshDescriptionUtil.h"
 
-#define LOCTEXT_NAMESPACE "UMeshInspectorTool"
+#include "ExplicitUseGeometryMathTypes.h"		// using UE::Geometry::(math types)
+using namespace UE::Geometry;
 
+#define LOCTEXT_NAMESPACE "UMeshInspectorTool"
 
 /*
  * ToolBuilder
diff --git a/Engine/Plugins/Experimental/MeshModelingToolset/Source/MeshModelingTools/Private/MeshSpaceDeformerTool.cpp b/Engine/Plugins/Experimental/MeshModelingToolset/Source/MeshModelingTools/Private/MeshSpaceDeformerTool.cpp
--- a/Engine/Plugins/Experimental/MeshModelingToolset/Source/MeshModelingTools/Private/MeshSpaceDeformerTool.cpp
+++ b/Engine/Plugins/Experimental/MeshModelingToolset/Source/MeshModelingTools/Private/MeshSpaceDeformerTool.cpp
@@ -27,11 +27,11 @@
 #include "CoreMinimal.h"
 #include "Math/Matrix.h"
 
-
+#include "ExplicitUseGeometryMathTypes.h"		// using UE::Geometry::(math types)
+using namespace UE::Geometry;
 
 #define LOCTEXT_NAMESPACE "MeshSpaceDeformerTool"
 
-
 /*
  * ToolBuilder
  */
diff --git a/Engine/Plugins/Experimental/MeshModelingToolset/Source/MeshModelingTools/Private/Physics/ExtractCollisionGeometryTool.cpp b/Engine/Plugins/Experimental/MeshModelingToolset/Source/MeshModelingTools/Private/Physics/ExtractCollisionGeometryTool.cpp
--- a/Engine/Plugins/Experimental/MeshModelingToolset/Source/MeshModelingTools/Private/Physics/ExtractCollisionGeometryTool.cpp
+++ b/Engine/Plugins/Experimental/MeshModelingToolset/Source/MeshModelingTools/Private/Physics/ExtractCollisionGeometryTool.cpp
@@ -27,11 +27,11 @@
 #include "Engine/Classes/PhysicsEngine/BodySetup.h"
 #include "Engine/Classes/PhysicsEngine/AggregateGeom.h"
 
+#include "ExplicitUseGeometryMathTypes.h"		// using UE::Geometry::(math types)
+using namespace UE::Geometry;
 
 #define LOCTEXT_NAMESPACE "UExtractCollisionGeometryTool"
 
-
-
 bool UExtractCollisionGeometryToolBuilder::CanBuildTool(const FToolBuilderState& SceneState) const
 {
 	int32 NumStaticMeshes = ToolBuilderUtil::CountComponents(SceneState, [&](UActorComponent* Comp) { return Cast<UStaticMeshComponent>(Comp) != nullptr; });
diff --git a/Engine/Plugins/Experimental/MeshModelingToolset/Source/MeshModelingTools/Private/Physics/PhysicsInspectorTool.cpp b/Engine/Plugins/Experimental/MeshModelingToolset/Source/MeshModelingTools/Private/Physics/PhysicsInspectorTool.cpp
--- a/Engine/Plugins/Experimental/MeshModelingToolset/Source/MeshModelingTools/Private/Physics/PhysicsInspectorTool.cpp
+++ b/Engine/Plugins/Experimental/MeshModelingToolset/Source/MeshModelingTools/Private/Physics/PhysicsInspectorTool.cpp
@@ -15,11 +15,11 @@
 #include "Engine/Classes/PhysicsEngine/BodySetup.h"
 #include "Engine/Classes/PhysicsEngine/AggregateGeom.h"
 
+#include "ExplicitUseGeometryMathTypes.h"		// using UE::Geometry::(math types)
+using namespace UE::Geometry;
 
 #define LOCTEXT_NAMESPACE "UPhysicsInspectorTool"
 
-
-
 bool UPhysicsInspectorToolBuilder::CanBuildTool(const FToolBuilderState& SceneState) const
 {
 	int32 NumStaticMeshes = ToolBuilderUtil::CountComponents(SceneState, [&](UActorComponent* Comp) { return Cast<UStaticMeshComponent>(Comp) != nullptr; });
diff --git a/Engine/Plugins/Experimental/MeshModelingToolset/Source/MeshModelingTools/Private/Physics/SetCollisionGeometryTool.cpp b/Engine/Plugins/Experimental/MeshModelingToolset/Source/MeshModelingTools/Private/Physics/SetCollisionGeometryTool.cpp
--- a/Engine/Plugins/Experimental/MeshModelingToolset/Source/MeshModelingTools/Private/Physics/SetCollisionGeometryTool.cpp
+++ b/Engine/Plugins/Experimental/MeshModelingToolset/Source/MeshModelingTools/Private/Physics/SetCollisionGeometryTool.cpp
@@ -24,9 +24,10 @@
 
 #include "Async/ParallelFor.h"
 
-#define LOCTEXT_NAMESPACE "USetCollisionGeometryTool"
-
+#include "ExplicitUseGeometryMathTypes.h"		// using UE::Geometry::(math types)
+using namespace UE::Geometry;
 
+#define LOCTEXT_NAMESPACE "USetCollisionGeometryTool"
 
 bool USetCollisionGeometryToolBuilder::CanBuildTool(const FToolBuilderState& SceneState) const
 {
diff --git a/Engine/Plugins/Experimental/MeshModelingToolset/Source/MeshModelingTools/Private/PlaneCutTool.cpp b/Engine/Plugins/Experimental/MeshModelingToolset/Source/MeshModelingTools/Private/PlaneCutTool.cpp
--- a/Engine/Plugins/Experimental/MeshModelingToolset/Source/MeshModelingTools/Private/PlaneCutTool.cpp
+++ b/Engine/Plugins/Experimental/MeshModelingToolset/Source/MeshModelingTools/Private/PlaneCutTool.cpp
@@ -30,8 +30,10 @@
 
 #include "Misc/MessageDialog.h"
 
-#define LOCTEXT_NAMESPACE "UPlaneCutTool"
+#include "ExplicitUseGeometryMathTypes.h"		// using UE::Geometry::(math types)
+using namespace UE::Geometry;
 
+#define LOCTEXT_NAMESPACE "UPlaneCutTool"
 
 /*
  * ToolBuilder
diff --git a/Engine/Plugins/Experimental/MeshModelingToolset/Source/MeshModelingTools/Private/PositionPlaneGizmo.cpp b/Engine/Plugins/Experimental/MeshModelingToolset/Source/MeshModelingTools/Private/PositionPlaneGizmo.cpp
--- a/Engine/Plugins/Experimental/MeshModelingToolset/Source/MeshModelingTools/Private/PositionPlaneGizmo.cpp
+++ b/Engine/Plugins/Experimental/MeshModelingToolset/Source/MeshModelingTools/Private/PositionPlaneGizmo.cpp
@@ -7,7 +7,8 @@
 #include "Engine/Engine.h"  // for GEngine - @todo remove this?
 #include "ToolSceneQueriesUtil.h"
 
-
+#include "ExplicitUseGeometryMathTypes.h"		// using UE::Geometry::(math types)
+using namespace UE::Geometry;
 
 UInteractiveGizmo* UPositionPlaneGizmoBuilder::BuildGizmo(const FToolBuilderState& SceneState) const
 {
@@ -19,9 +20,6 @@ UInteractiveGizmo* UPositionPlaneGizmoBuilder::BuildGizmo(const FToolBuilderStat
 
 
 
-
-
-
 void UPositionPlaneGizmo::SetWorld(UWorld* World)
 {
 	TargetWorld = World;
diff --git a/Engine/Plugins/Experimental/MeshModelingToolset/Source/MeshModelingTools/Private/Properties/MeshAnalysisProperties.cpp b/Engine/Plugins/Experimental/MeshModelingToolset/Source/MeshModelingTools/Private/Properties/MeshAnalysisProperties.cpp
--- a/Engine/Plugins/Experimental/MeshModelingToolset/Source/MeshModelingTools/Private/Properties/MeshAnalysisProperties.cpp
+++ b/Engine/Plugins/Experimental/MeshModelingToolset/Source/MeshModelingTools/Private/Properties/MeshAnalysisProperties.cpp
@@ -9,11 +9,12 @@
 
 #define LOCTEXT_NAMESPACE "UMeshAnalysisProperites"
 
+using namespace UE::Geometry;
 
 void UMeshAnalysisProperties::Update(const FDynamicMesh3& MeshIn, const FTransform& Transform)
 {
 	FTriangleMeshAdapterd TransformedMesh = MeshAdapterUtil::MakeTransformedDynamicMeshAdapter(&MeshIn, Transform);
-	FVector2d VolArea = TMeshQueries<FTriangleMeshAdapterd>::GetVolumeArea(TransformedMesh);
+	UE::Geometry::FVector2d VolArea = TMeshQueries<FTriangleMeshAdapterd>::GetVolumeArea(TransformedMesh);
 	this->SurfaceArea = FString::Printf(TEXT("%.2f m^2"), VolArea.Y / 10000);
 	this->Volume = FString::Printf(TEXT("%.2f m^3"), VolArea.X / 1000000);
 }
diff --git a/Engine/Plugins/Experimental/MeshModelingToolset/Source/MeshModelingTools/Private/RemeshMeshTool.cpp b/Engine/Plugins/Experimental/MeshModelingToolset/Source/MeshModelingTools/Private/RemeshMeshTool.cpp
--- a/Engine/Plugins/Experimental/MeshModelingToolset/Source/MeshModelingTools/Private/RemeshMeshTool.cpp
+++ b/Engine/Plugins/Experimental/MeshModelingToolset/Source/MeshModelingTools/Private/RemeshMeshTool.cpp
@@ -17,8 +17,10 @@
 
 #include "SceneManagement.h" // for FPrimitiveDrawInterface
 
-#define LOCTEXT_NAMESPACE "URemeshMeshTool"
+#include "ExplicitUseGeometryMathTypes.h"		// using UE::Geometry::(math types)
+using namespace UE::Geometry;
 
+#define LOCTEXT_NAMESPACE "URemeshMeshTool"
 
 /*
  * ToolBuilder
diff --git a/Engine/Plugins/Experimental/MeshModelingToolset/Source/MeshModelingTools/Private/RemoveOccludedTrianglesTool.cpp b/Engine/Plugins/Experimental/MeshModelingToolset/Source/MeshModelingTools/Private/RemoveOccludedTrianglesTool.cpp
--- a/Engine/Plugins/Experimental/MeshModelingToolset/Source/MeshModelingTools/Private/RemoveOccludedTrianglesTool.cpp
+++ b/Engine/Plugins/Experimental/MeshModelingToolset/Source/MeshModelingTools/Private/RemoveOccludedTrianglesTool.cpp
@@ -25,11 +25,11 @@
 #include "Misc/ScopedSlowTask.h"
 #endif
 
-
+#include "ExplicitUseGeometryMathTypes.h"		// using UE::Geometry::(math types)
+using namespace UE::Geometry;
 
 #define LOCTEXT_NAMESPACE "URemoveOccludedTrianglesTool"
 
-
 /*
  * ToolBuilder
  */
diff --git a/Engine/Plugins/Experimental/MeshModelingToolset/Source/MeshModelingTools/Private/Tests/DisplaceMeshTool_Tests.inl b/Engine/Plugins/Experimental/MeshModelingToolset/Source/MeshModelingTools/Private/Tests/DisplaceMeshTool_Tests.inl
--- a/Engine/Plugins/Experimental/MeshModelingToolset/Source/MeshModelingTools/Private/Tests/DisplaceMeshTool_Tests.inl
+++ b/Engine/Plugins/Experimental/MeshModelingToolset/Source/MeshModelingTools/Private/Tests/DisplaceMeshTool_Tests.inl
@@ -56,7 +56,7 @@ IMPLEMENT_SIMPLE_AUTOMATION_TEST(UDisplaceMeshTestConstant,
 	UTEST_TRUE("Input mesh has vertex at index 0", InputMesh->IsVertex(0));
 
 	// This is done internally during DisplaceOp::CalculateResult, but do it here so we can verify normals up front
-	FMeshNormals InputMeshNormals(InputMesh.Get());
+	UE::Geometry::FMeshNormals InputMeshNormals(InputMesh.Get());
 	InputMeshNormals.ComputeVertexNormals();
 	UTEST_TRUE("Input mesh normals point in +Z direction", FMath::IsNearlyZero((InputMeshNormals.GetNormals()[0] - FVector3d{ 0,0,1 }).Length()));
 
@@ -268,7 +268,7 @@ IMPLEMENT_SIMPLE_AUTOMATION_TEST(UDisplaceMeshTestDirectionalFilter,
 	// Create a box mesh, random noise, and a directional filter pointing in the +Y direction
 	TSharedPtr<FDynamicMesh3, ESPMode::ThreadSafe> InputMesh = CreateTestBoxMesh();
 
-	FMeshNormals InputMeshNormals(InputMesh.Get());
+	UE::Geometry::FMeshNormals InputMeshNormals(InputMesh.Get());
 	InputMeshNormals.ComputeVertexNormals();
 
 	EDisplaceMeshToolDisplaceType DisplaceType = EDisplaceMeshToolDisplaceType::RandomNoise;
diff --git a/Engine/Plugins/Experimental/MeshModelingToolset/Source/MeshModelingTools/Private/TransformMeshesTool.cpp b/Engine/Plugins/Experimental/MeshModelingToolset/Source/MeshModelingTools/Private/TransformMeshesTool.cpp
--- a/Engine/Plugins/Experimental/MeshModelingToolset/Source/MeshModelingTools/Private/TransformMeshesTool.cpp
+++ b/Engine/Plugins/Experimental/MeshModelingToolset/Source/MeshModelingTools/Private/TransformMeshesTool.cpp
@@ -16,8 +16,12 @@
 #include "Components/PrimitiveComponent.h"
 #include "Engine/World.h"
 
+#include "ExplicitUseGeometryMathTypes.h"		// using UE::Geometry::(math types)
+using namespace UE::Geometry;
+
 #define LOCTEXT_NAMESPACE "UTransformMeshesTool"
 
+
 /*
  * ToolBuilder
  */
diff --git a/Engine/Plugins/Experimental/MeshModelingToolset/Source/MeshModelingTools/Private/VolumeToMeshTool.cpp b/Engine/Plugins/Experimental/MeshModelingToolset/Source/MeshModelingTools/Private/VolumeToMeshTool.cpp
--- a/Engine/Plugins/Experimental/MeshModelingToolset/Source/MeshModelingTools/Private/VolumeToMeshTool.cpp
+++ b/Engine/Plugins/Experimental/MeshModelingToolset/Source/MeshModelingTools/Private/VolumeToMeshTool.cpp
@@ -20,8 +20,10 @@
 
 #include "Model.h"
 
-#define LOCTEXT_NAMESPACE "UVolumeToMeshTool"
+#include "ExplicitUseGeometryMathTypes.h"		// using UE::Geometry::(math types)
+using namespace UE::Geometry;
 
+#define LOCTEXT_NAMESPACE "UVolumeToMeshTool"
 
 /*
  * ToolBuilder
diff --git a/Engine/Plugins/Experimental/MeshModelingToolset/Source/MeshModelingTools/Private/WeldMeshEdgesTool.cpp b/Engine/Plugins/Experimental/MeshModelingToolset/Source/MeshModelingTools/Private/WeldMeshEdgesTool.cpp
--- a/Engine/Plugins/Experimental/MeshModelingToolset/Source/MeshModelingTools/Private/WeldMeshEdgesTool.cpp
+++ b/Engine/Plugins/Experimental/MeshModelingToolset/Source/MeshModelingTools/Private/WeldMeshEdgesTool.cpp
@@ -13,8 +13,10 @@
 #include "SceneManagement.h" // for FPrimitiveDrawInterface
 #include "Async/ParallelFor.h"
 
-#define LOCTEXT_NAMESPACE "UWeldMeshEdgesTool"
+#include "ExplicitUseGeometryMathTypes.h"		// using UE::Geometry::(math types)
+using namespace UE::Geometry;
 
+#define LOCTEXT_NAMESPACE "UWeldMeshEdgesTool"
 
 /*
  * ToolBuilder
diff --git a/Engine/Plugins/Experimental/MeshModelingToolset/Source/MeshModelingTools/Public/AddPatchTool.h b/Engine/Plugins/Experimental/MeshModelingToolset/Source/MeshModelingTools/Public/AddPatchTool.h
--- a/Engine/Plugins/Experimental/MeshModelingToolset/Source/MeshModelingTools/Public/AddPatchTool.h
+++ b/Engine/Plugins/Experimental/MeshModelingToolset/Source/MeshModelingTools/Public/AddPatchTool.h
@@ -13,7 +13,7 @@
 
 class IAssetGenerationAPI;
 
-class FDynamicMesh3;
+PREDECLARE_USE_GEOMETRY_CLASS(FDynamicMesh3);
 
 /**
  *
@@ -120,14 +120,14 @@ protected:
 
 	FBox WorldBounds;
 
-	FFrame3f ShapeFrame;
+	UE::Geometry::FFrame3f ShapeFrame;
 	bool bPreviewValid = true;
 
 	void UpdatePreviewPosition(const FInputDeviceRay& ClickPos);
 	void UpdatePreviewMesh();
 
-	TUniquePtr<FDynamicMesh3> BaseMesh;
+	TUniquePtr<UE::Geometry::FDynamicMesh3> BaseMesh;
 	void GeneratePreviewBaseMesh();
 
-	void GeneratePlane(FDynamicMesh3* OutMesh);
+	void GeneratePlane(UE::Geometry::FDynamicMesh3* OutMesh);
 };
diff --git a/Engine/Plugins/Experimental/MeshModelingToolset/Source/MeshModelingTools/Public/AddPrimitiveTool.h b/Engine/Plugins/Experimental/MeshModelingToolset/Source/MeshModelingTools/Public/AddPrimitiveTool.h
--- a/Engine/Plugins/Experimental/MeshModelingToolset/Source/MeshModelingTools/Public/AddPrimitiveTool.h
+++ b/Engine/Plugins/Experimental/MeshModelingToolset/Source/MeshModelingTools/Public/AddPrimitiveTool.h
@@ -13,7 +13,7 @@
 
 #include "AddPrimitiveTool.generated.h"
 
-class FDynamicMesh3;
+PREDECLARE_USE_GEOMETRY_CLASS(FDynamicMesh3);
 class IAssetGenerationAPI;
 
 /**
@@ -515,7 +515,7 @@ protected:
 	IAssetGenerationAPI* AssetAPI;
 
 	void UpdatePreviewPosition(const FInputDeviceRay& ClickPos);
-	FFrame3f ShapeFrame;
+	UE::Geometry::FFrame3f ShapeFrame;
 
 	void UpdatePreviewMesh();
 };
diff --git a/Engine/Plugins/Experimental/MeshModelingToolset/Source/MeshModelingTools/Public/AlignObjectsTool.h b/Engine/Plugins/Experimental/MeshModelingToolset/Source/MeshModelingTools/Public/AlignObjectsTool.h
--- a/Engine/Plugins/Experimental/MeshModelingToolset/Source/MeshModelingTools/Public/AlignObjectsTool.h
+++ b/Engine/Plugins/Experimental/MeshModelingToolset/Source/MeshModelingTools/Public/AlignObjectsTool.h
@@ -102,7 +102,9 @@ UCLASS()
 class MESHMODELINGTOOLS_API UAlignObjectsTool : public UMultiSelectionTool, public IClickDragBehaviorTarget
 {
 	GENERATED_BODY()
-
+	using FVector3d = UE::Geometry::FVector3d;
+	using FTransform3d = UE::Geometry::FTransform3d;
+	using FAxisAlignedBox3d = UE::Geometry::FAxisAlignedBox3d;
 public:
 	UAlignObjectsTool();
 
