Geometry Processing: Mesh Conversions.  Allow updating a MeshDescription UVs when the underlying mesh hasn't changed.  Enable this use in the UV editing tools.

diff --git a/Engine/Plugins/Experimental/GeometryProcessing/Source/MeshConversion/Private/DynamicMeshToMeshDescription.cpp b/Engine/Plugins/Experimental/GeometryProcessing/Source/MeshConversion/Private/DynamicMeshToMeshDescription.cpp
--- a/Engine/Plugins/Experimental/GeometryProcessing/Source/MeshConversion/Private/DynamicMeshToMeshDescription.cpp
+++ b/Engine/Plugins/Experimental/GeometryProcessing/Source/MeshConversion/Private/DynamicMeshToMeshDescription.cpp
@@ -64,13 +64,6 @@ void FDynamicMeshToMeshDescription::UpdateAttributes(const FDynamicMesh3* MeshIn
 {
 	check(MeshIn->IsCompactV());
 
-	if (bUpdateUVs)
-	{
-		// currently we can't update the shared UVs on a MeshDescription
-		// Updating only per-instance UVs will result in the two representations for UVs to be out of sync. 
-		// @todo UpdateUV path below when MeshDescription supports updating the shared UVs 
-		check(0);
-	}
 
 	FStaticMeshAttributes Attributes(MeshOut);
 
@@ -118,13 +111,72 @@ void FDynamicMeshToMeshDescription::UpdateAttributes(const FDynamicMesh3* MeshIn
 			if (MeshIn->HasAttributes())
 			{
 				check(MeshIn->TriangleCount() == MeshOut.Triangles().Num())
-				for (int UVLayerIndex = 0, NumLayers = MeshIn->Attributes()->NumUVLayers(); UVLayerIndex < NumLayers; UVLayerIndex++)
+				int32 NumLayers = MeshIn->Attributes()->NumUVLayers();			
+				MeshOut.SetNumUVChannels(NumLayers); // This resets MeshDescription's internal TriangleUV array
+				
+				for (int UVLayerIndex = 0; UVLayerIndex < NumLayers; UVLayerIndex++)
 				{
-					DynamicMeshToMeshDescriptionConversionHelper::SetAttributesFromOverlay(MeshIn, MeshOut, InstanceAttrib, MeshIn->Attributes()->GetUVLayer(UVLayerIndex), UVLayerIndex);
+					const FDynamicMeshUVOverlay* UVOverlay = MeshIn->Attributes()->GetUVLayer(UVLayerIndex);
+					// update the vertex Attribute UVs
+					DynamicMeshToMeshDescriptionConversionHelper::SetAttributesFromOverlay(MeshIn, MeshOut, InstanceAttrib, UVOverlay, UVLayerIndex);
+
+					// rebuild the shared UVs
+					FUVArray& UVArray = MeshOut.UVs(UVLayerIndex);
+					UVArray.Reset(); // delete existing UV buffer
+					UVArray.Reserve(UVOverlay->ElementCount());
+					
+					// rebuild the UV buffer
+					TArray<FUVID> ElIDToUVIDMap;  
+					ElIDToUVIDMap.Reserve(UVOverlay->ElementCount());
+					for (int32 ElID = 0, MaxID = UVOverlay->MaxElementID(); ElID < MaxID; ++ElID)
+					{
+						if (!UVOverlay->IsElement(ElID))
+						{
+							continue;
+						}
+
+						FVector2f UVValue = UVOverlay->GetElement(ElID);
+						FVector2D UVValue2D(UVValue.X, UVValue.Y);
+						FUVID UVID = UVArray.Add();
+						ElIDToUVIDMap.Insert(UVID, ElID);
+						UVArray.GetAttributes().GetAttributesRef<FVector2D>(MeshAttribute::UV::UVCoordinate)[UVID] = UVValue2D;
+					}
+
+					for (const FTriangleID TriangleID : MeshOut.Triangles().GetElementIDs())
+					{
+						int TriID = TriangleID.GetValue(); // assumes the same TriIDs in both meshes.
+						FIndex3i ElIDs = UVOverlay->GetTriangle(TriID);
+						TArray<FUVID, TFixedAllocator<3>> MDTri; 
+						MDTri.Add(ElIDToUVIDMap[ElIDs[0]]);
+						MDTri.Add(ElIDToUVIDMap[ElIDs[1]]);
+						MDTri.Add(ElIDToUVIDMap[ElIDs[2]]);
+						MeshOut.SetTriangleUVIndices(TriangleID, MDTri, UVLayerIndex);
+					}
+				}
+				if (0)
+				{ 
+					// Verify the shared UVs and per-vertexinstance UVs match
+					for (int UVLayerIndex = 0; UVLayerIndex < NumLayers; UVLayerIndex++)
+					{ 
+						for (const FTriangleID TriangleID : MeshOut.Triangles().GetElementIDs())
+						{
+							TArrayView<const FVertexInstanceID> TriWedges = MeshOut.GetTriangleVertexInstances(TriangleID);
+							TArrayView<FUVID> UVTri = MeshOut.GetTriangleUVIndices(TriangleID, UVLayerIndex);
+
+							for (int32 i = 0; i < 3; ++i)
+							{
+								// UV from shared
+								FVector2D SharedUV = MeshOut.UVs(UVLayerIndex).GetAttributes().GetAttributesRef<FVector2D>(MeshAttribute::UV::UVCoordinate)[UVTri[i]];
+								FVector2D WedgeUV = InstanceAttrib.Get(TriWedges[i], UVLayerIndex);
+								check(SharedUV == WedgeUV);
+							}
+						}
+					}
 				}
 			}
 			else
 			{
+				// [todo] correctly build shared UVs?
 				check(MeshIn->VertexCount() == MeshOut.Vertices().Num());
 				for (int VertID : MeshIn->VertexIndicesItr())
 				{
diff --git a/Engine/Plugins/Experimental/MeshModelingToolset/Source/MeshModelingTools/Private/UVProjectionTool.cpp b/Engine/Plugins/Experimental/MeshModelingToolset/Source/MeshModelingTools/Private/UVProjectionTool.cpp
--- a/Engine/Plugins/Experimental/MeshModelingToolset/Source/MeshModelingTools/Private/UVProjectionTool.cpp
+++ b/Engine/Plugins/Experimental/MeshModelingToolset/Source/MeshModelingTools/Private/UVProjectionTool.cpp
@@ -356,16 +356,7 @@ void UUVProjectionTool::GenerateAsset(const TArray<FDynamicMeshOpResult>& Result
 		ComponentTargets[ComponentIdx]->CommitMesh([&Results, &ComponentIdx, this](const FPrimitiveComponentTarget::FCommitParams& CommitParams)
 		{
 			FDynamicMeshToMeshDescription Converter;
-			//if (??) // TODO check if UV topology changed?  or remove all traces of this if statement (may be safe to assume it almost always changes w/ a uv projection op)
-			{
-				// full conversion if topology changed
-				Converter.Convert(Results[ComponentIdx].Mesh.Get(), *CommitParams.MeshDescription);
-			}
-			//else
-			//{
-			//	// otherwise just copy attributes
-			//	Converter.UpdateAttributes(Results[ComponentIdx]->Mesh.Get(), *MeshDescription, false, true);
-			//}
+			Converter.UpdateAttributes(Results[ComponentIdx].Mesh.Get(), *CommitParams.MeshDescription, false, false, true/*update uvs*/);
 		});
 	}
 
diff --git a/Engine/Plugins/Experimental/MeshModelingToolset/Source/MeshModelingToolsEditorOnly/Private/ParameterizeMeshTool.cpp b/Engine/Plugins/Experimental/MeshModelingToolset/Source/MeshModelingToolsEditorOnly/Private/ParameterizeMeshTool.cpp
--- a/Engine/Plugins/Experimental/MeshModelingToolset/Source/MeshModelingToolsEditorOnly/Private/ParameterizeMeshTool.cpp
+++ b/Engine/Plugins/Experimental/MeshModelingToolset/Source/MeshModelingToolsEditorOnly/Private/ParameterizeMeshTool.cpp
@@ -200,7 +200,7 @@ void UParameterizeMeshTool::Shutdown(EToolShutdownType ShutdownType)
 		ComponentTarget->CommitMesh([DynamicMeshResult](const FPrimitiveComponentTarget::FCommitParams& CommitParams)
 		{
 			FDynamicMeshToMeshDescription Converter;
-			Converter.Convert(DynamicMeshResult, *CommitParams.MeshDescription);
+			Converter.UpdateAttributes(DynamicMeshResult, *CommitParams.MeshDescription, false, false, true/*update uvs*/);
 		});
 		GetToolManager()->EndUndoTransaction();
 	}
diff --git a/Engine/Plugins/Experimental/MeshModelingToolset/Source/MeshModelingToolsEditorOnly/Private/UVLayoutTool.cpp b/Engine/Plugins/Experimental/MeshModelingToolset/Source/MeshModelingToolsEditorOnly/Private/UVLayoutTool.cpp
--- a/Engine/Plugins/Experimental/MeshModelingToolset/Source/MeshModelingToolsEditorOnly/Private/UVLayoutTool.cpp
+++ b/Engine/Plugins/Experimental/MeshModelingToolset/Source/MeshModelingToolsEditorOnly/Private/UVLayoutTool.cpp
@@ -357,7 +357,7 @@ void UUVLayoutTool::GenerateAsset(const TArray<FDynamicMeshOpResult>& Results)
 		ComponentTargets[ComponentIdx]->CommitMesh([&Results, &ComponentIdx, this](const FPrimitiveComponentTarget::FCommitParams& CommitParams)
 		{
 			FDynamicMeshToMeshDescription Converter;
-			Converter.Convert(Results[ComponentIdx].Mesh.Get(), *CommitParams.MeshDescription);
+			Converter.UpdateAttributes(Results[ComponentIdx].Mesh.Get(), *CommitParams.MeshDescription, false, false, true/*update uvs*/);
 		});
 	}
 
diff --git a/Engine/Plugins/Experimental/MeshModelingToolset/Source/ModelingComponents/Private/SimpleDynamicMeshComponent.cpp b/Engine/Plugins/Experimental/MeshModelingToolset/Source/ModelingComponents/Private/SimpleDynamicMeshComponent.cpp
--- a/Engine/Plugins/Experimental/MeshModelingToolset/Source/ModelingComponents/Private/SimpleDynamicMeshComponent.cpp
+++ b/Engine/Plugins/Experimental/MeshModelingToolset/Source/ModelingComponents/Private/SimpleDynamicMeshComponent.cpp
@@ -179,11 +179,9 @@ void USimpleDynamicMeshComponent::ApplyTransform(const FTransform3d& Transform,
 
 void USimpleDynamicMeshComponent::Bake(FMeshDescription* MeshDescription, bool bHaveModifiedTopology, const FConversionToMeshDescriptionOptions& ConversionOptions)
 {
-	// shared UVs on MeshDescription can't currently be updated. this forces us to rebuild the mesh.
-	const bool bModifiedUVsOrTopology = ( ConversionOptions.bUpdateUVs || bHaveModifiedTopology ); 
 	
 	FDynamicMeshToMeshDescription Converter(ConversionOptions);
-	if (bModifiedUVsOrTopology == false && Converter.HaveMatchingElementCounts(Mesh.Get(), MeshDescription))
+	if (bHaveModifiedTopology == false && Converter.HaveMatchingElementCounts(Mesh.Get(), MeshDescription))
 	{
 		if (ConversionOptions.bUpdatePositions)
 		{
