<description: restricted, no permission to view>

diff --git a/Engine/Plugins/Experimental/GeometryProcessing/Source/MeshConversion/Private/DynamicMeshToMeshDescription.cpp b/Engine/Plugins/Experimental/GeometryProcessing/Source/MeshConversion/Private/DynamicMeshToMeshDescription.cpp
--- a/Engine/Plugins/Experimental/GeometryProcessing/Source/MeshConversion/Private/DynamicMeshToMeshDescription.cpp
+++ b/Engine/Plugins/Experimental/GeometryProcessing/Source/MeshConversion/Private/DynamicMeshToMeshDescription.cpp
@@ -126,9 +126,10 @@ void FDynamicMeshToMeshDescription::UpdateAttributes(const FDynamicMesh3* MeshIn
 					UVArray.Reserve(UVOverlay->ElementCount());
 					
 					// rebuild the UV buffer
-					TArray<FUVID> ElIDToUVIDMap;  
-					ElIDToUVIDMap.Reserve(UVOverlay->ElementCount());
-					for (int32 ElID = 0, MaxID = UVOverlay->MaxElementID(); ElID < MaxID; ++ElID)
+					TArray<FUVID> ElIDToUVIDMap;
+					int32 MaxID = UVOverlay->MaxElementID();
+					ElIDToUVIDMap.Reserve(MaxID);
+					for (int32 ElID = 0; ElID < MaxID; ++ElID)
 					{
 						if (!UVOverlay->IsElement(ElID))
 						{
