Geometry Processing: Mesh Conversion: UVs. Fix potential bug when dynamic mesh IDs are sparse.

diff --git a/Engine/Plugins/Experimental/GeometryProcessing/Source/MeshConversion/Private/DynamicMeshToMeshDescription.cpp b/Engine/Plugins/Experimental/GeometryProcessing/Source/MeshConversion/Private/DynamicMeshToMeshDescription.cpp
--- a/Engine/Plugins/Experimental/GeometryProcessing/Source/MeshConversion/Private/DynamicMeshToMeshDescription.cpp
+++ b/Engine/Plugins/Experimental/GeometryProcessing/Source/MeshConversion/Private/DynamicMeshToMeshDescription.cpp
@@ -126,9 +126,9 @@ void FDynamicMeshToMeshDescription::UpdateAttributes(const FDynamicMesh3* MeshIn
 					UVArray.Reserve(UVOverlay->ElementCount());
 					
 					// rebuild the UV buffer
-					TArray<FUVID> ElIDToUVIDMap;
-					int32 MaxID = UVOverlay->MaxElementID();
-					ElIDToUVIDMap.Reserve(MaxID);
+					int32 MaxID = UVOverlay->MaxElementID(); // the true maxid +1 
+					TArray<FUVID> ElIDToUVIDMap;  
+					ElIDToUVIDMap.AddUninitialized(MaxID);
 					for (int32 ElID = 0; ElID < MaxID; ++ElID)
 					{
 						if (!UVOverlay->IsElement(ElID))
@@ -139,7 +139,7 @@ void FDynamicMeshToMeshDescription::UpdateAttributes(const FDynamicMesh3* MeshIn
 						FVector2f UVValue = UVOverlay->GetElement(ElID);
 						FVector2D UVValue2D(UVValue.X, UVValue.Y);
 						FUVID UVID = UVArray.Add();
-						ElIDToUVIDMap.Insert(UVID, ElID);
+						ElIDToUVIDMap[ElID] = UVID;
 						UVArray.GetAttributes().GetAttributesRef<FVector2D>(MeshAttribute::UV::UVCoordinate)[UVID] = UVValue2D;
 					}
 
