Mesh Modeling : Mesh Conversion - fix potential bug with conversion of non-manifold mesh description to dynamic mesh

diff --git a/Engine/Plugins/Experimental/GeometryProcessing/Source/MeshConversion/Private/MeshDescriptionToDynamicMesh.cpp b/Engine/Plugins/Experimental/GeometryProcessing/Source/MeshConversion/Private/MeshDescriptionToDynamicMesh.cpp
--- a/Engine/Plugins/Experimental/GeometryProcessing/Source/MeshConversion/Private/MeshDescriptionToDynamicMesh.cpp
+++ b/Engine/Plugins/Experimental/GeometryProcessing/Source/MeshConversion/Private/MeshDescriptionToDynamicMesh.cpp
@@ -195,9 +195,8 @@ void FMeshDescriptionToDynamicMesh::Convert(const FMeshDescription* MeshIn, FDyn
 	AddedTriangles.SetNum(MeshIn->Triangles().Num());
 
 	// allocate the TriIDMap (maps from DynamicMesh triangle ID to MeshDescription FTriangleID)
-	// reserve space for the data without setting array num()
-	TriIDMap.Reserve(MeshIn->Triangles().Num());
-
+	TriIDMap.AddUninitialized(MeshIn->Triangles().Num());
+	
 
 	// Iterate over triangles in the Mesh Description
 	// NOTE: If you change the iteration order here, please update the corresponding iteration in FDynamicMeshToMeshDescription::UpdateAttributes, 
@@ -288,27 +287,19 @@ void FMeshDescriptionToDynamicMesh::Convert(const FMeshDescription* MeshIn, FDyn
 				bDuplicate[2] = true;
 				bDuplicate[0] = true;
 			}
-			if (bDuplicate[0])
-			{
-				const FVector Position = VertexPositions[TriangleVertexIDs[0]];
-				const int32 NewVertIdx = MeshOut.AppendVertex(Position);
-				VertexIDs[0] = NewVertIdx;
-				VertIDMap.Insert(TriangleVertexIDs[0], NewVertIdx); // note 'insert', because the array may need to grow
-			}
-			if (bDuplicate[1])
-			{
-				const FVector Position = VertexPositions[TriangleVertexIDs[1]];
-				const int32 NewVertIdx = MeshOut.AppendVertex(Position);
-				VertexIDs[1] = NewVertIdx;
-				VertIDMap.Insert(TriangleVertexIDs[1], NewVertIdx);
-			}
-			if (bDuplicate[2])
+			for (int32 i = 0; i < 3; ++i)
 			{
-				const FVector Position = VertexPositions[TriangleVertexIDs[2]];
-				const int32 NewVertIdx = MeshOut.AppendVertex(Position);
-				VertexIDs[2] = NewVertIdx;
-				VertIDMap.Insert(TriangleVertexIDs[2], NewVertIdx);
+				if (bDuplicate[i])
+				{
+					const FVertexID& TriangleVertID = TriangleVertexIDs[i];
+					const FVector Position = VertexPositions[TriangleVertID];
+					const int32 NewVertIdx = MeshOut.AppendVertex(Position);
+					VertexIDs[i] = NewVertIdx;
+					VertIDMap.SetNumUninitialized(NewVertIdx + 1);
+					VertIDMap[NewVertIdx] = TriangleVertID;
+				}
 			}
+			
 
 			NewTriangleID = MeshOut.AppendTriangle(VertexIDs, GroupID);
 			checkSlow(NewTriangleID != FDynamicMesh3::NonManifoldID);
@@ -316,8 +307,8 @@ void FMeshDescriptionToDynamicMesh::Convert(const FMeshDescription* MeshIn, FDyn
 
 		checkSlow(NewTriangleID >= 0);
 		AddedTriangles[NewTriangleID] = TriData;
-		// .Insert(Value, Index)
-		TriIDMap.Insert(TriangleID, NewTriangleID);
+		
+		TriIDMap[NewTriangleID] = TriangleID;
 	
 		// copy triangle instance colors to per-vertex colors. Currently this just uses last-instance as
 		// we do not support per-triangle-vertex colors.
@@ -333,6 +324,12 @@ void FMeshDescriptionToDynamicMesh::Convert(const FMeshDescription* MeshIn, FDyn
 		}
 	}
 
+	int32 MaxTriID = MeshOut.MaxTriangleID(); // really MaxTriID+1
+	// if the source mesh had duplicates then TriIDMap will be too long, this will trim excess
+	TriIDMap.SetNum(MaxTriID);
+	
+
+
 	if (bFoundNonDefaultVertexColor == false)
 	{
 		MeshOut.DiscardVertexColors();
@@ -476,16 +473,20 @@ void FMeshDescriptionToDynamicMesh::Convert(const FMeshDescription* MeshIn, FDyn
 						
 					// map to translate UVIds from FUVID::int32() to DynamicOverlay Index.
 					TArray<int32> UVIndexMap;
-					UVIndexMap.Reserve(UVs.GetArraySize() + 1);
-
+					{
+						// pre-compute max UVID in the mesh description
+						int32 MaxUVElID = -1;
+						for (FUVID UVID : UVs.GetElementIDs())
+						{
+							MaxUVElID = FMath::Max(MaxUVElID, UVID.GetValue());
+						}
+						UVIndexMap.AddUninitialized(MaxUVElID+1);
+					}
 					for (FUVID UVID : UVs.GetElementIDs())
 					{
 						const FVector2D UVvalue = UVCoordinates[UVID];
 						int32 NewIndex = UVOverlay->AppendElement(FVector2f(UVvalue));
-
-						// forced to use Insert since this array might have to resize (unfortunately we can only guess the max UVID in the mesh description)
-						UVIndexMap.Insert(NewIndex, UVID.GetValue());
-
+						UVIndexMap[UVID.GetValue()] = NewIndex;
 					}
 
 					// copy uv "index buffer"
@@ -773,7 +774,8 @@ static void CopyTangents_Internal(const FMeshDescription* SourceMesh, const FDyn
 void FMeshDescriptionToDynamicMesh::CopyTangents(const FMeshDescription* SourceMesh, const FDynamicMesh3* TargetMesh, TMeshTangents<float>* TangentsOut)
 {
 	if (!ensureMsgf(bCalculateMaps, TEXT("Cannot CopyTangents unless Maps were calculated"))) return;
-	if (!ensureMsgf(TriIDMap.Num() == TargetMesh->TriangleCount(), TEXT("Tried to CopyTangents to mesh with different triangle count"))) return;
+	int32 TargetTriCount = TargetMesh->TriangleCount();
+	if (!ensureMsgf(TriIDMap.Num() == TargetTriCount, TEXT("Tried to CopyTangents to mesh with different triangle count"))) return;
 	CopyTangents_Internal<float>(SourceMesh, TargetMesh, TangentsOut, TriIDMap);
 }
 
