GeometryProcessing: Mesh Conversion - Add conversion from SkeletalMeshLODModel to Dynamic Mesh.   Introduce a ToDynamicMesh to expose functionality for making FDynamicMesh3 from different sources.

diff --git a/Engine/Plugins/Experimental/GeometryProcessing/Source/MeshConversion/Private/SkeletalMeshLODModelToDynamicMesh.cpp b/Engine/Plugins/Experimental/GeometryProcessing/Source/MeshConversion/Private/SkeletalMeshLODModelToDynamicMesh.cpp
new file mode 100644
--- /dev/null
+++ b/Engine/Plugins/Experimental/GeometryProcessing/Source/MeshConversion/Private/SkeletalMeshLODModelToDynamicMesh.cpp
@@ -0,0 +1,392 @@
+// Copyright Epic Games, Inc. All Rights Reserved. 
+
+#include "SkeletalMeshLODModelToDynamicMesh.h"
+#include "Rendering/SkeletalMeshLODModel.h"
+#include "ToDynamicMesh.h"
+
+#include "ExplicitUseGeometryMathTypes.h"		// using UE::Geometry::(math types)
+using namespace UE::Geometry;
+
+#if WITH_EDITOR
+
+class FSkeletalMeshLODModelWrapper
+{
+public:
+
+	typedef int32 TriIDType;
+	typedef int32 VertIDType;
+	typedef int32 WedgeIDType;
+
+	typedef int32 UVIDType;
+	typedef int32 NormalIDType;
+
+	FSkeletalMeshLODModelWrapper(const  FSkeletalMeshLODModel& MeshIn, bool bUseDisabledSections) :
+		bIncludeDisabledSections(bUseDisabledSections),
+		Mesh(&MeshIn)
+	{
+		auto SkipSection = [bUseDisabledSections](const FSkelMeshSection& Section)->bool
+		{
+			return (!bUseDisabledSections && Section.bDisabled);
+		};
+	
+		// pre-count triangles and verts
+		int32 TriCount = 0;
+		int32 VertCount = 0;
+		for (const FSkelMeshSection& Section : Mesh->Sections)
+		{
+			if (SkipSection(Section))
+			{
+				continue;
+			}
+
+			VertCount += (int32)Section.NumVertices;
+			TriCount += (int32)Section.NumTriangles;
+		}
+		
+		// construct a list of all valid VertIDs for this mesh.
+		VertIDs.Reserve(VertCount);
+		for (const FSkelMeshSection& Section : Mesh->Sections)
+		{
+			if (SkipSection(Section))
+			{
+				continue;
+			}
+
+			const int32 BaseVertexIndex = Section.BaseVertexIndex;
+			const int32 NumSectionVtx = Section.SoftVertices.Num();
+			for (int32 SectionVtxIndex = 0, VtxIndex = BaseVertexIndex; SectionVtxIndex < NumSectionVtx; ++SectionVtxIndex, ++VtxIndex)
+			{
+				VertIDs.Add(VtxIndex);
+			}
+		}
+
+		// construct list of all valid TriIDs for this mesh.  
+		TriIDs.Reserve(TriCount);
+		for (int32 SectionIdx = 0; SectionIdx < Mesh->Sections.Num(); ++SectionIdx)
+		{
+			const FSkelMeshSection& Section = Mesh->Sections[SectionIdx];
+
+			if (SkipSection(Section))
+			{
+				continue;
+			}
+			const int32 BaseIndex = Section.BaseIndex;
+			const int32 BaseTriIdx = BaseIndex / 3;
+			const int32 NumSectionTris = (int32)Section.NumTriangles;
+			for (int32 i = 0, j = BaseTriIdx; i < NumSectionTris; ++i, ++j)
+			{
+				TriIDs.Add(j);
+			}
+		}
+
+
+		const int32 MaxVertID = VertIDs.Last();
+		VertIDToSectionID.Init(-1, MaxVertID + 1);// value only used for debugging 
+		for (int32 SectionIdx = 0; SectionIdx < Mesh->Sections.Num(); ++SectionIdx)
+		{
+			const FSkelMeshSection& Section = Mesh->Sections[SectionIdx];
+			if (SkipSection(Section))
+			{
+				continue;
+			}
+
+			const int32 BaseVertexIndex = (int32)Section.BaseVertexIndex;
+			const int32 NumSectionVtx = Section.SoftVertices.Num();
+			for (int32 SectionVtxIndex = 0, VtxIndex = BaseVertexIndex; SectionVtxIndex < NumSectionVtx; ++SectionVtxIndex, ++VtxIndex)
+			{
+				VertIDToSectionID[VtxIndex] = SectionIdx;
+			}
+		}
+
+		const int32 MaxTriID = TriIDs.Last();
+		TriIDToSectionID.Init(-1, MaxTriID + 1);
+		for (int32 SectionIdx = 0; SectionIdx < Mesh->Sections.Num(); ++SectionIdx)
+		{
+			const FSkelMeshSection& Section = Mesh->Sections[SectionIdx];
+			if (SkipSection(Section))
+			{
+				continue;
+			}
+
+			const int32 BaseIndex = (int32)Section.BaseIndex;
+			const int32 BaseTriID = BaseIndex / 3;
+			const int32 NumSectionTris = (int32)Section.NumTriangles;
+			for (int32 i = 0; i < NumSectionTris; ++i)
+			{
+				TriIDToSectionID[i + BaseTriID] = SectionIdx;
+			}
+		}
+
+	}
+
+	
+	int32 NumTris() const
+	{
+		return TriIDs.Num();
+	}
+
+	int32 NumVerts() const
+	{
+		return VertIDs.Num();
+	}
+
+	int32 NumUVLayers() const
+	{
+		return Mesh->NumTexCoords;
+	}
+
+	// --"Vertex Buffer" info
+	const TArray<VertIDType>& GetVertIDs() const
+	{
+		return VertIDs;
+	}
+	FVector3d GetPosition(VertIDType VtxID) const
+	{
+		int32 SectionID = VertIDToSectionID[VtxID];
+		checkSlow(SectionID != -1);
+		const FSkelMeshSection& Section = Mesh->Sections[SectionID];
+		const int32 BaseVertexIndex = (int32)Section.BaseVertexIndex;
+
+		return FVector3d(Section.SoftVertices[VtxID - BaseVertexIndex].Position);
+	}
+
+	// --"Index Buffer" info
+	const TArray<TriIDType>& GetTriIDs() const
+	{
+		return TriIDs;
+	}
+	bool GetTri(TriIDType TriID, VertIDType& VID0, VertIDType& VID1, VertIDType& VID2) const
+	{
+		int32 Offset = TriID * 3;
+		VID0 = Mesh->IndexBuffer[Offset];
+		VID1 = Mesh->IndexBuffer[Offset + 1];
+		VID2 = Mesh->IndexBuffer[Offset + 2];
+
+		// maybe check IndexBuffer.Num() > Offset+2 ? 
+		return true;
+	}
+	
+	bool HasNormals() const
+	{
+		return true;
+	}
+
+	bool HasTangents() const
+	{
+		return true;
+	}
+
+	bool HasBiTangents() const
+	{
+		return true;
+	}
+
+	//-- Access to per-wedge attributes --//
+	void GetWedgeIDs(const TriIDType& TriID, WedgeIDType& WID0, WedgeIDType& WID1, WedgeIDType& WID2) const
+	{
+		int32 Offset = 3 * TriID;
+		WID0 = Offset;
+		WID1 = Offset + 1;
+		WID2 = Offset + 2;
+	}
+
+	FVector2f GetWedgeUV(int32 UVLayerIndex, WedgeIDType WID) const
+	{
+		return (FVector2f)GetWedgeVertexInstance(WID).UVs[UVLayerIndex];	
+	}
+
+	FVector3f GetWedgeNormal(WedgeIDType WID) const
+	{
+		const FVector4& TangentZ = GetWedgeVertexInstance(WID).TangentZ;
+		return FVector3f(TangentZ.X, TangentZ.Y, TangentZ.Z);
+	}
+
+	FVector3f GetWedgeTangent(WedgeIDType WID) const
+	{
+		return FVector3f(GetWedgeVertexInstance(WID).TangentX);
+	}
+
+	FVector3f GetWedgeBiTangent(WedgeIDType WID) const
+	{
+		return FVector3f(GetWedgeVertexInstance(WID).TangentY);
+	}
+	
+
+
+	int32 GetMaterialIndex(TriIDType TriID) const
+	{
+		int32 SectionID = TriIDToSectionID[TriID];
+		const FSkelMeshSection& Section = Mesh->Sections[SectionID];
+		return Section.MaterialIndex;
+	}
+
+	
+	
+	//-- null implementation of shared attributes: Skeletal mesh model doesn't use these --//
+	const TArray<int32>& GetUVIDs(int32 LayerID) const { return EmptyArray; }
+	FVector2f GetUV(int32 LayerID, UVIDType UVID) const { check(0); return FVector2f(); }
+	bool GetUVTri(int32 LayerID, const TriIDType&, UVIDType& ID0, UVIDType& ID1, UVIDType& ID2) const { ID0 = ID1 = ID2 = UVIDType(-1); return false;}
+	
+	const TArray<int32>& GetNormalIDs() const { return EmptyArray; }
+	FVector3f GetNormal(NormalIDType ID) const { check(0); return FVector3f(); }
+	bool GetNormalTri(const TriIDType&, NormalIDType& ID0, NormalIDType& ID1, NormalIDType& ID2) const { ID0 = ID1 = ID2 = NormalIDType(-1); return false; }
+
+	const TArray<int32>& GetTangentIDs() const { return EmptyArray; }
+	FVector3f GetTangent(NormalIDType ID) const { check(0); return FVector3f(); }
+	bool GetTangentTri(const TriIDType&, NormalIDType& ID0, NormalIDType& ID1, NormalIDType& ID2) const { ID0 = ID1 = ID2 = NormalIDType(-1); return false; }
+
+	
+	const TArray<int32>& GetBiTangentIDs() const { return EmptyArray; }
+	FVector3f GetBiTangent(NormalIDType ID) const {check(0); return FVector3f(); }
+	bool GetBiTangentTri(const TriIDType&, NormalIDType& ID0, NormalIDType& ID1, NormalIDType& ID2) const { ID0 = ID1 = ID2 = NormalIDType(-1); return false; }
+
+	
+	// -- additional methods, not required by the conversion interface
+	const  FSkeletalMeshLODModel& GetSrcMesh() const
+	{
+		return *Mesh;
+	}
+
+	
+	FLinearColor GetWedgeColor(const WedgeIDType WID) const
+	{
+		const FColor& Color = GetWedgeVertexInstance(WID).Color;
+		return Color.ReinterpretAsLinear();
+	}
+
+private:
+
+	inline const FSoftSkinVertex& GetWedgeVertexInstance(WedgeIDType WID) const
+	{
+		int32 VertID = GetVertID(WID);
+		int32 SectionID = VertIDToSectionID[VertID];
+		checkSlow(SectionID != -1);
+		const FSkelMeshSection& Section = Mesh->Sections[SectionID];
+		const int32 BaseVertexIndex = (int32)Section.BaseVertexIndex;
+		return Section.SoftVertices[VertID - BaseVertexIndex];
+	}
+
+	FSkeletalMeshLODModelWrapper();
+	FSkeletalMeshLODModelWrapper(const FSkeletalMeshLODModelWrapper&);
+
+	TArray<int32> EmptyArray;
+	
+	// Convert the WedgeID the ID of the corresponding position vertex.
+	inline const VertIDType GetVertID(int32 WedgeID) const
+	{
+		return Mesh->IndexBuffer[WedgeID];
+	}
+
+private:
+
+	bool   bIncludeDisabledSections;
+	TArray<VertIDType> VertIDToSectionID;
+	TArray<TriIDType> TriIDToSectionID;
+	TArray<TriIDType> TriIDs;
+	TArray<VertIDType> VertIDs;
+	
+	const  FSkeletalMeshLODModel* Mesh;
+};
+
+
+
+
+
+void FSkeletalMeshLODModelToDynamicMesh::Convert(const  FSkeletalMeshLODModel* MeshIn, FDynamicMesh3& MeshOut, bool bCopyTangents)
+{
+
+	
+
+	const bool bIncludeDisabledSections = true;
+	FSkeletalMeshLODModelWrapper ModelWrapper(*MeshIn, bIncludeDisabledSections);
+
+	if (bPrintDebugMessages)
+	{
+		UE_LOG(LogTemp, Warning, TEXT("FSkeletalMeshLODModelToDynamicMesh:FSkeletalMeshLODModel verts %d  instances %d"), ModelWrapper.NumVerts(), 3*ModelWrapper.NumTris());
+	}
+
+	// Making default GroupIDs for the mesh
+	auto TriToGroupID = [&ModelWrapper](const int32& SrcTriID)->int32 {return (ModelWrapper.GetMaterialIndex(SrcTriID) + 1); };
+	
+	// Actual conversion.
+	TToDynamicMesh<FSkeletalMeshLODModelWrapper> SkeletalToDynamicMesh;
+	if (bDisableAttributes)
+	{ 
+		SkeletalToDynamicMesh.ConvertWOAttributes(MeshOut, ModelWrapper, TriToGroupID );
+	}
+	else
+	{ 
+		auto TriToMaterialID = [&ModelWrapper](const int32& SrcTriID)->int32 { return ModelWrapper.GetMaterialIndex(SrcTriID); };
+		SkeletalToDynamicMesh.Convert(MeshOut, ModelWrapper, TriToGroupID, TriToMaterialID, bCopyTangents);						
+	}
+
+	// Special code for vertex colors.  Currently DynamicMesh3 doesn't use an overlay for vertex colors 
+	// thus we do not support per-triangle-vertex colors. 
+	// NB: This just uses last-seen wedge color to define the vertex color.
+	// NB: The "alpha" channel is dropped.
+	if (bEnableOutputVertexColors)
+	{
+		MeshOut.EnableVertexColors(FVector3f::One());
+		bool bFoundNonDefaultVertexColor = false;
+		for (int32 TriangleID : MeshOut.TriangleIndicesItr())
+		{
+			int32 SrcWIDs[3];
+			ModelWrapper.GetWedgeIDs(SkeletalToDynamicMesh.ToSrcTriIDMap[TriangleID], SrcWIDs[0], SrcWIDs[1], SrcWIDs[2]);
+			FIndex3i VIDs = MeshOut.GetTriangle(TriangleID);
+			for (int32 j = 0; j < 3; ++j)
+			{ 
+				FLinearColor WedgeColor = ModelWrapper.GetWedgeColor(SrcWIDs[j]);
+				FVector3f WedgeColor3(WedgeColor); // Color.A is lost here.
+				bFoundNonDefaultVertexColor |= (WedgeColor3 != FVector3f::One());
+				MeshOut.SetVertexColor(VIDs[j], WedgeColor3);
+			}
+		}
+		if (bFoundNonDefaultVertexColor == false)
+		{
+			MeshOut.DiscardVertexColors();
+		}
+	}
+
+	if (!bEnableOutputGroups)
+	{
+		MeshOut.DiscardTriangleGroups();
+	}
+	FDateTime Time_AfterAttribs = FDateTime::Now();
+
+	// move maps to the calling class or let them get destroyed when SkeletalToDynamicMesh goes out of scope
+	if (bCalculateMaps)
+	{
+		Swap(TriIDMap, SkeletalToDynamicMesh.ToSrcTriIDMap);
+		Swap(VertIDMap, SkeletalToDynamicMesh.ToSrcVertIDMap);
+	}
+
+	
+	if (bPrintDebugMessages)
+	{
+		int NumUVLayers = ModelWrapper.NumUVLayers();
+		UE_LOG(LogTemp, Warning, TEXT("FSkeletalMeshLODModelToDynamicMesh:  Conversion Timing: Triangles %fs   Attributes %fs"),
+			(SkeletalToDynamicMesh.Time_AfterTriangles - SkeletalToDynamicMesh.Time_AfterVertices).GetTotalSeconds(), 
+			(Time_AfterAttribs   - SkeletalToDynamicMesh.Time_AfterTriangles).GetTotalSeconds());
+
+		
+		int NumUVs = (MeshOut.HasAttributes() && NumUVLayers > 0) ? MeshOut.Attributes()->PrimaryUV()->MaxElementID() : 0;
+		int NumNormals = 0; 
+		if (MeshOut.HasAttributes())
+		{
+			(MeshOut.Attributes()->PrimaryNormals() != nullptr) ? MeshOut.Attributes()->PrimaryNormals()->MaxElementID() : 0;
+		}
+
+		UE_LOG(LogTemp, Warning, TEXT("FSkeletalMeshLODModelToDynamicMesh:  FDynamicMesh verts %d triangles %d (primary) uvs %d normals %d"), MeshOut.MaxVertexID(), MeshOut.MaxTriangleID(), NumUVs, NumNormals);
+	}
+			
+}
+
+#else
+
+void  FSkeletalMeshLODModelToDynamicMesh::Convert(const  FSkeletalMeshLODModel* MeshIn, FDynamicMesh3& MeshOut, bool bCopyTangents)
+{
+	// Conversion only supported with editor.
+	check(0);
+}
+
+#endif  // end with editor
\ No newline at end of file
diff --git a/Engine/Plugins/Experimental/GeometryProcessing/Source/MeshConversion/Private/ToDynamicMesh.h b/Engine/Plugins/Experimental/GeometryProcessing/Source/MeshConversion/Private/ToDynamicMesh.h
new file mode 100644
--- /dev/null
+++ b/Engine/Plugins/Experimental/GeometryProcessing/Source/MeshConversion/Private/ToDynamicMesh.h
@@ -0,0 +1,649 @@
+// Copyright Epic Games, Inc. All Rights Reserved.
+
+#pragma once
+
+#include "CoreMinimal.h"
+#include "DynamicMesh3.h"
+#include "DynamicMeshAttributeSet.h"
+#include "DynamicMeshOverlay.h"
+#include "VectorTypes.h"
+#include "Async/Async.h"
+
+namespace UE
+{
+namespace Geometry
+{ 
+
+
+/**
+*
+* Class used to convert a mesh without attributes (e.g. normals, uvs etc) to a FDynamicMesh3
+*
+* The Source Mesh has to implement this interface
+*
+*	struct FSrcMeshInterface
+*	{
+*       // ID types:  must be castable to int32
+*		typedef SrcVertIDType    VertIDType;
+*		typedef SrcTriIDType     TriIDType;
+*
+*		// accounting.
+*		int32 NumTris() const;
+*		int32 NumVerts() const;
+*
+*	    // --"Vertex Buffer" info
+*		const Iterable_VertIDType& GetVertIDs() const;
+*		const FVector GetPosition(const VertIDType VtxID) const;
+*
+*       // --"Index Buffer" info
+*		const Iterable_TriIDType& GetTriIDs() const
+*		// return false if this TriID is not contained in mesh.
+*		bool GetTri(const TriIDType TriID, VertIDType& VID0, VertIDType& VID1, VertIDType& VID2) const;
+*	};
+*
+*/
+template<typename SrcMeshType>
+class TToDynamicMeshBase
+{
+public:
+
+	using SrcVertIDType = typename SrcMeshType::VertIDType;
+	using SrcTriIDType = typename SrcMeshType::TriIDType;
+
+
+	TToDynamicMeshBase()
+	{
+	}
+
+
+	// Mapping data captured by conversion
+	TArray<SrcVertIDType> ToSrcVertIDMap;
+	TArray<SrcTriIDType>  ToSrcTriIDMap;
+
+	FDateTime Time_AfterVertices;
+	FDateTime Time_AfterTriangles;
+
+	// Produces a DynamicMesh3 w/o attributes from the MeshIn.
+	void Convert(FDynamicMesh3& MeshOut, const SrcMeshType& MeshIn, TFunctionRef<int32(const SrcTriIDType& SrcTrID)> GroupIDFunction)
+	{
+		MeshOut.Clear();
+		MeshOut.DiscardAttributes();
+
+		MeshOut.EnableTriangleGroups(0);
+
+		ToSrcVertIDMap.Reset();
+		ToSrcTriIDMap.Reset();
+
+		const int32 NumSrcVerts = MeshIn.NumVerts();
+		const int32 NumSrcTris = MeshIn.NumTris();
+
+		if (NumSrcVerts == 0 || NumSrcTris == 0)
+		{
+			return;
+		}
+		// allocate for map data
+		
+		ToSrcVertIDMap.AddUninitialized(NumSrcVerts);
+		ToSrcTriIDMap.AddUninitialized(NumSrcTris);
+		// copy vertex positions. Later we may have to append duplicate vertices to resolve non-manifold structures.
+		int32 MaxSrcVertID = -1;
+		for (const SrcVertIDType& SrcVertID : MeshIn.GetVertIDs())
+		{
+			const FVector3d Position = MeshIn.GetPosition(SrcVertID);
+			const int32 DstVertID = MeshOut.AppendVertex(Position);
+			ToSrcVertIDMap[DstVertID] = SrcVertID;
+
+			MaxSrcVertID = FMath::Max((int32)SrcVertID, MaxSrcVertID);
+		}
+		// map the other direction
+		TArray<int32> FromSrcVertIDMap;
+		FromSrcVertIDMap.AddUninitialized(MaxSrcVertID+1);
+		for (int32 i = 0; i < ToSrcVertIDMap.Num(); ++i)
+		{
+			int32 SrcVertID = (int32)ToSrcVertIDMap[i];
+			FromSrcVertIDMap[SrcVertID] = i;
+		}
+
+		Time_AfterVertices = FDateTime::Now();
+
+		// copy the index buffer.  Note, this may add vertices if a non-manifold state has to be resolved.
+		for (const SrcTriIDType& SrcTriID : MeshIn.GetTriIDs())
+		{
+
+			const int32 DstGroupID = GroupIDFunction(SrcTriID);
+
+			SrcVertIDType SrcTriVerts[3];
+			MeshIn.GetTri(SrcTriID, SrcTriVerts[0], SrcTriVerts[1], SrcTriVerts[2]);
+	
+			FIndex3i DstTriVerts(FromSrcVertIDMap[SrcTriVerts[0]],
+				                 FromSrcVertIDMap[SrcTriVerts[1]],
+				                 FromSrcVertIDMap[SrcTriVerts[2]]);
+
+			// attemp to add tri. this may fail if DynamicMesh is more topologically restrictive than the src mesh.
+			int32 DstTriangleID = MeshOut.AppendTriangle(DstTriVerts, DstGroupID);
+
+			//-- already seen this triangle for some reason.. or the src mesh had a degenerate tri
+			if (DstTriangleID == FDynamicMesh3::DuplicateTriangleID || DstTriangleID == FDynamicMesh3::InvalidID)
+			{
+				continue;
+			}
+
+			//-- non manifold 
+			// if append failed due to non-manifold, duplicate verts
+			if (DstTriangleID == FDynamicMesh3::NonManifoldID)
+			{
+				// determine which verts need to be duplicated
+				bool bDuplicate[3] = { false, false, false };
+				for (int i = 0; i < 3; ++i)
+				{
+					int ii = (i + 1) % 3;
+					int EID = MeshOut.FindEdge(DstTriVerts[i], DstTriVerts[ii]);
+					if (EID != FDynamicMesh3::InvalidID && MeshOut.IsBoundaryEdge(EID) == false)
+					{
+						bDuplicate[i] = true;
+						bDuplicate[ii] = true;
+					}
+				}
+				for (int i = 0; i < 3; ++i)
+				{
+					if (bDuplicate[i])
+					{
+						const FVector3d Position = MeshOut.GetVertex(DstTriVerts[i]);
+						const int32 NewDstVertID = MeshOut.AppendVertex(Position);
+						DstTriVerts[i] = NewDstVertID;
+						// grow the map to make room
+						ToSrcVertIDMap.SetNumUninitialized(NewDstVertID + 1);
+
+						ToSrcVertIDMap[NewDstVertID] = DstTriVerts[i]; 
+					}
+				}
+
+				// add the fixed tri
+				DstTriangleID = MeshOut.AppendTriangle(DstTriVerts, DstGroupID);
+				checkSlow(DstTriangleID != FDynamicMesh3::NonManifoldID);
+			}
+
+			// record in map
+			ToSrcTriIDMap[DstTriangleID] = SrcTriID;
+		}
+
+		const int32 MaxDstTriID = MeshOut.MaxTriangleID(); // really MaxTriID+1
+		// if the source mesh had duplicates then TriIDMap will be too long, this will trim excess
+		ToSrcTriIDMap.SetNum(MaxDstTriID);
+		Time_AfterTriangles = FDateTime::Now();
+	}
+};
+
+
+// Used for exact attribute value welding.
+template <typename AttrType>
+struct TVertexAttr
+{
+	int VID;
+	AttrType AttrValue;
+	bool operator==(const TVertexAttr& o) const
+	{
+		return VID == o.VID && AttrValue == o.AttrValue;
+	}
+};
+
+template <typename AttrType>
+FORCEINLINE uint32 GetTypeHash(const TVertexAttr<AttrType>& Vector)
+{
+	// ugh copied from FVector clearly should not be using CRC for hash!!
+	return FCrc::MemCrc32(&Vector, sizeof(Vector));
+}
+
+template <typename AttrType> struct TOverlayTraits {};
+template<> struct TOverlayTraits<FVector3f> { typedef FDynamicMeshNormalOverlay OverlayType; };
+template<> struct TOverlayTraits<FVector2f> { typedef FDynamicMeshUVOverlay OverlayType; };
+
+// Welder used for exact attribute value welding when constructing overlay
+template <typename AttrType>
+class TAttrWelder
+{
+public:
+	using OverlayType = typename TOverlayTraits<AttrType>::OverlayType;
+	using FVertexAttr = TVertexAttr<AttrType>;
+
+	TMap<FVertexAttr, int> UniqueVertexAttrs;
+	OverlayType* Overlay;
+
+	TAttrWelder() : Overlay(nullptr){}
+
+	TAttrWelder(OverlayType* OverlayIn)
+	{
+		check(OverlayIn);
+		Overlay = OverlayIn;
+	}
+
+	int FindOrAddUnique(const AttrType& AttrValue, int VertexID)
+	{
+		FVertexAttr VertAttr = { VertexID, AttrValue };
+
+		const int32* FoundIndex = UniqueVertexAttrs.Find(VertAttr);
+		if (FoundIndex != nullptr)
+		{
+			return *FoundIndex;
+		}
+
+		int32 NewIndex = Overlay->AppendElement(AttrValue);
+		UniqueVertexAttrs.Add(VertAttr, NewIndex);
+		return NewIndex;
+	}
+};
+
+/**
+*
+* Class used to convert a mesh with attributes (e.g. normals, uvs etc) to a FDynamicMesh3
+*
+* The Source Mesh has to implement this interface
+*
+*	struct FSrcMeshInterface
+*	{
+*       // ID types:  must be castable to int32
+*		typedef SrcVertIDType    VertIDType;
+*		typedef SrcTriIDType     TriIDType;
+*		typedef SrcUVIDType      UVIDType;
+*		typedef SrcNormalIDType  NormalIDType;
+*       typedef SrcWedgeIDType   WedgeIDType;
+*
+*		// accounting.
+*		int32 NumTris() const;
+*		int32 NumVerts() const;
+*       int32 NumUVLayers() const;
+* 
+*	    // --"Vertex Buffer" info 
+*		const Iterable_VertIDType& GetVertIDs() const;
+*		const FVector3d GetPosition(const VertIDType VtxID) const;
+*
+*       // --"Index Buffer" info
+*		const Iterable_TriIDType& GetTriIDs() const
+*		// return false if this TriID is not contained in mesh.
+*		bool GetTri(const TriIDType TriID, VertIDType& VID0, VertIDType& VID1, VertIDType& VID2) const;
+* 
+*		bool HasNormals() const;
+*		bool HasTangents() cosnt;
+*		bool HasBiTangents() const;
+*	
+*       // Each triangle corner is a wedge
+*       void GetWedgeID(const TriIDType& TriID, WedgeIDType& WID0, WedgeIDType& WID1, WedgeIDType& WID2) const;
+*		
+*    	// attribute access per-wedge
+*		// NB:  ToDynamicMesh will attempt to weld identical attributes that are associated with the same vertex
+*		FVector2f GetWedgeUV(int32 UVLayerIndex, WedgeIDType WID) const;
+*		FVector3f GetWedgeNormal(WedgeIDType WID) const;
+*		FVector3f GetWedgeTangent(WedgeIDType WID) const;
+*		FVector3f GetWedgeBiTangent(WedgeIDType WID) const;
+* 
+*		// attribute access that exploits shared attributes. 
+*		// each group of shared attributes presents itself as a mesh with its own attribute vertex buffer.
+*		// NB:  If the mesh has no shared Attr attributes, then Get{Attr}IDs() should return an empty array.
+*       // NB:  Get{Attr}Tri() functions should return false if the triangle is not set in the attribute mesh. 
+*       
+*       const TArray<UVIDType>& GetUVIDs(int32 LayerID) const;
+*		FVector2f GetUV(int32 LayerID, UVIDType UVID) const;
+*		bool GetUVTri(int32 LayerID, const TriIDType& TID, UVIDType& ID0, UVIDType& ID1, UVIDType& ID2) const;
+*
+*       const TArray<NormalIDType>& GetNormalIDs() const;
+*		FVector3f GetNormal(NormalIDType ID) const;
+*		bool GetNormalTri(const TriIDType& TID, NormalIDType& NID0, NormalIDType& NID1, NormalIDType& NID2) const;
+*
+*		const TArray<NormalIDType>& GetTangentIDs() const;
+* 		FVector3f GetTangent(NormalIDType ID) const;
+*		bool GetTangentTri(const TriIDType& TID, NormalIDType& NID0, NormalIDType& NID1, NormalIDType& NID2) const;
+* 
+* 		const TArray<NormalIDType>& GetBiTangentIDs() const;
+*		FVector3f GetBiTangent(NormalIDType ID) const;
+*       bool GetBiTangentTri(const TriIDType& TID, NormalIDType& NID0, NormalIDType& NID1, NormalIDType& NID2) const;
+* 
+*	};
+*
+*/
+template <typename SrcMeshType>
+class TToDynamicMesh : public TToDynamicMeshBase<SrcMeshType>
+{
+public:
+	typedef TToDynamicMeshBase<SrcMeshType>  MyBase;
+	using SrcTriIDType    = typename MyBase::SrcTriIDType;
+	using SrcVertIDType   = typename MyBase::SrcVertIDType;
+	using SrcUVIDType     = typename SrcMeshType::UVIDType;
+	using SrcNormalIDType = typename SrcMeshType::NormalIDType;
+	using SrcWedgeIDType  = typename SrcMeshType::WedgeIDType;
+
+	TToDynamicMesh() :MyBase() {}
+
+
+
+	// Convert To FDynamicMesh w/o attributes
+	void ConvertWOAttributes(FDynamicMesh3& MeshOut,
+		const SrcMeshType& MeshIn,
+		TFunctionRef<int32(const SrcTriIDType& SrcTrID)> GroupIDFunction)
+	{
+		MyBase::Convert(MeshOut, MeshIn, GroupIDFunction);
+	}
+
+
+	// Convert To FDynamicMesh.  Will Copy GroupID to the mesh, and will create overlays with UVs, Normal, MaterialID (additionally Tangents and BiTangents if requested)
+	void Convert(FDynamicMesh3& MeshOut,
+		const SrcMeshType& MeshIn,
+		TFunctionRef<int32(const SrcTriIDType& SrcTrID)> GroupIDFunction,
+		TFunctionRef<int32(const SrcTriIDType& SrcTrID)> MaterialIDFunction,
+		bool bCopyTangents)
+	{
+
+		MyBase::Convert(MeshOut, MeshIn, GroupIDFunction); // convert the mesh.  Must call before populate overlays
+		PopulateOverlays(MeshOut, MeshIn, MaterialIDFunction, bCopyTangents); // convert the attributes
+	}
+
+protected:
+
+
+	// Populates overlays for UVs, Normal, Material ID. Also Tangent and BiTangent if bCopyTangents == true.  
+	// NB: This does not copy any polygroup information.  
+	void PopulateOverlays(FDynamicMesh3& MeshOut, const SrcMeshType& MeshIn, TFunctionRef<int32(const SrcTriIDType& SrcTrID)> MaterialIDFunction, bool bCopyTangents)
+	{
+		// Attributes we want to transfer
+		FDynamicMeshNormalOverlay* NormalOverlay = nullptr;
+		FDynamicMeshNormalOverlay* TangentOverlay = nullptr;
+		FDynamicMeshNormalOverlay* BiTangentOverlay = nullptr;
+		FDynamicMeshMaterialAttribute* MaterialIDAttrib = nullptr;
+
+		// only copy tangents if they exist
+		bCopyTangents = bCopyTangents && MeshIn.HasTangents() && MeshIn.HasBiTangents();
+
+		// -- Create Overlays
+		const int32 NumUVLayers = MeshIn.NumUVLayers();
+
+		MeshOut.EnableAttributes(); // by default 1-UV layer and 1-normal layer
+
+		// create any addition Normal overlays and reserve space.
+		int32 NumRequiredNormalLayers = (bCopyTangents) ? 3 : 1;
+		if (MeshOut.Attributes()->NumNormalLayers() < NumRequiredNormalLayers)
+		{
+			MeshOut.Attributes()->SetNumNormalLayers(NumRequiredNormalLayers);
+		}
+		for (int32 i = 0; i < NumRequiredNormalLayers; ++i)
+		{
+			MeshOut.Attributes()->GetNormalLayer(i)->InitializeTriangles(MeshOut.MaxTriangleID());
+		}
+
+		NormalOverlay = MeshOut.Attributes()->PrimaryNormals();
+		if (bCopyTangents)
+		{
+			TangentOverlay   = MeshOut.Attributes()->PrimaryTangents();
+			BiTangentOverlay = MeshOut.Attributes()->PrimaryBiTangents();
+		}
+
+		// create UV overlays 
+		MeshOut.Attributes()->SetNumUVLayers(NumUVLayers);
+		// reserve space in any new UV layers.
+		for (int32 i = 1; i < NumUVLayers; ++i)
+		{
+			MeshOut.Attributes()->GetUVLayer(i)->InitializeTriangles(MeshOut.MaxTriangleID());
+		}
+
+
+		// always enable Material ID if there are any attributes
+		MeshOut.Attributes()->EnableMaterialID();
+		MaterialIDAttrib = MeshOut.Attributes()->GetMaterialID();
+
+		// do the conversions.
+		// we will populate all the attributes simultaneously, hold on to futures in this array and then Wait for them at the end
+		TArray<TFuture<void>> Pending;
+
+
+
+		// populate UV overlays
+		for (int UVLayerIndex = 0; UVLayerIndex < NumUVLayers; UVLayerIndex++)
+		{
+			auto UVFuture = Async(EAsyncExecution::ThreadPool, [&, UVLayerIndex]()
+				{
+					FDynamicMeshUVOverlay* Overlay = MeshOut.Attributes()->GetUVLayer(UVLayerIndex);
+
+					PopulateOverlay<FVector2f, SrcUVIDType>(
+						Overlay,
+						MeshIn.GetUVIDs(UVLayerIndex),
+						[&MeshIn, UVLayerIndex](const SrcUVIDType& UVID)->FVector2f { return MeshIn.GetUV(UVLayerIndex, UVID); },
+						[&MeshIn, UVLayerIndex](const SrcTriIDType& TriID, SrcUVIDType& UV0, SrcUVIDType& UV1, SrcUVIDType& UV2)->bool {return MeshIn.GetUVTri(UVLayerIndex, TriID, UV0, UV1, UV2); },
+						[&MeshIn, UVLayerIndex](const SrcWedgeIDType& WID)->FVector2f {return MeshIn.GetWedgeUV(UVLayerIndex, WID); },
+						[&MeshIn](const SrcTriIDType& TriID, SrcWedgeIDType& WID0, SrcWedgeIDType& WID1, SrcWedgeIDType& WID2)->void { MeshIn.GetWedgeIDs(TriID, WID0, WID1, WID2); }
+					);
+				});
+			Pending.Add(MoveTemp(UVFuture));
+		}
+		// populate Normal overlay
+		if (NormalOverlay != nullptr)
+		{
+			auto NormalFuture = Async(EAsyncExecution::ThreadPool, [&]()
+				{
+					PopulateOverlay<FVector3f, SrcNormalIDType>(
+						NormalOverlay,
+						MeshIn.GetNormalIDs(),
+						[&MeshIn](const SrcNormalIDType& NID)->FVector3f { return MeshIn.GetNormal(NID); },
+						[&MeshIn](const SrcTriIDType& TriID, SrcNormalIDType& N0, SrcNormalIDType& N1, SrcNormalIDType& N2)->bool {return MeshIn.GetNormalTri(TriID, N0, N1, N2); },
+						[&MeshIn](const SrcWedgeIDType& WID)->FVector3f {return MeshIn.GetWedgeNormal(WID); },
+						[&MeshIn](const SrcTriIDType& TriID, SrcWedgeIDType& WID0, SrcWedgeIDType& WID1, SrcWedgeIDType& WID2)->void { MeshIn.GetWedgeIDs(TriID, WID0, WID1, WID2); }
+					);
+				});
+			Pending.Add(MoveTemp(NormalFuture));
+		}
+		// populate Tangent overlay
+		if (TangentOverlay != nullptr)
+		{
+			auto TangentFuture = Async(EAsyncExecution::ThreadPool, [&]()
+				{
+					PopulateOverlay<FVector3f, SrcNormalIDType>(
+						TangentOverlay,
+						MeshIn.GetTangentIDs(),
+						[&MeshIn](const SrcNormalIDType& NID)->FVector3f { return MeshIn.GetTangent(NID); },
+						[&MeshIn](const SrcTriIDType& TriID, SrcNormalIDType& N0, SrcNormalIDType& N1, SrcNormalIDType& N2)->bool {return MeshIn.GetTangentTri(TriID, N0, N1, N2); },
+						[&MeshIn](const SrcWedgeIDType& WID)->FVector3f {return MeshIn.GetWedgeTangent(WID); },
+						[&MeshIn](const SrcTriIDType& TriID, SrcWedgeIDType& WID0, SrcWedgeIDType& WID1, SrcWedgeIDType& WID2)->void { MeshIn.GetWedgeIDs(TriID, WID0, WID1, WID2); }
+					);
+				});
+			Pending.Add(MoveTemp(TangentFuture));
+		}
+		//populate BiTangent overlay
+		if (BiTangentOverlay != nullptr)
+		{
+			auto BiTangentFuture = Async(EAsyncExecution::ThreadPool, [&]()
+				{
+					PopulateOverlay<FVector3f, SrcNormalIDType>(
+						BiTangentOverlay,
+						MeshIn.GetTangentIDs(),
+						[&MeshIn](const SrcNormalIDType& NID)->FVector3f { return MeshIn.GetBiTangent(NID); },
+						[&MeshIn](const SrcTriIDType& TriID, SrcNormalIDType& N0, SrcNormalIDType& N1, SrcNormalIDType& N2)->bool {return MeshIn.GetBiTangentTri(TriID, N0, N1, N2); },
+						[&MeshIn](const SrcWedgeIDType& WID)->FVector3f {return MeshIn.GetWedgeBiTangent(WID); },
+						[&MeshIn](const SrcTriIDType& TriID, SrcWedgeIDType& WID0, SrcWedgeIDType& WID1, SrcWedgeIDType& WID2)->void { MeshIn.GetWedgeIDs(TriID, WID0, WID1, WID2); }
+					);
+
+				});
+			Pending.Add(MoveTemp(BiTangentFuture));
+		}
+		//populate MaterialID overlay
+		if (MaterialIDAttrib != nullptr)
+		{
+			auto MaterialFuture = Async(EAsyncExecution::ThreadPool, [&]()
+				{
+					for (int32 TriangleID : MeshOut.TriangleIndicesItr())
+					{
+						SrcTriIDType SrcTriID = MyBase::ToSrcTriIDMap[TriangleID];
+						MaterialIDAttrib->SetValue(TriangleID, MaterialIDFunction(SrcTriID));
+					}
+				});
+			Pending.Add(MoveTemp(MaterialFuture));
+		}
+
+		// wait for all work to be done
+		for (TFuture<void>& Future : Pending)
+		{
+			Future.Wait();
+		}
+	}
+
+
+	// Access the attributes in the Source Mesh on a per-wedge granularity (each corner of each triangle is a wedge), and weld them based on strict equality when creating the overlay.
+	template<typename AttrType>
+	void PopulateOverlayFromWedgeAttr(typename TOverlayTraits<AttrType>::OverlayType* Overlay,
+		TFunctionRef<AttrType(const SrcWedgeIDType& WedgeID)> WedgeAttrs,
+		TFunctionRef<void(const SrcTriIDType& SrcTriID, SrcWedgeIDType& WID0, SrcWedgeIDType& WID1, SrcWedgeIDType& WID2)> TriToWedges,
+		bool bWeldIdenticalAttrs = true)
+	{
+
+		FDynamicMesh3& MeshOut = *Overlay->GetParentMesh();
+		TAttrWelder<AttrType> AttrWelder(Overlay);
+
+		for (int32 TriangleID : MeshOut.TriangleIndicesItr())
+		{
+			// skip overlay triangles that a prior method has populated.
+			// NB: this function may have been called after PopulateOverlayFromSharedAttr
+			if (Overlay->IsSetTriangle(TriangleID))
+			{
+				continue;
+			}
+
+			const SrcTriIDType& SrcTriID = MyBase::ToSrcTriIDMap[TriangleID];
+
+			SrcWedgeIDType WID[3];
+			TriToWedges(SrcTriID, WID[0], WID[1], WID[2]);
+
+			FIndex3i Tri = MeshOut.GetTriangle(TriangleID);
+			FIndex3i AttrTri;
+			for (int j = 0; j < 3; ++j)
+			{
+				const AttrType& AttrValue = WedgeAttrs(WID[j]);
+				if (bWeldIdenticalAttrs)
+				{
+					AttrTri[j] = AttrWelder.FindOrAddUnique(AttrValue, Tri[j]);
+				}
+				else
+				{
+					AttrTri[j] = Overlay->AppendElement(AttrValue);
+				}
+			}
+			Overlay->SetTriangle(TriangleID, AttrTri);
+		}
+	}
+
+
+	template <typename AttrType, typename SrcAttrIDType>
+	void PopulateOverlayFromSharedAttr(typename TOverlayTraits<AttrType>::OverlayType* Overlay,
+		const TArray<SrcAttrIDType>& SrcAttrIDs,
+		TFunctionRef<AttrType(const SrcAttrIDType& AttrID)>  SharedAttrValues,
+		TFunctionRef<bool(const SrcTriIDType& SrcTriID, SrcAttrIDType& AttrID0, SrcAttrIDType& AttrID1, SrcAttrIDType& AttrID2)> AttrTriIndices)
+	{
+		FDynamicMesh3& MeshOut = *Overlay->GetParentMesh();
+
+		// return if no shared attributes.
+		if (SrcAttrIDs.Num() == 0)
+		{
+			return;
+		}
+
+		
+		int32 MaxSrcAttrID = -1;
+		for (const SrcAttrIDType& SrcAttrID : SrcAttrIDs)
+		{
+			MaxSrcAttrID = FMath::Max(MaxSrcAttrID, (int32)SrcAttrID);
+		}
+		TArray<int32> FromSrcAttrIDMap;
+		FromSrcAttrIDMap.Init(-1, MaxSrcAttrID+1);
+		// copy attr values into the overlay - these are vertices in the attr mesh.
+		for (const SrcAttrIDType& SrcAttrID : SrcAttrIDs)
+		{
+			const AttrType AttrValue = SharedAttrValues(SrcAttrID);
+			const int32 NewIndex = Overlay->AppendElement(AttrValue);
+			int32 SrcAttrID32 = (int32)SrcAttrID;
+			FromSrcAttrIDMap[SrcAttrID32] = NewIndex;
+		}
+
+		// use the attr index buffer to make triangles in the overlay
+		for (int32 TriID : MeshOut.TriangleIndicesItr())
+		{
+			const SrcTriIDType SrcTriID = MyBase::ToSrcTriIDMap[TriID];
+			SrcAttrIDType SrcAttrTri[3];
+			bool bHasValidAttrTri = AttrTriIndices(SrcTriID, SrcAttrTri[0], SrcAttrTri[1], SrcAttrTri[2]);
+
+			if (!bHasValidAttrTri)
+			{
+				// don't set this tri in the overlay
+				continue;
+			}
+
+			// translate to Overlay indicies 
+			FIndex3i AttrTri(FromSrcAttrIDMap[(int32)SrcAttrTri[0]],
+				FromSrcAttrIDMap[(int32)SrcAttrTri[1]],
+				FromSrcAttrIDMap[(int32)SrcAttrTri[2]]);
+
+			///--  We have to do some clean-up on the shared Attrs because the MeshIn format might support wild stuff.. --///
+			{
+				// MeshIn may attach multiple mesh vertices to the same attribute element.  DynamicMesh does not.
+				// if we have already used this element for a different mesh vertex, split it.
+				const FIndex3i ParentTriangle = MeshOut.GetTriangle(TriID);
+				for (int i = 0; i < 3; ++i)
+				{
+					int32 ParentVID = Overlay->GetParentVertex(AttrTri[i]);
+					if (ParentVID != FDynamicMesh3::InvalidID && ParentVID != ParentTriangle[i])
+					{
+						const AttrType AttrValue = SharedAttrValues(AttrTri[i]);
+						AttrTri[i] = Overlay->AppendElement(AttrValue);
+					}
+				}
+
+				// MeshIn may have degenerate attr tris.  Dynamic Mesh does not.
+				// if the attr tri is degenerate we split the degenerate attr edge by adding two new attrs
+				// in its place, or if it is totally degenerate we add 3 new attrs
+
+				if (AttrTri[0] == AttrTri[1] && AttrTri[0] == AttrTri[2])
+				{
+					const AttrType AttrValue = SharedAttrValues(AttrTri[0]);
+					AttrTri[0] = Overlay->AppendElement(AttrValue);
+					AttrTri[1] = Overlay->AppendElement(AttrValue);
+					AttrTri[2] = Overlay->AppendElement(AttrValue);
+
+				}
+				else
+				{
+					if (AttrTri[0] == AttrTri[1])
+					{
+						const AttrType AttrValue = SharedAttrValues(AttrTri[0]);
+						AttrTri[0] = Overlay->AppendElement(AttrValue);
+						AttrTri[1] = Overlay->AppendElement(AttrValue);
+					}
+					if (AttrTri[0] == AttrTri[2])
+					{
+						const AttrType AttrValue = SharedAttrValues(AttrTri[0]);
+						AttrTri[0] = Overlay->AppendElement(AttrValue);
+						AttrTri[2] = Overlay->AppendElement(AttrValue);
+					}
+					if (AttrTri[1] == AttrTri[2])
+					{
+						const AttrType AttrValue = SharedAttrValues(AttrTri[1]);
+						AttrTri[1] = Overlay->AppendElement(AttrValue);
+						AttrTri[2] = Overlay->AppendElement(AttrValue);
+					}
+				}
+			}
+
+			// set the triangle in the overlay
+			Overlay->SetTriangle(TriID, AttrTri);
+
+		}
+	}
+
+	template <typename AttrType, typename SrcAttrIDType>
+	void PopulateOverlay(typename TOverlayTraits<AttrType>::OverlayType* Overlay,
+		const TArray<SrcAttrIDType>& SrcAttrIDs,
+		TFunctionRef<AttrType(const SrcAttrIDType& AttrID)>  SharedAttrValues,
+		TFunctionRef<bool(const SrcTriIDType& SrcTriID, SrcAttrIDType& AttrID0, SrcAttrIDType& AttrID1, SrcAttrIDType& AttrID2)> AttrTriIndices,
+		TFunctionRef<AttrType(const SrcWedgeIDType& WegdeID)> WedgeAttrs,
+		TFunctionRef<void(const SrcTriIDType& SrcTriID, SrcWedgeIDType& WID0, SrcWedgeIDType& WID1, SrcWedgeIDType& WID2)> TriToWedges)
+	{
+		PopulateOverlayFromSharedAttr(Overlay, SrcAttrIDs, SharedAttrValues, AttrTriIndices);
+		// Populate the overlay tris that were left empty by the shared attrs
+		PopulateOverlayFromWedgeAttr(Overlay, WedgeAttrs, TriToWedges);
+	}
+};
+
+} } // end namespace UE::Geometry
\ No newline at end of file
diff --git a/Engine/Plugins/Experimental/GeometryProcessing/Source/MeshConversion/Public/SkeletalMeshLODModelToDynamicMesh.h b/Engine/Plugins/Experimental/GeometryProcessing/Source/MeshConversion/Public/SkeletalMeshLODModelToDynamicMesh.h
new file mode 100644
--- /dev/null
+++ b/Engine/Plugins/Experimental/GeometryProcessing/Source/MeshConversion/Public/SkeletalMeshLODModelToDynamicMesh.h
@@ -0,0 +1,55 @@
+// Copyright Epic Games, Inc. All Rights Reserved.
+
+#pragma once
+
+#include "CoreMinimal.h"
+#include "DynamicMesh3.h"
+
+// forward 
+class FSkeletalMeshLODModel;
+
+using UE::Geometry::FDynamicMesh3;
+
+/**
+ * Convert FSkeletalMeshLODModel to FDynamicMesh3
+ *
+ */
+class MESHCONVERSION_API FSkeletalMeshLODModelToDynamicMesh
+{
+public:
+	/** If true, will print some possibly-helpful debugging spew to output log */
+	bool bPrintDebugMessages = false;
+
+	/** Should we initialize triangle groups on output mesh */
+	bool bEnableOutputGroups = true;
+
+	/** Should we initialize vertex colors on output mesh */
+	bool bEnableOutputVertexColors = true;
+
+	/** Should we calculate conversion index maps */
+	bool bCalculateMaps = true;
+
+	/** Ignore all mesh attributes (e.g. UV/Normal layers, material groups) */
+	bool bDisableAttributes = false;
+
+
+	/** map from DynamicMesh triangle ID to FSkeletalMeshLODModel TriIdx*/
+	TArray<int32> TriIDMap;
+
+	/**
+	* map from DynamicMesh vertex Id to  FSkeletalMeshLODModel FVertexID.
+	* NB: due to vertex splitting, multiple DynamicMesh vertex ids
+	* may map to the same  FSkeletalMeshLODModel FVertexID.
+	*  ( a vertex split is a result of reconciling non-manifold  FSkeletalMeshLODModel vertex )
+	*/
+	TArray<int32> VertIDMap;
+
+
+	/**
+	 * Default conversion of MeshDescription to DynamicMesh
+	 * @param bCopyTangents  - if bDisableAttributes is false, this requests the tangent plane vectors (tangent and bitangent)
+	 *                          be stored as overlays in the MeshOut DynamicAttributeSet, provided they exist on the MeshIn
+	 */
+	void Convert(const  FSkeletalMeshLODModel* MeshIn, FDynamicMesh3& MeshOut, bool bCopyTangents = false);
+
+};
\ No newline at end of file
