Add an optional FProgressCancel to FConvexHull3

#rnx

diff --git a/Engine/Plugins/Experimental/GeometryProcessing/Source/GeometryAlgorithms/Private/ConvexHull3.cpp b/Engine/Plugins/Experimental/GeometryProcessing/Source/GeometryAlgorithms/Private/ConvexHull3.cpp
--- a/Engine/Plugins/Experimental/GeometryProcessing/Source/GeometryAlgorithms/Private/ConvexHull3.cpp
+++ b/Engine/Plugins/Experimental/GeometryProcessing/Source/GeometryAlgorithms/Private/ConvexHull3.cpp
@@ -642,6 +642,10 @@ bool TConvexHull3<RealType>::Solve(int32 NumPoints, TFunctionRef<FVector3<RealTy
 	Connectivity.InitVisibility(Hull, NumPoints, GetPointFunc, FilterFunc);
 	while (true)
 	{
+		if (Progress && (NumHullPoints % 100) == 0 && Progress->Cancelled())
+		{
+			return false;
+		}
 		FIndex2i Visible = Connectivity.ChooseVisiblePoint();
 		if (Visible.A == -1)
 		{
diff --git a/Engine/Plugins/Experimental/GeometryProcessing/Source/GeometryAlgorithms/Public/ConvexHull3.h b/Engine/Plugins/Experimental/GeometryProcessing/Source/GeometryAlgorithms/Public/ConvexHull3.h
--- a/Engine/Plugins/Experimental/GeometryProcessing/Source/GeometryAlgorithms/Public/ConvexHull3.h
+++ b/Engine/Plugins/Experimental/GeometryProcessing/Source/GeometryAlgorithms/Public/ConvexHull3.h
@@ -8,6 +8,7 @@
 #include "LineTypes.h"
 #include "PlaneTypes.h"
 #include "HalfspaceTypes.h"
+#include "Util/ProgressCancel.h"
 
 namespace UE {
 namespace Geometry {
@@ -166,6 +167,10 @@ public:
 		return NumHullPoints;
 	}
 
+
+	/** Set this to be able to cancel running operation */
+	FProgressCancel* Progress = nullptr;
+
 protected:
 
 	int32 Dimension;
