Mesh Modeling : Mesh Conversions - Update Shared UVs on MeshDescription.  Disable uv indexing in the MeshDescription during the update.

diff --git a/Engine/Plugins/Experimental/GeometryProcessing/Source/MeshConversion/Private/DynamicMeshToMeshDescription.cpp b/Engine/Plugins/Experimental/GeometryProcessing/Source/MeshConversion/Private/DynamicMeshToMeshDescription.cpp
--- a/Engine/Plugins/Experimental/GeometryProcessing/Source/MeshConversion/Private/DynamicMeshToMeshDescription.cpp
+++ b/Engine/Plugins/Experimental/GeometryProcessing/Source/MeshConversion/Private/DynamicMeshToMeshDescription.cpp
@@ -110,6 +110,8 @@ void FDynamicMeshToMeshDescription::UpdateAttributes(const FDynamicMesh3* MeshIn
 		{
 			if (MeshIn->HasAttributes())
 			{
+				MeshOut.SuspendUVIndexing();
+
 				check(MeshIn->TriangleCount() == MeshOut.Triangles().Num())
 				int32 NumLayers = MeshIn->Attributes()->NumUVLayers();			
 				MeshOut.SetNumUVChannels(NumLayers); // This resets MeshDescription's internal TriangleUV array
@@ -154,6 +156,9 @@ void FDynamicMeshToMeshDescription::UpdateAttributes(const FDynamicMesh3* MeshIn
 						MeshOut.SetTriangleUVIndices(TriangleID, MDTri, UVLayerIndex);
 					}
 				}
+
+				MeshOut.ResumeUVIndexing();
+
 				if (0)
 				{ 
 					// Verify the shared UVs and per-vertexinstance UVs match
