ModelingTools: Add floating and spiral stair generators.

#rnx

diff --git a/Engine/Plugins/Experimental/GeometryProcessing/Source/GeometricObjects/Private/Generators/StairGenerator.cpp b/Engine/Plugins/Experimental/GeometryProcessing/Source/GeometricObjects/Private/Generators/StairGenerator.cpp
new file mode 100644
--- /dev/null
+++ b/Engine/Plugins/Experimental/GeometryProcessing/Source/GeometricObjects/Private/Generators/StairGenerator.cpp
@@ -0,0 +1,1304 @@
+// Copyright Epic Games, Inc. All Rights Reserved.
+
+#include "Generators/StairGenerator.h"
+
+using namespace UE::Geometry;
+
+/**
+ * FStairGenerator
+ */
+
+FMeshShapeGenerator& FStairGenerator::Generate()
+{
+	Reset();
+
+	switch (StairStyle)
+	{
+	case EStairStyle::Floating:
+		return GenerateFloatingStairs();
+	case EStairStyle::Solid:
+	default:
+		return GenerateSolidStairs();
+	}
+}
+
+/**
+ * Generate a solid stair mesh.
+ *
+ * Stair topology composition. (side view cross section)
+ *
+ * (EStairStyle::Solid)
+ *                  .___.
+ *                  |   |          | z (height)
+ *              .___|___|          |
+ *              |   |   |          |_____ x (depth)
+ *          .___|___|___|           \
+ *          |   |   |   |            \ y (width)
+ *  Row --> |___|___|___|
+ *
+ *          ^-- Column
+ *
+ */
+FMeshShapeGenerator& FStairGenerator::GenerateSolidStairs()
+{
+	auto TriangleNumber = [](int x)
+	{
+		return (x * (x + 1) / 2);
+	};
+
+	const int NumConnectQuads = 4 * NumSteps;
+	const int NumCornerVerts = 4 * (NumSteps + 1);
+	NumVertsPerSide = TriangleNumber(NumSteps + 1) + NumSteps;
+	NumVerts = 2 * NumVertsPerSide;
+	NumQuadsPerSide = TriangleNumber(NumSteps);
+	NumQuads = 2 * NumQuadsPerSide + NumConnectQuads;
+
+	const int NumSideQuadAttrs = NumVerts;
+	const int NumConnectQuadCornerAttrs = 2 * NumCornerVerts;
+	const int NumConnectQuadInteriorAttrs = 4 * (NumSteps - 1);
+	NumAttrs = NumConnectQuadCornerAttrs + NumConnectQuadInteriorAttrs + NumSideQuadAttrs;
+	SetBufferSizes(NumVerts, 2 * NumQuads, NumAttrs, NumAttrs);
+
+	// Ordered Side lists for iteration.
+	const ESide RightLeftSides[2] = { ESide::Right, ESide::Left };
+	const ESide AllSides[6] = { ESide::Right, ESide::Left, ESide::Front, ESide::Top, ESide::Back, ESide::Bottom };
+
+	// Generate vertices by vertical column per Right/Left side.
+	LeftSideColumnId = (NumSteps + 1);
+	VertexIds.SetNum(2 * LeftSideColumnId);
+	VertexIdsToColumnRow.SetNum(NumVerts);
+	int VertexId = 0;
+	for (const ESide& Side : RightLeftSides)
+	{
+		int StartColumnId = (Side == ESide::Right ? 0 : LeftSideColumnId);
+		for (int VertexColumn = 0; VertexColumn < NumSteps + 1; VertexColumn++)
+		{
+			int NumVertices = (VertexColumn == NumSteps ? VertexColumn + 1 : VertexColumn + 2);
+			VertexIds[StartColumnId + VertexColumn].SetNum(NumVertices);
+			for (int VertexRow = 0; VertexRow < NumVertices; VertexRow++)
+			{
+				Vertices[VertexId] = GenerateVertex(Side, VertexColumn, VertexRow);
+				VertexIds[StartColumnId + VertexColumn][VertexRow] = VertexId;
+				VertexIdsToColumnRow[VertexId] = FIndex2i(StartColumnId + VertexColumn, VertexRow);
+				VertexId++;
+			}
+		}
+	}
+
+	// Generate quad representation for each side.
+	//
+	// Compute these in respective groups for later normal,
+	// UV & polygroup assignment.
+	FaceDesc.SetNum(4 * NumQuads);
+	int FaceDescId = 0;
+	for (const ESide& Side : AllSides)
+	{
+		switch (Side)
+		{
+		case ESide::Right:
+		case ESide::Left:
+		{
+			const int StartColId = Side == ESide::Right ? 0 : LeftSideColumnId;
+			if (Side == ESide::Right)
+			{
+				RightStartFaceId = FaceDescId;
+			}
+			else
+			{
+				LeftStartFaceId = FaceDescId;
+			}
+			for (int Step = 0; Step < NumSteps; Step++)
+			{
+				int CurrentColId = StartColId + Step;
+				for (int VertexRow = 0; VertexRow < VertexIds[CurrentColId].Num() - 1; VertexRow++)
+				{
+					FaceDesc[FaceDescId++] = (VertexIds[CurrentColId][VertexRow]);
+					FaceDesc[FaceDescId++] = (VertexIds[CurrentColId][VertexRow + 1]);
+					FaceDesc[FaceDescId++] = (VertexIds[CurrentColId + 1][VertexRow + 1]);
+					FaceDesc[FaceDescId++] = (VertexIds[CurrentColId + 1][VertexRow]);
+				}
+			}
+			break;
+		}
+		case ESide::Front:
+		{
+			FrontStartFaceId = FaceDescId;
+			for (int Step = 0; Step < NumSteps; Step++)
+			{
+				int LastStepRow = VertexIds[Step].Num() - 1;
+				FaceDesc[FaceDescId++] = (VertexIds[Step + LeftSideColumnId][LastStepRow - 1]);
+				FaceDesc[FaceDescId++] = (VertexIds[Step + LeftSideColumnId][LastStepRow]);
+				FaceDesc[FaceDescId++] = (VertexIds[Step][LastStepRow]);
+				FaceDesc[FaceDescId++] = (VertexIds[Step][LastStepRow - 1]);
+			}
+			break;
+		}
+		case ESide::Top:
+		{
+			TopStartFaceId = FaceDescId;
+			for (int Step = 0; Step < NumSteps; Step++)
+			{
+				int LastStepRow = VertexIds[Step].Num() - 1;
+				FaceDesc[FaceDescId++] = (VertexIds[Step + LeftSideColumnId][LastStepRow]);
+				FaceDesc[FaceDescId++] = (VertexIds[Step + LeftSideColumnId + 1][LastStepRow]);
+				FaceDesc[FaceDescId++] = (VertexIds[Step + 1][LastStepRow]);
+				FaceDesc[FaceDescId++] = (VertexIds[Step][LastStepRow]);
+			}
+			break;
+		}
+		case ESide::Back:
+		{
+			BackStartFaceId = FaceDescId;
+			for (int Step = 0; Step < NumSteps; Step++)
+			{
+				FaceDesc[FaceDescId++] = (VertexIds[NumSteps][Step]);
+				FaceDesc[FaceDescId++] = (VertexIds[NumSteps][Step + 1]);
+				FaceDesc[FaceDescId++] = (VertexIds[LeftSideColumnId + NumSteps][Step + 1]);
+				FaceDesc[FaceDescId++] = (VertexIds[LeftSideColumnId + NumSteps][Step]);
+			}
+			break;
+		}
+		case ESide::Bottom:
+		{
+			BottomStartFaceId = FaceDescId;
+			for (int Step = 0; Step < NumSteps; Step++)
+			{
+				FaceDesc[FaceDescId++] = (VertexIds[Step][0]);
+				FaceDesc[FaceDescId++] = (VertexIds[Step + 1][0]);
+				FaceDesc[FaceDescId++] = (VertexIds[Step + 1 + LeftSideColumnId][0]);
+				FaceDesc[FaceDescId++] = (VertexIds[Step + LeftSideColumnId][0]);
+			}
+			break;
+		}
+		}
+	}
+	LastFaceId = FaceDescId;
+
+	// Compute Normals & UVs
+	NormalDesc.SetNum(4 * NumQuads);
+	int NormalId = 0;
+
+	UVDesc.SetNum(4 * NumQuads);
+	int UVId = 0;
+	float MaxDimension = GetMaxDimension();
+	float UVScale = (bScaleUVByAspectRatio) ? (1.0f / (float)MaxDimension) : 1.0f;
+
+	const int NumFrontFaceVertex = TopStartFaceId - FrontStartFaceId;
+	const int NumTopFaceVertex = BackStartFaceId - TopStartFaceId;
+	const int NumBackFaceVertex = BottomStartFaceId - BackStartFaceId;
+	const int NumBottomFaceVertex = LastFaceId - BottomStartFaceId;
+	for (const ESide& Side : AllSides)
+	{
+		switch (Side)
+		{
+		case ESide::Right:
+		case ESide::Left:
+		{
+			const int StartVertexId = (Side == ESide::Right ? 0 : NumVertsPerSide);
+			for (int SideVertex = 0; SideVertex < NumVertsPerSide; SideVertex++)
+			{
+				int VId = StartVertexId + SideVertex;
+				Normals[NormalId++] = GenerateNormal(Side, VId);
+				UVs[UVId++] = GenerateUV(Side, SideVertex, VId, UVScale);
+			}
+			// For the Side Quads, the Normal/UV descriptor is identical to the Face descriptor.
+			const int StartSideFaceVertex = (Side == ESide::Right ? 0 : 4 * NumQuadsPerSide);
+			const int NumSideFaceVertex = 4 * NumQuadsPerSide;
+			for (int SideFaceVertex = 0; SideFaceVertex < NumSideFaceVertex; SideFaceVertex++)
+			{
+				NormalDesc[StartSideFaceVertex + SideFaceVertex] = FaceDesc[StartSideFaceVertex + SideFaceVertex];
+				UVDesc[StartSideFaceVertex + SideFaceVertex] = FaceDesc[StartSideFaceVertex + SideFaceVertex];
+			}
+			break;
+		}
+		case ESide::Front:
+		case ESide::Top:
+		{
+			const int NumFaceVertex = (Side == ESide::Front ? NumFrontFaceVertex : NumTopFaceVertex);
+			const int StartFaceId = (Side == ESide::Front ? FrontStartFaceId : TopStartFaceId);
+			for (int FaceVertexId = 0; FaceVertexId < NumFaceVertex; FaceVertexId++)
+			{
+				int VId = FaceDesc[StartFaceId + FaceVertexId];
+				Normals[NormalId] = GenerateNormal(Side, VId);
+				NormalDesc[StartFaceId + FaceVertexId] = NormalId;
+				NormalId++;
+
+				UVs[UVId] = GenerateUV(Side, FaceVertexId / 4, VId, UVScale);
+				UVDesc[StartFaceId + FaceVertexId] = UVId;
+				UVId++;
+			}
+			break;
+		}
+		case ESide::Back:
+		case ESide::Bottom:
+		{
+			const int NormalStartId = NormalId;
+			const int UVStartId = UVId;
+			const int NumFaceVertex = (Side == ESide::Back ? NumBackFaceVertex : NumBottomFaceVertex);
+			const int StartFaceId = (Side == ESide::Back ? BackStartFaceId : BottomStartFaceId);
+			for (int FaceVertexId = 0; FaceVertexId < NumFaceVertex; FaceVertexId += 4)
+			{
+				Normals[NormalId++] = GenerateNormal(Side, FaceDesc[StartFaceId + FaceVertexId]);
+				Normals[NormalId++] = GenerateNormal(Side, FaceDesc[StartFaceId + FaceVertexId + 3]);
+
+				const int SideFaceId = FaceVertexId / 4;
+				UVs[UVId++] = GenerateUV(Side, SideFaceId, FaceDesc[StartFaceId + FaceVertexId], UVScale);
+				UVs[UVId++] = GenerateUV(Side, SideFaceId, FaceDesc[StartFaceId + FaceVertexId + 3], UVScale);
+
+				// For the last face, add the last/cap vertices.
+				if (FaceVertexId + 4 >= NumFaceVertex)
+				{
+					Normals[NormalId++] = GenerateNormal(Side, FaceDesc[StartFaceId + FaceVertexId + 1]);
+					Normals[NormalId++] = GenerateNormal(Side, FaceDesc[StartFaceId + FaceVertexId + 2]);
+
+					UVs[UVId++] = GenerateUV(Side, SideFaceId, FaceDesc[StartFaceId + FaceVertexId + 1], UVScale);
+					UVs[UVId++] = GenerateUV(Side, SideFaceId, FaceDesc[StartFaceId + FaceVertexId + 2], UVScale);
+				}
+			}
+			// Assign the normals/UV.
+			int FaceVertexId = StartFaceId;
+			int FaceNormalId = NormalStartId;
+			int FaceUVId = UVStartId;
+			for (int FaceId = 0; FaceId < NumSteps; FaceId++)
+			{
+				// Face descriptor index order for back/bottom
+				//
+				//         1.___.2
+				//          |   |
+				//          |___|
+				//         0     3
+				//
+				// The normals are generated along each vertex row
+				// in order: (0, 3, 1, 2). Thus the respective normal/UV
+				// IDs follow: (0, 2, 3, 1).
+				//
+				NormalDesc[FaceVertexId] = FaceNormalId;
+				NormalDesc[FaceVertexId + 1] = FaceNormalId + 2;
+				NormalDesc[FaceVertexId + 2] = FaceNormalId + 3;
+				NormalDesc[FaceVertexId + 3] = FaceNormalId + 1;
+
+				UVDesc[FaceVertexId] = FaceUVId;
+				UVDesc[FaceVertexId + 1] = FaceUVId + 2;
+				UVDesc[FaceVertexId + 2] = FaceUVId + 3;
+				UVDesc[FaceVertexId + 3] = FaceUVId + 1;
+
+				FaceVertexId += 4;
+				FaceNormalId += 2;
+				FaceUVId += 2;
+			}
+			break;
+		}
+		}
+	}
+
+	// Triangulate quad mesh into the output mesh.
+	constexpr int Tris[2][3] = { {0, 3, 2}, {2, 1, 0} };
+	int CurrentTriId = 0;
+	int GroupId = 0;
+	for (int FaceId = 0; FaceId < NumQuads; FaceId++)
+	{
+		ESide Side = FaceToSide(FaceId);
+		if (bPolygroupPerQuad)
+		{
+			GroupId = FaceId;
+		}
+		else
+		{
+			switch (Side)
+			{
+			case ESide::Right:
+				GroupId = 0;
+				break;
+			case ESide::Left:
+				GroupId = 1;
+				break;
+			case ESide::Front:
+			case ESide::Top:
+			{
+				// Each actual step (front/top) face are grouped together
+				check(NumFrontFaceVertex == NumTopFaceVertex);
+				check(NumFrontFaceVertex % 4 == 0);
+				int StartFaceId = Side == ESide::Front ? FrontStartFaceId : TopStartFaceId;
+				int LocalFaceId = FaceId - StartFaceId / 4;
+				GroupId = LocalFaceId + 4;
+				break;
+			}
+			case ESide::Back:
+				GroupId = 2;
+				break;
+			case ESide::Bottom:
+				GroupId = 3;
+				break;
+			}
+		}
+		const bool Flipped = (Side == ESide::Left);
+		for (int TriId = 0; TriId < 2; TriId++)
+		{
+			const int AId = 4 * FaceId + Tris[TriId][0];
+			const int BId = 4 * FaceId + Tris[TriId][1];
+			const int CId = 4 * FaceId + Tris[TriId][2];
+			const int FA = FaceDesc[AId];
+			const int FB = FaceDesc[BId];
+			const int FC = FaceDesc[CId];
+			SetTriangle(CurrentTriId, FA, FB, FC, Flipped);
+			const int NA = NormalDesc[AId];
+			const int NB = NormalDesc[BId];
+			const int NC = NormalDesc[CId];
+			SetTriangleNormals(CurrentTriId, NA, NB, NC, Flipped);
+			const int UVA = UVDesc[AId];
+			const int UVB = UVDesc[BId];
+			const int UVC = UVDesc[CId];
+			SetTriangleUVs(CurrentTriId, UVA, UVB, UVC, Flipped);
+			SetTrianglePolygon(CurrentTriId, GroupId);
+			CurrentTriId++;
+		}
+	}
+
+	return *this;
+}
+
+/**
+ * Generate a floating stair mesh.
+ *
+ * Stair topology composition. (side view cross section)
+ *
+ * (EStairStyle::Floating)
+ *                  .___.
+ *                  |   |          | z (height)
+ *              .___|___|          |
+ *              |   |   |          |_____ x (depth)
+ *          .___|___|___|           \
+ *          |   |   |                \ y (width)
+ *  Row --> |___|___|
+ *
+ *          ^-- Column
+ *
+ */
+FMeshShapeGenerator& FStairGenerator::GenerateFloatingStairs()
+{
+	const int NumConnectQuads = 4 * NumSteps;
+	const int NumCornerVerts = 8 * NumSteps - 4;
+	NumVertsPerSide = 4 * NumSteps;
+	NumVerts = 2 * NumVertsPerSide;
+	NumQuadsPerSide = 2 * NumSteps - 1;
+	NumQuads = 2 * NumQuadsPerSide + NumConnectQuads;
+
+	const int NumSideQuadAttrs = NumVerts;
+	const int NumConnectQuadCornerAttrs = 2 * NumCornerVerts;
+	const int NumConnectQuadInteriorAttrs = 2 * (NumSteps - 1);
+	NumAttrs = NumConnectQuadCornerAttrs + NumConnectQuadInteriorAttrs + NumSideQuadAttrs;
+	SetBufferSizes(NumVerts, 2 * NumQuads, NumAttrs, NumAttrs);
+
+	// Ordered Side lists for iteration.
+	const ESide RightLeftSides[2] = { ESide::Right, ESide::Left };
+	const ESide AllSides[6] = { ESide::Right, ESide::Left, ESide::Front, ESide::Top, ESide::Back, ESide::Bottom };
+
+	// Generate vertices by vertical column per Right/Left side.
+	LeftSideColumnId = (NumSteps + 1);
+	VertexIds.SetNum(2 * LeftSideColumnId);
+	VertexIdsToColumnRow.SetNum(NumVerts);
+	int VertexId = 0;
+	for (const ESide& Side : RightLeftSides)
+	{
+		int StartColumnId = (Side == ESide::Right ? 0 : LeftSideColumnId);
+		for (int VertexColumn = 0; VertexColumn < NumSteps + 1; VertexColumn++)
+		{
+			int NumVertices = 4;
+			if (VertexColumn == 0)
+			{
+				NumVertices = 2;
+			}
+			else if (VertexColumn == 1 || VertexColumn == NumSteps)
+			{
+				NumVertices = 3;
+			}
+			VertexIds[StartColumnId + VertexColumn].SetNum(NumVertices);
+			for (int VertexRow = 0; VertexRow < NumVertices; VertexRow++)
+			{
+				Vertices[VertexId] = GenerateVertex(Side, VertexColumn, VertexRow);
+				VertexIds[StartColumnId + VertexColumn][VertexRow] = VertexId;
+				VertexIdsToColumnRow[VertexId] = FIndex2i(StartColumnId + VertexColumn, VertexRow);
+				VertexId++;
+			}
+		}
+	}
+
+	// Generate quad representation for each side.
+	//
+	// Compute these in respective groups for later normal,
+	// UV & polygroup assignment.
+	FaceDesc.SetNum(4 * NumQuads);
+	int FaceDescId = 0;
+	for (const ESide& Side : AllSides)
+	{
+		switch (Side)
+		{
+		case ESide::Right:
+		case ESide::Left:
+		{
+			const int StartColId = Side == ESide::Right ? 0 : LeftSideColumnId;
+			if (Side == ESide::Right)
+			{
+				RightStartFaceId = FaceDescId;
+			}
+			else
+			{
+				LeftStartFaceId = FaceDescId;
+			}
+			for (int Step = 0; Step < NumSteps; Step++)
+			{
+				int CurrentColId = StartColId + Step;
+				if (Step < 2)
+				{
+					for (int VertexRow = 0; VertexRow < VertexIds[CurrentColId].Num() - 1; VertexRow++)
+					{
+						FaceDesc[FaceDescId++] = (VertexIds[CurrentColId][VertexRow]);
+						FaceDesc[FaceDescId++] = (VertexIds[CurrentColId][VertexRow + 1]);
+						FaceDesc[FaceDescId++] = (VertexIds[CurrentColId + 1][VertexRow + 1]);
+						FaceDesc[FaceDescId++] = (VertexIds[CurrentColId + 1][VertexRow]);
+					}
+				}
+				else
+				{
+					for (int VertexRow = 1; VertexRow < VertexIds[CurrentColId].Num() - 1; VertexRow++)
+					{
+						FaceDesc[FaceDescId++] = (VertexIds[CurrentColId][VertexRow]);
+						FaceDesc[FaceDescId++] = (VertexIds[CurrentColId][VertexRow + 1]);
+						FaceDesc[FaceDescId++] = (VertexIds[CurrentColId + 1][VertexRow]);
+						FaceDesc[FaceDescId++] = (VertexIds[CurrentColId + 1][VertexRow - 1]);
+					}
+				}
+			}
+			break;
+		}
+		case ESide::Front:
+		{
+			FrontStartFaceId = FaceDescId;
+			for (int Step = 0; Step < NumSteps; Step++)
+			{
+				int LastStepRow = VertexIds[Step].Num() - 1;
+				FaceDesc[FaceDescId++] = (VertexIds[Step + LeftSideColumnId][LastStepRow - 1]);
+				FaceDesc[FaceDescId++] = (VertexIds[Step + LeftSideColumnId][LastStepRow]);
+				FaceDesc[FaceDescId++] = (VertexIds[Step][LastStepRow]);
+				FaceDesc[FaceDescId++] = (VertexIds[Step][LastStepRow - 1]);
+			}
+			break;
+		}
+		case ESide::Top:
+		{
+			TopStartFaceId = FaceDescId;
+			for (int Step = 0; Step < NumSteps; Step++)
+			{
+				int LastStepRow = VertexIds[Step].Num() - 1;
+				int NextStepRow = Step < NumSteps - 1 ? VertexIds[Step + 1].Num() - 2 : VertexIds[Step + 1].Num() - 1;
+				FaceDesc[FaceDescId++] = (VertexIds[Step + LeftSideColumnId][LastStepRow]);
+				FaceDesc[FaceDescId++] = (VertexIds[Step + LeftSideColumnId + 1][NextStepRow]);
+				FaceDesc[FaceDescId++] = (VertexIds[Step + 1][NextStepRow]);
+				FaceDesc[FaceDescId++] = (VertexIds[Step][LastStepRow]);
+			}
+			break;
+		}
+		case ESide::Back:
+		{
+			BackStartFaceId = FaceDescId;
+			for (int Step = 0; Step < NumSteps; Step++)
+			{
+				int VertexCol = Step + 2 < NumSteps ? Step + 2 : NumSteps;
+				int VertexRow = Step < NumSteps - 1 ? 0 : 1;
+				FaceDesc[FaceDescId++] = (VertexIds[VertexCol][VertexRow]);
+				FaceDesc[FaceDescId++] = (VertexIds[VertexCol][VertexRow + 1]);
+				FaceDesc[FaceDescId++] = (VertexIds[LeftSideColumnId + VertexCol][VertexRow + 1]);
+				FaceDesc[FaceDescId++] = (VertexIds[LeftSideColumnId + VertexCol][VertexRow]);
+			}
+			break;
+		}
+		case ESide::Bottom:
+		{
+			BottomStartFaceId = FaceDescId;
+			for (int Step = 0; Step < NumSteps; Step++)
+			{
+				int LastStepRow = Step < 2 ? 0 : 1;
+				FaceDesc[FaceDescId++] = (VertexIds[Step][LastStepRow]);
+				FaceDesc[FaceDescId++] = (VertexIds[Step + 1][0]);
+				FaceDesc[FaceDescId++] = (VertexIds[Step + 1 + LeftSideColumnId][0]);
+				FaceDesc[FaceDescId++] = (VertexIds[Step + LeftSideColumnId][LastStepRow]);
+			}
+			break;
+		}
+		}
+	}
+	LastFaceId = FaceDescId;
+
+	auto IsUnsharedFace = [this](ESide Side, int SideFaceId) -> bool
+	{
+		switch (Side)
+		{
+		case ESide::Right:
+		case ESide::Left:
+			return false;
+		case ESide::Front:
+		case ESide::Top:
+			return true;
+		case ESide::Back:
+			return SideFaceId < NumSteps - 2;
+		case ESide::Bottom:
+			return SideFaceId > 1;
+		}
+		return true;
+	};
+
+	// Compute Normals & UVs
+	NormalDesc.SetNum(4 * NumQuads);
+	int NormalId = 0;
+
+	UVDesc.SetNum(4 * NumQuads);
+	int UVId = 0;
+	float MaxDimension = GetMaxDimension();
+	float UVScale = (bScaleUVByAspectRatio) ? (1.0f / (float)MaxDimension) : 1.0f;
+
+	const int NumFrontFaceVertex = TopStartFaceId - FrontStartFaceId;
+	const int NumTopFaceVertex = BackStartFaceId - TopStartFaceId;
+	const int NumBackFaceVertex = BottomStartFaceId - BackStartFaceId;
+	const int NumBottomFaceVertex = LastFaceId - BottomStartFaceId;
+	for (const ESide& Side : AllSides)
+	{
+		switch (Side)
+		{
+		case ESide::Right:
+		case ESide::Left:
+		{
+			const int StartVertexId = (Side == ESide::Right ? 0 : NumVertsPerSide);
+			for (int SideVertex = 0; SideVertex < NumVertsPerSide; SideVertex++)
+			{
+				int VId = StartVertexId + SideVertex;
+				Normals[NormalId++] = GenerateNormal(Side, VId);
+				UVs[UVId++] = GenerateUV(Side, SideVertex, VId, UVScale);
+			}
+			// For the Side Quads, the Normal/UV descriptor is identical to the Face descriptor.
+			const int StartSideFaceVertex = (Side == ESide::Right ? 0 : 4 * NumQuadsPerSide);
+			const int NumSideFaceVertex = 4 * NumQuadsPerSide;
+			for (int SideFaceVertex = 0; SideFaceVertex < NumSideFaceVertex; SideFaceVertex++)
+			{
+				NormalDesc[StartSideFaceVertex + SideFaceVertex] = FaceDesc[StartSideFaceVertex + SideFaceVertex];
+				UVDesc[StartSideFaceVertex + SideFaceVertex] = FaceDesc[StartSideFaceVertex + SideFaceVertex];
+			}
+			break;
+		}
+		case ESide::Front:
+		case ESide::Top:
+		{
+			const int NumFaceVertex = (Side == ESide::Front ? NumFrontFaceVertex : NumTopFaceVertex);
+			const int StartFaceId = (Side == ESide::Front ? FrontStartFaceId : TopStartFaceId);
+			for (int FaceVertexId = 0; FaceVertexId < NumFaceVertex; FaceVertexId++)
+			{
+				int VId = FaceDesc[StartFaceId + FaceVertexId];
+				Normals[NormalId] = GenerateNormal(Side, VId);
+				NormalDesc[StartFaceId + FaceVertexId] = NormalId;
+				NormalId++;
+
+				UVs[UVId] = GenerateUV(Side, FaceVertexId / 4, VId, UVScale);
+				UVDesc[StartFaceId + FaceVertexId] = UVId;
+				UVId++;
+			}
+			break;
+		}
+		case ESide::Back:
+		case ESide::Bottom:
+		{
+			const int NormalStartId = NormalId;
+			const int UVStartId = UVId;
+			const int NumFaceVertex = (Side == ESide::Back ? NumBackFaceVertex : NumBottomFaceVertex);
+			const int StartFaceId = (Side == ESide::Back ? BackStartFaceId : BottomStartFaceId);
+			for (int FaceVertexId = 0; FaceVertexId < NumFaceVertex; FaceVertexId += 4)
+			{
+				const int SideFaceId = FaceVertexId / 4;
+				if (IsUnsharedFace(Side, SideFaceId))
+				{
+					Normals[NormalId++] = GenerateNormal(Side, FaceDesc[StartFaceId + FaceVertexId]);
+					Normals[NormalId++] = GenerateNormal(Side, FaceDesc[StartFaceId + FaceVertexId + 3]);
+					Normals[NormalId++] = GenerateNormal(Side, FaceDesc[StartFaceId + FaceVertexId + 1]);
+					Normals[NormalId++] = GenerateNormal(Side, FaceDesc[StartFaceId + FaceVertexId + 2]);
+
+					UVs[UVId++] = GenerateUV(Side, SideFaceId, FaceDesc[StartFaceId + FaceVertexId], UVScale);
+					UVs[UVId++] = GenerateUV(Side, SideFaceId, FaceDesc[StartFaceId + FaceVertexId + 3], UVScale);
+					UVs[UVId++] = GenerateUV(Side, SideFaceId, FaceDesc[StartFaceId + FaceVertexId + 1], UVScale);
+					UVs[UVId++] = GenerateUV(Side, SideFaceId, FaceDesc[StartFaceId + FaceVertexId + 2], UVScale);
+				}
+				else
+				{
+					Normals[NormalId++] = GenerateNormal(Side, FaceDesc[StartFaceId + FaceVertexId]);
+					Normals[NormalId++] = GenerateNormal(Side, FaceDesc[StartFaceId + FaceVertexId + 3]);
+
+					UVs[UVId++] = GenerateUV(Side, SideFaceId, FaceDesc[StartFaceId + FaceVertexId], UVScale);
+					UVs[UVId++] = GenerateUV(Side, SideFaceId, FaceDesc[StartFaceId + FaceVertexId + 3], UVScale);
+
+					// For the last shared face, add the last/cap vertices.
+					if ((Side == ESide::Back && SideFaceId == NumSteps - 1) ||
+						(Side == ESide::Bottom && SideFaceId == 1))
+					{
+						Normals[NormalId++] = GenerateNormal(Side, FaceDesc[StartFaceId + FaceVertexId + 1]);
+						Normals[NormalId++] = GenerateNormal(Side, FaceDesc[StartFaceId + FaceVertexId + 2]);
+
+						UVs[UVId++] = GenerateUV(Side, SideFaceId, FaceDesc[StartFaceId + FaceVertexId + 1], UVScale);
+						UVs[UVId++] = GenerateUV(Side, SideFaceId, FaceDesc[StartFaceId + FaceVertexId + 2], UVScale);
+					}
+				}
+			}
+			// Assign the normals/UVs.
+			int FaceVertexId = StartFaceId;
+			int FaceNormalId = NormalStartId;
+			int FaceUVId = UVStartId;
+			for (int FaceId = 0; FaceId < NumSteps; FaceId++)
+			{
+				// Face descriptor index order for back/bottom
+				//
+				//         1.___.2
+				//          |   |
+				//          |___|
+				//         0     3
+				//
+				// The normals/UVs are generated along each vertex row
+				// in order: (0, 3, 1, 2). Thus the respective normal/UV
+				// IDs follow: (0, 2, 3, 1).
+				//
+				NormalDesc[FaceVertexId] = FaceNormalId;
+				NormalDesc[FaceVertexId + 1] = FaceNormalId + 2;
+				NormalDesc[FaceVertexId + 2] = FaceNormalId + 3;
+				NormalDesc[FaceVertexId + 3] = FaceNormalId + 1;
+				UVDesc[FaceVertexId] = FaceUVId;
+				UVDesc[FaceVertexId + 1] = FaceUVId + 2;
+				UVDesc[FaceVertexId + 2] = FaceUVId + 3;
+				UVDesc[FaceVertexId + 3] = FaceUVId + 1;
+				FaceVertexId += 4;
+
+				// The last shared face on the bottom should increment like
+				// an unshared face so that subsequent bottom faces reference
+				// the correct attributes.
+				if (IsUnsharedFace(Side, FaceId) || (Side == ESide::Bottom && FaceId == 1))
+				{
+					FaceNormalId += 4;
+					FaceUVId += 4;
+				}
+				else
+				{
+					FaceNormalId += 2;
+					FaceUVId += 2;
+				}
+			}
+			break;
+		}
+		}
+	}
+
+	// Triangulate quad mesh into the output mesh.
+	constexpr int Tris[2][3] = { {0, 3, 2}, {2, 1, 0} };
+	int CurrentTriId = 0;
+	int GroupId = 0;
+	for (int FaceId = 0; FaceId < NumQuads; FaceId++)
+	{
+		ESide Side = FaceToSide(FaceId);
+		if (bPolygroupPerQuad)
+		{
+			GroupId = FaceId;
+		}
+		else
+		{
+			switch (Side)
+			{
+			case ESide::Right:
+				GroupId = 0;
+				break;
+			case ESide::Left:
+				GroupId = 1;
+				break;
+			case ESide::Front:
+			case ESide::Top:
+			{
+				// Each actual step (front/top) face are grouped together
+				check(NumFrontFaceVertex == NumTopFaceVertex);
+				check(NumFrontFaceVertex % 4 == 0);
+				int StartFaceId = Side == ESide::Front ? FrontStartFaceId : TopStartFaceId;
+				int LocalFaceId = FaceId - StartFaceId / 4;
+				GroupId = LocalFaceId + 2;
+				break;
+			}
+			case ESide::Back:
+			case ESide::Bottom:
+			{
+				// Each actual step (back/bottom) face are grouped together
+				check(NumBackFaceVertex == NumBackFaceVertex);
+				check(NumBackFaceVertex % 4 == 0);
+				int StartFaceId = Side == ESide::Back ? BackStartFaceId : BottomStartFaceId;
+				int LocalFaceId = FaceId - StartFaceId / 4;
+				if (Side == ESide::Back && LocalFaceId == NumSteps - 1)
+				{
+					LocalFaceId--;
+				}
+				else if (Side == ESide::Bottom)
+				{
+					LocalFaceId = LocalFaceId > 0 ? LocalFaceId - 1 : LocalFaceId;
+				}
+				GroupId = LocalFaceId + 2 + NumFrontFaceVertex / 4;
+				break;
+			}
+			}
+		}
+		const bool Flipped = (Side == ESide::Left);
+		for (int TriId = 0; TriId < 2; TriId++)
+		{
+			const int AId = 4 * FaceId + Tris[TriId][0];
+			const int BId = 4 * FaceId + Tris[TriId][1];
+			const int CId = 4 * FaceId + Tris[TriId][2];
+			const int FA = FaceDesc[AId];
+			const int FB = FaceDesc[BId];
+			const int FC = FaceDesc[CId];
+			SetTriangle(CurrentTriId, FA, FB, FC, Flipped);
+			const int NA = NormalDesc[AId];
+			const int NB = NormalDesc[BId];
+			const int NC = NormalDesc[CId];
+			SetTriangleNormals(CurrentTriId, NA, NB, NC, Flipped);
+			const int UVA = UVDesc[AId];
+			const int UVB = UVDesc[BId];
+			const int UVC = UVDesc[CId];
+			SetTriangleUVs(CurrentTriId, UVA, UVB, UVC, Flipped);
+			SetTrianglePolygon(CurrentTriId, GroupId);
+			CurrentTriId++;
+		}
+	}
+
+	return *this;
+}
+
+/**
+ * Helper to identify the side based on the face index.
+ *
+ * This method is only valid after the face description array has been
+ * initialized.
+ */
+FStairGenerator::ESide FStairGenerator::FaceToSide(int FaceId)
+{
+	ESide Side = ESide::Right;
+	const int FaceVertexId = FaceId * 4;
+	if (FaceVertexId >= RightStartFaceId && FaceVertexId < LeftStartFaceId)
+	{
+		Side = ESide::Right;
+	}
+	else if (FaceVertexId >= LeftStartFaceId && FaceVertexId < FrontStartFaceId)
+	{
+		Side = ESide::Left;
+	}
+	else if (FaceVertexId >= FrontStartFaceId && FaceVertexId < TopStartFaceId)
+	{
+		Side = ESide::Front;
+	}
+	else if (FaceVertexId >= TopStartFaceId && FaceVertexId < BackStartFaceId)
+	{
+		Side = ESide::Top;
+	}
+	else if (FaceVertexId >= BackStartFaceId && FaceVertexId < BottomStartFaceId)
+	{
+		Side = ESide::Back;
+	}
+	else if (FaceVertexId >= BottomStartFaceId && FaceVertexId < LastFaceId)
+	{
+		Side = ESide::Bottom;
+	}
+	else
+	{
+		check(false);
+	}
+	return Side;
+}
+
+
+/**
+ * Reset state data on the generator.
+ *
+ * This is invoked at the head of the Generate() method.
+ */
+void FStairGenerator::Reset()
+{
+	if (NumSteps <= 0)
+	{
+		NumSteps = 1;
+	}
+
+	VertexIds.Reset();
+	VertexIdsToColumnRow.Reset();
+	FaceDesc.Reset();
+	NormalDesc.Reset();
+	UVDesc.Reset();
+
+	NumQuadsPerSide = 0;
+	NumQuads = 0;
+	NumVertsPerSide = 0;
+	NumVerts = 0;
+	NumAttrs = 0;
+
+	RightSideColumnId = 0;
+	LeftSideColumnId = 0;
+
+	RightStartFaceId = 0;
+	LeftStartFaceId = 0;
+	FrontStartFaceId = 0;
+	TopStartFaceId = 0;
+	BackStartFaceId = 0;
+	BottomStartFaceId = 0;
+	LastFaceId = 0;
+}
+
+
+/**
+ * FLinearStairGenerator
+ */
+
+
+FVector3d FLinearStairGenerator::GenerateVertex(ESide Side, int VertexColumn, int VertexRow)
+{
+	if (Side != ESide::Right && Side != ESide::Left)
+	{
+		// Vertices are only generated for Right & Left sides.
+		check(false);
+	}
+	float X = VertexColumn * StepDepth;
+	float Y = (Side == ESide::Right ? 0.5 * StepWidth : -0.5 * StepWidth);
+	float Z = VertexRow * StepHeight;
+	return FVector3d(X, Y, Z);
+}
+
+FVector3f FLinearStairGenerator::GenerateNormal(ESide Side, int VertexId)
+{
+	FVector3f N;
+	switch (Side)
+	{
+	case ESide::Right:
+		N = FVector3f::UnitY();
+		break;
+	case ESide::Left:
+		N = -FVector3f::UnitY();
+		break;
+	case ESide::Front:
+		N = -FVector3f::UnitX();
+		break;
+	case ESide::Top:
+		N = FVector3f::UnitZ();
+		break;
+	case ESide::Back:
+		N = FVector3f::UnitX();
+		break;
+	case ESide::Bottom:
+		N = -FVector3f::UnitZ();
+		break;
+	default:
+		check(false);
+		break;
+	}
+	return N;
+}
+
+FVector2f FLinearStairGenerator::GenerateUV(ESide Side, int Step, int VertexId, float UVScale)
+{
+	const int Col = VertexIdsToColumnRow[VertexId].A;
+	const int Row = VertexIdsToColumnRow[VertexId].B;
+	FVector2f UV(0.0f, 0.0f);
+	switch (Side)
+	{
+	case ESide::Right:
+	case ESide::Left:
+	{
+		const float UScale = NumSteps * StepDepth * UVScale;
+		const float VScale = NumSteps * StepHeight * UVScale;
+		UV.X = FMathf::Lerp(-0.5f, 0.5f, float(Col % LeftSideColumnId) / float(NumSteps + 1)) * UScale + 0.5f;
+		UV.Y = FMathf::Lerp(-0.5f, 0.5f, float(Row) / float(NumSteps + 1)) * VScale + 0.5f;
+		break;
+	}
+	case ESide::Front:
+	{
+		const float UScale = StepWidth * UVScale;
+		const float VScale = StepHeight * UVScale;
+		UV.X = Col < LeftSideColumnId ? 0.5f : -0.5f;
+		UV.Y = Row == VertexIds[Col].Num() - 1 ? 0.5f : -0.5f;
+		UV.X = UV.X * UScale + 0.5f;
+		UV.Y = UV.Y * VScale + 0.5f;
+		break;
+	}
+	case ESide::Top:
+	{
+		const float UScale = StepWidth * UVScale;
+		const float VScale = StepDepth * UVScale;
+		UV.X = Col < LeftSideColumnId ? 0.5f : -0.5f;
+		UV.Y = (Col % LeftSideColumnId) > Step ? -0.5f : 0.5f;
+		UV.X = UV.X * UScale + 0.5f;
+		UV.Y = UV.Y * VScale + 0.5f;
+		break;
+	}
+	case ESide::Back:
+	{
+		const float UScale = StepWidth * UVScale;
+		const float VScale = NumSteps * StepHeight * UVScale;
+		UV.X = Col < LeftSideColumnId ? 0.5f : -0.5f;
+		UV.Y = FMathf::Lerp(-0.5f, 0.5f, float(Row) / float(NumSteps + 1));
+		UV.X = UV.X * UScale + 0.5f;
+		UV.Y = UV.Y * VScale + 0.5f;
+		break;
+	}
+	case ESide::Bottom:
+	{
+		const float UScale = StepWidth * UVScale;
+		const float VScale = NumSteps * StepDepth * UVScale;
+		UV.X = Col < LeftSideColumnId ? 0.5f : -0.5f;
+		UV.Y = FMathf::Lerp(-0.5f, 0.5f, float(Col % LeftSideColumnId) / float(NumSteps + 1));
+		UV.X = UV.X * UScale + 0.5f;
+		UV.Y = UV.Y * VScale + 0.5f;
+		break;
+	}
+	}
+	return UV;
+}
+
+float FLinearStairGenerator::GetMaxDimension()
+{
+	return FMathf::Max(FMathf::Max(NumSteps * StepDepth, NumSteps * StepHeight), StepWidth);
+}
+
+
+/**
+ * FFloatingStairGenerator
+ */
+
+FVector3d FFloatingStairGenerator::GenerateVertex(ESide Side, int VertexColumn, int VertexRow)
+{
+	if (Side != ESide::Right && Side != ESide::Left)
+	{
+		// Vertices are only generated for Right & Left sides.
+		check(false);
+	}
+	float X = VertexColumn * StepDepth;
+	float Y = (Side == ESide::Right ? 0.5 * StepWidth : -0.5 * StepWidth);
+	float Z = VertexColumn > 1 ? ((VertexColumn - 2) + VertexRow) * StepHeight : VertexRow * StepHeight;
+	return FVector3d(X, Y, Z);
+}
+
+FVector2f FFloatingStairGenerator::GenerateUV(ESide Side, int Step, int VertexId, float UVScale)
+{
+	const int Col = VertexIdsToColumnRow[VertexId].A;
+	const int Row = VertexIdsToColumnRow[VertexId].B;
+	FVector2f UV(0.0f, 0.0f);
+	switch (Side)
+	{
+	case ESide::Right:
+	case ESide::Left:
+	{
+		const float UScale = NumSteps * StepDepth * UVScale;
+		const float VScale = NumSteps * StepHeight * UVScale;
+		float X = float(Col % LeftSideColumnId);
+		float Y = X < 3.0f ? float(Row) : float(Row) + X - 2.0f;
+		X /= float(NumSteps + 1);
+		Y /= float(NumSteps + 1);
+		UV.X = FMathf::Lerp(-0.5f, 0.5f, X) * UScale + 0.5f;
+		UV.Y = FMathf::Lerp(-0.5f, 0.5f, Y) * VScale + 0.5f;
+		break;
+	}
+	case ESide::Back:
+	{
+		const float UScale = StepWidth * UVScale;
+		const float VScale = NumSteps * StepHeight * UVScale;
+		float Y = Step < NumSteps - 1 ? float(Row + Step) : float(Row + Step - 1);
+		Y /= float(NumSteps + 1);
+		UV.X = Col < LeftSideColumnId ? 0.5f : -0.5f;
+		UV.Y = FMathf::Lerp(-0.5f, 0.5f, Y);
+		UV.X = UV.X * UScale + 0.5f;
+		UV.Y = UV.Y * VScale + 0.5f;
+		break;
+	}
+	default:
+	{
+		UV = Super::GenerateUV(Side, Step, VertexId, UVScale);
+		break;
+	}
+	}
+	return UV;
+}
+
+
+/**
+ * FCurvedStairGenerator
+ */
+
+void FCurvedStairGenerator::Reset()
+{
+	Super::Reset();
+
+	bIsClockwise = CurveAngle > 0.0f;
+	CurveRadians = CurveAngle * TMathUtilConstants<float>::DegToRad;
+	CurveRadiansPerStep = CurveRadians / NumSteps;
+	OuterRadius = InnerRadius + StepWidth;
+	RadiusRatio = OuterRadius / InnerRadius;
+	BackNormal = FVector3f::Zero();
+}
+
+FVector3d FCurvedStairGenerator::GenerateVertex(ESide Side, int VertexColumn, int VertexRow)
+{
+	if (Side != ESide::Right && Side != ESide::Left)
+	{
+		// Vertices are only generated for Right & Left sides.
+		check(false);
+	}
+	float X = 0.0f;
+	float Y = 0.0f;
+	float Z = VertexRow * StepHeight;
+
+	float XCoeff = FMathf::Cos(VertexColumn * CurveRadiansPerStep);
+	float YCoeff = FMathf::Sin(VertexColumn * CurveRadiansPerStep);
+	if (bIsClockwise)
+	{
+		X = (Side == ESide::Right ? XCoeff * InnerRadius : XCoeff * OuterRadius);
+		Y = (Side == ESide::Right ? YCoeff * InnerRadius : YCoeff * OuterRadius);
+	}
+	else
+	{
+		X = (Side == ESide::Right ? XCoeff * -OuterRadius : XCoeff * -InnerRadius);
+		Y = (Side == ESide::Right ? YCoeff * -OuterRadius : YCoeff * -InnerRadius);
+	}
+	return FVector3d(X, Y, Z);
+}
+
+FVector3f FCurvedStairGenerator::GenerateNormal(ESide Side, int VertexId)
+{
+	const int Col = VertexIdsToColumnRow[VertexId].A;
+	const int Row = VertexIdsToColumnRow[VertexId].B;
+
+	FVector3f N;
+	switch (Side)
+	{
+	case ESide::Right:
+	case ESide::Left:
+	{
+		float X = FMathf::Cos((Col % (NumSteps + 1)) * CurveRadiansPerStep);
+		float Y = FMathf::Sin((Col % (NumSteps + 1)) * CurveRadiansPerStep);
+		N = (Side == ESide::Right ? FVector3f(-X, -Y, 0.0f) : FVector3f(X, Y, 0.0f));
+		N.Normalize();
+		break;
+	}
+	case ESide::Front:
+	{
+		float X = FMathf::Cos((Col % (NumSteps + 1)) * CurveRadiansPerStep);
+		float Y = FMathf::Sin((Col % (NumSteps + 1)) * CurveRadiansPerStep);
+		N = FVector3f(Y, -X, 0.0f);
+		N.Normalize();
+		break;
+	}
+	case ESide::Top:
+	{
+		N = FVector3f::UnitZ();
+		break;
+	}
+	case ESide::Back:
+	{
+		if (BackNormal == FVector3f::Zero())
+		{
+			float X = FMathf::Cos(NumSteps * CurveRadiansPerStep);
+			float Y = FMathf::Sin(NumSteps * CurveRadiansPerStep);
+			BackNormal = FVector3f(-Y, X, 0.0f);
+		}
+		N = BackNormal;
+		break;
+	}
+	case ESide::Bottom:
+	{
+		N = -FVector3f::UnitZ();
+		break;
+	}
+	}
+	return N;
+}
+
+FVector2f FCurvedStairGenerator::GenerateUV(ESide Side, int Step, int VertexId, float UVScale)
+{
+	const int Col = VertexIdsToColumnRow[VertexId].A;
+	const int Row = VertexIdsToColumnRow[VertexId].B;
+	FVector2f UV;
+	switch (Side)
+	{
+	case ESide::Right:
+	case ESide::Left:
+	{
+		const float UScale = OuterRadius * UVScale;
+		const float VScale = NumSteps * StepHeight * UVScale;
+		UV.X = FMathf::Lerp(-0.5f, 0.5f, float(Col % LeftSideColumnId) / float(NumSteps + 1));
+		UV.Y = FMathf::Lerp(-0.5f, 0.5f, float(Row) / float(NumSteps + 1));
+
+		// Proportionally scale the UVs of each side according
+		// to the RadiusRatio and ensure [0,1].
+		if (RadiusRatio * UScale > 1.0f)
+		{
+			UV.X /= (Side == ESide::Left ? UScale : UScale * RadiusRatio);
+		}
+		else if (Side == ESide::Left)
+		{
+			UV.X *= RadiusRatio;
+		}
+		UV.X = UV.X * UScale + 0.5f;
+		UV.Y = UV.Y * VScale + 0.5f;
+		break;
+	}
+	case ESide::Front:
+	{
+		const float UScale = StepWidth * UVScale;
+		const float VScale = StepHeight * UVScale;
+		UV.X = Col < LeftSideColumnId ? 0.5f : -0.5f;
+		UV.Y = Row == VertexIds[Col].Num() - 1 ? -0.5f : 0.5f;
+		UV.X = UV.X * UScale + 0.5f;
+		UV.Y = UV.Y * VScale + 0.5f;
+		break;
+	}
+	case ESide::Top:
+	{
+		const float UScale = StepWidth * UVScale;
+		const float VScale = OuterRadius / NumSteps * UVScale;
+		UV.X = Col < LeftSideColumnId ? 0.5f : -0.5f;
+		UV.Y = (Col % LeftSideColumnId) > Step ? -0.5f : 0.5f;
+
+		// Proportionally scale the outer edge of top faces
+		// to avoid stretching while ensuring [0,1].
+		if (RadiusRatio * VScale > 1.0f)
+		{
+			UV.Y /= (Col >= LeftSideColumnId ? VScale : VScale * RadiusRatio);
+		}
+		else if (Col >= LeftSideColumnId)
+		{
+			UV.Y *= RadiusRatio;
+		}
+
+		UV.X = UV.X * UScale + 0.5f;
+		UV.Y = UV.Y * VScale + 0.5f;
+		break;
+	}
+	case ESide::Back:
+	{
+		const float UScale = StepWidth * UVScale;
+		const float VScale = NumSteps * StepHeight * UVScale;
+		UV.X = Col < LeftSideColumnId ? 0.5f : -0.5f;
+		UV.Y = FMathf::Lerp(-0.5f, 0.5f, float(Row) / float(NumSteps + 1));
+		UV.X = UV.X * UScale + 0.5f;
+		UV.Y = UV.Y * VScale + 0.5f;
+		break;
+	}
+	case ESide::Bottom:
+	{
+		const float UScale = StepWidth * UVScale;
+		const float VScale = OuterRadius * UVScale;
+		UV.X = Col < LeftSideColumnId ? 0.5f : -0.5f;
+		UV.Y = FMathf::Lerp(-0.5f, 0.5f, float(Col % LeftSideColumnId) / float(NumSteps + 1));
+
+		// Proportionally scale outer edge of bottom faces
+		// to avoid stretching while ensuring [0,1].
+		if (RadiusRatio * VScale > 1.0f)
+		{
+			UV.Y /= (Col >= LeftSideColumnId ? VScale : VScale * RadiusRatio);
+		}
+		else if (Col >= LeftSideColumnId)
+		{
+			UV.Y *= RadiusRatio;
+		}
+		UV.X = UV.X * UScale + 0.5f;
+		UV.Y = UV.Y * VScale + 0.5f;
+		break;
+	}
+	default:
+	{
+		check(false);
+		break;
+	}
+	}
+	return UV;
+}
+
+float FCurvedStairGenerator::GetMaxDimension()
+{
+	float MaxDepth = FMathf::Abs(CurveRadians) * OuterRadius;
+	return FMathf::Max(FMathf::Max(MaxDepth, NumSteps * StepHeight), StepWidth);
+}
+
+
+/**
+ * FSpiralStairGenerator
+ */
+
+FVector3d FSpiralStairGenerator::GenerateVertex(ESide Side, int VertexColumn, int VertexRow)
+{
+	if (Side != ESide::Right && Side != ESide::Left)
+	{
+		// Vertices are only generated for Right & Left sides.
+		check(false);
+	}
+	float X = 0.0f;
+	float Y = 0.0f;
+	float Z = VertexColumn > 1 ? ((VertexColumn - 2) + VertexRow) * StepHeight : VertexRow * StepHeight;
+
+	float XCoeff = FMathf::Cos(VertexColumn * CurveRadiansPerStep);
+	float YCoeff = FMathf::Sin(VertexColumn * CurveRadiansPerStep);
+	if (bIsClockwise)
+	{
+		X = (Side == ESide::Right ? XCoeff * InnerRadius : XCoeff * OuterRadius);
+		Y = (Side == ESide::Right ? YCoeff * InnerRadius : YCoeff * OuterRadius);
+	}
+	else
+	{
+		X = (Side == ESide::Right ? XCoeff * -OuterRadius : XCoeff * -InnerRadius);
+		Y = (Side == ESide::Right ? YCoeff * -OuterRadius : YCoeff * -InnerRadius);
+	}
+	return FVector3d(X, Y, Z);
+}
+
+FVector2f FSpiralStairGenerator::GenerateUV(ESide Side, int Step, int VertexId, float UVScale)
+{
+	const int Col = VertexIdsToColumnRow[VertexId].A;
+	const int Row = VertexIdsToColumnRow[VertexId].B;
+	FVector2f UV;
+	switch (Side)
+	{
+	case ESide::Right:
+	case ESide::Left:
+	{
+		const float UScale = OuterRadius * UVScale;
+		const float VScale = NumSteps * StepHeight * UVScale;
+		float X = float(Col % LeftSideColumnId);
+		float Y = X < 3.0f ? float(Row) : float(Row) + X - 2.0f;
+		UV.X = FMathf::Lerp(-0.5f, 0.5f, X / float(NumSteps + 1));
+		UV.Y = FMathf::Lerp(-0.5f, 0.5f, Y / float(NumSteps + 1));
+
+		// Proportionally scale the UVs of each side according
+		// to the RadiusRatio and ensure [0,1].
+		if (RadiusRatio * UScale > 1.0f)
+		{
+			UV.X /= (Side == ESide::Left ? UScale : UScale * RadiusRatio);
+		}
+		else if (Side == ESide::Left)
+		{
+			UV.X *= RadiusRatio;
+		}
+		UV.X = UV.X * UScale + 0.5f;
+		UV.Y = UV.Y * VScale + 0.5f;
+		break;
+	}
+	case ESide::Back:
+	{
+		const float UScale = StepWidth * UVScale;
+		const float VScale = NumSteps * StepHeight * UVScale;
+		float Y = Step < NumSteps - 1 ? float(Row + Step) : float(Row + Step - 1);
+		Y /= float(NumSteps + 1);
+		UV.X = Col < LeftSideColumnId ? 0.5f : -0.5f;
+		UV.Y = FMathf::Lerp(-0.5f, 0.5f, Y);
+		UV.X = UV.X * UScale + 0.5f;
+		UV.Y = UV.Y * VScale + 0.5f;
+		break;
+	}
+	default:
+	{
+		UV = Super::GenerateUV(Side, Step, VertexId, UVScale);
+		break;
+	}
+	}
+	return UV;
+}
diff --git a/Engine/Plugins/Experimental/GeometryProcessing/Source/GeometricObjects/Public/Generators/StairGenerator.h b/Engine/Plugins/Experimental/GeometryProcessing/Source/GeometricObjects/Public/Generators/StairGenerator.h
--- a/Engine/Plugins/Experimental/GeometryProcessing/Source/GeometricObjects/Public/Generators/StairGenerator.h
+++ b/Engine/Plugins/Experimental/GeometryProcessing/Source/GeometricObjects/Public/Generators/StairGenerator.h
@@ -15,7 +15,7 @@ namespace Geometry
 /**
  * Base Stair mesh generator class.
  */
-class /*GEOMETRICOBJECTS_API*/ FStairGenerator : public FMeshShapeGenerator
+class GEOMETRICOBJECTS_API FStairGenerator : public FMeshShapeGenerator
 {
 public:
 	/** If true (default), UVs are scaled so that there is no stretching. If false, UVs are scaled to fill unit square */
@@ -31,7 +31,7 @@ public:
 	float StepHeight = 20.0f;
 
 	/** The number of steps in this staircase. */
-	int NumSteps = 5;
+	int NumSteps = 8;
 
 	enum class ESide
 	{
@@ -43,492 +43,35 @@ public:
 		Bottom	// -Z
 	};
 
-	enum class EBaseStyle
-	{
-		SolidBase,	// Each step connects to the floor.
-		HollowBase	// Each step only extends one step over to support the next step.
-	};
-	/** The style of the base of the stairs */
-	EBaseStyle BaseStyle = EBaseStyle::SolidBase;
-
 public:
 	virtual ~FStairGenerator() {}
 
 	/** Generate the mesh */
-	virtual FMeshShapeGenerator& Generate() override
-	{
-		Reset();
-
-		switch (BaseStyle)
-		{
-		case EBaseStyle::HollowBase:
-			return GenerateHollowBase();
-		case EBaseStyle::SolidBase:
-		default:
-			return GenerateSolidBase();
-		}
-	}
+	virtual FMeshShapeGenerator& Generate() override;
 
-	/**
-	 * Generate a stair mesh with a solid base.
-	 *
-	 * Stair topology composition. (side view cross section)
-	 *
-	 * (EBaseStyle::SolidBase)
-	 *                  .___.   
-	 *                  |   |          | z (height)
-	 *              .___|___|          | 
-	 *              |   |   |          |_____ x (depth)
-	 *          .___|___|___|           \
-	 *          |   |   |   |            \ y (width)
-	 *  Row --> |___|___|___|
-	 *
-	 *          ^-- Column
-	 *
-	 */
-	virtual FMeshShapeGenerator& GenerateSolidBase()
-	{
-		auto TriangleNumber = [](int x)
-		{
-			return (x * (x + 1) / 2);
-		};
-
-		const int NumConnectQuads = 4 * NumSteps;
-		const int NumCornerVerts = 4 * (NumSteps + 1);
-		NumVertsPerSide = TriangleNumber(NumSteps + 1) + NumSteps;
-		NumVerts = 2 * NumVertsPerSide;
-		NumQuadsPerSide = TriangleNumber(NumSteps);
-		NumQuads = 2 * NumQuadsPerSide + NumConnectQuads;
-
-		const int NumSideQuadAttrs = NumVerts;
-		const int NumConnectQuadCornerAttrs = 2 * NumCornerVerts;
-		const int NumConnectQuadInteriorAttrs = 4 * (NumSteps - 1);
-		NumAttrs = NumConnectQuadCornerAttrs + NumConnectQuadInteriorAttrs + NumSideQuadAttrs;
-		SetBufferSizes(NumVerts, 2 * NumQuads, NumAttrs, NumAttrs);
-
-		// Ordered Side lists for iteration.
-		const ESide RightLeftSides[2] = { ESide::Right, ESide::Left };
-		const ESide AllSides[6] = { ESide::Right, ESide::Left, ESide::Front, ESide::Top, ESide::Back, ESide::Bottom };
-
-		// Generate vertices by vertical column per Right/Left side.
-		LeftSideColumnId = (NumSteps + 1);
-		VertexIds.SetNum(2 * LeftSideColumnId);
-		VertexIdsToColumnRow.SetNum(NumVerts);
-		int VertexId = 0;
-		for (const ESide& Side : RightLeftSides)
-		{
-			int StartColumnId = (Side == ESide::Right ? 0 : LeftSideColumnId);
-			for (int VertexColumn = 0; VertexColumn < NumSteps + 1; VertexColumn++)
-			{
-				int NumVertices = (VertexColumn == NumSteps ? VertexColumn + 1 : VertexColumn + 2);
-				VertexIds[StartColumnId + VertexColumn].SetNum(NumVertices);
-				for (int VertexRow = 0; VertexRow < NumVertices; VertexRow++)
-				{
-					Vertices[VertexId] = GenerateVertex(Side, VertexColumn, VertexRow);
-					VertexIds[StartColumnId + VertexColumn][VertexRow] = VertexId;
-					VertexIdsToColumnRow[VertexId] = FIndex2i(StartColumnId + VertexColumn, VertexRow);
-					VertexId++;
-				}
-			}
-		}
-
-		// Generate quad representation for each side.
-		//
-		// Compute these in respective groups for later normal,
-		// UV & polygroup assignment.
-		FaceDesc.SetNum(4 * NumQuads);
-		int FaceDescId = 0;
-		for (const ESide& Side : AllSides)
-		{
-			switch (Side)
-			{
-			case ESide::Right:
-			{
-				const int StartColId = 0;
-				RightStartFaceId = FaceDescId;
-				for (int Step = 0; Step < NumSteps; Step++)
-				{
-					int CurrentColId = StartColId + Step;
-					for (int VertexRow = 0; VertexRow < VertexIds[CurrentColId].Num() - 1; VertexRow++)
-					{
-						FaceDesc[FaceDescId++] = (VertexIds[CurrentColId][VertexRow]);
-						FaceDesc[FaceDescId++] = (VertexIds[CurrentColId][VertexRow + 1]);
-						FaceDesc[FaceDescId++] = (VertexIds[CurrentColId + 1][VertexRow + 1]);
-						FaceDesc[FaceDescId++] = (VertexIds[CurrentColId + 1][VertexRow]);
-					}
-				}
-				break;
-			}
-			case ESide::Left:
-			{
-				const int StartColId = LeftSideColumnId;
-				LeftStartFaceId = FaceDescId;
-				for (int Step = 0; Step < NumSteps; Step++)
-				{
-					int CurrentColId = StartColId + Step;
-					for (int VertexRow = 0; VertexRow < VertexIds[CurrentColId].Num() - 1; VertexRow++)
-					{
-						// Flipped face order from ESide::Right
-						FaceDesc[FaceDescId++] = (VertexIds[CurrentColId][VertexRow]);
-						FaceDesc[FaceDescId++] = (VertexIds[CurrentColId + 1][VertexRow]);
-						FaceDesc[FaceDescId++] = (VertexIds[CurrentColId + 1][VertexRow + 1]);
-						FaceDesc[FaceDescId++] = (VertexIds[CurrentColId][VertexRow + 1]);
-					}
-				}
-				break;
-			}
-			case ESide::Front:
-			{
-				FrontStartFaceId = FaceDescId;
-				for (int Step = 0; Step < NumSteps; Step++)
-				{
-					int LastStepRow = VertexIds[Step].Num() - 1;
-					FaceDesc[FaceDescId++] = (VertexIds[Step + LeftSideColumnId][LastStepRow - 1]);
-					FaceDesc[FaceDescId++] = (VertexIds[Step + LeftSideColumnId][LastStepRow]);
-					FaceDesc[FaceDescId++] = (VertexIds[Step][LastStepRow]);
-					FaceDesc[FaceDescId++] = (VertexIds[Step][LastStepRow - 1]);
-				}
-				break;
-			}
-			case ESide::Top:
-			{
-				TopStartFaceId = FaceDescId;
-				for (int Step = 0; Step < NumSteps; Step++)
-				{
-					int LastStepRow = VertexIds[Step].Num() - 1;
-					FaceDesc[FaceDescId++] = (VertexIds[Step + LeftSideColumnId][LastStepRow]);
-					FaceDesc[FaceDescId++] = (VertexIds[Step + LeftSideColumnId + 1][LastStepRow]);
-					FaceDesc[FaceDescId++] = (VertexIds[Step + 1][LastStepRow]);
-					FaceDesc[FaceDescId++] = (VertexIds[Step][LastStepRow]);
-				}
-				break;
-			}
-			case ESide::Back:
-			{
-				BackStartFaceId = FaceDescId;
-				for (int Step = 0; Step < NumSteps; Step++)
-				{
-					FaceDesc[FaceDescId++] = (VertexIds[NumSteps][Step]);
-					FaceDesc[FaceDescId++] = (VertexIds[NumSteps][Step + 1]);
-					FaceDesc[FaceDescId++] = (VertexIds[LeftSideColumnId + NumSteps][Step + 1]);
-					FaceDesc[FaceDescId++] = (VertexIds[LeftSideColumnId + NumSteps][Step]);
-				}
-				break;
-			}
-			case ESide::Bottom:
-			{
-				BottomStartFaceId = FaceDescId;
-				for (int Step = 0; Step < NumSteps; Step++)
-				{
-					FaceDesc[FaceDescId++] = (VertexIds[Step][0]);
-					FaceDesc[FaceDescId++] = (VertexIds[Step + 1][0]);
-					FaceDesc[FaceDescId++] = (VertexIds[Step + 1 + LeftSideColumnId][0]);
-					FaceDesc[FaceDescId++] = (VertexIds[Step + LeftSideColumnId][0]);
-				}
-				break;
-			}
-			}
-		}
-		LastFaceId = FaceDescId;
-
-		// Compute Normals
-		NormalDesc.SetNum(4 * NumQuads);
-		int NormalId = 0;
-		const int NumFrontFaceVertex = TopStartFaceId - FrontStartFaceId;
-		const int NumTopFaceVertex = BackStartFaceId - TopStartFaceId;
-		const int NumBackFaceVertex = BottomStartFaceId - BackStartFaceId;
-		const int NumBottomFaceVertex = LastFaceId - BottomStartFaceId;
-		for (const ESide& Side : AllSides)
-		{
-			switch (Side)
-			{
-			case ESide::Right:
-			case ESide::Left:
-			{
-				const int StartVertexId = (Side == ESide::Right ? 0 : NumVertsPerSide);
-				for (int SideVertex = 0; SideVertex < NumVertsPerSide; SideVertex++)
-				{
-					int VId = StartVertexId + SideVertex;
-					Normals[NormalId++] = GenerateNormal(Side, VId);
-				}
-				// For the Side Quads, the Normal descriptor is identical to the Face descriptor.
-				const int StartSideFaceVertex = (Side == ESide::Right ? 0 : 4 * NumQuadsPerSide);
-				const int NumSideFaceVertex = 4 * NumQuadsPerSide;
-				for (int SideFaceVertex = 0; SideFaceVertex < NumSideFaceVertex; SideFaceVertex++)
-				{
-					NormalDesc[StartSideFaceVertex + SideFaceVertex] = FaceDesc[StartSideFaceVertex + SideFaceVertex];
-				}
-				break;
-			}
-			case ESide::Front:
-			case ESide::Top:
-			{
-				const int NumFaceVertex = (Side == ESide::Front ? NumFrontFaceVertex : NumTopFaceVertex);
-				const int StartFaceId = (Side == ESide::Front ? FrontStartFaceId : TopStartFaceId);
-				for (int FaceVertexId = 0; FaceVertexId < NumFaceVertex; FaceVertexId++)
-				{
-					int VId = FaceDesc[StartFaceId + FaceVertexId];
-					Normals[NormalId] = GenerateNormal(Side, VId);
-					NormalDesc[StartFaceId + FaceVertexId] = NormalId;
-					NormalId++;
-				}
-				break;
-			}
-			case ESide::Back:
-			case ESide::Bottom:
-			{
-				// Generate the normals.
-				const int NormalStartId = NormalId;
-				const int NumFaceVertex = (Side == ESide::Back ? NumBackFaceVertex : NumBottomFaceVertex);
-				const int StartFaceId = (Side == ESide::Back ? BackStartFaceId : BottomStartFaceId);
-				for (int FaceVertexId = 0; FaceVertexId < NumFaceVertex; FaceVertexId += 4)
-				{
-					Normals[NormalId++] = GenerateNormal(Side, FaceDesc[StartFaceId + FaceVertexId]);
-					Normals[NormalId++] = GenerateNormal(Side, FaceDesc[StartFaceId + FaceVertexId + 3]);
-
-					// For the last face, add the last/cap vertices.
-					if (FaceVertexId + 4 >= NumFaceVertex)
-					{
-						Normals[NormalId++] = GenerateNormal(Side, FaceDesc[StartFaceId + FaceVertexId + 1]);
-						Normals[NormalId++] = GenerateNormal(Side, FaceDesc[StartFaceId + FaceVertexId + 2]);
-					}
-				}
-				// Assign the normals via the normal descriptor.
-				int FaceVertexId = StartFaceId;
-				int FaceNormalId = NormalStartId;
-				for (int FaceId = 0; FaceId < NumSteps; FaceId++)
-				{
-					// Face descriptor index order for back/bottom
-					//
-					//         1.___.2
-					//          |   |
-					//          |___|
-					//         0     3
-					//
-					// The normals are generated along each vertex row
-					// in order: (0, 3, 1, 2). Thus the respective normal
-					// IDs follow: (0, 2, 3, 1).
-					//
-					NormalDesc[FaceVertexId] = FaceNormalId;
-					NormalDesc[FaceVertexId + 1] = FaceNormalId + 2;
-					NormalDesc[FaceVertexId + 2] = FaceNormalId + 3;
-					NormalDesc[FaceVertexId + 3] = FaceNormalId + 1;
-					FaceVertexId += 4;
-					FaceNormalId += 2;
-				}
-				break;
-			}
-			}
-		}
-
-		// Compute UVs
-		UVDesc.SetNum(4 * NumQuads);
-		int UVId = 0;
-		float MaxDimension = GetMaxDimension();
-		float UVScale = (bScaleUVByAspectRatio) ? (1.0f / (float)MaxDimension) : 1.0f;
-		for (const ESide& Side : AllSides)
-		{
-			switch (Side)
-			{
-			case ESide::Right:
-			case ESide::Left:
-			{
-				const int StartVertexId = (Side == ESide::Right ? 0 : NumVertsPerSide);
-				for (int SideVertex = 0; SideVertex < NumVertsPerSide; SideVertex++)
-				{
-					int VId = StartVertexId + SideVertex;
-					UVs[UVId++] = GenerateUV(Side, SideVertex, VId, UVScale);
-				}
-				// For the Side Quads, the UV descriptor is identical to the Face descriptor.
-				const int StartSideFaceVertex = (Side == ESide::Right ? 0 : 4 * NumQuadsPerSide);
-				const int NumSideFaceVertex = 4 * NumQuadsPerSide;
-				for (int SideFaceVertex = 0; SideFaceVertex < NumSideFaceVertex; SideFaceVertex++)
-				{
-					UVDesc[StartSideFaceVertex + SideFaceVertex] = FaceDesc[StartSideFaceVertex + SideFaceVertex];
-				}
-				break;
-			}
-			case ESide::Front:
-			case ESide::Top:
-			{
-				const int NumFaceVertex = (Side == ESide::Front ? NumFrontFaceVertex : NumTopFaceVertex);
-				const int StartFaceId = (Side == ESide::Front ? FrontStartFaceId : TopStartFaceId);
-				for (int FaceVertexId = 0; FaceVertexId < NumFaceVertex; FaceVertexId++)
-				{
-					int VId = FaceDesc[StartFaceId + FaceVertexId];
-					UVs[UVId] = GenerateUV(Side, FaceVertexId / 4, VId, UVScale);
-					UVDesc[StartFaceId + FaceVertexId] = UVId;
-					UVId++;
-				}
-				break;
-			}
-			case ESide::Back:
-			case ESide::Bottom:
-			{
-				// Generate the UVs.
-				const int UVStartId = UVId;
-				const int NumFaceVertex = (Side == ESide::Back ? NumBackFaceVertex : NumBottomFaceVertex);
-				const int StartFaceId = (Side == ESide::Back ? BackStartFaceId : BottomStartFaceId);
-				for (int FaceVertexId = 0; FaceVertexId < NumFaceVertex; FaceVertexId += 4)
-				{
-					const int SideFaceId = FaceVertexId / 4;
-					UVs[UVId++] = GenerateUV(Side, SideFaceId, FaceDesc[StartFaceId + FaceVertexId], UVScale);
-					UVs[UVId++] = GenerateUV(Side, SideFaceId, FaceDesc[StartFaceId + FaceVertexId + 3], UVScale);
-
-					// For the last face, add the last/cap vertices.
-					if (FaceVertexId + 4 >= NumFaceVertex)
-					{
-						UVs[UVId++] = GenerateUV(Side, SideFaceId, FaceDesc[StartFaceId + FaceVertexId + 1], UVScale);
-						UVs[UVId++] = GenerateUV(Side, SideFaceId, FaceDesc[StartFaceId + FaceVertexId + 2], UVScale);
-					}
-				}
-				// Assign the UVs via the normal descriptor.
-				int FaceVertexId = StartFaceId;
-				int FaceUVId = UVStartId;
-				for (int FaceId = 0; FaceId < NumSteps; FaceId++)
-				{
-					// Face descriptor index order for back/bottom
-					//
-					//         1.___.2
-					//          |   |
-					//          |___|
-					//         0     3
-					//
-					// The UVs are generated along each vertex row
-					// in order: (0, 3, 1, 2). Thus the respective UV
-					// IDs follow: (0, 2, 3, 1).
-					//
-					UVDesc[FaceVertexId] = FaceUVId;
-					UVDesc[FaceVertexId + 1] = FaceUVId + 2;
-					UVDesc[FaceVertexId + 2] = FaceUVId + 3;
-					UVDesc[FaceVertexId + 3] = FaceUVId + 1;
-					FaceVertexId += 4;
-					FaceUVId += 2;
-				}
-				break;
-			}
-			}
-		}
-
-		// Triangulate quad mesh into the output mesh.
-		auto FaceToSide = [this](int FaceId) -> ESide
-		{
-			ESide Side = ESide::Right;
-			const int FaceVertexId = FaceId * 4;
-			if (FaceVertexId >= RightStartFaceId && FaceVertexId < LeftStartFaceId)
-			{
-				Side = ESide::Right;
-			}
-			else if (FaceVertexId >= LeftStartFaceId && FaceVertexId < FrontStartFaceId)
-			{
-				Side = ESide::Left;
-			}
-			else if (FaceVertexId >= FrontStartFaceId && FaceVertexId < TopStartFaceId)
-			{
-				Side = ESide::Front;
-			}
-			else if (FaceVertexId >= TopStartFaceId && FaceVertexId < BackStartFaceId)
-			{
-				Side = ESide::Top;
-			}
-			else if (FaceVertexId >= BackStartFaceId && FaceVertexId < BottomStartFaceId)
-			{
-				Side = ESide::Back;
-			}
-			else if (FaceVertexId >= BottomStartFaceId && FaceVertexId < LastFaceId)
-			{
-				Side = ESide::Bottom;
-			}
-			else
-			{
-				check(false);
-			}
-			return Side;
-		};
-		constexpr int Tris[2][3] = { {0, 3, 2}, {2, 1, 0} };
-		int CurrentTriId = 0;
-		int GroupId = 0;
-		for (int FaceId = 0; FaceId < NumQuads; FaceId++)
-		{
-			ESide Side = FaceToSide(FaceId);
-			if (bPolygroupPerQuad)
-			{
-				GroupId = FaceId;
-			}
-			else
-			{
-				switch (Side)
-				{
-				case ESide::Right:
-					GroupId = 0;
-					break;
-				case ESide::Left:
-					GroupId = 1;
-					break;
-				case ESide::Front:
-				case ESide::Top:
-					// Each actual step (front/top) face are grouped
-					// as a polygroup.
-					check(NumFrontFaceVertex == NumTopFaceVertex);
-					check(NumFrontFaceVertex % 4 == 0);
-					GroupId = FaceId % (NumFrontFaceVertex / 4) + 2;
-					break;
-				case ESide::Back:
-					GroupId = 2 + NumFrontFaceVertex / 4 + 1;
-					break;
-				case ESide::Bottom:
-					GroupId = 2 + NumFrontFaceVertex / 4 + 2;
-					break;
-				}
-			}
-			for (int TriId = 0; TriId < 2; TriId++)
-			{
-				const int AId = 4 * FaceId + Tris[TriId][0];
-				const int BId = 4 * FaceId + Tris[TriId][1];
-				const int CId = 4 * FaceId + Tris[TriId][2];
-				const int FA = FaceDesc[AId];
-				const int FB = FaceDesc[BId];
-				const int FC = FaceDesc[CId];
-				SetTriangle(CurrentTriId, FA, FB, FC);
-				const int NA = NormalDesc[AId];
-				const int NB = NormalDesc[BId];
-				const int NC = NormalDesc[CId];
-				SetTriangleNormals(CurrentTriId, NA, NB, NC);
-				const int UVA = UVDesc[AId];
-				const int UVB = UVDesc[BId];
-				const int UVC = UVDesc[CId];
-				SetTriangleUVs(CurrentTriId, UVA, UVB, UVC);
-				SetTrianglePolygon(CurrentTriId, GroupId);
-				CurrentTriId++;
-			}
-		}
-
-		return *this;
-	}
+private:
+	/** Generate solid style stairs */
+	FMeshShapeGenerator& GenerateSolidStairs();
+
+	/** Generate floating style stairs */
+	FMeshShapeGenerator& GenerateFloatingStairs();
 
 	/**
-	 * Generate a hollow base stair mesh.
-	 * 
-	 * Stair topology composition. (side view cross section)
-	 *
-	 * (EBaseStyle::HollowBase)
-	 *                  .___.
-	 *                  |   |          | z (height)
-	 *              .___|___|          |
-	 *              |   |   |          |_____ x (depth)
-	 *          .___|___|___|           \
-	 *          |   |   |                \ y (width)
-	 *  Row --> |___|___|
-	 *
-	 *          ^-- Column
+	 * Helper to identify the side based on the face index.
 	 *
+	 * This method is only valid after the face description array has been
+	 * initialized.
 	 */
-	virtual FMeshShapeGenerator& GenerateHollowBase()
+	ESide FaceToSide(int FaceId);
+
+protected:
+	enum class EStairStyle
 	{
-		// TODO: Not yet implemented.
-		check(false);
-		return *this;
-	}
+		Solid,		// Each step is supported by the floor.
+		Floating	// Each step only supported by previous step.
+	};
+	/** The style of the stairs */
+	EStairStyle StairStyle = EStairStyle::Solid;
 
 protected:
 	TArray<TArray<int>> VertexIds;
@@ -558,34 +101,8 @@ protected:
 	int LastFaceId = 0;
 
 protected:
-	/**
-	 * Reset state data on the generator.
-	 * 
-	 * This is invoked at the head of the Generate() method.
-	 */
-	virtual void Reset()
-	{
-		VertexIds.Reset();
-		VertexIdsToColumnRow.Reset();
-		FaceDesc.Reset();
-		NormalDesc.Reset();
-		UVDesc.Reset();
-
-		NumQuadsPerSide = 0;
-		NumQuads = 0;
-		NumVertsPerSide = 0;
-		NumVerts = 0;
-		NumAttrs = 0;
-
-		RightSideColumnId = 0;
-		LeftSideColumnId = 0;
-
-		FrontStartFaceId = 0;
-		TopStartFaceId = 0;
-		BackStartFaceId = 0;
-		BottomStartFaceId = 0;
-		LastFaceId = 0;
-	}
+	/** Reset state data on the generator */
+	virtual void Reset();
 
 	/**
 	 * Returns a vertex position.
@@ -593,7 +110,7 @@ protected:
 	 * The method is provided the Right or Left side and the
 	 * corresponding vertex column/row index to compute. Column
 	 * and row indices refer to the matrix-like ordered vertex layout.
-	 * See diagram in GenerateSolidBase() / GenerateHollowBase().
+	 * See diagram in GenerateSolidStairs() / GenerateFloatingStairs().
 	 *
 	 * This generator only generates vertices for the Right & Left
 	 * Sides of the stairs.
@@ -645,134 +162,53 @@ protected:
 /**
  * Generate an oriented Linear Stair mesh.
  */
-class /*GEOMETRICOBJECTS_API*/ FLinearStairGenerator : public FStairGenerator
+class GEOMETRICOBJECTS_API FLinearStairGenerator : public FStairGenerator
 {
 public:
 	/** The depth of each step. */
 	float StepDepth = 30.0f;
 
 public:
+	FLinearStairGenerator()
+	{
+		StairStyle = EStairStyle::Solid;
+	}
 	virtual ~FLinearStairGenerator() {}
 
 protected:
-	virtual FVector3d GenerateVertex(ESide Side, int VertexColumn, int VertexRow) override
-	{
-		if (Side != ESide::Right && Side != ESide::Left)
-		{
-			// Vertices are only generated for Right & Left sides.
-			check(false);
-		}
-		float X = VertexColumn * StepDepth;
-		float Y = (Side == ESide::Right ? 0.5 * StepWidth : -0.5 * StepWidth);
-		float Z = VertexRow * StepHeight;
-		return FVector3d(X, Y, Z);
-	}
+	typedef FStairGenerator Super;
 
-	virtual FVector3f GenerateNormal(ESide Side, int VertexId) override
-	{
-		FVector3f N;
-		switch (Side)
-		{
-		case ESide::Right:
-			N = FVector3f::UnitY();
-			break;
-		case ESide::Left:
-			N = -FVector3f::UnitY();
-			break;
-		case ESide::Front:
-			N = -FVector3f::UnitX();
-			break;
-		case ESide::Top:
-			N = FVector3f::UnitZ();
-			break;
-		case ESide::Back:
-			N = FVector3f::UnitX();
-			break;
-		case ESide::Bottom:
-			N = -FVector3f::UnitZ();
-			break;
-		default:
-			check(false);
-			break;
-		}
-		return N;
-	}
+protected:
+	virtual FVector3d GenerateVertex(ESide Side, int VertexColumn, int VertexRow) override;
+	virtual FVector3f GenerateNormal(ESide Side, int VertexId) override;
+	virtual FVector2f GenerateUV(ESide Side, int Step, int VertexId, float UVScale) override;
+	virtual float GetMaxDimension() override;
+};
 
-	virtual FVector2f GenerateUV(ESide Side, int Step, int VertexId, float UVScale) override
+/**
+ * Generate an oriented Floating Stair mesh.
+ */
+class GEOMETRICOBJECTS_API FFloatingStairGenerator : public FLinearStairGenerator
+{
+public:
+	FFloatingStairGenerator()
 	{
-		const int Col = VertexIdsToColumnRow[VertexId].A;
-		const int Row = VertexIdsToColumnRow[VertexId].B;
-		FVector2f UV(0.0f, 0.0f);
-		switch (Side)
-		{
-		case ESide::Right:
-		case ESide::Left:
-		{
-			const float UScale = NumSteps * StepDepth * UVScale;
-			const float VScale = NumSteps * StepHeight * UVScale;
-			UV.X = FMathf::Lerp(-0.5f, 0.5f, float(Col % LeftSideColumnId) / float(NumSteps + 1)) * UScale + 0.5f;
-			UV.Y = FMathf::Lerp(-0.5f, 0.5f, float(Row) / float(NumSteps + 1)) * VScale + 0.5f;
-			break;
-		}
-		case ESide::Front:
-		{
-			const float UScale = StepWidth * UVScale;
-			const float VScale = StepHeight * UVScale;
-			UV.X = Col < LeftSideColumnId ? 0.5f : -0.5f;
-			UV.Y = Row > Step ? -0.5f : 0.5f;
-			UV.X = UV.X * UScale + 0.5f;
-			UV.Y = UV.Y * VScale + 0.5f;
-			break;
-		}
-		case ESide::Top:
-		{
-			const float UScale = StepWidth * UVScale;
-			const float VScale = StepDepth * UVScale;
-			UV.X = Col < LeftSideColumnId ? 0.5f : -0.5f;
-			UV.Y = (Col % LeftSideColumnId) > Step ? -0.5f : 0.5f;
-			UV.X = UV.X * UScale + 0.5f;
-			UV.Y = UV.Y * VScale + 0.5f;
-			break;
-		}
-		case ESide::Back:
-		{
-			const float UScale = StepWidth * UVScale;
-			const float VScale = NumSteps * StepHeight * UVScale;
-			UV.X = Col < LeftSideColumnId ? 0.5f : -0.5f;
-			UV.Y = FMathf::Lerp(-0.5f, 0.5f, float(Row) / float(NumSteps + 1));
-			UV.X = UV.X * UScale + 0.5f;
-			UV.Y = UV.Y * VScale + 0.5f;
-			break;
-		}
-		case ESide::Bottom:
-		{
-			const float UScale = StepWidth * UVScale;
-			const float VScale = NumSteps * StepDepth * UVScale;
-			UV.X = Col < LeftSideColumnId ? 0.5f : -0.5f;
-			UV.Y = FMathf::Lerp(-0.5f, 0.5f, float(Col % LeftSideColumnId) / float(NumSteps + 1));
-			UV.X = UV.X * UScale + 0.5f;
-			UV.Y = UV.Y * VScale + 0.5f;
-			break;
-		}
-		default:
-		{
-			check(false);
-			break;
-		}
-		}
-		return UV;
+		StairStyle = EStairStyle::Floating;
 	}
+	virtual ~FFloatingStairGenerator() {}
 
-	virtual float GetMaxDimension()
-	{
-		return FMathf::Max(FMathf::Max(NumSteps * StepDepth, NumSteps * StepHeight), StepWidth);
-	}
+protected:
+	typedef FLinearStairGenerator Super;
+
+protected:
+	virtual FVector3d GenerateVertex(ESide Side, int VertexColumn, int VertexRow) override;
+	virtual FVector2f GenerateUV(ESide Side, int Step, int VertexId, float UVScale) override;
 };
 
 /**
  * Generate an oriented Curved Stair mesh.
  */
-class /*GEOMETRICOBJECTS_API*/ FCurvedStairGenerator : public FStairGenerator
+class GEOMETRICOBJECTS_API FCurvedStairGenerator : public FStairGenerator
 {
 public:
 	/** Inner radius of the curved staircase */
@@ -782,10 +218,14 @@ public:
 	float CurveAngle = 90.0f;
 
 public:
+	FCurvedStairGenerator()
+	{
+		StairStyle = EStairStyle::Solid;
+	}
 	virtual ~FCurvedStairGenerator() {}
 
 protected:
-	typedef FStairGenerator ParentClass;
+	typedef FStairGenerator Super;
 
 	/** Precompute/cached data */
 	bool bIsClockwise = true;
@@ -796,208 +236,33 @@ protected:
 	FVector3f BackNormal = FVector3f::Zero();
 
 protected:
-	virtual void Reset()
-	{
-		ParentClass::Reset();
-
-		bIsClockwise = CurveAngle > 0.0f;
-		CurveRadians = CurveAngle * TMathUtilConstants<float>::DegToRad;
-		CurveRadiansPerStep = CurveRadians / NumSteps;
-		OuterRadius = InnerRadius + StepWidth;
-		RadiusRatio = OuterRadius / InnerRadius;
-		BackNormal = FVector3f::Zero();
-	}
+	virtual void Reset() override;
+	virtual FVector3d GenerateVertex(ESide Side, int VertexColumn, int VertexRow) override;
+	virtual FVector3f GenerateNormal(ESide Side, int VertexId) override;
+	virtual FVector2f GenerateUV(ESide Side, int Step, int VertexId, float UVScale) override;
+	virtual float GetMaxDimension() override;
+};
 
-	virtual FVector3d GenerateVertex(ESide Side, int VertexColumn, int VertexRow) override
-	{
-		if (Side != ESide::Right && Side != ESide::Left)
-		{
-			// Vertices are only generated for Right & Left sides.
-			check(false);
-		}
-		float X = 0.0f;
-		float Y = 0.0f;
-		float Z = VertexRow * StepHeight;
-
-		float XCoeff = FMathf::Cos(VertexColumn * CurveRadiansPerStep);
-		float YCoeff = FMathf::Sin(VertexColumn * CurveRadiansPerStep);
-		if (bIsClockwise)
-		{
-			X = (Side == ESide::Right ? XCoeff * InnerRadius : XCoeff * OuterRadius);
-			Y = (Side == ESide::Right ? YCoeff * InnerRadius : YCoeff * OuterRadius);
-		}
-		else
-		{
-			X = (Side == ESide::Right ? XCoeff * -OuterRadius : XCoeff * -InnerRadius);
-			Y = (Side == ESide::Right ? YCoeff * -OuterRadius : YCoeff * -InnerRadius);
-		}
-		return FVector3d(X, Y, Z);
-	}
 
-	virtual FVector3f GenerateNormal(ESide Side, int VertexId) override
+/**
+ * Generate an oriented Curved Stair mesh.
+ */
+class GEOMETRICOBJECTS_API FSpiralStairGenerator : public FCurvedStairGenerator
+{
+public:
+	FSpiralStairGenerator()
 	{
-		const int Col = VertexIdsToColumnRow[VertexId].A;
-		const int Row = VertexIdsToColumnRow[VertexId].B;
-
-		FVector3f N;
-		switch (Side)
-		{
-		case ESide::Right:
-		case ESide::Left:
-		{
-			float X = FMathf::Cos((Col % (NumSteps + 1)) * CurveRadiansPerStep);
-			float Y = FMathf::Sin((Col % (NumSteps + 1)) * CurveRadiansPerStep);
-			N = (Side == ESide::Right ? FVector3f(-X, -Y, 0.0f) : FVector3f(X, Y, 0.0f));
-			N.Normalize();
-			break;
-		}
-		case ESide::Front:
-		{
-			float X = FMathf::Cos((Col % (NumSteps + 1)) * CurveRadiansPerStep);
-			float Y = FMathf::Sin((Col % (NumSteps + 1)) * CurveRadiansPerStep);
-			N = FVector3f(Y, -X, 0.0f);
-			N.Normalize();
-			break;
-		}
-		case ESide::Top:
-		{
-			N = FVector3f::UnitZ();
-			break;
-		}
-		case ESide::Back:
-		{
-			if (BackNormal == FVector3f::Zero())
-			{
-				float X = FMathf::Cos(NumSteps * CurveRadiansPerStep);
-				float Y = FMathf::Sin(NumSteps * CurveRadiansPerStep);
-				BackNormal = FVector3f(-Y, X, 0.0f);
-			}
-			N = BackNormal;
-			break;
-		}
-		case ESide::Bottom:
-		{
-			N = -FVector3f::UnitZ();
-			break;
-		}
-		default:
-		{
-			check(false);
-			break;
-		}
-		}
-		return N;
+		StairStyle = EStairStyle::Floating;
 	}
+	virtual ~FSpiralStairGenerator() {}
 
-	virtual FVector2f GenerateUV(ESide Side, int Step, int VertexId, float UVScale) override
-	{
-		const int Col = VertexIdsToColumnRow[VertexId].A;
-		const int Row = VertexIdsToColumnRow[VertexId].B;
-		FVector2f UV;
-		switch (Side)
-		{
-		case ESide::Right:
-		case ESide::Left:
-		{
-			const float UScale = OuterRadius * UVScale;
-			const float VScale = NumSteps * StepHeight * UVScale;
-			UV.X = FMathf::Lerp(-0.5f, 0.5f, float(Col % LeftSideColumnId) / float(NumSteps + 1));
-			UV.Y = FMathf::Lerp(-0.5f, 0.5f, float(Row) / float(NumSteps + 1));
-
-			// Proportionally scale the UVs of each side according
-			// to the RadiusRatio and ensure [0,1].
-			if (RadiusRatio * UScale > 1.0f)
-			{
-				UV.X /= (Side == ESide::Left ? UScale : UScale * RadiusRatio);
-			}
-			else if (Side == ESide::Left)
-			{
-				UV.X *= RadiusRatio;
-			}
-			UV.X = UV.X * UScale + 0.5f;
-			UV.Y = UV.Y * VScale + 0.5f;
-			break;
-		}
-		case ESide::Front:
-		{
-			const float UScale = StepWidth * UVScale;
-			const float VScale = StepHeight * UVScale;
-			UV.X = Col < LeftSideColumnId ? 0.5f : -0.5f;
-			UV.Y = Row > Step ? -0.5f : 0.5f;
-			UV.X = UV.X * UScale + 0.5f;
-			UV.Y = UV.Y * VScale + 0.5f;
-			break;
-		}
-		case ESide::Top:
-		{
-			const float UScale = StepWidth * UVScale;
-			const float VScale = OuterRadius / NumSteps * UVScale;
-			UV.X = Col < LeftSideColumnId ? 0.5f : -0.5f;
-			UV.Y = (Col % LeftSideColumnId) > Step ? -0.5f : 0.5f;
-
-			// Proportionally scale the outer edge of top faces
-			// to avoid stretching while ensuring [0,1].
-			if (RadiusRatio * VScale > 1.0f)
-			{
-				UV.Y /= (Col >= LeftSideColumnId ? VScale : VScale * RadiusRatio);
-			}
-			else if (Col >= LeftSideColumnId)
-			{
-				UV.Y *= RadiusRatio;
-			}
-			
-			UV.X = UV.X * UScale + 0.5f;
-			UV.Y = UV.Y * VScale + 0.5f;
-			break;
-		}
-		case ESide::Back:
-		{
-			const float UScale = StepWidth * UVScale;
-			const float VScale = NumSteps * StepHeight * UVScale;
-			UV.X = Col < LeftSideColumnId ? 0.5f : -0.5f;
-			UV.Y = FMathf::Lerp(-0.5f, 0.5f, float(Row) / float(NumSteps + 1));
-			UV.X = UV.X * UScale + 0.5f;
-			UV.Y = UV.Y * VScale + 0.5f;
-			break;
-		}
-		case ESide::Bottom:
-		{
-			const float UScale = StepWidth * UVScale;
-			const float VScale = OuterRadius * UVScale;
-			UV.X = Col < LeftSideColumnId ? 0.5f : -0.5f;
-			UV.Y = FMathf::Lerp(-0.5f, 0.5f, float(Col % LeftSideColumnId) / float(NumSteps + 1));
-
-			// Proportionally scale outer edge of bottom faces
-			// to avoid stretching while ensuring [0,1].
-			if (RadiusRatio * VScale > 1.0f)
-			{
-				UV.Y /= (Col >= LeftSideColumnId ? VScale : VScale * RadiusRatio);
-			}
-			else if (Col >= LeftSideColumnId)
-			{
-				UV.Y *= RadiusRatio;
-			}
-			UV.X = UV.X * UScale + 0.5f;
-			UV.Y = UV.Y * VScale + 0.5f;
-			break;
-		}
-		default:
-		{
-			check(false);
-			break;
-		}
-		}
-		return UV;
-	}
+protected:
+	typedef FCurvedStairGenerator Super;
 
-	virtual float GetMaxDimension()
-	{
-		float MaxDepth = FMathf::Abs(CurveRadians) * OuterRadius;
-		return FMathf::Max(FMathf::Max(MaxDepth, NumSteps * StepHeight), StepWidth);
-	}
+protected:
+	virtual FVector3d GenerateVertex(ESide Side, int VertexColumn, int VertexRow) override;
+	virtual FVector2f GenerateUV(ESide Side, int Step, int VertexId, float UVScale) override;
 };
 
-
-
 } // end namespace UE::Geometry
-} // end namespace UE
\ No newline at end of file
+} // end namespace UE
diff --git a/Engine/Plugins/Experimental/MeshModelingToolset/Source/MeshModelingTools/Private/AddPrimitiveTool.cpp b/Engine/Plugins/Experimental/MeshModelingToolset/Source/MeshModelingTools/Private/AddPrimitiveTool.cpp
--- a/Engine/Plugins/Experimental/MeshModelingToolset/Source/MeshModelingTools/Private/AddPrimitiveTool.cpp
+++ b/Engine/Plugins/Experimental/MeshModelingToolset/Source/MeshModelingTools/Private/AddPrimitiveTool.cpp
@@ -651,6 +651,18 @@ void UAddStairsPrimitiveTool::GenerateMesh(FDynamicMesh3* OutMesh) const
 		OutMesh->Copy(&StairGen);
 		break;
 	}
+	case EProceduralStairsType::Floating:
+	{
+		FFloatingStairGenerator StairGen;
+		StairGen.StepWidth = StairSettings->StepWidth;
+		StairGen.StepHeight = StairSettings->StepHeight;
+		StairGen.StepDepth = StairSettings->StepDepth;
+		StairGen.NumSteps = StairSettings->NumSteps;
+		StairGen.bPolygroupPerQuad = (ShapeSettings->PolygroupMode == EMakeMeshPolygroupMode::PerQuad);
+		StairGen.Generate();
+		OutMesh->Copy(&StairGen);
+		break;
+	}
 	case EProceduralStairsType::Curved:
 	{
 		FCurvedStairGenerator StairGen;
@@ -664,6 +676,19 @@ void UAddStairsPrimitiveTool::GenerateMesh(FDynamicMesh3* OutMesh) const
 		OutMesh->Copy(&StairGen);
 		break;
 	}
+	case EProceduralStairsType::Spiral:
+	{
+		FSpiralStairGenerator StairGen;
+		StairGen.StepWidth = StairSettings->StepWidth;
+		StairGen.StepHeight = StairSettings->StepHeight;
+		StairGen.NumSteps = StairSettings->NumSteps;
+		StairGen.InnerRadius = StairSettings->InnerRadius;
+		StairGen.CurveAngle = StairSettings->SpiralAngle;
+		StairGen.bPolygroupPerQuad = (ShapeSettings->PolygroupMode == EMakeMeshPolygroupMode::PerQuad);
+		StairGen.Generate();
+		OutMesh->Copy(&StairGen);
+		break;
+	}
 	}
 }
 
diff --git a/Engine/Plugins/Experimental/MeshModelingToolset/Source/MeshModelingTools/Public/AddPrimitiveTool.h b/Engine/Plugins/Experimental/MeshModelingToolset/Source/MeshModelingTools/Public/AddPrimitiveTool.h
--- a/Engine/Plugins/Experimental/MeshModelingToolset/Source/MeshModelingTools/Public/AddPrimitiveTool.h
+++ b/Engine/Plugins/Experimental/MeshModelingToolset/Source/MeshModelingTools/Public/AddPrimitiveTool.h
@@ -374,9 +374,13 @@ UENUM()
 enum class EProceduralStairsType
 {
 	/** Create a linear staircase */
-	Linear UMETA(DisplayName = "Linear Stairs"),
+	Linear UMETA(DisplayName = "Linear"),
+	/** Create a floating staircase */
+	Floating UMETA(DisplayName = "Floating"),
 	/** Create a curved staircase */
-	Curved UMETA(DisplayName = "Curved Stairs")
+	Curved UMETA(DisplayName = "Curved"),
+	/** Create a spiral staircase */
+	Spiral UMETA(DisplayName = "Spiral")
 };
 
 UCLASS()
@@ -401,15 +405,19 @@ public:
 	float StepHeight = 20.0f;
 
 	/** Depth of each step for linear stair shapes */
-	UPROPERTY(EditAnywhere, Category = "ShapeSettings", meta = (DisplayName = "Step Depth", UIMin = "1.0", UIMax = "1000.0", ClampMin = "0.0001", ClampMax = "1000000.0", ProceduralShapeSetting, EditCondition = "StairsType == EProceduralStairsType::Linear", EditConditionHides))
+	UPROPERTY(EditAnywhere, Category = "ShapeSettings", meta = (DisplayName = "Step Depth", UIMin = "1.0", UIMax = "1000.0", ClampMin = "0.0001", ClampMax = "1000000.0", ProceduralShapeSetting, EditCondition = "StairsType == EProceduralStairsType::Linear || StairsType == EProceduralStairsType::Floating", EditConditionHides))
 	float StepDepth = 30.0f;
 
 	/** Angular length of curved stair shapes */
 	UPROPERTY(EditAnywhere, Category = "ShapeSettings", meta = (DisplayName = "Curve Angle", UIMin = "-360.0", UIMax = "360.0", ClampMin = "-360.0", ClampMax = "360.0", ProceduralShapeSetting, EditCondition = "StairsType == EProceduralStairsType::Curved", EditConditionHides))
 	float CurveAngle = 90.0f;
 
+	/** Angular length of spiral stair shapes */
+	UPROPERTY(EditAnywhere, Category = "ShapeSettings", meta = (DisplayName = "Spiral Angle", UIMin = "-3600.0", UIMax = "3600.0", ClampMin = "-360000.0", ClampMax = "360000.0", ProceduralShapeSetting, EditCondition = "StairsType == EProceduralStairsType::Spiral", EditConditionHides))
+	float SpiralAngle = 90.0f;
+
 	/** Inner radius for curved/spiral stair shapes */
-	UPROPERTY(EditAnywhere, Category = "ShapeSettings", meta = (DisplayName = "Inner Radius", UIMin = "1.0", UIMax = "1000.0", ClampMin = "0.0001", ClampMax = "1000000.0", ProceduralShapeSetting, EditCondition = "StairsType == EProceduralStairsType::Curved", EditConditionHides))
+	UPROPERTY(EditAnywhere, Category = "ShapeSettings", meta = (DisplayName = "Inner Radius", UIMin = "1.0", UIMax = "1000.0", ClampMin = "0.0001", ClampMax = "1000000.0", ProceduralShapeSetting, EditCondition = "StairsType == EProceduralStairsType::Curved || StairsType == EProceduralStairsType::Spiral", EditConditionHides))
 	float InnerRadius = 150.0f;
 };
 
