GeometryProcessing: FVector2/3/4 refactoring. Move various member functions to standalone functions in UE::Geometry:: namespace, to simplify upcoming conversion to core vector types. Update all call sites.
#rnx

diff --git a/Engine/Plugins/Experimental/Enterprise/DataprepGeometryOperations/Source/Private/DataprepGeometryOperations.cpp b/Engine/Plugins/Experimental/Enterprise/DataprepGeometryOperations/Source/Private/DataprepGeometryOperations.cpp
--- a/Engine/Plugins/Experimental/Enterprise/DataprepGeometryOperations/Source/Private/DataprepGeometryOperations.cpp
+++ b/Engine/Plugins/Experimental/Enterprise/DataprepGeometryOperations/Source/Private/DataprepGeometryOperations.cpp
@@ -319,7 +319,7 @@ void UDataprepBakeTransformOperation::OnExecution_Implementation(const FDataprep
 				RemainingUniformScale = AbsScales[BestSubIdx];
 				if (RemainingUniformScale <= FLT_MIN)
 				{
-					RemainingUniformScale = AbsScales.MaxAbsElement();
+					RemainingUniformScale = UE::Geometry::MaxAbsElement(AbsScales);
 				}
 			}
 			switch (BakeScale)
diff --git a/Engine/Plugins/Experimental/GeometryProcessing/Source/DynamicMesh/Private/DynamicMesh3_Edits.cpp b/Engine/Plugins/Experimental/GeometryProcessing/Source/DynamicMesh/Private/DynamicMesh3_Edits.cpp
--- a/Engine/Plugins/Experimental/GeometryProcessing/Source/DynamicMesh/Private/DynamicMesh3_Edits.cpp
+++ b/Engine/Plugins/Experimental/GeometryProcessing/Source/DynamicMesh/Private/DynamicMesh3_Edits.cpp
@@ -896,19 +896,19 @@ EMeshResult FDynamicMesh3::SplitEdge(int eab, FEdgeSplitInfo& SplitInfo, double
 	if (IsBoundaryEdge(eab))
 	{
 		// create vertex
-		FVector3d vNew = FVector3d::Lerp(GetVertex(a), GetVertex(b), split_t);
+		FVector3d vNew = Lerp(GetVertex(a), GetVertex(b), split_t);
 		int f = AppendVertex(vNew);
 		if (HasVertexNormals())
 		{
-			SetVertexNormal(f, FVector3f::Lerp(GetVertexNormal(a), GetVertexNormal(b), (float)split_t).Normalized());
+			SetVertexNormal(f, Normalized(Lerp(GetVertexNormal(a), GetVertexNormal(b), (float)split_t)) );
 		}
 		if (HasVertexColors())
 		{
-			SetVertexColor(f, FVector3f::Lerp(GetVertexColor(a), GetVertexColor(b), (float)split_t));
+			SetVertexColor(f, Lerp(GetVertexColor(a), GetVertexColor(b), (float)split_t));
 		}
 		if (HasVertexUVs())
 		{
-			SetVertexUV(f, FVector2f::Lerp(GetVertexUV(a), GetVertexUV(b), (float)split_t));
+			SetVertexUV(f, Lerp(GetVertexUV(a), GetVertexUV(b), (float)split_t));
 		}
 
 		// look up edge bc, which needs to be modified
@@ -973,19 +973,19 @@ EMeshResult FDynamicMesh3::SplitEdge(int eab, FEdgeSplitInfo& SplitInfo, double
 		}
 
 		// create vertex
-		FVector3d vNew = FVector3d::Lerp(GetVertex(a), GetVertex(b), split_t);
+		FVector3d vNew = Lerp(GetVertex(a), GetVertex(b), split_t);
 		int f = AppendVertex(vNew);
 		if (HasVertexNormals())
 		{
-			SetVertexNormal(f, FVector3f::Lerp(GetVertexNormal(a), GetVertexNormal(b), (float)split_t).Normalized());
+			SetVertexNormal(f, Normalized(Lerp(GetVertexNormal(a), GetVertexNormal(b), (float)split_t)) );
 		}
 		if (HasVertexColors())
 		{
-			SetVertexColor(f, FVector3f::Lerp(GetVertexColor(a), GetVertexColor(b), (float)split_t));
+			SetVertexColor(f, Lerp(GetVertexColor(a), GetVertexColor(b), (float)split_t));
 		}
 		if (HasVertexUVs())
 		{
-			SetVertexUV(f, FVector2f::Lerp(GetVertexUV(a), GetVertexUV(b), (float)split_t));
+			SetVertexUV(f, Lerp(GetVertexUV(a), GetVertexUV(b), (float)split_t));
 		}
 
 		// look up edges that we are going to need to update
@@ -1582,18 +1582,18 @@ EMeshResult FDynamicMesh3::CollapseEdge(int vKeep, int vRemove, double collapse_
 	}
 
 	// set kept vertex to interpolated collapse position
-	SetVertex(vKeep, FVector3d::Lerp(KeptPos, RemovedPos, collapse_t));
+	SetVertex(vKeep, Lerp(KeptPos, RemovedPos, collapse_t));
 	if (HasVertexUVs())
 	{
-		SetVertexUV(vKeep, FVector2f::Lerp(GetVertexUV(vKeep), RemovedUV, collapse_t));
+		SetVertexUV(vKeep, Lerp(GetVertexUV(vKeep), RemovedUV, (float)collapse_t));
 	}
 	if (HasVertexNormals())
 	{
-		SetVertexNormal(vKeep, FVector3f::Lerp(GetVertexNormal(vKeep), RemovedNormal, collapse_t).Normalized());
+		SetVertexNormal(vKeep, Normalized(Lerp(GetVertexNormal(vKeep), RemovedNormal, (float)collapse_t)) );
 	}
 	if (HasVertexColors())
 	{
-		SetVertexColor(vKeep, FVector3f::Lerp(GetVertexColor(vKeep), RemovedColor, collapse_t));
+		SetVertexColor(vKeep, Lerp(GetVertexColor(vKeep), RemovedColor, (float)collapse_t));
 	}
 
 	CollapseInfo.KeptVertex = vKeep;
diff --git a/Engine/Plugins/Experimental/GeometryProcessing/Source/DynamicMesh/Private/DynamicMesh3_Queries.cpp b/Engine/Plugins/Experimental/GeometryProcessing/Source/DynamicMesh/Private/DynamicMesh3_Queries.cpp
--- a/Engine/Plugins/Experimental/GeometryProcessing/Source/DynamicMesh/Private/DynamicMesh3_Queries.cpp
+++ b/Engine/Plugins/Experimental/GeometryProcessing/Source/DynamicMesh/Private/DynamicMesh3_Queries.cpp
@@ -686,7 +686,7 @@ FVector3d FDynamicMesh3::GetEdgeNormal(int eID) const
 		if (Tris[1] != InvalidID)
 		{
 			n += GetTriNormal(Tris[1]);
-			n.Normalize();
+			Normalize(n);
 		}
 		return n;
 	}
@@ -741,7 +741,7 @@ FFrame3d FDynamicMesh3::GetVertexFrame(int vID, bool bFrameNormalY) const
 	int ovi = GetOtherEdgeVertex(eid, vID);
 	FVector3d ov = Vertices[ovi];
 	FVector3d edge = (ov - v);
-	edge.Normalize();
+	Normalize(edge);
 
 	FVector3d other = normal.Cross(edge);
 	edge = other.Cross(normal);
@@ -797,7 +797,7 @@ FVector3d FDynamicMesh3::GetTriBaryNormal(int tID, double bary0, double bary1, d
 	const FIndex3i& tIDs = Triangles[tID];
 	const TDynamicVector<FVector3f>& normalsR = VertexNormals.GetValue();
 	FVector3d n = FVector3d(bary0 * normalsR[tIDs[0]] + bary1 * normalsR[tIDs[1]] + bary2 * normalsR[tIDs[2]]);
-	n.Normalize();
+	Normalize(n);
 	return n;
 }
 
@@ -819,7 +819,7 @@ void FDynamicMesh3::GetTriBaryPoint(int tID, double bary0, double bary1, double
 		const TDynamicVector<FVector3f>& normalsR = this->VertexNormals.GetValue();
 		vinfo.Normal =
 		    (float)bary0 * normalsR[tIDs[0]] + (float)bary1 * normalsR[tIDs[1]] + (float)bary2 * normalsR[tIDs[2]];
-		vinfo.Normal.Normalize();
+		Normalize(vinfo.Normal);
 	}
 	vinfo.bHaveC = HasVertexColors();
 	if (vinfo.bHaveC)
@@ -853,11 +853,11 @@ FFrame3d FDynamicMesh3::GetTriFrame(int tID, int nEdge) const
 	                               Vertices[tIDs[(nEdge + 2) % 3]]};
 
 	FVector3d edge1 = TriVerts[1] - TriVerts[0];
-	edge1.Normalize();
+	Normalize(edge1);
 	FVector3d edge2 = TriVerts[2] - TriVerts[1];
-	edge2.Normalize();
+	Normalize(edge2);
 	FVector3d normal = edge2.Cross(edge1);
-	normal.Normalize();
+	Normalize(normal);
 	FVector3d other  = normal.Cross(edge1);
 	FVector3d center = (TriVerts[0] + TriVerts[1] + TriVerts[2]) / 3;
 
@@ -884,15 +884,15 @@ double FDynamicMesh3::GetTriInternalAngleR(int tID, int i) const
 	const FVector3d TV[3] = {Vertices[Triangle[0]], Vertices[Triangle[1]], Vertices[Triangle[2]]};
 	if (i == 0)
 	{
-		return (TV[1] - TV[0]).Normalized().AngleR((TV[2] - TV[0]).Normalized());
+		return AngleR( Normalized(TV[1] - TV[0]), Normalized(TV[2] - TV[0]) );
 	}
 	else if (i == 1)
 	{
-		return (TV[0] - TV[1]).Normalized().AngleR((TV[2] - TV[1]).Normalized());
+		return AngleR( Normalized(TV[0] - TV[1]), Normalized(TV[2] - TV[1]) );
 	}
 	else
 	{
-		return (TV[0] - TV[2]).Normalized().AngleR((TV[1] - TV[2]).Normalized());
+		return AngleR( Normalized(TV[0] - TV[2]), Normalized(TV[1] - TV[2]) );
 	}
 }
 
@@ -900,13 +900,13 @@ FVector3d FDynamicMesh3::GetTriInternalAnglesR(int tID) const
 {
 	const FIndex3i& Triangle = Triangles[tID];
 	const FVector3d TV[3] = {Vertices[Triangle[0]], Vertices[Triangle[1]], Vertices[Triangle[2]]};
-	FVector3d ABhat = (TV[1] - TV[0]).Normalized();
-	FVector3d BChat = (TV[2] - TV[1]).Normalized();
-	FVector3d AChat = (TV[2] - TV[0]).Normalized();
+	FVector3d ABhat = Normalized(TV[1] - TV[0]);
+	FVector3d BChat = Normalized(TV[2] - TV[1]);
+	FVector3d AChat = Normalized(TV[2] - TV[0]);
 	return FVector3d(
-		ABhat.AngleR(AChat),	// AB vs AC angle
-		(-ABhat).AngleR(BChat), // BA vs BC angle
-		(AChat).AngleR(BChat) // AC vs BC angle (same as CA vs CB, so don't need to negate both)
+		AngleR(ABhat, AChat),	// AB vs AC angle
+		AngleR(-ABhat, BChat), // BA vs BC angle
+		AngleR(AChat, BChat) // AC vs BC angle (same as CA vs CB, so don't need to negate both)
 		);
 }
 
diff --git a/Engine/Plugins/Experimental/GeometryProcessing/Source/DynamicMesh/Private/DynamicMeshEditor.cpp b/Engine/Plugins/Experimental/GeometryProcessing/Source/DynamicMesh/Private/DynamicMeshEditor.cpp
--- a/Engine/Plugins/Experimental/GeometryProcessing/Source/DynamicMesh/Private/DynamicMeshEditor.cpp
+++ b/Engine/Plugins/Experimental/GeometryProcessing/Source/DynamicMesh/Private/DynamicMeshEditor.cpp
@@ -726,7 +726,7 @@ FVector3f FDynamicMeshEditor::ComputeAndSetQuadNormal(const FIndex2i& QuadTris,
 	{
 		Normal = (FVector3f)Mesh->GetTriNormal(QuadTris.A);
 		Normal += (FVector3f)Mesh->GetTriNormal(QuadTris.B);
-		Normal.Normalize();
+		Normalize(Normal);
 	}
 	SetQuadNormals(QuadTris, Normal);
 	return Normal;
@@ -870,8 +870,8 @@ void FDynamicMeshEditor::SetTubeNormals(const TArray<int>& Triangles, const TArr
 		FVector3f Corners[4]{ (FVector3f)Mesh->GetVertex(VertexIDs1[M1[0]]), (FVector3f)Mesh->GetVertex(VertexIDs1[M1[1]]), (FVector3f)Mesh->GetVertex(VertexIDs2[M2[0]]), (FVector3f)Mesh->GetVertex(VertexIDs2[M2[1]]) };
 		FVector3f Edges[2]{ Corners[1] - Corners[0], Corners[3] - Corners[2] };
 		FVector3f Across = Corners[2] - Corners[0];
-		MatchedEdgeNormals[0][LastMatchedIdx] = Edges[0].Cross(Across).Normalized();
-		MatchedEdgeNormals[1][LastMatchedIdx] = Edges[1].Cross(Across).Normalized();
+		MatchedEdgeNormals[0][LastMatchedIdx] = Normalized(Edges[0].Cross(Across));
+		MatchedEdgeNormals[1][LastMatchedIdx] = Normalized(Edges[1].Cross(Across));
 	}
 
 	TArray<FVector3f> MatchedVertNormals[2];
@@ -882,8 +882,8 @@ void FDynamicMeshEditor::SetTubeNormals(const TArray<int>& Triangles, const TArr
 	VertNormals[1].SetNum(VertexIDs2.Num());
 	for (int LastMatchedIdx = NumMatched - 1, Idx = 0; Idx < NumMatched; LastMatchedIdx = Idx++)
 	{
-		MatchedVertNormals[0][Idx] = (MatchedEdgeNormals[0][LastMatchedIdx] + MatchedEdgeNormals[0][Idx]).Normalized();
-		MatchedVertNormals[1][Idx] = (MatchedEdgeNormals[1][LastMatchedIdx] + MatchedEdgeNormals[1][Idx]).Normalized();
+		MatchedVertNormals[0][Idx] = Normalized(MatchedEdgeNormals[0][LastMatchedIdx] + MatchedEdgeNormals[0][Idx]);
+		MatchedVertNormals[1][Idx] = Normalized(MatchedEdgeNormals[1][LastMatchedIdx] + MatchedEdgeNormals[1][Idx]);
 	}
 
 	TMap<int, int> VertToElID;
diff --git a/Engine/Plugins/Experimental/GeometryProcessing/Source/DynamicMesh/Private/GroupTopology.cpp b/Engine/Plugins/Experimental/GeometryProcessing/Source/DynamicMesh/Private/GroupTopology.cpp
--- a/Engine/Plugins/Experimental/GeometryProcessing/Source/DynamicMesh/Private/GroupTopology.cpp
+++ b/Engine/Plugins/Experimental/GeometryProcessing/Source/DynamicMesh/Private/GroupTopology.cpp
@@ -365,7 +365,7 @@ FVector3d FGroupTopology::GetEdgeMidpoint(int32 GroupEdgeID, double* ArcLengthOu
 		double a = (*PerVertexLengthsOut)[k-1], b = (*PerVertexLengthsOut)[k];
 		double t = (Len - a) / (b - a);
 		FVector3d A(Mesh->GetVertex(Vertices[k-1])), B(Mesh->GetVertex(Vertices[k]));
-		return FVector3d::Lerp(A, B, t);
+		return Lerp(A, B, t);
 	}
 
 	// compute arclen and then walk forward until we get halfway
@@ -383,7 +383,7 @@ FVector3d FGroupTopology::GetEdgeMidpoint(int32 GroupEdgeID, double* ArcLengthOu
 		{
 			double t = (Len - AccumLength) / (NewLen - AccumLength);
 			FVector3d A(Mesh->GetVertex(Vertices[k - 1])), B(Mesh->GetVertex(Vertices[k]));
-			return FVector3d::Lerp(A, B, t);
+			return Lerp(A, B, t);
 		}
 		AccumLength = NewLen;
 	}
@@ -627,7 +627,7 @@ bool FGroupTopology::GetGroupEdgeTangent(int GroupEdgeID, FVector3d& TangentOut)
 	FVector3d EndPos = Mesh->GetVertex(Edge.Span.Vertices[Edge.Span.Vertices.Num()-1]);
 	if (StartPos.DistanceSquared(EndPos) > 100 * FMathd::ZeroTolerance)
 	{
-		TangentOut = (EndPos - StartPos).Normalized();
+		TangentOut = Normalized(EndPos - StartPos);
 		return true;
 	}
 	else
@@ -650,7 +650,7 @@ FFrame3d FGroupTopology::GetGroupFrame(int32 GroupID) const
 		Normal += Mesh->GetTriNormal(tid);
 	}
 	Centroid /= (double)Face.Triangles.Num();
-	Normal.Normalize();
+	Normalize(Normal);
 	return FFrame3d(Centroid, Normal);
 }
 
@@ -715,7 +715,7 @@ FFrame3d FGroupTopology::GetSelectionFrame(const FGroupTopologySelection& Select
 	if (AccumCount > 0)
 	{
 		AccumulatedOrigin /= (double)AccumCount;
-		AccumulatedNormal.Normalize();
+		Normalize(AccumulatedNormal);
 
 		// We set our frame Z to be accumulated normal, and the other two axes are unconstrained, so
 		// we want to set them to something that will make our frame generally more useful. If the normal
@@ -728,7 +728,7 @@ FFrame3d FGroupTopology::GetSelectionFrame(const FGroupTopologySelection& Select
 		{
 			// Otherwise, let's place one of the other axes into the XY plane so that the frame is more
 			// useful for translation. We somewhat arbitrarily choose Y for this. 
-			FVector3d FrameY = AccumulatedNormal.Cross(FVector3d::UnitZ()).Normalized(); // orthogonal to world Z and frame Z 
+			FVector3d FrameY = Normalized(AccumulatedNormal.Cross(FVector3d::UnitZ())); // orthogonal to world Z and frame Z 
 			FVector3d FrameX = FrameY.Cross(AccumulatedNormal); // safe to not normalize because already orthogonal
 			AccumulatedFrame = FFrame3d(AccumulatedOrigin, FrameX, FrameY, AccumulatedNormal);
 		}
diff --git a/Engine/Plugins/Experimental/GeometryProcessing/Source/DynamicMesh/Private/MeshBoundaryLoops.cpp b/Engine/Plugins/Experimental/GeometryProcessing/Source/DynamicMesh/Private/MeshBoundaryLoops.cpp
--- a/Engine/Plugins/Experimental/GeometryProcessing/Source/DynamicMesh/Private/MeshBoundaryLoops.cpp
+++ b/Engine/Plugins/Experimental/GeometryProcessing/Source/DynamicMesh/Private/MeshBoundaryLoops.cpp
@@ -366,7 +366,7 @@ FVector3d FMeshBoundaryLoops::GetVertexNormal(int vid)
 	{
 		n += Mesh->GetTriNormal(ti);
 	}
-	n.Normalize();
+	Normalize(n);
 	return n;
 }
 
diff --git a/Engine/Plugins/Experimental/GeometryProcessing/Source/DynamicMesh/Private/MeshCurvature.cpp b/Engine/Plugins/Experimental/GeometryProcessing/Source/DynamicMesh/Private/MeshCurvature.cpp
--- a/Engine/Plugins/Experimental/GeometryProcessing/Source/DynamicMesh/Private/MeshCurvature.cpp
+++ b/Engine/Plugins/Experimental/GeometryProcessing/Source/DynamicMesh/Private/MeshCurvature.cpp
@@ -85,15 +85,15 @@ double TGaussianCurvature(const FDynamicMesh3& mesh, int32 v_i, GetPositionFuncT
 
 		FVector3d Vij = Vj - Vi;
 		FVector3d Nij(Vij);
-		double LenSqrVij = Nij.Normalize();
+		double LenSqrVij = Normalize(Nij);
 		LenSqrVij *= LenSqrVij;
 
 		FVector3d Vik = Vk - Vi;
 		FVector3d Nik(Vik);
-		double LenSqrVik = Nik.Normalize();
+		double LenSqrVik = Normalize(Nik);
 		LenSqrVik *= LenSqrVik;
 
-		double OpeningAngle = Nij.AngleR(Nik);
+		double OpeningAngle = AngleR(Nij, Nik);
 		AngleSum += OpeningAngle;
 
 		if (VectorUtil::IsObtuse(Vi, Vj, Vk))
diff --git a/Engine/Plugins/Experimental/GeometryProcessing/Source/DynamicMesh/Private/MeshNormals.cpp b/Engine/Plugins/Experimental/GeometryProcessing/Source/DynamicMesh/Private/MeshNormals.cpp
--- a/Engine/Plugins/Experimental/GeometryProcessing/Source/DynamicMesh/Private/MeshNormals.cpp
+++ b/Engine/Plugins/Experimental/GeometryProcessing/Source/DynamicMesh/Private/MeshNormals.cpp
@@ -74,7 +74,7 @@ void FMeshNormals::Compute_FaceAvg_AreaWeighted()
 
 	for (int VertIdx : Mesh->VertexIndicesItr())
 	{
-		Normals[VertIdx].Normalize();
+		Normalize(Normals[VertIdx]);
 	}
 }
 
@@ -103,7 +103,7 @@ void FMeshNormals::Compute_FaceAvg(bool bWeightByArea, bool bWeightByAngle)
 
 	for (int VertIdx : Mesh->VertexIndicesItr())
 	{
-		Normals[VertIdx].Normalize();
+		Normalize(Normals[VertIdx]);
 	}
 }
 
@@ -152,7 +152,7 @@ void FMeshNormals::Compute_Overlay_FaceAvg(const FDynamicMeshNormalOverlay* Norm
 
 	for (int ElemIdx : NormalOverlay->ElementIndicesItr())
 	{
-		Normals[ElemIdx].Normalize();
+		Normalize(Normals[ElemIdx]);
 	}
 }
 
@@ -178,7 +178,7 @@ void FMeshNormals::Compute_Overlay_FaceAvg_AreaWeighted(const FDynamicMeshNormal
 
 	for (int ElemIdx : NormalOverlay->ElementIndicesItr())
 	{
-		Normals[ElemIdx].Normalize();
+		Normalize(Normals[ElemIdx]);
 	}
 }
 
@@ -237,7 +237,7 @@ FVector3d FMeshNormals::ComputeVertexNormal(const FDynamicMesh3& Mesh, int VertI
 		int32 j = IndexUtil::FindTriIndex(VertIdx, Triangle);
 		SumNormal += TriNormal * TriNormalWeights[j];
 	});
-	return SumNormal.Normalized();
+	return Normalized(SumNormal);
 }
 
 
@@ -257,7 +257,7 @@ FVector3d FMeshNormals::ComputeVertexNormal(const FDynamicMesh3& Mesh, int32 Ver
 			NormalSum += TriNormal * TriNormalWeights[j];
 		}
 	});
-	return NormalSum.Normalized();
+	return Normalized(NormalSum);
 }
 
 
@@ -278,7 +278,7 @@ FVector3d FMeshNormals::ComputeOverlayNormal(const FDynamicMesh3& Mesh, const FD
 		}
 	});
 
-	return (Count > 0) ? SumNormal.Normalized() : FVector3d::Zero();
+	return (Count > 0) ? Normalized(SumNormal) : FVector3d::Zero();
 }
 
 
diff --git a/Engine/Plugins/Experimental/GeometryProcessing/Source/DynamicMesh/Private/MeshRegionBoundaryLoops.cpp b/Engine/Plugins/Experimental/GeometryProcessing/Source/DynamicMesh/Private/MeshRegionBoundaryLoops.cpp
--- a/Engine/Plugins/Experimental/GeometryProcessing/Source/DynamicMesh/Private/MeshRegionBoundaryLoops.cpp
+++ b/Engine/Plugins/Experimental/GeometryProcessing/Source/DynamicMesh/Private/MeshRegionBoundaryLoops.cpp
@@ -349,7 +349,7 @@ FVector3d FMeshRegionBoundaryLoops::GetVertexNormal(int vid)
 	{
 		n += Mesh->GetTriNormal(ti);
 	}
-	n.Normalize();
+	Normalize(n);
 	return n;
 }
 
diff --git a/Engine/Plugins/Experimental/GeometryProcessing/Source/DynamicMesh/Private/MeshSimplification.cpp b/Engine/Plugins/Experimental/GeometryProcessing/Source/DynamicMesh/Private/MeshSimplification.cpp
--- a/Engine/Plugins/Experimental/GeometryProcessing/Source/DynamicMesh/Private/MeshSimplification.cpp
+++ b/Engine/Plugins/Experimental/GeometryProcessing/Source/DynamicMesh/Private/MeshSimplification.cpp
@@ -835,7 +835,7 @@ static bool IsCollapsableDevelopableEdge(const FDynamicMesh3& Mesh, int32 Collap
 	// if we are not planar, we need to find the 'other' developable edge at RemoveV.
 	// This edge must be aligned w/ our collapse edge and have the same normals
 	FVector3d A = Mesh.GetVertex(RemoveV), B = Mesh.GetVertex(KeepV);
-	FVector3d EdgeDir(B - A); EdgeDir.Normalize();
+	FVector3d EdgeDir(B - A); Normalize(EdgeDir);
 	int32 FoldEdges = 0, FlatEdges = 0, OtherEdges = 0;
 	for (int32 eid : Mesh.VtxEdgesItr(RemoveV))
 	{
@@ -852,7 +852,7 @@ static bool IsCollapsableDevelopableEdge(const FDynamicMesh3& Mesh, int32 Collap
 			FIndex2i OtherEdgeV = Mesh.GetEdgeV(eid);
 			int32 OtherV = IndexUtil::FindEdgeOtherVertex(OtherEdgeV, RemoveV);
 			FVector3d C = Mesh.GetVertex(OtherV);
-			if ((A-C).Normalized().Dot(EdgeDir) > DotTolerance)
+			if ( Normalized(A-C).Dot(EdgeDir) > DotTolerance)
 			{
 				if ((Normal3.Dot(Normal1) > DotTolerance && Normal4.Dot(Normal2) > DotTolerance) ||
 					(Normal3.Dot(Normal2) > DotTolerance && Normal4.Dot(Normal1) > DotTolerance))
@@ -1568,7 +1568,7 @@ void TMeshSimplification<FAttrBasedQuadricErrord>::OnEdgeCollapse(int edgeID, in
 	Quadric.ComputeAttributes(collapse_pt, UpdatedNormald);
 
 	FVector3f UpdatedNormal(UpdatedNormald.X, UpdatedNormald.Y, UpdatedNormald.Z);
-	UpdatedNormal.Normalize();
+	Normalize(UpdatedNormal);
 
 	if (NormalOverlay != nullptr)
 	{
diff --git a/Engine/Plugins/Experimental/GeometryProcessing/Source/DynamicMesh/Private/MeshTangents.cpp b/Engine/Plugins/Experimental/GeometryProcessing/Source/DynamicMesh/Private/MeshTangents.cpp
--- a/Engine/Plugins/Experimental/GeometryProcessing/Source/DynamicMesh/Private/MeshTangents.cpp
+++ b/Engine/Plugins/Experimental/GeometryProcessing/Source/DynamicMesh/Private/MeshTangents.cpp
@@ -141,7 +141,7 @@ void TMeshTangents<RealType>::ComputeTriangleTangents(const FDynamicMeshUVOverla
 		{
 			double BitangentSign = VectorUtil::BitangentSign(TriNormal, Tangent, Bitangent);
 			Bitangent = VectorUtil::Bitangent(TriNormal, Tangent, BitangentSign);
-			Bitangent.Normalize();
+			Normalize(Bitangent);
 		}
 
 		Tangents[TriangleID] = (FVector3<RealType>)Tangent;
@@ -195,7 +195,7 @@ bool TMeshTangents<RealType>::CopyToOverlays(FDynamicMesh3& MeshToSet)
 
 static FVector3d PlaneProjectionNormalized(const FVector3d& Vector, const FVector3d& PlaneNormal)
 {
-	return (Vector - Vector.Dot(PlaneNormal) * PlaneNormal).Normalized();
+	return Normalized(Vector - Vector.Dot(PlaneNormal) * PlaneNormal);
 }
 
 
@@ -236,8 +236,8 @@ void TMeshTangents<RealType>::ComputeSeparatePerTriangleTangents(const FDynamicM
 			FVector3d ReconsBitangent = VectorUtil::Bitangent(VtxNormal, ProjectedTangent, BitangentSign);
 
 			SetPerTriangleTangent(TriangleID, j, 
-				((FVector3<RealType>)ProjectedTangent).Normalized(),
-				((FVector3<RealType>)ReconsBitangent).Normalized() );
+				Normalized((FVector3<RealType>)ProjectedTangent),
+				Normalized((FVector3<RealType>)ReconsBitangent) );
 		}
 
 	} /*, EParallelForFlags::ForceSingleThread */);
@@ -350,7 +350,7 @@ void TMeshTangents<RealType>::ComputeMikkTStyleTangents(const FDynamicMeshNormal
 				Mesh->GetTriVertices(TriangleID, TriVerts[0], TriVerts[1], TriVerts[2]);
 				FVector3d EdgeV1 = TriVerts[(Idx+1)%3] - TriVerts[Idx];
 				FVector3d EdgeV2 = TriVerts[(Idx+2)%3] - TriVerts[Idx];
-				Angle = EdgeV1.Normalized().AngleR(EdgeV2.Normalized());
+				Angle = AngleR(Normalized(EdgeV1), Normalized(EdgeV2));
 
 				// weight by angle
 				ProjTangent *= Angle;
@@ -373,8 +373,8 @@ void TMeshTangents<RealType>::ComputeMikkTStyleTangents(const FDynamicMeshNormal
 			}
 			else
 			{
-				double DotTangent = AccumTangent[FoundIndex].Normalized().Dot(ProjTangent);
-				double DotBitangent = AccumBitangent[FoundIndex].Normalized().Dot(ProjBitangent);
+				double DotTangent = Normalized(AccumTangent[FoundIndex]).Dot(ProjTangent);
+				double DotBitangent = Normalized(AccumBitangent[FoundIndex]).Dot(ProjBitangent);
 				if (DotTangent <= HardcodedSplitDotThresh || DotBitangent <= HardcodedSplitDotThresh )
 				{
 					// If angle between new and existing tangents is too large, we split into a new tangent group.
@@ -402,8 +402,8 @@ void TMeshTangents<RealType>::ComputeMikkTStyleTangents(const FDynamicMeshNormal
 		int32 NumTangents = AccumTangent.Num();
 		for (int32 k = 0; k < NumTangents; ++k)
 		{
-			AccumTangent[k].Normalize();
-			AccumBitangent[k].Normalize();
+			Normalize(AccumTangent[k]);
+			Normalize(AccumBitangent[k]);
 			if (AccumMagnitudeAngle[k].Z > 0)
 			{
 				AccumMagnitudeAngle[k].X /= AccumMagnitudeAngle[k].Z;
@@ -423,7 +423,7 @@ void TMeshTangents<RealType>::ComputeMikkTStyleTangents(const FDynamicMeshNormal
 
 			// todo: catch failures here?
 
-			Tangents[TriVertIndex] = ((FVector3<RealType>)Tangent).Normalized();
+			Tangents[TriVertIndex] = Normalized((FVector3<RealType>)Tangent);
 
 			FIndex3i NormTriangle = NormalOverlay->GetTriangle(TriangleID);
 			FVector3d TriVertNormal = NormalOverlay->IsSetTriangle(TriangleID) ?
@@ -432,7 +432,7 @@ void TMeshTangents<RealType>::ComputeMikkTStyleTangents(const FDynamicMeshNormal
 			// convert bitangent
 			double BinormalSign = VectorUtil::BitangentSign(TriVertNormal, Tangent, Bitangent);
 			FVector3d ReconsBitangent = VectorUtil::Bitangent(TriVertNormal, Tangent, BinormalSign);
-			Bitangents[TriVertIndex] = ((FVector3<RealType>)ReconsBitangent).Normalized();
+			Bitangents[TriVertIndex] = Normalized((FVector3<RealType>)ReconsBitangent);
 
 			//FVector3d OrigBitangent = AccumBitangent[UniqueIndex];
 			//FVector3<RealType> BTResult = Bitangents[3 * TriangleID + TriIndex];
diff --git a/Engine/Plugins/Experimental/GeometryProcessing/Source/DynamicMesh/Private/MeshTransforms.cpp b/Engine/Plugins/Experimental/GeometryProcessing/Source/DynamicMesh/Private/MeshTransforms.cpp
--- a/Engine/Plugins/Experimental/GeometryProcessing/Source/DynamicMesh/Private/MeshTransforms.cpp
+++ b/Engine/Plugins/Experimental/GeometryProcessing/Source/DynamicMesh/Private/MeshTransforms.cpp
@@ -127,7 +127,7 @@ void MeshTransforms::ApplyTransform(FDynamicMesh3& Mesh, const FTransform3d& Tra
 		{
 			FVector3f Normal = Mesh.GetVertexNormal(vid);
 			Normal = (FVector3f)Transform.TransformNormal((FVector3d)Normal);
-			Mesh.SetVertexNormal(vid, Normal.Normalized());
+			Mesh.SetVertexNormal(vid, Normalized(Normal));
 		}
 	}
 	if (Mesh.HasAttributes())
@@ -137,7 +137,7 @@ void MeshTransforms::ApplyTransform(FDynamicMesh3& Mesh, const FTransform3d& Tra
 		{
 			FVector3f Normal = Normals->GetElement(elemid);
 			Normal = (FVector3f)Transform.TransformNormal((FVector3d)Normal);
-			Normals->SetElement(elemid, Normal.Normalized());
+			Normals->SetElement(elemid, Normalized(Normal));
 		}
 	}
 }
@@ -161,7 +161,7 @@ void MeshTransforms::ApplyTransformInverse(FDynamicMesh3& Mesh, const FTransform
 			{
 				FVector3f Normal = Mesh.GetVertexNormal(vid);
 				Normal = (FVector3f)Transform.InverseTransformNormal((FVector3d)Normal);
-				Mesh.SetVertexNormal(vid, Normal.Normalized());
+				Mesh.SetVertexNormal(vid, Normalized(Normal));
 			}
 		}
 	});
@@ -178,7 +178,7 @@ void MeshTransforms::ApplyTransformInverse(FDynamicMesh3& Mesh, const FTransform
 				{
 					FVector3f Normal = Normals->GetElement(elemid);
 					Normal = (FVector3f)Transform.InverseTransformNormal((FVector3d)Normal);
-					Normals->SetElement(elemid, Normal.Normalized());
+					Normals->SetElement(elemid, Normalized(Normal));
 				}
 			});
 		}
@@ -206,7 +206,7 @@ void MeshTransforms::ApplyTransform(FDynamicMesh3& Mesh,
 			{
 				FVector3f Normal = Mesh.GetVertexNormal(vid);
 				Normal = NormalTransform(Normal);
-				Mesh.SetVertexNormal(vid, Normal.Normalized());
+				Mesh.SetVertexNormal(vid, Normalized(Normal));
 			}
 		}
 	});
@@ -223,7 +223,7 @@ void MeshTransforms::ApplyTransform(FDynamicMesh3& Mesh,
 				{
 					FVector3f Normal = Normals->GetElement(elemid);
 					Normal = NormalTransform(Normal);
-					Normals->SetElement(elemid, Normal.Normalized());
+					Normals->SetElement(elemid, Normalized(Normal));
 				}
 			});
 		}
diff --git a/Engine/Plugins/Experimental/GeometryProcessing/Source/DynamicMesh/Private/MeshWeights.cpp b/Engine/Plugins/Experimental/GeometryProcessing/Source/DynamicMesh/Private/MeshWeights.cpp
--- a/Engine/Plugins/Experimental/GeometryProcessing/Source/DynamicMesh/Private/MeshWeights.cpp
+++ b/Engine/Plugins/Experimental/GeometryProcessing/Source/DynamicMesh/Private/MeshWeights.cpp
@@ -64,7 +64,7 @@ FVector3d TMeanValueCentroid(const FDynamicMesh3& mesh, int32 v_i, GetPositionFu
 
 		FVector3d Vj = GetPositionFunc(v_j);
 		FVector3d vVj = (Vj - Vi);
-		double len_vVj = vVj.Normalize();
+		double len_vVj = Normalize(vVj);
 		// TODO: is this the right thing to do? if vertices are coincident,
 		//   weight of this vertex should be very high!
 		if (len_vVj < FMathf::ZeroTolerance)
@@ -74,12 +74,12 @@ FVector3d TMeanValueCentroid(const FDynamicMesh3& mesh, int32 v_i, GetPositionFu
 			wSum += w_ij;
 			continue;
 		}
-		FVector3d vVdelta = (GetPositionFunc(opp_v1) - Vi).Normalized();
+		FVector3d vVdelta = Normalized(GetPositionFunc(opp_v1) - Vi);
 		double w_ij = VectorUtil::VectorTanHalfAngle(vVj, vVdelta);
 
 		if (opp_v2 != FDynamicMesh3::InvalidID)
 		{
-			FVector3d vVgamma = (GetPositionFunc(opp_v2) - Vi).Normalized();
+			FVector3d vVgamma = Normalized(GetPositionFunc(opp_v2) - Vi);
 			w_ij += VectorUtil::VectorTanHalfAngle(vVj, vVgamma);
 		}
 
@@ -266,8 +266,8 @@ double TVoronoiArea(const FDynamicMesh3& mesh, int32 v_i, GetPositionFuncType Ge
 			FVector3d Vik = Vk - Vi;
 			double Dot = Vij.Dot(Vik);
 			double areaT = 0.5 * FMathd::Sqrt(Vij.SquaredLength() * Vik.SquaredLength() - Dot * Dot);
-			Vij.Normalize(); Vik.Normalize();
-			areaSum += (Vij.AngleR(Vik) > FMathd::HalfPi) ?    // obtuse at v_i ?
+			Normalize(Vij); Normalize(Vik);
+			areaSum += (AngleR(Vij, Vik) > FMathd::HalfPi) ?    // obtuse at v_i ?
 				(areaT * 0.5) : (areaT * 0.25);
 		}
 		else
diff --git a/Engine/Plugins/Experimental/GeometryProcessing/Source/DynamicMesh/Private/Operations/EmbedSurfacePath.cpp b/Engine/Plugins/Experimental/GeometryProcessing/Source/DynamicMesh/Private/Operations/EmbedSurfacePath.cpp
--- a/Engine/Plugins/Experimental/GeometryProcessing/Source/DynamicMesh/Private/Operations/EmbedSurfacePath.cpp
+++ b/Engine/Plugins/Experimental/GeometryProcessing/Source/DynamicMesh/Private/Operations/EmbedSurfacePath.cpp
@@ -146,7 +146,7 @@ bool WalkMeshPlanar(const FDynamicMesh3* Mesh, int StartTri, FVector3d StartPt,
 	auto PtTriPlaneSignedDist = [](FVector3d Pt, const FTriangle3d& Tri)
 	{
 		FVector3d Normal = Tri.Normal();
-		Normal.Normalize(FMathd::Epsilon);
+		Normalize(Normal,FMathd::Epsilon);
 		return (Pt - Tri.V[0]).Dot(Normal);
 	};
 
diff --git a/Engine/Plugins/Experimental/GeometryProcessing/Source/DynamicMesh/Private/Operations/ExtrudeMesh.cpp b/Engine/Plugins/Experimental/GeometryProcessing/Source/DynamicMesh/Private/Operations/ExtrudeMesh.cpp
--- a/Engine/Plugins/Experimental/GeometryProcessing/Source/DynamicMesh/Private/Operations/ExtrudeMesh.cpp
+++ b/Engine/Plugins/Experimental/GeometryProcessing/Source/DynamicMesh/Private/Operations/ExtrudeMesh.cpp
@@ -158,7 +158,7 @@ bool FExtrudeMesh::ApplyExtrude(FExtrusionInfo& Region, FMeshNormals* UseNormals
 					if (k == 0)
 					{
 						FVector3d FirstEdge = Mesh->GetVertex(BaseLoop.Vertices[1]) - Mesh->GetVertex(BaseLoop.Vertices[0]);
-						FirstEdge.Normalize();
+						Normalize(FirstEdge);
 						FirstProjectFrame = FFrame3d(FVector3d::Zero(), (FVector3d)Normal);
 						FirstProjectFrame.ConstrainedAlignAxis(0, FirstEdge, (FVector3d)Normal);
 						FrameUp = FirstProjectFrame.GetAxis(1);
diff --git a/Engine/Plugins/Experimental/GeometryProcessing/Source/DynamicMesh/Private/Operations/FFDLattice.cpp b/Engine/Plugins/Experimental/GeometryProcessing/Source/DynamicMesh/Private/Operations/FFDLattice.cpp
--- a/Engine/Plugins/Experimental/GeometryProcessing/Source/DynamicMesh/Private/Operations/FFDLattice.cpp
+++ b/Engine/Plugins/Experimental/GeometryProcessing/Source/DynamicMesh/Private/Operations/FFDLattice.cpp
@@ -24,7 +24,7 @@ FFFDLattice::FFFDLattice(const FVector3i& InDims, const FDynamicMesh3& Mesh, flo
 	FVector3d Extents = InitialBounds.Extents();
 
 	// Compute padding based on maximum component of the diagonal. Avoids problems when one or more component is zero.
-	double MaxDiagonal = InitialBounds.Diagonal().MaxElement();
+	double MaxDiagonal = MaxElement(InitialBounds.Diagonal());
 	Extents = Extents + (0.5 * ClampedPadding * MaxDiagonal);
 
 	InitialBounds.Min = Center - Extents;
@@ -208,7 +208,7 @@ void FFFDLattice::GetRotatedOverlayNormals(const TArray<FVector3d>& LatticeContr
 		// resulting normal will be flipped. So we instead multiply by det(J)*transpose(inv(J)) to get the sign right.
 		FMatrix3d InvJacobian = Jacobian.DeterminantTimesInverseTranspose();
 		OutNormals[OverlayElementID] = FMatrix3f(InvJacobian) * NormalOverlay->GetElement(OverlayElementID);
-		OutNormals[OverlayElementID].Normalize();
+		Normalize(OutNormals[OverlayElementID]);
 	};
 
 	ParallelFor(ElementCount, InterpolationJob, ParallelForFlags);
@@ -254,7 +254,7 @@ void FFFDLattice::GetRotatedMeshVertexNormals(const TArray<FVector3d>& LatticeCo
 		// resulting normal will be flipped. So we instead multiply by det(J)*transpose(inv(J)) to get the sign right.
 		FMatrix3d InvJacobian = Jacobian.DeterminantTimesInverseTranspose();
 		OutNormals[VertexID] = FMatrix3f( InvJacobian ) * OriginalNormals[VertexID];
-		OutNormals[VertexID].Normalize();
+		Normalize(OutNormals[VertexID]);
 	};
 
 	ParallelFor(MaxVertexID, InterpolationJob, ParallelForFlags);
diff --git a/Engine/Plugins/Experimental/GeometryProcessing/Source/DynamicMesh/Private/Operations/GroupEdgeInserter.cpp b/Engine/Plugins/Experimental/GeometryProcessing/Source/DynamicMesh/Private/Operations/GroupEdgeInserter.cpp
--- a/Engine/Plugins/Experimental/GeometryProcessing/Source/DynamicMesh/Private/Operations/GroupEdgeInserter.cpp
+++ b/Engine/Plugins/Experimental/GeometryProcessing/Source/DynamicMesh/Private/Operations/GroupEdgeInserter.cpp
@@ -412,7 +412,7 @@ void InsertNewVertexEndpoints(
 			bool bVertexIsOriginal = (CurrentVid == SpanVids[NextIndex - 1]);
 			if (!bVertexIsOriginal)
 			{
-				SplitPoint.Tangent = (Params.Mesh->GetVertex(SpanVids[NextIndex]) - Params.Mesh->GetVertex(CurrentVid)).Normalized();
+				SplitPoint.Tangent = Normalized(Params.Mesh->GetVertex(SpanVids[NextIndex]) - Params.Mesh->GetVertex(CurrentVid));
 			}
 			else
 			{
@@ -420,13 +420,13 @@ void InsertNewVertexEndpoints(
 				SplitPoint.Tangent = FVector3d::Zero();
 				if (NextIndex > 1)
 				{
-					SplitPoint.Tangent += (VertexPosition - Params.Mesh->GetVertex(SpanVids[NextIndex - 2])).Normalized();
+					SplitPoint.Tangent += Normalized(VertexPosition - Params.Mesh->GetVertex(SpanVids[NextIndex - 2]));
 				}
 				if (NextIndex < SpanVids.Num())
 				{
-					SplitPoint.Tangent += (Params.Mesh->GetVertex(SpanVids[NextIndex]) - VertexPosition).Normalized();
+					SplitPoint.Tangent += Normalized(Params.Mesh->GetVertex(SpanVids[NextIndex]) - VertexPosition);
 				}
-				SplitPoint.Tangent.Normalize();
+				Normalize(SplitPoint.Tangent);
 			}
 		}
 		else
@@ -453,7 +453,7 @@ void InsertNewVertexEndpoints(
 			// Assemble the actual output
 			SplitPoint.ElementID = CurrentVid;
 			SplitPoint.bIsVertex = true; // we made the vertex that is this target
-			SplitPoint.Tangent = (Params.Mesh->GetVertex(SpanVids[NextIndex]) - Params.Mesh->GetVertex(CurrentVid)).Normalized();
+			SplitPoint.Tangent = Normalized(Params.Mesh->GetVertex(SpanVids[NextIndex]) - Params.Mesh->GetVertex(CurrentVid));
 		}
 
 		EndPointsOut.Add(SplitPoint);
@@ -1134,11 +1134,11 @@ bool GetPlaneCutPath(const FDynamicMesh3& Mesh, int32 GroupID,
 	FVector3d EndPosition = EndPoint.bIsVertex ? Mesh.GetVertex(EndPoint.ElementID)
 		: Mesh.GetEdgePoint(EndPoint.ElementID, EndPoint.EdgeTValue);
 
-	FVector3d InPlaneVector = (EndPosition - StartPosition).Normalized();
+	FVector3d InPlaneVector = Normalized(EndPosition - StartPosition);
 
 	// Get components of the two tangents that are orthogonal to the vector between the points.
-	FVector3d NormalA = (StartPoint.Tangent - StartPoint.Tangent.Dot(InPlaneVector) * InPlaneVector).Normalized();
-	FVector3d NormalB = (StartPoint.Tangent - StartPoint.Tangent.Dot(InPlaneVector) * InPlaneVector).Normalized();
+	FVector3d NormalA = Normalized(StartPoint.Tangent - StartPoint.Tangent.Dot(InPlaneVector) * InPlaneVector);
+	FVector3d NormalB = Normalized(StartPoint.Tangent - StartPoint.Tangent.Dot(InPlaneVector) * InPlaneVector);
 
 	// Make the vectors be in the same half space so that their average represents the closer average of the
 	// corresponding lines.
@@ -1148,7 +1148,7 @@ bool GetPlaneCutPath(const FDynamicMesh3& Mesh, int32 GroupID,
 	}
 
 	// Do the averaging
-	FVector CutPlaneNormal = (FVector)(NormalA + NormalB).Normalized();
+	FVector CutPlaneNormal = (FVector)Normalized(NormalA + NormalB);
 	if (CutPlaneNormal.IsZero())
 	{
 		// This likely shouldn't happen, since it means that the two endpoint tangents are colinear
diff --git a/Engine/Plugins/Experimental/GeometryProcessing/Source/DynamicMesh/Private/Operations/InsetMeshRegion.cpp b/Engine/Plugins/Experimental/GeometryProcessing/Source/DynamicMesh/Private/Operations/InsetMeshRegion.cpp
--- a/Engine/Plugins/Experimental/GeometryProcessing/Source/DynamicMesh/Private/Operations/InsetMeshRegion.cpp
+++ b/Engine/Plugins/Experimental/GeometryProcessing/Source/DynamicMesh/Private/Operations/InsetMeshRegion.cpp
@@ -110,7 +110,7 @@ bool FInsetMeshRegion::ApplyInset(FInsetInfo& Region, FMeshNormals* UseNormals)
 			const FDynamicMesh3::FEdge EdgeVT = Mesh->GetEdge(LoopPair.InnerEdges[k]);
 			FVector3d A = Mesh->GetVertex(EdgeVT.Vert[0]);
 			FVector3d B = Mesh->GetVertex(EdgeVT.Vert[1]);
-			FVector3d EdgeDir = (A - B).Normalized();
+			FVector3d EdgeDir = Normalized(A - B);
 			FVector3d Midpoint = (A + B) * 0.5;
 			int32 EdgeTri = EdgeVT.Tri[0];
 			FVector3d Normal, Centroid; double Area;
@@ -340,7 +340,7 @@ bool FInsetMeshRegion::ApplyInset(FInsetInfo& Region, FMeshNormals* UseNormals)
 				if (k == 0)
 				{
 					FVector3d FirstEdge = Mesh->GetVertex(BaseLoopV[1]) - Mesh->GetVertex(BaseLoopV[0]);
-					FirstEdge.Normalize();
+					Normalize(FirstEdge);
 					FirstProjectFrame = FFrame3d(FVector3d::Zero(), (FVector3d)Normal);
 					FirstProjectFrame.ConstrainedAlignAxis(0, FirstEdge, (FVector3d)Normal);
 					FrameUp = FirstProjectFrame.GetAxis(1);
diff --git a/Engine/Plugins/Experimental/GeometryProcessing/Source/DynamicMesh/Private/Operations/MeshBoolean.cpp b/Engine/Plugins/Experimental/GeometryProcessing/Source/DynamicMesh/Private/Operations/MeshBoolean.cpp
--- a/Engine/Plugins/Experimental/GeometryProcessing/Source/DynamicMesh/Private/Operations/MeshBoolean.cpp
+++ b/Engine/Plugins/Experimental/GeometryProcessing/Source/DynamicMesh/Private/Operations/MeshBoolean.cpp
@@ -767,9 +767,9 @@ bool FMeshBoolean::CollapseWouldHurtTriangleQuality(
 			FVector3d Edge2(Verts[2] - Verts[0]);
 			FVector3d VCross(Edge2.Cross(Edge1));
 
-			// TODO: does this tolerance make a difference?  if not, set to zero and remove the VCross.Normalize()
+			// TODO: does this tolerance make a difference?  if not, set to zero and remove the Normalize(VCross)
 			double EdgeFlipTolerance = 1.e-5;
-			double Area2 = VCross.Normalize();
+			double Area2 = Normalize(VCross);
 			if (TryToImproveTriQualityThreshold > 0)
 			{
 				FVector3d Edge3(Verts[2] - Verts[1]);
@@ -849,7 +849,7 @@ bool FMeshBoolean::CollapseWouldChangeShapeOrUVs(
 				int OtherV = IndexUtil::FindEdgeOtherVertex(OtherEdgeV, RemoveV);
 				FVector3d OtherVPos = Mesh.GetVertex(OtherV);
 				FVector3d OtherEdgeDir = OtherVPos - RemoveVPos;
-				if (OtherEdgeDir.Normalize() == 0)
+				if (Normalize(OtherEdgeDir) == 0)
 				{
 					// collapsing degenerate edges above should prevent this
 					bHasBadEdge = true;
@@ -876,7 +876,7 @@ bool FMeshBoolean::CollapseWouldChangeShapeOrUVs(
 					FVector2f OtherUV = Mesh.GetVertexUV(OtherV);
 					FVector2f RemoveUV = Mesh.GetVertexUV(RemoveV);
 					FVector2f KeepUV = Mesh.GetVertexUV(KeepV);
-					if (FVector2f::Lerp(OtherUV, KeepUV, LerpT).DistanceSquared(RemoveUV) > UVEqualThresholdSq)
+					if (Lerp(OtherUV, KeepUV, LerpT).DistanceSquared(RemoveUV) > UVEqualThresholdSq)
 					{
 						bHasBadEdge = true;
 						return;
@@ -934,7 +934,7 @@ bool FMeshBoolean::CollapseWouldChangeShapeOrUVs(
 						FVector2f OtherUV = UVs->GetElement(OtherE);
 						FVector2f RemoveUV = UVs->GetElement(RemoveE);
 						FVector2f KeepUV = UVs->GetElement(KeepE);
-						if (FVector2f::Lerp(OtherUV, KeepUV, LerpT).DistanceSquared(RemoveUV) > UVEqualThresholdSq)
+						if (Lerp(OtherUV, KeepUV, LerpT).DistanceSquared(RemoveUV) > UVEqualThresholdSq)
 						{
 							bHasBadEdge = true;
 							return;
@@ -1054,7 +1054,7 @@ void FMeshBoolean::SimplifyAlongNewEdges(int NumMeshesToProcess, FDynamicMesh3*
 				FVector3d RemoveVPos = CutMesh[0]->GetVertex(Edge.Vert[RemoveVIdx]);
 				FVector3d KeepVPos = CutMesh[0]->GetVertex(Edge.Vert[KeepVIdx]);
 				FVector3d EdgeDir = KeepVPos - RemoveVPos;
-				if (EdgeDir.Normalize() == 0) // 0 is returned as a special case when the edge was too short to normalize
+				if (Normalize(EdgeDir) == 0) // 0 is returned as a special case when the edge was too short to normalize
 				{
 					// collapsing degenerate edges above should prevent this
 					ensure(!bCollapseDegenerateEdgesOnCut);
diff --git a/Engine/Plugins/Experimental/GeometryProcessing/Source/DynamicMesh/Private/Operations/MeshMeshCut.cpp b/Engine/Plugins/Experimental/GeometryProcessing/Source/DynamicMesh/Private/Operations/MeshMeshCut.cpp
--- a/Engine/Plugins/Experimental/GeometryProcessing/Source/DynamicMesh/Private/Operations/MeshMeshCut.cpp
+++ b/Engine/Plugins/Experimental/GeometryProcessing/Source/DynamicMesh/Private/Operations/MeshMeshCut.cpp
@@ -343,7 +343,7 @@ namespace MeshCut
 				FMeshSurfacePath SurfacePath(Mesh);
 				int StartTID = Mesh->GetVtxSingleTriangle(PtA.ElemID); // TODO: would be faster to have a PlanarWalk call that takes a start vertex ID !
 				FVector3d WalkPlaneNormal = BaseFaceNormals[Seg.BaseTID].Cross(PtB.Pos - PtA.Pos);
-				if (ensure(WalkPlaneNormal.Normalize() > 0))
+				if (ensure(Normalize(WalkPlaneNormal) > 0))
 				{
 					bool bWalkSuccess = SurfacePath.AddViaPlanarWalk(StartTID,
 						Mesh->GetVertex(PtA.ElemID), -1, PtB.ElemID,
diff --git a/Engine/Plugins/Experimental/GeometryProcessing/Source/DynamicMesh/Private/Operations/MeshMirror.cpp b/Engine/Plugins/Experimental/GeometryProcessing/Source/DynamicMesh/Private/Operations/MeshMirror.cpp
--- a/Engine/Plugins/Experimental/GeometryProcessing/Source/DynamicMesh/Private/Operations/MeshMirror.cpp
+++ b/Engine/Plugins/Experimental/GeometryProcessing/Source/DynamicMesh/Private/Operations/MeshMirror.cpp
@@ -24,7 +24,7 @@ void FMeshMirror::Mirror(FProgressCancel* Progress)
 	}
 
 	// We assume the normal is unit length everywhere, so it's worth forcing it to be just in case.
-	PlaneNormal.Normalize();
+	Normalize(PlaneNormal);
 
 	// Change the existing vertex positions and normals to be mirrored.
 	for (int32 Vid : Mesh->VertexIndicesItr())
@@ -86,7 +86,7 @@ void FMeshMirror::MirrorAndAppend(FProgressCancel* Progress)
 	}
 
 	// We assume the normal is unit length everywhere, so it's worth forcing it to be just in case.
-	PlaneNormal.Normalize();
+	Normalize(PlaneNormal);
 
 	// Start by creating and appending mirrored vertices. VidToMirrorMap will map each original
 	// vertex to the id of its mirror (which may be itself if it is not duplicated due to welding).
@@ -354,7 +354,7 @@ void MirrorAndAppendVerticesWithWelding(FDynamicMesh3& Mesh, const FVector3d& Pl
 
 				// For bowtie vertices, the projection of the normal onto the plane might be 0. The normal here is nonsensical anyway,
 				// so it doesn't matter what we do. Let's just leave it as is if we didn't get a projection.
-				if (NewNormal.Normalize() != 0)
+				if (Normalize(NewNormal) != 0)
 				{
 					// Projection worked
 					Mesh.SetVertexNormal(Vid, NewNormal);
diff --git a/Engine/Plugins/Experimental/GeometryProcessing/Source/DynamicMesh/Private/Operations/MeshSelfUnion.cpp b/Engine/Plugins/Experimental/GeometryProcessing/Source/DynamicMesh/Private/Operations/MeshSelfUnion.cpp
--- a/Engine/Plugins/Experimental/GeometryProcessing/Source/DynamicMesh/Private/Operations/MeshSelfUnion.cpp
+++ b/Engine/Plugins/Experimental/GeometryProcessing/Source/DynamicMesh/Private/Operations/MeshSelfUnion.cpp
@@ -557,7 +557,7 @@ void FMeshSelfUnion::SimplifyAlongNewEdges(TArray<int>& CutBoundaryEdges, TMap<i
 				FVector3d RemoveVPos = Mesh->GetVertex(Edge.Vert[RemoveVIdx]);
 				FVector3d KeepVPos = Mesh->GetVertex(Edge.Vert[KeepVIdx]);
 				FVector3d EdgeDir = KeepVPos - RemoveVPos;
-				if (EdgeDir.Normalize() == 0) // 0 is returned as a special case when the edge was too short to normalize
+				if (Normalize(EdgeDir) == 0) // 0 is returned as a special case when the edge was too short to normalize
 				{
 					// collapsing degenerate edges above should prevent this
 					ensure(!bCollapseDegenerateEdgesOnCut);
diff --git a/Engine/Plugins/Experimental/GeometryProcessing/Source/DynamicMesh/Private/Operations/MinimalHoleFiller.cpp b/Engine/Plugins/Experimental/GeometryProcessing/Source/DynamicMesh/Private/Operations/MinimalHoleFiller.cpp
--- a/Engine/Plugins/Experimental/GeometryProcessing/Source/DynamicMesh/Private/Operations/MinimalHoleFiller.cpp
+++ b/Engine/Plugins/Experimental/GeometryProcessing/Source/DynamicMesh/Private/Operations/MinimalHoleFiller.cpp
@@ -85,7 +85,7 @@ namespace
 
 	double EdgeFlipMetric(FVector3d Normal0, FVector3d Normal1, double FlipDotTol)
 	{
-		return (FlipDotTol == 0) ? Normal0.Dot(Normal1) : Normal0.Normalized().Dot(Normal1.Normalized());
+		return (FlipDotTol == 0) ? Normal0.Dot(Normal1) : Normalized(Normal0).Dot(Normalized(Normal1));
 	}
 
 	void GetEdgeFlipTris(const FDynamicMesh3& Mesh, int EdgeID,
diff --git a/Engine/Plugins/Experimental/GeometryProcessing/Source/DynamicMesh/Private/Operations/OffsetMeshRegion.cpp b/Engine/Plugins/Experimental/GeometryProcessing/Source/DynamicMesh/Private/Operations/OffsetMeshRegion.cpp
--- a/Engine/Plugins/Experimental/GeometryProcessing/Source/DynamicMesh/Private/Operations/OffsetMeshRegion.cpp
+++ b/Engine/Plugins/Experimental/GeometryProcessing/Source/DynamicMesh/Private/Operations/OffsetMeshRegion.cpp
@@ -288,7 +288,7 @@ bool FOffsetMeshRegion::ApplyOffset(FOffsetInfo& Region, FMeshNormals* UseNormal
 				if (k == 0)
 				{
 					FVector3d FirstEdge = Mesh->GetVertex(BaseLoopV[1]) - Mesh->GetVertex(BaseLoopV[0]);
-					FirstEdge.Normalize();
+					Normalize(FirstEdge);
 					FirstProjectFrame = FFrame3d(FVector3d::Zero(), (FVector3d)Normal);
 					FirstProjectFrame.ConstrainedAlignAxis(0, FirstEdge, (FVector3d)Normal);
 					FrameUp = FirstProjectFrame.GetAxis(1);
diff --git a/Engine/Plugins/Experimental/GeometryProcessing/Source/DynamicMesh/Private/Parameterization/MeshUVTransforms.cpp b/Engine/Plugins/Experimental/GeometryProcessing/Source/DynamicMesh/Private/Parameterization/MeshUVTransforms.cpp
--- a/Engine/Plugins/Experimental/GeometryProcessing/Source/DynamicMesh/Private/Parameterization/MeshUVTransforms.cpp
+++ b/Engine/Plugins/Experimental/GeometryProcessing/Source/DynamicMesh/Private/Parameterization/MeshUVTransforms.cpp
@@ -113,7 +113,7 @@ void UE::MeshUVTransforms::MakeSeamsDisjoint(FDynamicMeshUVOverlay* UVOverlay)
 			Result.Y = Random.GetFraction() * 2.0f - 1.0f;
 			Magnitude = Result.SquaredLength();
 		} while (Magnitude > 1.0f || Magnitude < FMathf::ZeroTolerance);
-		Result.Normalize();
+		Normalize(Result);
 		return JitterScale * Result;
 	};
 
diff --git a/Engine/Plugins/Experimental/GeometryProcessing/Source/DynamicMesh/Private/Sampling/MeshCurvatureMapBaker.cpp b/Engine/Plugins/Experimental/GeometryProcessing/Source/DynamicMesh/Private/Sampling/MeshCurvatureMapBaker.cpp
--- a/Engine/Plugins/Experimental/GeometryProcessing/Source/DynamicMesh/Private/Sampling/MeshCurvatureMapBaker.cpp
+++ b/Engine/Plugins/Experimental/GeometryProcessing/Source/DynamicMesh/Private/Sampling/MeshCurvatureMapBaker.cpp
@@ -151,11 +151,10 @@ void FMeshCurvatureMapBaker::Bake_Single()
 		double Curvature = CurvatureSampleFunction(Sample);
 
 		double Sign = FMathd::Sign(Curvature);
-		double T = ClampRange.GetT(FMathd::Abs(Curvature));
+		float T = (float)ClampRange.GetT(FMathd::Abs(Curvature));
 
 		FVector3f CurvatureColor = (Sign < 0) ?
-			FVector3f::Lerp(ZeroColor, NegativeColor, T)
-			: FVector3f::Lerp(ZeroColor, PositiveColor, T);
+			Lerp(ZeroColor, NegativeColor, T) : Lerp(ZeroColor, PositiveColor, T);
 
 		ResultBuilder->SetPixel(Coords, CurvatureColor);
 	});
diff --git a/Engine/Plugins/Experimental/GeometryProcessing/Source/DynamicMesh/Private/Sampling/MeshImageBakingCache.cpp b/Engine/Plugins/Experimental/GeometryProcessing/Source/DynamicMesh/Private/Sampling/MeshImageBakingCache.cpp
--- a/Engine/Plugins/Experimental/GeometryProcessing/Source/DynamicMesh/Private/Sampling/MeshImageBakingCache.cpp
+++ b/Engine/Plugins/Experimental/GeometryProcessing/Source/DynamicMesh/Private/Sampling/MeshImageBakingCache.cpp
@@ -235,7 +235,7 @@ bool FMeshImageBakingCache::ValidateCache()
 		{
 			//FVector3d BaseTriNormal = Mesh->GetTriNormal(SampleInfo.TriangleIndex);
 			NormalOverlay->GetTriBaryInterpolate<double>(SampleInfo.TriangleIndex, &SampleInfo.BaryCoords[0], &ValueOut.BaseNormal[0]);
-			ValueOut.BaseNormal.Normalize();
+			Normalize(ValueOut.BaseNormal);
 			FVector3d RayDir = ValueOut.BaseNormal;
 
 			ValueOut.BaseSample = SampleInfo;
diff --git a/Engine/Plugins/Experimental/GeometryProcessing/Source/DynamicMesh/Private/Sampling/MeshNormalMapBaker.cpp b/Engine/Plugins/Experimental/GeometryProcessing/Source/DynamicMesh/Private/Sampling/MeshNormalMapBaker.cpp
--- a/Engine/Plugins/Experimental/GeometryProcessing/Source/DynamicMesh/Private/Sampling/MeshNormalMapBaker.cpp
+++ b/Engine/Plugins/Experimental/GeometryProcessing/Source/DynamicMesh/Private/Sampling/MeshNormalMapBaker.cpp
@@ -32,7 +32,7 @@ void FMeshNormalMapBaker::Bake()
 			// sample normal on detail mesh
 			FVector3d DetailNormal;
 			DetailNormalOverlay->GetTriBaryInterpolate<double>(DetailTriID, &SampleData.DetailBaryCoords[0], &DetailNormal.X);
-			DetailNormal.Normalize();
+			Normalize(DetailNormal);
 
 			// compute normal in tangent space
 			double dx = DetailNormal.Dot(BaseTangentX);
diff --git a/Engine/Plugins/Experimental/GeometryProcessing/Source/DynamicMesh/Private/Sampling/MeshOcclusionMapBaker.cpp b/Engine/Plugins/Experimental/GeometryProcessing/Source/DynamicMesh/Private/Sampling/MeshOcclusionMapBaker.cpp
--- a/Engine/Plugins/Experimental/GeometryProcessing/Source/DynamicMesh/Private/Sampling/MeshOcclusionMapBaker.cpp
+++ b/Engine/Plugins/Experimental/GeometryProcessing/Source/DynamicMesh/Private/Sampling/MeshOcclusionMapBaker.cpp
@@ -28,7 +28,7 @@ void FMeshOcclusionMapBaker::Bake()
 		FVector3d P = Points[k];
 		if (P.Z > 0)
 		{
-			RayDirections.Add(P.Normalized());
+			RayDirections.Add( Normalized(P) );
 		}
 	}
 
@@ -52,7 +52,7 @@ void FMeshOcclusionMapBaker::Bake()
 			//FVector3d DetailTriNormal = DetailMesh.GetTriNormal(DetailTriID);
 			FVector3d DetailTriNormal;
 			DetailNormalOverlay->GetTriBaryInterpolate<double>(DetailTriID, &SampleData.DetailBaryCoords[0], &DetailTriNormal.X);
-			DetailTriNormal.Normalize();
+			Normalize(DetailTriNormal);
 
 			FVector3d DetailBaryCoords = SampleData.DetailBaryCoords;
 			FVector3d DetailPos = DetailMesh->GetTriBaryPoint(DetailTriID, DetailBaryCoords.X, DetailBaryCoords.Y, DetailBaryCoords.Z);
diff --git a/Engine/Plugins/Experimental/GeometryProcessing/Source/DynamicMesh/Private/Sampling/MeshPropertyMapBaker.cpp b/Engine/Plugins/Experimental/GeometryProcessing/Source/DynamicMesh/Private/Sampling/MeshPropertyMapBaker.cpp
--- a/Engine/Plugins/Experimental/GeometryProcessing/Source/DynamicMesh/Private/Sampling/MeshPropertyMapBaker.cpp
+++ b/Engine/Plugins/Experimental/GeometryProcessing/Source/DynamicMesh/Private/Sampling/MeshPropertyMapBaker.cpp
@@ -90,7 +90,7 @@ void FMeshPropertyMapBaker::Bake()
 						{
 							FVector3d DetailNormal;
 							DetailNormalOverlay->GetTriBaryInterpolate<double>(DetailTriID, &SampleData.DetailBaryCoords[0], &DetailNormal.X);
-							DetailNormal.Normalize();
+							Normalize(DetailNormal);
 							Color = NormalToColor(DetailNormal);
 						}
 					}
diff --git a/Engine/Plugins/Experimental/GeometryProcessing/Source/DynamicMesh/Private/ShapeApproximation/MeshSimpleShapeApproximation.cpp b/Engine/Plugins/Experimental/GeometryProcessing/Source/DynamicMesh/Private/ShapeApproximation/MeshSimpleShapeApproximation.cpp
--- a/Engine/Plugins/Experimental/GeometryProcessing/Source/DynamicMesh/Private/ShapeApproximation/MeshSimpleShapeApproximation.cpp
+++ b/Engine/Plugins/Experimental/GeometryProcessing/Source/DynamicMesh/Private/ShapeApproximation/MeshSimpleShapeApproximation.cpp
@@ -319,9 +319,9 @@ void FMeshSimpleShapeApproximation::Generate_ProjectedHulls(FSimpleShapeSet3d& S
 		if (AxisMode == EProjectedHullAxisMode::SmallestBoxDimension)
 		{
 			FAxisAlignedBox3d Bounds = Mesh.GetBounds();
-			int32 AxisIndex = Bounds.Diagonal().MinAbsElementIndex();
-			check(Bounds.Diagonal().MinAbsElement() == Bounds.Diagonal()[AxisIndex]);
-			ProjectionPlane = FFrame3d(FVector3d::Zero(), FVector3d::MakeUnit(AxisIndex));
+			int32 AxisIndex = MinAbsElementIndex(Bounds.Diagonal());
+			check(MinAbsElement(Bounds.Diagonal()) == Bounds.Diagonal()[AxisIndex]);
+			ProjectionPlane = FFrame3d(FVector3d::Zero(), MakeUnitVector3<double>(AxisIndex));
 		}
 		else if (AxisMode == EProjectedHullAxisMode::SmallestVolume)
 		{
@@ -345,7 +345,7 @@ void FMeshSimpleShapeApproximation::Generate_ProjectedHulls(FSimpleShapeSet3d& S
 		}
 		else
 		{
-			ProjectionPlane = FFrame3d(FVector3d::Zero(), FVector3d::MakeUnit((int32)AxisMode));
+			ProjectionPlane = FFrame3d(FVector3d::Zero(), MakeUnitVector3<double>((int32)AxisMode));
 		}
 
 		FMeshProjectionHull Hull(&Mesh);
diff --git a/Engine/Plugins/Experimental/GeometryProcessing/Source/DynamicMesh/Private/ShapeApproximation/SimpleShapeSet3.cpp b/Engine/Plugins/Experimental/GeometryProcessing/Source/DynamicMesh/Private/ShapeApproximation/SimpleShapeSet3.cpp
--- a/Engine/Plugins/Experimental/GeometryProcessing/Source/DynamicMesh/Private/ShapeApproximation/SimpleShapeSet3.cpp
+++ b/Engine/Plugins/Experimental/GeometryProcessing/Source/DynamicMesh/Private/ShapeApproximation/SimpleShapeSet3.cpp
@@ -329,7 +329,7 @@ static void TransformCapsuleShape(FCapsuleShape3d& CapsuleShape, const FTransfor
 
 	CapsuleShape.Capsule.Segment.Center = 0.5 * (P0 + P1);
 	CapsuleShape.Capsule.Segment.Direction = (P1 - P0);
-	CapsuleShape.Capsule.Segment.Extent = CapsuleShape.Capsule.Segment.Direction.Normalize() * 0.5;
+	CapsuleShape.Capsule.Segment.Extent = Normalize(CapsuleShape.Capsule.Segment.Direction) * 0.5;
 
 	double CurRadius = CapsuleShape.Capsule.Radius;
 	FFrame3d CapsuleFrame(CapsuleShape.Capsule.Segment.Center, CapsuleShape.Capsule.Segment.Direction);
@@ -382,7 +382,7 @@ static void TransformCapsuleShape(FCapsuleShape3d& CapsuleShape, const TArray<FT
 
 	CapsuleShape.Capsule.Segment.Center = 0.5 * (P0 + P1);
 	CapsuleShape.Capsule.Segment.Direction = (P1 - P0);
-	CapsuleShape.Capsule.Segment.Extent = CapsuleShape.Capsule.Segment.Direction.Normalize() * 0.5;
+	CapsuleShape.Capsule.Segment.Extent = Normalize(CapsuleShape.Capsule.Segment.Direction) * 0.5;
 	double RadiusScale = NewSideVec.Length() / InitialSideVec.Length();
 	CapsuleShape.Capsule.Radius *= RadiusScale;
 }
diff --git a/Engine/Plugins/Experimental/GeometryProcessing/Source/DynamicMesh/Public/MeshRefinerBase.h b/Engine/Plugins/Experimental/GeometryProcessing/Source/DynamicMesh/Public/MeshRefinerBase.h
--- a/Engine/Plugins/Experimental/GeometryProcessing/Source/DynamicMesh/Public/MeshRefinerBase.h
+++ b/Engine/Plugins/Experimental/GeometryProcessing/Source/DynamicMesh/Public/MeshRefinerBase.h
@@ -159,7 +159,7 @@ protected:
 		else
 		{
 			double ZeroTolerance = FMathd::ZeroTolerance;
-			return Direction0.Normalized(ZeroTolerance).Dot(Direction1.Normalized(ZeroTolerance));
+			return Normalized(Direction0, ZeroTolerance).Dot(Normalized(Direction1, ZeroTolerance));
 		}
 	}
 
diff --git a/Engine/Plugins/Experimental/GeometryProcessing/Source/DynamicMesh/Public/MeshTangents.h b/Engine/Plugins/Experimental/GeometryProcessing/Source/DynamicMesh/Public/MeshTangents.h
--- a/Engine/Plugins/Experimental/GeometryProcessing/Source/DynamicMesh/Public/MeshTangents.h
+++ b/Engine/Plugins/Experimental/GeometryProcessing/Source/DynamicMesh/Public/MeshTangents.h
@@ -134,10 +134,10 @@ public:
 		int32 k = TriangleID * 3;
 		if (k >= 0 && (k+2) < Tangents.Num())
 		{
-			TangentOut = FVector3<RealType>::Blend3(Tangents[k], Tangents[k+1], Tangents[k+2], BaryCoords.X, BaryCoords.Y, BaryCoords.Z);
-			TangentOut.Normalize();
-			BitangentOut = FVector3<RealType>::Blend3(Bitangents[k], Bitangents[k+1], Bitangents[k+2], BaryCoords.X, BaryCoords.Y, BaryCoords.Z);
-			BitangentOut.Normalize();
+			TangentOut = Blend3(Tangents[k], Tangents[k+1], Tangents[k+2], BaryCoords.X, BaryCoords.Y, BaryCoords.Z);
+			Normalize(TangentOut);
+			BitangentOut = Blend3(Bitangents[k], Bitangents[k+1], Bitangents[k+2], BaryCoords.X, BaryCoords.Y, BaryCoords.Z);
+			Normalize(BitangentOut);
 		}
 		else
 		{
diff --git a/Engine/Plugins/Experimental/GeometryProcessing/Source/GeometricObjects/Private/CompGeom/PolygonTriangulation.cpp b/Engine/Plugins/Experimental/GeometryProcessing/Source/GeometricObjects/Private/CompGeom/PolygonTriangulation.cpp
--- a/Engine/Plugins/Experimental/GeometryProcessing/Source/GeometricObjects/Private/CompGeom/PolygonTriangulation.cpp
+++ b/Engine/Plugins/Experimental/GeometryProcessing/Source/GeometricObjects/Private/CompGeom/PolygonTriangulation.cpp
@@ -194,7 +194,7 @@ void PolygonTriangulation::ComputePolygonPlane(const TArray<FVector3<T>>& Vertex
 		PlaneNormalOut.Y += (PositionJ.Z - PositionI.Z) * (PositionI.X + PositionJ.X);
 		PlaneNormalOut.Z += (PositionJ.X - PositionI.X) * (PositionI.Y + PositionJ.Y);
 	}
-	PlaneNormalOut.Normalize();
+	Normalize(PlaneNormalOut);
 	PlanePointOut /= (T)NumVertices;
 }
 
diff --git a/Engine/Plugins/Experimental/GeometryProcessing/Source/GeometricObjects/Private/Generators/StairGenerator.cpp b/Engine/Plugins/Experimental/GeometryProcessing/Source/GeometricObjects/Private/Generators/StairGenerator.cpp
--- a/Engine/Plugins/Experimental/GeometryProcessing/Source/GeometricObjects/Private/Generators/StairGenerator.cpp
+++ b/Engine/Plugins/Experimental/GeometryProcessing/Source/GeometricObjects/Private/Generators/StairGenerator.cpp
@@ -1077,7 +1077,7 @@ FVector3f FCurvedStairGenerator::GenerateNormal(ESide Side, int VertexId)
 		float X = FMathf::Cos((Col % (NumSteps + 1)) * CurveRadiansPerStep);
 		float Y = FMathf::Sin((Col % (NumSteps + 1)) * CurveRadiansPerStep);
 		N = (Side == ESide::Right ? FVector3f(-X, -Y, 0.0f) : FVector3f(X, Y, 0.0f));
-		N.Normalize();
+		Normalize(N);
 		break;
 	}
 	case ESide::Front:
@@ -1085,7 +1085,7 @@ FVector3f FCurvedStairGenerator::GenerateNormal(ESide Side, int VertexId)
 		float X = FMathf::Cos((Col % (NumSteps + 1)) * CurveRadiansPerStep);
 		float Y = FMathf::Sin((Col % (NumSteps + 1)) * CurveRadiansPerStep);
 		N = FVector3f(Y, -X, 0.0f);
-		N.Normalize();
+		Normalize(N);
 		break;
 	}
 	case ESide::Top:
diff --git a/Engine/Plugins/Experimental/GeometryProcessing/Source/GeometricObjects/Private/Generators/SweepGenerator.cpp b/Engine/Plugins/Experimental/GeometryProcessing/Source/GeometricObjects/Private/Generators/SweepGenerator.cpp
--- a/Engine/Plugins/Experimental/GeometryProcessing/Source/GeometricObjects/Private/Generators/SweepGenerator.cpp
+++ b/Engine/Plugins/Experimental/GeometryProcessing/Source/GeometricObjects/Private/Generators/SweepGenerator.cpp
@@ -52,7 +52,7 @@ int32 GetPolygonGroup(EProfileSweepPolygonGrouping PolygonGroupingMode, int32 Sw
 void FProfileSweepGenerator::AdjustNormalsForTriangle(int32 TriIndex, int32 FirstIndex, int32 SecondIndex, int32 ThirdIndex,
 	TArray<FVector3d> &WeightedNormals)
 {
-	FVector3d AbNormalized = (Vertices[SecondIndex] - Vertices[FirstIndex]).Normalized();
+	FVector3d AbNormalized = Normalized(Vertices[SecondIndex] - Vertices[FirstIndex]);
 	AdjustNormalsForTriangle(TriIndex, FirstIndex, SecondIndex, ThirdIndex, WeightedNormals, AbNormalized);
 }
 
@@ -73,14 +73,14 @@ void FProfileSweepGenerator::AdjustNormalsForTriangle(int32 TriIndex, int32 Firs
 	// per vertex later- we avoid adding the result into the per-vertex sums as we go along to avoid
 	// concurrent writes since this is done in parallel.
 
-	FVector3d Bc = (Vertices[ThirdIndex] - Vertices[SecondIndex]).Normalized();
-	FVector3d Ac = (Vertices[ThirdIndex] - Vertices[FirstIndex]).Normalized();
-	FVector3d TriangleNormal = Ac.Cross(AbNormalized).Normalized();
+	FVector3d Bc = Normalized(Vertices[ThirdIndex] - Vertices[SecondIndex]);
+	FVector3d Ac = Normalized(Vertices[ThirdIndex] - Vertices[FirstIndex]);
+	FVector3d TriangleNormal = Normalized(Ac.Cross(AbNormalized));
 
 	// Note that AngleR requires normalized inputs
-	WeightedNormals[TriIndex * 3] = TriangleNormal * AbNormalized.AngleR(Ac);
-	WeightedNormals[TriIndex * 3 + 1] = TriangleNormal * Bc.AngleR(-AbNormalized);
-	WeightedNormals[TriIndex * 3 + 2] = TriangleNormal * Ac.AngleR(Bc);
+	WeightedNormals[TriIndex * 3] = TriangleNormal * AngleR(AbNormalized,Ac);
+	WeightedNormals[TriIndex * 3 + 1] = TriangleNormal * AngleR(Bc, -AbNormalized);
+	WeightedNormals[TriIndex * 3 + 2] = TriangleNormal * AngleR(Ac, Bc);
 
 	// We can safely hook the triangle normals up, even though their calculation is incomplete.
 	SetTriangleNormals(TriIndex, FirstIndex, SecondIndex, ThirdIndex);
@@ -322,7 +322,7 @@ FMeshShapeGenerator& FProfileSweepGenerator::Generate()
 						|| Vertices[CurrentVert].DistanceSquared(Vertices[BottomRightVert])
 							/ Vertices[BottomVert].DistanceSquared(Vertices[RightVert]) <= (1 + DiagonalTolerance)*(1 + DiagonalTolerance))
 					{
-						FVector3d DiagonalDown = (Vertices[BottomRightVert] - Vertices[CurrentVert]).Normalized();
+						FVector3d DiagonalDown = Normalized(Vertices[BottomRightVert] - Vertices[CurrentVert]);
 
 						SetTriangle(TriIndex, CurrentVert, BottomRightVert, RightVert);
 						AdjustNormalsForTriangle(TriIndex, CurrentVert, BottomRightVert, RightVert, WeightedNormals, DiagonalDown);
@@ -344,7 +344,7 @@ FMeshShapeGenerator& FProfileSweepGenerator::Generate()
 					}
 					else
 					{
-						FVector3d DiagonalUp = (Vertices[RightVert] - Vertices[BottomVert]).Normalized();
+						FVector3d DiagonalUp = Normalized(Vertices[RightVert] - Vertices[BottomVert]);
 						SetTriangle(TriIndex, BottomVert, RightVert, CurrentVert);
 						AdjustNormalsForTriangle(TriIndex, BottomVert, RightVert, CurrentVert, WeightedNormals, DiagonalUp);
 						SetTriangleUVs(TriIndex,
@@ -391,7 +391,7 @@ FMeshShapeGenerator& FProfileSweepGenerator::Generate()
 	// Normalize and set them
 	ParallelFor(NumVerts, [this, &NormalSums](int VertIndex)
 	{
-		SetNormal(VertIndex, (FVector3f)(NormalSums[VertIndex].Normalized()), VertIndex);
+		SetNormal(VertIndex, (FVector3f)(Normalized(NormalSums[VertIndex])), VertIndex);
 	});
 
 	if (Progress && Progress->Cancelled())
diff --git a/Engine/Plugins/Experimental/GeometryProcessing/Source/GeometricObjects/Private/Spatial/PointSetHashTable.cpp b/Engine/Plugins/Experimental/GeometryProcessing/Source/GeometricObjects/Private/Spatial/PointSetHashTable.cpp
--- a/Engine/Plugins/Experimental/GeometryProcessing/Source/GeometricObjects/Private/Spatial/PointSetHashTable.cpp
+++ b/Engine/Plugins/Experimental/GeometryProcessing/Source/GeometricObjects/Private/Spatial/PointSetHashTable.cpp
@@ -59,7 +59,7 @@ bool FPointSetHashtable::FindPointsInBall(const FVector3d& QueryPt, double Query
 	// if we are close enough to cell border we need to check nbrs
 	// [TODO] could iterate over fewer cells here, if r is bounded by CellSize,
 	// then we should only ever need to look at 3, depending on which octant we are in.
-	if ((QueryPt - QueryCellCenter).MaxAbsElement() + QueryRadius > HalfCellSize)
+	if (MaxAbsElement(QueryPt - QueryCellCenter) + QueryRadius > HalfCellSize)
 	{
 		for (int ci = 0; ci < 26; ++ci)
 		{
@@ -70,7 +70,7 @@ bool FPointSetHashtable::FindPointsInBall(const FVector3d& QueryPt, double Query
 				QueryCellCenter.X + HalfCellSize*NbrOffset.X - QueryPt.X,
 				QueryCellCenter.Y + HalfCellSize*NbrOffset.Y - QueryPt.Y,
 				QueryCellCenter.Z + HalfCellSize*NbrOffset.Z - QueryPt.Z);
-			if (PtToFaceCenter.MinAbsElement() > QueryRadius)
+			if (MinAbsElement(PtToFaceCenter) > QueryRadius)
 			{
 				continue;
 			}
diff --git a/Engine/Plugins/Experimental/GeometryProcessing/Source/GeometricObjects/Private/Util/RevolveUtil.cpp b/Engine/Plugins/Experimental/GeometryProcessing/Source/GeometricObjects/Private/Util/RevolveUtil.cpp
--- a/Engine/Plugins/Experimental/GeometryProcessing/Source/GeometricObjects/Private/Util/RevolveUtil.cpp
+++ b/Engine/Plugins/Experimental/GeometryProcessing/Source/GeometricObjects/Private/Util/RevolveUtil.cpp
@@ -113,7 +113,7 @@ void RevolveUtil::MakeProfileCurveMidpointOfFirstStep(TArray<FVector3d>& Profile
 	for (FVector3d& Point : ProfileCurve)
 	{
 		FVector3d PointWrtAxisOrigin = Point - RevolutionAxisOrigin;
-		FVector3d AdjustmentDirection = -RevolutionAxisDirection.Cross(PointWrtAxisOrigin).Normalized();
+		FVector3d AdjustmentDirection = -Normalized(RevolutionAxisDirection.Cross(PointWrtAxisOrigin));
 		double DistToAxis = PointWrtAxisOrigin.Distance(PointWrtAxisOrigin.Dot(RevolutionAxisDirection)* RevolutionAxisDirection);
 		
 		Point += AdjustmentScale * DistToAxis * AdjustmentDirection;
diff --git a/Engine/Plugins/Experimental/GeometryProcessing/Source/GeometricObjects/Public/Curve/CurveUtil.h b/Engine/Plugins/Experimental/GeometryProcessing/Source/GeometricObjects/Public/Curve/CurveUtil.h
--- a/Engine/Plugins/Experimental/GeometryProcessing/Source/GeometricObjects/Public/Curve/CurveUtil.h
+++ b/Engine/Plugins/Experimental/GeometryProcessing/Source/GeometricObjects/Public/Curve/CurveUtil.h
@@ -36,7 +36,7 @@ struct TCurveUtil
 			EndIdx = FMath::Min(NV - 1, EndIdx);
 			StartIdx = FMath::Max(0, StartIdx);
 		}
-		return (Vertices[EndIdx] - Vertices[StartIdx]).Normalized();
+		return Normalized(Vertices[EndIdx] - Vertices[StartIdx]);
 	}
 
 
diff --git a/Engine/Plugins/Experimental/GeometryProcessing/Source/GeometricObjects/Public/Curve/DynamicGraph2.h b/Engine/Plugins/Experimental/GeometryProcessing/Source/GeometricObjects/Public/Curve/DynamicGraph2.h
--- a/Engine/Plugins/Experimental/GeometryProcessing/Source/GeometricObjects/Public/Curve/DynamicGraph2.h
+++ b/Engine/Plugins/Experimental/GeometryProcessing/Source/GeometricObjects/Public/Curve/DynamicGraph2.h
@@ -210,16 +210,16 @@ public:
 		FVector2<T> A = vertices.AsVector2(nbra);
 		FVector2<T> B = vertices.AsVector2(nbrb);
 		A -= V;
-		if (A.Normalize() == 0)
+		if (Normalize(A) == 0)
 		{
 			return InvalidValue;
 		}
 		B -= V;
-		if (B.Normalize() == 0)
+		if (Normalize(B) == 0)
 		{
 			return InvalidValue;
 		}
-		return A.AngleD(B);
+		return AngleD(A, B);
 	}
 
 protected:
diff --git a/Engine/Plugins/Experimental/GeometryProcessing/Source/GeometricObjects/Public/Distance/DistLine3Triangle3.h b/Engine/Plugins/Experimental/GeometryProcessing/Source/GeometricObjects/Public/Distance/DistLine3Triangle3.h
--- a/Engine/Plugins/Experimental/GeometryProcessing/Source/GeometricObjects/Public/Distance/DistLine3Triangle3.h
+++ b/Engine/Plugins/Experimental/GeometryProcessing/Source/GeometricObjects/Public/Distance/DistLine3Triangle3.h
@@ -56,7 +56,7 @@ public:
 		FVector3<Real> edge0 = Triangle.V[1] - Triangle.V[0];
 		FVector3<Real> edge1 = Triangle.V[2] - Triangle.V[0];
 		FVector3<Real> normal = edge0.Cross(edge1);
-		normal.Normalize();
+		Normalize(normal);
 		Real NdD = normal.Dot(Line.Direction);
 		if (TMathUtil<Real>::Abs(NdD) > TMathUtil<Real>::ZeroTolerance)
 		{
diff --git a/Engine/Plugins/Experimental/GeometryProcessing/Source/GeometricObjects/Public/Generators/BoxSphereGenerator.h b/Engine/Plugins/Experimental/GeometryProcessing/Source/GeometricObjects/Public/Generators/BoxSphereGenerator.h
--- a/Engine/Plugins/Experimental/GeometryProcessing/Source/GeometricObjects/Public/Generators/BoxSphereGenerator.h
+++ b/Engine/Plugins/Experimental/GeometryProcessing/Source/GeometricObjects/Public/Generators/BoxSphereGenerator.h
@@ -55,7 +55,7 @@ public:
 				double sz = z * FMathd::Sqrt(1.0 - x2*0.5 - y2*0.5 + x2*y2/3.0);
 				V = sx*AxX + sy*AxY + sz*AxZ;
 			}
-			V.Normalize();
+			Normalize(V);
 			Vertices[VertIdx] = V;
 		}
 
diff --git a/Engine/Plugins/Experimental/GeometryProcessing/Source/GeometricObjects/Public/Generators/GridBoxMeshGenerator.h b/Engine/Plugins/Experimental/GeometryProcessing/Source/GeometricObjects/Public/Generators/GridBoxMeshGenerator.h
--- a/Engine/Plugins/Experimental/GeometryProcessing/Source/GeometricObjects/Public/Generators/GridBoxMeshGenerator.h
+++ b/Engine/Plugins/Experimental/GeometryProcessing/Source/GeometricObjects/Public/Generators/GridBoxMeshGenerator.h
@@ -113,7 +113,7 @@ public:
 					
 					for (int EdgeVert = 1; EdgeVert + 1 < EdgeLen; EdgeVert++)
 					{
-						Vertices[CurrentVertIndex] = FVector3d::Lerp(Corners[0], Corners[1], EdgeVert * Nscale[Dim]);
+						Vertices[CurrentVertIndex] = Lerp(Corners[0], Corners[1], EdgeVert * Nscale[Dim]);
 						for (int WhichFace = 0; WhichFace < 2; WhichFace++) // each edge is shared by two faces (w/ major axes of subdim 0 and subdim 1 respectively)
 						{
 							int FaceInd = FaceInds[WhichFace];
@@ -156,7 +156,7 @@ public:
 					};
 					for (int D1 = 1; D1 + 1 < N[D[1][Dim]]; D1++)
 					{
-						Vertices[CurrentVertIndex] = FVector3d::Lerp(Edges[0], Edges[1], D1 * Nscale[D[1][Dim]]);
+						Vertices[CurrentVertIndex] = Lerp(Edges[0], Edges[1], D1 * Nscale[D[1][Dim]]);
 						FaceVertIndices[MajorFaceInd][ToFaceV(Dim, D0, D1)] = CurrentVertIndex;
 						CurrentVertIndex++;
 					}
@@ -164,7 +164,7 @@ public:
 			}
 		}
 
-		double MaxDimension = 2.0*Box.Extents.MaxAbsElement();
+		double MaxDimension = MaxAbsElement(2.0*Box.Extents);
 		float UVScale = (bScaleUVByAspectRatio) ? (1.0f / (float)MaxDimension) : 1.0f;
 
 		// create the face triangles and UVs+normals
diff --git a/Engine/Plugins/Experimental/GeometryProcessing/Source/GeometricObjects/Public/Generators/MarchingCubes.h b/Engine/Plugins/Experimental/GeometryProcessing/Source/GeometricObjects/Public/Generators/MarchingCubes.h
--- a/Engine/Plugins/Experimental/GeometryProcessing/Source/GeometricObjects/Public/Generators/MarchingCubes.h
+++ b/Engine/Plugins/Experimental/GeometryProcessing/Source/GeometricObjects/Public/Generators/MarchingCubes.h
@@ -848,7 +848,7 @@ protected:
 					b = PIso; fb = mid_f;
 				}
 			}
-			PIso = FVector3<double>::Lerp(a, b, 0.5);
+			PIso = Lerp(a, b, 0.5);
 
 		}
 		else
diff --git a/Engine/Plugins/Experimental/GeometryProcessing/Source/GeometricObjects/Public/Generators/MeshShapeGenerator.h b/Engine/Plugins/Experimental/GeometryProcessing/Source/GeometricObjects/Public/Generators/MeshShapeGenerator.h
--- a/Engine/Plugins/Experimental/GeometryProcessing/Source/GeometricObjects/Public/Generators/MeshShapeGenerator.h
+++ b/Engine/Plugins/Experimental/GeometryProcessing/Source/GeometricObjects/Public/Generators/MeshShapeGenerator.h
@@ -240,29 +240,29 @@ public:
 
 	static FVector3d BilinearInterp(const FVector3d &v00, const FVector3d &v10, const FVector3d &v11, const FVector3d &v01, double tx, double ty)
 	{
-		FVector3d a = FVector3d::Lerp(v00, v01, ty);
-		FVector3d b = FVector3d::Lerp(v10, v11, ty);
-		return FVector3d::Lerp(a, b, tx);
+		FVector3d a = Lerp(v00, v01, ty);
+		FVector3d b = Lerp(v10, v11, ty);
+		return Lerp(a, b, tx);
 	}
 
 	static FVector2d BilinearInterp(const FVector2d &v00, const FVector2d &v10, const FVector2d &v11, const FVector2d &v01, double tx, double ty)
 	{
-		FVector2d a = FVector2d::Lerp(v00, v01, ty);
-		FVector2d b = FVector2d::Lerp(v10, v11, ty);
-		return FVector2d::Lerp(a, b, tx);
+		FVector2d a = Lerp(v00, v01, ty);
+		FVector2d b = Lerp(v10, v11, ty);
+		return Lerp(a, b, tx);
 	}
 
 	static FVector2f BilinearInterp(const FVector2f &v00, const FVector2f &v10, const FVector2f &v11, const FVector2f &v01, float tx, float ty)
 	{
-		FVector2f a = FVector2f::Lerp(v00, v01, ty);
-		FVector2f b = FVector2f::Lerp(v10, v11, ty);
-		return FVector2f::Lerp(a, b, tx);
+		FVector2f a = Lerp(v00, v01, ty);
+		FVector2f b = Lerp(v10, v11, ty);
+		return Lerp(a, b, tx);
 	}
 
 
 	static FVector3i LinearInterp(const FVector3i &a, const FVector3i &b, double t)
 	{
-		FVector3d c = FVector3d::Lerp((FVector3d)a, (FVector3d)b, t);
+		FVector3d c = Lerp((FVector3d)a, (FVector3d)b, t);
 		return FVector3i((int)round(c.X), (int)round(c.Y), (int)round(c.Z));
 	}
 
diff --git a/Engine/Plugins/Experimental/GeometryProcessing/Source/GeometricObjects/Public/Generators/MinimalBoxMeshGenerator.h b/Engine/Plugins/Experimental/GeometryProcessing/Source/GeometricObjects/Public/Generators/MinimalBoxMeshGenerator.h
--- a/Engine/Plugins/Experimental/GeometryProcessing/Source/GeometricObjects/Public/Generators/MinimalBoxMeshGenerator.h
+++ b/Engine/Plugins/Experimental/GeometryProcessing/Source/GeometricObjects/Public/Generators/MinimalBoxMeshGenerator.h
@@ -37,7 +37,7 @@ public:
 			Vertices[i] = Box.GetCorner(i);
 		}
 
-		double MaxDimension = 2.0*Box.Extents.MaxAbsElement();
+		double MaxDimension = MaxAbsElement(2.0*Box.Extents);
 		float UVScale = (bScaleUVByAspectRatio) ? (1.0f / (float)MaxDimension) : 1.0f;
 
 		int TriIndex = 0;
diff --git a/Engine/Plugins/Experimental/GeometryProcessing/Source/GeometricObjects/Public/Generators/SweepGenerator.h b/Engine/Plugins/Experimental/GeometryProcessing/Source/GeometricObjects/Public/Generators/SweepGenerator.h
--- a/Engine/Plugins/Experimental/GeometryProcessing/Source/GeometricObjects/Public/Generators/SweepGenerator.h
+++ b/Engine/Plugins/Experimental/GeometryProcessing/Source/GeometricObjects/Public/Generators/SweepGenerator.h
@@ -422,7 +422,7 @@ public:
 		TArray<FVector2d> NormalSides; NormalSides.SetNum(NumX - 1);
 		for (int XIdx = 0; XIdx+1 < NumX; XIdx++)
 		{
-			NormalSides[XIdx] = (FVector2d(Radii[XIdx+1], Heights[XIdx+1]) - FVector2d(Radii[XIdx], Heights[XIdx])).Perp().Normalized();
+			NormalSides[XIdx] = Normalized( (FVector2d(Radii[XIdx+1], Heights[XIdx+1]) - FVector2d(Radii[XIdx], Heights[XIdx])).Perp() );
 		}
 		TArray<FVector2d> SmoothedNormalSides; SmoothedNormalSides.SetNum(NumX);
 		// smooth internal normals
@@ -430,7 +430,7 @@ public:
 		SmoothedNormalSides.Last() = NormalSides.Last();
 		for (int XIdx = 1; XIdx + 1 < NumX; XIdx++)
 		{
-			SmoothedNormalSides[XIdx] = (NormalSides[XIdx] + NormalSides[XIdx - 1]).Normalized();
+			SmoothedNormalSides[XIdx] = Normalized(NormalSides[XIdx] + NormalSides[XIdx - 1]);
 		}
 
 
@@ -717,7 +717,7 @@ public:
 
 		for (int k = 0; k < Normals.Num(); ++k)
 		{
-			Normals[k].Normalize();
+			Normalize(Normals[k]);
 		}
 
 		return *this;
diff --git a/Engine/Plugins/Experimental/GeometryProcessing/Source/GeometricObjects/Public/Image/ImageOccupancyMap.h b/Engine/Plugins/Experimental/GeometryProcessing/Source/GeometricObjects/Public/Image/ImageOccupancyMap.h
--- a/Engine/Plugins/Experimental/GeometryProcessing/Source/GeometricObjects/Public/Image/ImageOccupancyMap.h
+++ b/Engine/Plugins/Experimental/GeometryProcessing/Source/GeometricObjects/Public/Image/ImageOccupancyMap.h
@@ -106,7 +106,7 @@ public:
 			{
 				FVector3d A, B, C;
 				UVSpaceMesh.GetTriVertices(NearestTriID, A, B, C);
-				FTriangle2d UVTriangle(A.XY(), B.XY(), C.XY());
+				FTriangle2d UVTriangle(GetXY(A), GetXY(B), GetXY(C));
 
 				if (UVTriangle.IsInsideOrOn(UVPoint))
 				{
@@ -118,9 +118,9 @@ public:
 				else if (NearDistSqr < TexelDiag * TexelDiag)
 				{
 					FDistPoint3Triangle3d DistQuery = TMeshQueries<MeshType>::TriangleDistance(UVSpaceMesh, NearestTriID, UVPoint3d);
-					FVector2d NearestUV = DistQuery.ClosestTrianglePoint.XY();
+					FVector2d NearestUV = GetXY(DistQuery.ClosestTrianglePoint);
 					// nudge point into triangle to improve numerical behavior of things like barycentric coord calculation
-					NearestUV += (10.0 * FMathf::ZeroTolerance) * (NearestUV - UVPoint).Normalized();
+					NearestUV += (10.0 * FMathf::ZeroTolerance) * Normalized(NearestUV - UVPoint);
 
 					TexelType[LinearIdx] = InteriorTexel;
 					TexelQueryUV[LinearIdx] = (FVector2f)NearestUV;
@@ -130,7 +130,7 @@ public:
 				{
 					TexelType[LinearIdx] = GutterTexel;
 					FDistPoint3Triangle3d DistQuery = TMeshQueries<MeshType>::TriangleDistance(UVSpaceMesh, NearestTriID, UVPoint3d);
-					FVector2d NearestUV = DistQuery.ClosestTrianglePoint.XY();
+					FVector2d NearestUV = GetXY(DistQuery.ClosestTrianglePoint);
 					FVector2i NearestCoords = Dimensions.UVToCoords(NearestUV);
 					int64 NearestLinearIdx = Dimensions.GetIndex(NearestCoords);
 
diff --git a/Engine/Plugins/Experimental/GeometryProcessing/Source/GeometricObjects/Public/Implicit/CachingMeshSDF.h b/Engine/Plugins/Experimental/GeometryProcessing/Source/GeometricObjects/Public/Implicit/CachingMeshSDF.h
--- a/Engine/Plugins/Experimental/GeometryProcessing/Source/GeometricObjects/Public/Implicit/CachingMeshSDF.h
+++ b/Engine/Plugins/Experimental/GeometryProcessing/Source/GeometricObjects/Public/Implicit/CachingMeshSDF.h
@@ -139,7 +139,7 @@ public:
 		// figure out Origin & dimensions
 		FAxisAlignedBox3d Bounds = Spatial->GetBoundingBox();
 		
-		float MaxDim = (Bounds.Max - Bounds.Min + ExpandBounds * 2).MaxElement() + 4 * MaxOffsetDistance;
+		float MaxDim = MaxElement(Bounds.Max - Bounds.Min + ExpandBounds * 2) + 4 * MaxOffsetDistance;
 		if (!ensureMsgf(MaxDim / CellSize <= ApproxMaxCellsPerDimension - 4, TEXT("SDF resolution clamped to avoid excessive memory use")))
 		{
 			CellSize = MaxDim / (ApproxMaxCellsPerDimension - 4);
diff --git a/Engine/Plugins/Experimental/GeometryProcessing/Source/GeometricObjects/Public/Implicit/SweepingMeshSDF.h b/Engine/Plugins/Experimental/GeometryProcessing/Source/GeometricObjects/Public/Implicit/SweepingMeshSDF.h
--- a/Engine/Plugins/Experimental/GeometryProcessing/Source/GeometricObjects/Public/Implicit/SweepingMeshSDF.h
+++ b/Engine/Plugins/Experimental/GeometryProcessing/Source/GeometricObjects/Public/Implicit/SweepingMeshSDF.h
@@ -200,7 +200,7 @@ public:
 			ExactBandWidth = FMath::Max(ApproxMaxCellsPerDimension/2-1, 1);
 		}
 
-		float MaxDim = (Bounds.Max - Bounds.Min + ExpandBounds * 2).MaxElement();
+		float MaxDim = MaxElement(Bounds.Max - Bounds.Min + ExpandBounds * 2);
 		if (!ensureMsgf(MaxDim / CellSize <= ApproxMaxCellsPerDimension - 2 * ExactBandWidth, TEXT("SDF resolution clamped to avoid excessive memory use")))
 		{
 			CellSize = MaxDim / (ApproxMaxCellsPerDimension - 2 * ExactBandWidth);
diff --git a/Engine/Plugins/Experimental/GeometryProcessing/Source/GeometricObjects/Public/Intersection/IntrLine2Line2.h b/Engine/Plugins/Experimental/GeometryProcessing/Source/GeometricObjects/Public/Intersection/IntrLine2Line2.h
--- a/Engine/Plugins/Experimental/GeometryProcessing/Source/GeometricObjects/Public/Intersection/IntrLine2Line2.h
+++ b/Engine/Plugins/Experimental/GeometryProcessing/Source/GeometricObjects/Public/Intersection/IntrLine2Line2.h
@@ -111,7 +111,7 @@ public:
 
 		// if either line direction is not a normalized vector, 
 		//   results are garbage, so fail query
-		if (Line1.Direction.IsNormalized() == false || Line2.Direction.IsNormalized() == false) 
+		if (IsNormalized(Line1.Direction) == false || IsNormalized(Line2.Direction) == false) 
 		{
 			Type = EIntersectionType::Empty;
 			Result = EIntersectionResult::InvalidQuery;
diff --git a/Engine/Plugins/Experimental/GeometryProcessing/Source/GeometricObjects/Public/Intersection/IntrLine2Triangle2.h b/Engine/Plugins/Experimental/GeometryProcessing/Source/GeometricObjects/Public/Intersection/IntrLine2Triangle2.h
--- a/Engine/Plugins/Experimental/GeometryProcessing/Source/GeometricObjects/Public/Intersection/IntrLine2Triangle2.h
+++ b/Engine/Plugins/Experimental/GeometryProcessing/Source/GeometricObjects/Public/Intersection/IntrLine2Triangle2.h
@@ -85,7 +85,7 @@ public:
 
 		// if either Line Direction is not a normalized vector, 
 		//   results are garbage, so fail query
-		if (Line.Direction.IsNormalized() == false)
+		if (IsNormalized(Line.Direction) == false)
 		{
 			Type = EIntersectionType::Empty;
 			Result = EIntersectionResult::InvalidQuery;
diff --git a/Engine/Plugins/Experimental/GeometryProcessing/Source/GeometricObjects/Public/Intersection/IntrSegment2Segment2.h b/Engine/Plugins/Experimental/GeometryProcessing/Source/GeometricObjects/Public/Intersection/IntrSegment2Segment2.h
--- a/Engine/Plugins/Experimental/GeometryProcessing/Source/GeometricObjects/Public/Intersection/IntrSegment2Segment2.h
+++ b/Engine/Plugins/Experimental/GeometryProcessing/Source/GeometricObjects/Public/Intersection/IntrSegment2Segment2.h
@@ -146,7 +146,7 @@ public:
 
 		// if either segment direction is not a normalized vector, 
 		//   results are garbage, so fail query
-		if (Segment1.Direction.IsNormalized() == false || Segment2.Direction.IsNormalized() == false)
+		if (IsNormalized(Segment1.Direction) == false || IsNormalized(Segment2.Direction) == false)
 		{
 			Type = EIntersectionType::Empty;
 			Result = EIntersectionResult::InvalidQuery;
diff --git a/Engine/Plugins/Experimental/GeometryProcessing/Source/GeometricObjects/Public/Intersection/IntrSegment2Triangle2.h b/Engine/Plugins/Experimental/GeometryProcessing/Source/GeometricObjects/Public/Intersection/IntrSegment2Triangle2.h
--- a/Engine/Plugins/Experimental/GeometryProcessing/Source/GeometricObjects/Public/Intersection/IntrSegment2Triangle2.h
+++ b/Engine/Plugins/Experimental/GeometryProcessing/Source/GeometricObjects/Public/Intersection/IntrSegment2Triangle2.h
@@ -86,7 +86,7 @@ public:
 
 		// if either line direction is not a normalized vector, 
 		//   results are garbage, so fail query
-		if (Segment.Direction.IsNormalized() == false)
+		if (IsNormalized(Segment.Direction) == false)
 		{
 			Type = EIntersectionType::Empty;
 			Result = EIntersectionResult::InvalidQuery;
diff --git a/Engine/Plugins/Experimental/GeometryProcessing/Source/GeometricObjects/Public/Intersection/IntrTriangle3Triangle3.h b/Engine/Plugins/Experimental/GeometryProcessing/Source/GeometricObjects/Public/Intersection/IntrTriangle3Triangle3.h
--- a/Engine/Plugins/Experimental/GeometryProcessing/Source/GeometricObjects/Public/Intersection/IntrTriangle3Triangle3.h
+++ b/Engine/Plugins/Experimental/GeometryProcessing/Source/GeometricObjects/Public/Intersection/IntrTriangle3Triangle3.h
@@ -256,7 +256,7 @@ public:
 		E0[2] = Triangle0.V[0] - Triangle0.V[2];
 
 		// Get normal vector of Triangle0.
-		FVector3<Real> N0 = E0[0].UnitCross(E0[1]);
+		FVector3<Real> N0 = UnitCross(E0[0], E0[1]);
 
 		// Project Triangle1 onto normal line of Triangle0, test for separation.
 		double N0dT0V0 = N0.Dot(Triangle0.V[0]);
@@ -274,13 +274,13 @@ public:
 		E1[2] = Triangle1.V[0] - Triangle1.V[2];
 
 		// Get normal vector of Triangle1.
-		FVector3<Real> N1 = E1[0].UnitCross(E1[1]);
+		FVector3<Real> N1 = UnitCross(E1[0], E1[1]);
 
 		FVector3<Real> dir;
 		double min0, max0;
 		int i0, i1;
 
-		FVector3<Real> N0xN1 = N0.UnitCross(N1);
+		FVector3<Real> N0xN1 = UnitCross(N0, N1);
 		if (N0xN1.Dot(N0xN1) >= Tolerance)
 		{
 			// Triangles are not parallel.
@@ -299,7 +299,7 @@ public:
 			{
 				for (i0 = 0; i0 < 3; ++i0)
 				{
-					dir = E0[i0].UnitCross(E1[i1]);
+					dir = UnitCross(E0[i0], E1[i1]);
 					ProjectOntoAxis(Triangle0, dir, min0, max0);
 					ProjectOntoAxis(Triangle1, dir, min1, max1);
 					if (max0 < min1 - Tolerance || max1 < min0 - Tolerance)
@@ -322,7 +322,7 @@ public:
 			// Directions N0xE0[i0].
 			for (i0 = 0; i0 < 3; ++i0)
 			{
-				dir = N0.UnitCross(E0[i0]);
+				dir = UnitCross(N0, E0[i0]);
 				ProjectOntoAxis(Triangle0, dir, min0, max0);
 				ProjectOntoAxis(Triangle1, dir, min1, max1);
 				if (max0 < min1 - Tolerance || max1 < min0 - Tolerance)
@@ -334,7 +334,7 @@ public:
 			// Directions N1xE1[i1].
 			for (i1 = 0; i1 < 3; ++i1)
 			{
-				dir = N1.UnitCross(E1[i1]);
+				dir = UnitCross(N1, E1[i1]);
 				ProjectOntoAxis(Triangle0, dir, min0, max0);
 				ProjectOntoAxis(Triangle1, dir, min1, max1);
 				if (max0 < min1 - Tolerance || max1 < min0 - Tolerance)
@@ -361,7 +361,7 @@ public:
 		E0[2] = Triangle0.V[0] - Triangle0.V[2];
 
 		// Get normal vector of Triangle0.
-		FVector3<Real> N0 = E0[0].UnitCross(E0[1]);
+		FVector3<Real> N0 = UnitCross(E0[0], E0[1]);
 
 		// Project Triangle1 onto normal line of Triangle0, test for separation.
 		double N0dT0V0 = N0.Dot(Triangle0.V[0]);
@@ -378,13 +378,13 @@ public:
 		E1[2] = Triangle1.V[0] - Triangle1.V[2];
 
 		// Get normal vector of Triangle1.
-		FVector3<Real> N1 = E1[0].UnitCross(E1[1]);
+		FVector3<Real> N1 = UnitCross(E1[0], E1[1]);
 
 		FVector3<Real> dir;
 		double min0, max0;
 		int i0, i1;
 
-		FVector3<Real> N0xN1 = N0.UnitCross(N1);
+		FVector3<Real> N0xN1 = UnitCross(N0, N1);
 		if (N0xN1.Dot(N0xN1) >= Tolerance) {
 			// Triangles are not parallel.
 
@@ -399,7 +399,7 @@ public:
 			// Directions E0[i0]xE1[i1].
 			for (i1 = 0; i1 < 3; ++i1) {
 				for (i0 = 0; i0 < 3; ++i0) {
-					dir = E0[i0].UnitCross(E1[i1]);
+					dir = UnitCross(E0[i0], E1[i1]);
 					ProjectOntoAxis(Triangle0, dir, min0, max0);
 					ProjectOntoAxis(Triangle1, dir, min1, max1);
 					if (max0 < min1 - Tolerance || max1 < min0 - Tolerance) {
@@ -412,7 +412,7 @@ public:
 		else { // Triangles are parallel (and, in fact, coplanar).
 		 // Directions N0xE0[i0].
 			for (i0 = 0; i0 < 3; ++i0) {
-				dir = N0.UnitCross(E0[i0]);
+				dir = UnitCross(N0, E0[i0]);
 				ProjectOntoAxis(Triangle0, dir, min0, max0);
 				ProjectOntoAxis(Triangle1, dir, min1, max1);
 				if (max0 < min1 - Tolerance || max1 < min0 - Tolerance) {
@@ -422,7 +422,7 @@ public:
 
 			// Directions N1xE1[i1].
 			for (i1 = 0; i1 < 3; ++i1) {
-				dir = N1.UnitCross(E1[i1]);
+				dir = UnitCross(N1, E1[i1]);
 				ProjectOntoAxis(Triangle0, dir, min0, max0);
 				ProjectOntoAxis(Triangle1, dir, min1, max1);
 				if (max0 < min1 - Tolerance || max1 < min0 - Tolerance) {
@@ -543,7 +543,7 @@ public:
 			// Project onto yz-plane.
 			for (i = 0; i < 3; ++i)
 			{
-				projTri.V[i] = triangle.V[i].YZ();
+				projTri.V[i] = GetYZ(triangle.V[i]);
 				projEnd0.X = end0.Y;
 				projEnd0.Y = end0.Z;
 				projEnd1.X = end1.Y;
@@ -555,7 +555,7 @@ public:
 			// Project onto xz-plane.
 			for (i = 0; i < 3; ++i)
 			{
-				projTri.V[i] = triangle.V[i].XZ();
+				projTri.V[i] = GetXZ(triangle.V[i]);
 				projEnd0.X = end0.X;
 				projEnd0.Y = end0.Z;
 				projEnd1.X = end1.X;
@@ -567,7 +567,7 @@ public:
 			// Project onto xy-plane.
 			for (i = 0; i < 3; ++i)
 			{
-				projTri.V[i] = triangle.V[i].XY();
+				projTri.V[i] = GetXY(triangle.V[i]);
 				projEnd0.X = end0.X;
 				projEnd0.Y = end0.Y;
 				projEnd1.X = end1.X;
@@ -728,8 +728,8 @@ protected:
 			// Project onto yz-plane.
 			for (i = 0; i < 3; ++i)
 			{
-				projTri0.V[i] = tri0.V[i].YZ();
-				projTri1.V[i] = tri1.V[i].YZ();
+				projTri0.V[i] = GetYZ(tri0.V[i]);
+				projTri1.V[i] = GetYZ(tri1.V[i]);
 			}
 		}
 		else if (maxNormal == 1)
@@ -737,8 +737,8 @@ protected:
 			// Project onto xz-plane.
 			for (i = 0; i < 3; ++i)
 			{
-				projTri0.V[i] = tri0.V[i].XZ();
-				projTri1.V[i] = tri1.V[i].XZ();
+				projTri0.V[i] = GetXZ(tri0.V[i]);
+				projTri1.V[i] = GetXZ(tri1.V[i]);
 			}
 		}
 		else
@@ -746,8 +746,8 @@ protected:
 			// Project onto xy-plane.
 			for (i = 0; i < 3; ++i)
 			{
-				projTri0.V[i] = tri0.V[i].XY();
-				projTri1.V[i] = tri1.V[i].XY();
+				projTri0.V[i] = GetXY(tri0.V[i]);
+				projTri1.V[i] = GetXY(tri1.V[i]);
 			}
 		}
 
diff --git a/Engine/Plugins/Experimental/GeometryProcessing/Source/GeometricObjects/Public/LineTypes.h b/Engine/Plugins/Experimental/GeometryProcessing/Source/GeometricObjects/Public/LineTypes.h
--- a/Engine/Plugins/Experimental/GeometryProcessing/Source/GeometricObjects/Public/LineTypes.h
+++ b/Engine/Plugins/Experimental/GeometryProcessing/Source/GeometricObjects/Public/LineTypes.h
@@ -48,7 +48,7 @@ struct TLine2
 	 */
 	static TLine2<T> FromPoints(const FVector2<T>& Point0, const FVector2<T>& Point1)
 	{
-		return TLine2<T>(Point0, (Point1 - Point0).Normalized());
+		return TLine2<T>(Point0, Normalized(Point1 - Point0) );
 	}
 
 
@@ -173,7 +173,7 @@ struct TLine3
 	 */
 	static TLine3<T> FromPoints(const FVector3<T>& Point0, const FVector3<T>& Point1)
 	{
-		return TLine3<T>(Point0, (Point1 - Point0).Normalized());
+		return TLine3<T>(Point0, Normalized(Point1 - Point0) );
 	}
 
 
diff --git a/Engine/Plugins/Experimental/GeometryProcessing/Source/GeometricObjects/Public/MeshAdapterTransforms.h b/Engine/Plugins/Experimental/GeometryProcessing/Source/GeometricObjects/Public/MeshAdapterTransforms.h
--- a/Engine/Plugins/Experimental/GeometryProcessing/Source/GeometricObjects/Public/MeshAdapterTransforms.h
+++ b/Engine/Plugins/Experimental/GeometryProcessing/Source/GeometricObjects/Public/MeshAdapterTransforms.h
@@ -199,7 +199,7 @@ namespace MeshAdapterTransforms
 				{
 					FVector3f Normal = Mesh.GetNormal(elemid);
 					Normal = NormalTransform(Normal);
-					Mesh.SetNormal(elemid, Normal.Normalized());
+					Mesh.SetNormal(elemid, UE::Geometry::Normalized(Normal));
 				}
 			});
 		}
diff --git a/Engine/Plugins/Experimental/GeometryProcessing/Source/GeometricObjects/Public/PlaneTypes.h b/Engine/Plugins/Experimental/GeometryProcessing/Source/GeometricObjects/Public/PlaneTypes.h
--- a/Engine/Plugins/Experimental/GeometryProcessing/Source/GeometricObjects/Public/PlaneTypes.h
+++ b/Engine/Plugins/Experimental/GeometryProcessing/Source/GeometricObjects/Public/PlaneTypes.h
@@ -123,7 +123,7 @@ struct TPlane3
 		}
 
 		FVector3<RealType> DirectionVec = Point1 - Point0;
-		FVector3<RealType> Direction = DirectionVec.Normalized();
+		FVector3<RealType> Direction = Normalized(DirectionVec);
 		RealType Length = DirectionVec.Dot(Direction);
 
 		// test if segment is parallel to plane, if so, no intersection
diff --git a/Engine/Plugins/Experimental/GeometryProcessing/Source/GeometricObjects/Public/Polygon2.h b/Engine/Plugins/Experimental/GeometryProcessing/Source/GeometricObjects/Public/Polygon2.h
--- a/Engine/Plugins/Experimental/GeometryProcessing/Source/GeometricObjects/Public/Polygon2.h
+++ b/Engine/Plugins/Experimental/GeometryProcessing/Source/GeometricObjects/Public/Polygon2.h
@@ -182,7 +182,7 @@ public:
 	{
 		FVector2<T> next = Vertices[(VertexIndex + 1) % Vertices.Num()];
 		FVector2<T> prev = Vertices[VertexIndex == 0 ? Vertices.Num() - 1 : VertexIndex - 1];
-		return (next - prev).Normalized();
+		return Normalized(next - prev);
 	}
 
 
@@ -205,14 +205,14 @@ public:
 	{
 		FVector2<T> next = Vertices[(VertexIndex + 1) % Vertices.Num()];
 		FVector2<T> prev = Vertices[VertexIndex == 0 ? Vertices.Num() - 1 : VertexIndex - 1];
-		next -= Vertices[VertexIndex]; next.Normalize();
-		prev -= Vertices[VertexIndex]; prev.Normalize();
+		next -= Vertices[VertexIndex]; Normalize(next);
+		prev -= Vertices[VertexIndex]; Normalize(prev);
 
 		FVector2<T> n = (next.Perp() - prev.Perp());
-		T len = n.Normalize();
+		T len = Normalize(n);
 		if (len == 0) 
 		{
-			return (next + prev).Normalized();   // this gives right direction for degenerate angle
+			return Normalized(next + prev);   // this gives right direction for degenerate angle
 		}
 		else 
 		{
@@ -368,8 +368,8 @@ public:
 		ToNextOut = Vertices[(iVertex + 1) % N] - Vertices[iVertex];
 		if (bNormalize) 
 		{
-			ToPrevOut.Normalize();
-			ToNextOut.Normalize();
+			Normalize(ToPrevOut);
+			Normalize(ToNextOut);
 		}
 	}
 
@@ -381,7 +381,7 @@ public:
 	{
 		FVector2<T> e0, e1;
 		NeighbourVectors(iVertex, e0, e1, true);
-		return e0.AngleD(e1);
+		return AngleD(e0, e1);
 	}
 
 
@@ -658,7 +658,7 @@ public:
 
 		FVector2<T> n0 = GetNormal(iSeg);
 		FVector2<T> n1 = GetNormal((iSeg + 1) % Vertices.Num());
-		return ((T(1) - t) * n0 + t * n1).Normalized();
+		return Normalized((T(1) - t) * n0 + t * n1);
 	}
 
 
@@ -827,8 +827,8 @@ public:
 			FVector2<T> v = Vertices[k];
 			FVector2<T> next = Vertices[(k + 1) % Vertices.Num()];
 			FVector2<T> prev = Vertices[k == 0 ? Vertices.Num() - 1 : k - 1];
-			FVector2<T> dn = (next - v).Normalized();
-			FVector2<T> dp = (prev - v).Normalized();
+			FVector2<T> dn = Normalized(next - v);
+			FVector2<T> dp = Normalized(prev - v);
 			TLine2<T> ln(v + OffsetDistance * dn.Perp(), dn);
 			TLine2<T> lp(v - OffsetDistance * dp.Perp(), dp);
 
@@ -1031,9 +1031,9 @@ public:
 			FVector2<T> next = OldV[iNext];
 
 			FVector2<T> cp = prev - center;
-			T cpdist = cp.Normalize();
+			T cpdist = Normalize(cp);
 			FVector2<T> cn = next - center;
-			T cndist = cn.Normalize();
+			T cndist = Normalize(cn);
 
 			// if degenerate, skip this vert
 			if (cpdist < TMathUtil<T>::ZeroTolerance || cndist < TMathUtil<T>::ZeroTolerance) 
@@ -1042,7 +1042,7 @@ public:
 				continue;
 			}
 
-			T angle = cp.AngleD(cn);
+			T angle = AngleD(cp, cn);
 			// TODO document what this means sign-wise
 			// TODO re-test post Unreal port that this DotPerp is doing the right thing
 			T sign = cn.DotPerp(cp);
diff --git a/Engine/Plugins/Experimental/GeometryProcessing/Source/GeometricObjects/Public/Quaternion.h b/Engine/Plugins/Experimental/GeometryProcessing/Source/GeometricObjects/Public/Quaternion.h
--- a/Engine/Plugins/Experimental/GeometryProcessing/Source/GeometricObjects/Public/Quaternion.h
+++ b/Engine/Plugins/Experimental/GeometryProcessing/Source/GeometricObjects/Public/Quaternion.h
@@ -352,8 +352,8 @@ void TQuaternion<RealType>::SetFromTo(const FVector3<RealType>& From, const FVec
 	// [RMS] not ideal to explicitly normalize here, but if we don't,
 	//   output TQuaternion is not normalized and this causes problems,
 	//   eg like drift if we do repeated SetFromTo()
-	FVector3<RealType> from = From.Normalized(), to = To.Normalized();
-	FVector3<RealType> bisector = (from + to).Normalized();
+	FVector3<RealType> from = UE::Geometry::Normalized(From), to = UE::Geometry::Normalized(To);
+	FVector3<RealType> bisector = UE::Geometry::Normalized(from + to);
 	W = from.Dot(bisector);
 	if (W != 0) 
 	{
diff --git a/Engine/Plugins/Experimental/GeometryProcessing/Source/GeometricObjects/Public/RayTypes.h b/Engine/Plugins/Experimental/GeometryProcessing/Source/GeometricObjects/Public/RayTypes.h
--- a/Engine/Plugins/Experimental/GeometryProcessing/Source/GeometricObjects/Public/RayTypes.h
+++ b/Engine/Plugins/Experimental/GeometryProcessing/Source/GeometricObjects/Public/RayTypes.h
@@ -44,7 +44,7 @@ public:
 		this->Direction = Direction;
 		if (bIsNormalized == false)
 		{
-			this->Direction.Normalize();
+			Normalize(this->Direction);
 		}
 	}
 
diff --git a/Engine/Plugins/Experimental/GeometryProcessing/Source/GeometricObjects/Public/SegmentTypes.h b/Engine/Plugins/Experimental/GeometryProcessing/Source/GeometricObjects/Public/SegmentTypes.h
--- a/Engine/Plugins/Experimental/GeometryProcessing/Source/GeometricObjects/Public/SegmentTypes.h
+++ b/Engine/Plugins/Experimental/GeometryProcessing/Source/GeometricObjects/Public/SegmentTypes.h
@@ -37,7 +37,7 @@ public:
 		// set from endpoints 
 		Center = T(.5) * (Point0 + Point1);
 		Direction = Point1 - Point0;
-		Extent = T(.5) * Direction.Normalize();
+		Extent = T(.5) * Normalize(Direction);
 	}
 
 	/**
@@ -229,7 +229,7 @@ public:
 		}
 
 		// Lines are parallel.
-		diff.Normalize();
+		Normalize(diff);
 		T diffNDotPerpD1 = diff.DotPerp(OtherSegment.Direction);
 		if (TMathUtil<T>::Abs(diffNDotPerpD1) <= DotThresh)
 		{
@@ -312,7 +312,7 @@ protected:
 	{
 		Center = 0.5 * (p0 + p1);
 		Direction = p1 - p0;
-		Extent = 0.5 * Direction.Normalize();
+		Extent = 0.5 * Normalize(Direction);
 	}
 
 };
@@ -350,7 +350,7 @@ public:
 		// set from endpoints 
 		Center = T(.5) * (Point0 + Point1);
 		Direction = Point1 - Point0;
-		Extent = T(.5) * Direction.Normalize();
+		Extent = T(.5) * Normalize(Direction);
 	}
 
 	/**
@@ -533,7 +533,7 @@ protected:
 	{
 		Center = 0.5 * (p0 + p1);
 		Direction = p1 - p0;
-		Extent = 0.5 * Direction.Normalize();
+		Extent = 0.5 * Normalize(Direction);
 	}
 
 };
diff --git a/Engine/Plugins/Experimental/GeometryProcessing/Source/GeometricObjects/Public/TransformTypes.h b/Engine/Plugins/Experimental/GeometryProcessing/Source/GeometricObjects/Public/TransformTypes.h
--- a/Engine/Plugins/Experimental/GeometryProcessing/Source/GeometricObjects/Public/TransformTypes.h
+++ b/Engine/Plugins/Experimental/GeometryProcessing/Source/GeometricObjects/Public/TransformTypes.h
@@ -178,7 +178,7 @@ public:
 		const FVector3<RealType>& S = Scale3D;
 		RealType DetSign = FMathd::SignNonZero(S.X * S.Y * S.Z); // we only need to multiply by the sign of the determinant, rather than divide by it, since we normalize later anyway
 		FVector3<RealType> SafeInvS(S.Y*S.Z*DetSign, S.X*S.Z*DetSign, S.X*S.Y*DetSign);
-		return TransformVectorNoScale( (SafeInvS*Normal).Normalized() );
+		return TransformVectorNoScale( Normalized(SafeInvS*Normal) );
 	}
 
 
@@ -214,7 +214,7 @@ public:
 	 */
 	FVector3<RealType> InverseTransformNormal(const FVector3<RealType>& Normal) const
 	{
-		return InverseTransformVectorNoScale( (Scale3D*Normal).Normalized() );
+		return InverseTransformVectorNoScale( Normalized(Scale3D*Normal) );
 	}
 
 
diff --git a/Engine/Plugins/Experimental/GeometryProcessing/Source/GeometricObjects/Public/VectorTypes.h b/Engine/Plugins/Experimental/GeometryProcessing/Source/GeometricObjects/Public/VectorTypes.h
--- a/Engine/Plugins/Experimental/GeometryProcessing/Source/GeometricObjects/Public/VectorTypes.h
+++ b/Engine/Plugins/Experimental/GeometryProcessing/Source/GeometricObjects/Public/VectorTypes.h
@@ -223,96 +223,103 @@ struct FVector2
 		return *this;
 	}
 
-	constexpr bool IsNormalized()
+	constexpr bool operator==(const FVector2<T>& Other) const
 	{
-		return TMathUtil<T>::Abs((X * X + Y * Y) - 1) < TMathUtil<T>::ZeroTolerance;
+		return X == Other.X && Y == Other.Y;
 	}
 
-	// Angle in Degrees
-	T AngleD(const FVector2<T>& V2) const
+	constexpr bool operator!=(const FVector2<T>& Other) const
 	{
-		T DotVal = Dot(V2);
-		T ClampedDot = (DotVal < (T)-1) ? (T)-1 : ((DotVal > (T)1) ? (T)1 : DotVal);
-		return (T)(acos(ClampedDot) * (T)(180.0 / 3.14159265358979323846));
+		return X != Other.X || Y != Other.Y;
 	}
 
-	// Angle in Radians
-	T AngleR(const FVector2<T>& V2) const
-	{
-		T DotVal = Dot(V2);
-		T ClampedDot = (DotVal < (T)-1) ? (T)-1 : ((DotVal > (T)1) ? (T)1 : DotVal);
-		return (T)acos(ClampedDot);
-	}
 
-	// Angle in Radians
-	T SignedAngleR(const FVector2<T>& V2) const
+	// Note: This was called "IsLeft" in the Geometry3Sharp code (where it also went through Math.Sign)
+	// Returns >0 if C is to the left of the A->B Line, <0 if to the right, 0 if on the line
+	static T Orient(const FVector2<T>& A, const FVector2<T>& B, const FVector2<T>& C)
 	{
-		T DotVal = Dot(V2);
-		T ClampedDot = (DotVal < (T)-1) ? (T)-1 : ((DotVal > (T)1) ? (T)1 : DotVal);
-		T Direction = Cross(V2);
-		if (Direction*Direction < TMathUtil<T>::ZeroTolerance)
-		{
-			return (DotVal < 0) ? TMathUtil<T>::Pi : (T)0;
-		}
-		else
-		{
-			T Sign = Direction < 0 ? (T)-1 : (T)1;
-			return Sign * TMathUtil<T>::ACos(ClampedDot);
-		}
+		return (B - A).DotPerp(C - A);
 	}
+};
 
-	T Normalize(const T Epsilon = 0)
-	{
-		T length = Length();
-		if (length > Epsilon)
-		{
-			T invLength = ((T)1) / length;
-			X *= invLength;
-			Y *= invLength;
-			return length;
-		}
-		X = Y = (T)0;
-		return 0;
-	}
-
-	constexpr FVector2<T> Normalized(const T Epsilon = 0) const
-	{
-		T length = Length();
-		if (length > Epsilon)
-		{
-			T invLength = ((T)1) / length;
-			return FVector2<T>(X * invLength, Y * invLength);
-		}
-		return FVector2<T>((T)0, (T)0);
-	}
 
+template <typename T>
+constexpr bool IsNormalized(const FVector2<T>& Vector, const T Tolerance = TMathUtil<T>::ZeroTolerance)
+{
+	return TMathUtil<T>::Abs((Vector.X*Vector.X + Vector.Y*Vector.Y) - 1) < Tolerance;
+}
 
+template <typename T>
+T Normalize(FVector2<T>& Vector, const T Epsilon = 0)
+{
+	 T length = Vector.Length();
+	 if (length > Epsilon)
+	 {
+		 T invLength = ((T)1) / length;
+		 Vector.X *= invLength;
+		 Vector.Y *= invLength;
+		 return length;
+	 }
+	 Vector.X = Vector.Y = (T)0;
+	 return (T)0;
+}
 
-	constexpr bool operator==(const FVector2<T>& Other) const
+template <typename T>
+FVector2<T> Normalized(const FVector2<T>& Vector, const T Epsilon = 0)
+{
+	T length = Vector.Length();
+	if (length > Epsilon)
 	{
-		return X == Other.X && Y == Other.Y;
+		T invLength = ((T)1) / length;
+		return FVector2<T>(Vector.X*invLength, Vector.Y*invLength);
 	}
+	return FVector2<T>((T)0, (T)0);
+}
 
-	constexpr bool operator!=(const FVector2<T>& Other) const
-	{
-		return X != Other.X || Y != Other.Y;
-	}
 
+// Angle in Degrees
+template <typename T>
+T AngleD(const FVector2<T>& V1, const FVector2<T>& V2)
+{
+	T DotVal = V1.Dot(V2);
+	T ClampedDot = (DotVal < (T)-1) ? (T)-1 : ((DotVal > (T)1) ? (T)1 : DotVal);
+	return TMathUtil<T>::ACos(ClampedDot) * TMathUtil<T>::RadToDeg;
+}
+
+// Angle in Radians
+template <typename T>
+T AngleR(const FVector2<T>& V1, const FVector2<T>& V2)
+{
+	T DotVal = V1.Dot(V2);
+	T ClampedDot = (DotVal < (T)-1) ? (T)-1 : ((DotVal > (T)1) ? (T)1 : DotVal);
+	return TMathUtil<T>::ACos(ClampedDot);
+}
 
-	static FVector2<T> Lerp(const FVector2<T>& A, const FVector2<T>& B, T Alpha)
+// Angle in Radians
+template <typename T>
+T SignedAngleR(const FVector2<T>& V1, const FVector2<T>& V2)
+{
+	T DotVal = V1.Dot(V2);
+	T ClampedDot = (DotVal < (T)-1) ? (T)-1 : ((DotVal > (T)1) ? (T)1 : DotVal);
+	T Direction = V1.Cross(V2);
+	if (Direction * Direction < TMathUtil<T>::ZeroTolerance)
 	{
-		T OneMinusAlpha = (T)1 - Alpha;
-		return FVector2<T>(OneMinusAlpha * A.X + Alpha * B.X,
-			OneMinusAlpha * A.Y + Alpha * B.Y);
+		return (DotVal < 0) ? TMathUtil<T>::Pi : (T)0;
 	}
-
-	// Note: This was called "IsLeft" in the Geometry3Sharp code (where it also went through Math.Sign)
-	// Returns >0 if C is to the left of the A->B Line, <0 if to the right, 0 if on the line
-	static T Orient(const FVector2<T>& A, const FVector2<T>& B, const FVector2<T>& C)
+	else
 	{
-		return (B - A).DotPerp(C - A);
+		T Sign = Direction < 0 ? (T)-1 : (T)1;
+		return Sign * TMathUtil<T>::ACos(ClampedDot);
 	}
-};
+}
+
+template <typename T>
+FVector2<T> Lerp(const FVector2<T>& A, const FVector2<T>& B, T Alpha)
+{
+	T OneMinusAlpha = (T)1 - Alpha;
+	return FVector2<T>(OneMinusAlpha * A.X + Alpha * B.X,
+		OneMinusAlpha * A.Y + Alpha * B.Y);
+}
 
 
 /**
@@ -387,15 +394,6 @@ struct FVector3
 	{
 	}
 
-	explicit constexpr operator FColor() const
-	{
-		return FColor(
-			FMathf::Clamp((int)((float)X*255.0f), 0, 255),
-			FMathf::Clamp((int)((float)Y*255.0f), 0, 255),
-			FMathf::Clamp((int)((float)Z*255.0f), 0, 255));
-	}
-
-
 	static FVector3<T> Zero()
 	{
 		return FVector3<T>((T)0, (T)0, (T)0);
@@ -421,14 +419,6 @@ struct FVector3
 		return FVector3<T>(TNumericLimits<T>::Max(), TNumericLimits<T>::Max(), TNumericLimits<T>::Max());
 	}
 	
-	/** @return unit vector along axis X=0, Y=1, Z=2 */
-	constexpr static FVector3<T> MakeUnit(int32 Axis)
-	{
-		FVector3<T> UnitVec((T)0, (T)0, (T)0);
-		UnitVec[FMath::Clamp(Axis, 0, 2)] = (T)1;
-		return UnitVec;
-	}
-
 	FVector3<T>& operator=(const FVector3<T>& V2)
 	{
 		X = V2.X;
@@ -566,163 +556,210 @@ struct FVector3
 			X * V2.Y - Y * V2.X);
 	}
 
-	FVector3<T> UnitCross(const FVector3<T>& V2) const
+	constexpr bool operator==(const FVector3<T>& Other) const
 	{
-		FVector3<T> N = Cross(V2);
-		return N.Normalized();
+		return X == Other.X && Y == Other.Y && Z == Other.Z;
 	}
 
-	/**
-	 * Assumes this vector and V2 are both already normalized
-	 * @return the angle to vector V2 in degrees
-	 */
-	T AngleD(const FVector3<T>& V2) const
+	constexpr bool operator!=(const FVector3<T>& Other) const
 	{
-		T DotVal = Dot(V2);
-		T ClampedDot = (DotVal < (T)-1) ? (T)-1 : ((DotVal > (T)1) ? (T)1 : DotVal);
-		return (T)(acos(ClampedDot) * (T)(180.0 / 3.14159265358979323846));
+		return X != Other.X || Y != Other.Y || Z != Other.Z;
 	}
+};
 
-	/**
-	 * Assumes this vector and V2 are both already normalized
-	 * @return the angle to vector V2 in radians
-	 */
-	T AngleR(const FVector3<T>& V2) const
-	{
-		T DotVal = Dot(V2);
-		T ClampedDot = (DotVal < (T)-1) ? (T)-1 : ((DotVal > (T)1) ? (T)1 : DotVal);
-		return (T)acos(ClampedDot);
-	}
 
-	constexpr bool IsNormalized()
-	{
-		return TMathUtil<T>::Abs((X * X + Y * Y + Z * Z) - 1) < TMathUtil<T>::ZeroTolerance;
-	}
 
-	T Normalize(const T Epsilon = 0)
-	{
-		T length = Length();
-		if (length > Epsilon)
-		{
-			T invLength = ((T)1) / length;
-			X *= invLength;
-			Y *= invLength;
-			Z *= invLength;
-			return length;
-		}
-		X = Y = Z = (T)0;
-		return 0;
-	}
+/** @return unit vector along axis X=0, Y=1, Z=2 */
+template <typename T>
+constexpr FVector3<T> MakeUnitVector3(int32 Axis)
+{
+	FVector3<T> UnitVec((T)0, (T)0, (T)0);
+	UnitVec[FMath::Clamp(Axis, 0, 2)] = (T)1;
+	return UnitVec;
+}
 
-	constexpr FVector3<T> Normalized(const T Epsilon = 0) const
-	{
-		T length = Length();
-		if (length > Epsilon)
-		{
-			T invLength = ((T)1) / length;
-			return FVector3<T>(X * invLength, Y * invLength, Z * invLength);
-		}
-		return FVector3<T>((T)0, (T)0, (T)0);
-	}
 
-	constexpr T MaxElement() const
-	{
-		return TMathUtil<T>::Max3(X,Y,Z);
-	}
+template <typename T>
+constexpr bool IsNormalized(const FVector3<T>& Vector, const T Tolerance = TMathUtil<T>::ZeroTolerance)
+{
+	return TMathUtil<T>::Abs((Vector.X*Vector.X + Vector.Y*Vector.Y + Vector.Z*Vector.Z) - 1) < Tolerance;
+}
 
-	/** @return 0/1/2 index of maximum element */
-	constexpr int32 MaxElementIndex() const
-	{
-		return TMathUtil<T>::Max3Index(X, Y, Z);
-	}
 
-	constexpr T MinElement() const
+template<typename T>
+T Normalize(FVector3<T>& Vector, const T Epsilon = 0)
+{
+	T length = Vector.Length();
+	if (length > Epsilon)
 	{
-		return TMathUtil<T>::Min3(X,Y,Z);
+		T invLength = ((T)1) / length;
+		Vector.X *= invLength;
+		Vector.Y *= invLength;
+		Vector.Z *= invLength;
+		return length;
 	}
+	Vector.X = Vector.Y = Vector.Z = (T)0;
+	return (T)0;
+}
 
-	/** @return 0/1/2 index of minimum element */
-	constexpr int32 MinElementIndex() const
+template<typename T>
+constexpr FVector3<T> Normalized(const FVector3<T>& Vector, const T Epsilon = 0)
+{
+	T length = Vector.Length();
+	if (length > Epsilon)
 	{
-		return TMathUtil<T>::Min3Index(X, Y, Z);
+		T invLength = ((T)1) / length;
+		return FVector3<T>(Vector.X*invLength, Vector.Y*invLength, Vector.Z*invLength);
 	}
+	return FVector3<T>((T)0, (T)0, (T)0);
+}
 
-	constexpr friend FVector3 Min( const FVector3& V0, const FVector3& V1 )
-	{
-		return FVector3(TMathUtil<T>::Min(V0.X, V1.X),
-						TMathUtil<T>::Min(V0.Y, V1.Y),
-						TMathUtil<T>::Min(V0.Z, V1.Z));
-	}
 
-	constexpr friend FVector3 Max( const FVector3& V0, const FVector3& V1 )
-	{
-		return FVector3(TMathUtil<T>::Max(V0.X, V1.X),
-						TMathUtil<T>::Max(V0.Y, V1.Y),
-						TMathUtil<T>::Max(V0.Z, V1.Z));
-	}
+template<typename T>
+FVector3<T> UnitCross(const FVector3<T>& V1, const FVector3<T>& V2)
+{
+	FVector3<T> N = V1.Cross(V2);
+	return Normalized(N);
+}
 
-	constexpr T MaxAbsElement() const
-	{
-		return TMathUtil<T>::Max3(TMathUtil<T>::Abs(X), TMathUtil<T>::Abs(Y), TMathUtil<T>::Abs(Z));
-	}
+/**
+ * Computes the Angle between V1 and V2, assuming they are already normalized
+ * @return the (positive) angle between V1 and V2 in degrees
+ */
+template <typename T>
+T AngleD(const FVector3<T>& V1, const FVector3<T>& V2)
+{
+	T DotVal = V1.Dot(V2);
+	T ClampedDot = (DotVal < (T)-1) ? (T)-1 : ((DotVal > (T)1) ? (T)1 : DotVal);
+	return TMathUtil<T>::ACos(ClampedDot) * TMathUtil<T>::RadToDeg;
+}
 
-	/** @return 0/1/2 index of maximum absolute-value element */
-	constexpr T MaxAbsElementIndex() const
-	{
-		return TMathUtil<T>::Max3Index(TMathUtil<T>::Abs(X), TMathUtil<T>::Abs(Y), TMathUtil<T>::Abs(Z));
-	}
+/**
+ * Computes the Angle between V1 and V2, assuming they are already normalized
+ * @return the (positive) angle between V1 and V2 in radians
+ */
+template <typename T>
+T AngleR(const FVector3<T>& V1, const FVector3<T>& V2)
+{
+	T DotVal = V1.Dot(V2);
+	T ClampedDot = (DotVal < (T)-1) ? (T)-1 : ((DotVal > (T)1) ? (T)1 : DotVal);
+	return TMathUtil<T>::ACos(ClampedDot);
+}
 
-	constexpr T MinAbsElement() const
-	{
-		return TMathUtil<T>::Min3(TMathUtil<T>::Abs(X), TMathUtil<T>::Abs(Y), TMathUtil<T>::Abs(Z));
-	}
+template <typename T>
+constexpr FVector2<T> GetXY(const FVector3<T>& V)
+{
+	return FVector2<T>(V.X, V.Y);
+}
 
-	/** @return 0/1/2 index of minimum absolute-value element */
-	constexpr T MinAbsElementIndex() const
-	{
-		return TMathUtil<T>::Min3Index(TMathUtil<T>::Abs(X), TMathUtil<T>::Abs(Y), TMathUtil<T>::Abs(Z));
-	}
+template <typename T>
+constexpr FVector2<T> GetXZ(const FVector3<T>& V)
+{
+	return FVector2<T>(V.X, V.Z);
+}
 
-	constexpr FVector2<T> XY() const
-	{
-		return FVector2<T>(X, Y);
-	}
-	constexpr FVector2<T> XZ() const
-	{
-		return FVector2<T>(X, Z);
-	}
-	constexpr FVector2<T> YZ() const
-	{
-		return FVector2<T>(Y, Z);
-	}
+template <typename T>
+constexpr FVector2<T> GetYZ(const FVector3<T>& V)
+{
+	return FVector2<T>(V.Y, V.Z);
+}
 
-	static FVector3<T> Lerp(const FVector3<T>& A, const FVector3<T>& B, T Alpha)
-	{
-		T OneMinusAlpha = (T)1 - Alpha;
-		return FVector3<T>(OneMinusAlpha * A.X + Alpha * B.X,
-						   OneMinusAlpha * A.Y + Alpha * B.Y,
-						   OneMinusAlpha * A.Z + Alpha * B.Z);
-	}
 
-	static FVector3<T> Blend3(const FVector3<T>& A, const FVector3<T>& B, const FVector3<T>& C, const T& WeightA, const T& WeightB, const T& WeightC)
-	{
-		return FVector3<T>(
-			WeightA*A.X + WeightB*B.X + WeightC*C.X,
-			WeightA*A.Y + WeightB*B.Y + WeightC*C.Y,
-			WeightA*A.Z + WeightB*B.Z + WeightC*C.Z);
-	}
+template<typename T>
+constexpr FVector3<T> Min(const FVector3<T>& V0, const FVector3<T>& V1)
+{
+	return FVector3<T>(TMathUtil<T>::Min(V0.X, V1.X),
+		TMathUtil<T>::Min(V0.Y, V1.Y),
+		TMathUtil<T>::Min(V0.Z, V1.Z));
+}
 
+template<typename T>
+constexpr FVector3<T> Max(const FVector3<T>& V0, const FVector3<T>& V1)
+{
+	return FVector3<T>(TMathUtil<T>::Max(V0.X, V1.X),
+		TMathUtil<T>::Max(V0.Y, V1.Y),
+		TMathUtil<T>::Max(V0.Z, V1.Z));
+}
 
-	constexpr bool operator==(const FVector3<T>& Other) const
-	{
-		return X == Other.X && Y == Other.Y && Z == Other.Z;
-	}
 
-	constexpr bool operator!=(const FVector3<T>& Other) const
-	{
-		return X != Other.X || Y != Other.Y || Z != Other.Z;
-	}
-};
+template<typename T>
+constexpr T MaxElement(const FVector3<T>& Vector)
+{
+	return TMathUtil<T>::Max3(Vector.X, Vector.Y, Vector.Z);
+}
+
+/** @return 0/1/2 index of maximum element */
+template<typename T>
+constexpr int32 MaxElementIndex(const FVector3<T>& Vector)
+{
+	return TMathUtil<T>::Max3Index(Vector.X, Vector.Y, Vector.Z);
+}
+
+template<typename T>
+constexpr T MinElement(const FVector3<T>& Vector)
+{
+	return TMathUtil<T>::Min3(Vector.X, Vector.Y, Vector.Z);
+}
+
+/** @return 0/1/2 index of minimum element */
+template<typename T>
+constexpr int32 MinElementIndex(const FVector3<T>& Vector)
+{
+	return TMathUtil<T>::Min3Index(Vector.X, Vector.Y, Vector.Z);
+}
+
+template<typename T>
+constexpr T MaxAbsElement(const FVector3<T>& Vector)
+{
+	return TMathUtil<T>::Max3(TMathUtil<T>::Abs(Vector.X), TMathUtil<T>::Abs(Vector.Y), TMathUtil<T>::Abs(Vector.Z));
+}
+
+/** @return 0/1/2 index of maximum absolute-value element */
+template<typename T>
+constexpr T MaxAbsElementIndex(const FVector3<T>& Vector)
+{
+	return TMathUtil<T>::Max3Index(TMathUtil<T>::Abs(Vector.X), TMathUtil<T>::Abs(Vector.Y), TMathUtil<T>::Abs(Vector.Z));
+}
+
+template<typename T>
+constexpr T MinAbsElement(const FVector3<T>& Vector)
+{
+	return TMathUtil<T>::Min3(TMathUtil<T>::Abs(Vector.X), TMathUtil<T>::Abs(Vector.Y), TMathUtil<T>::Abs(Vector.Z));
+}
+
+/** @return 0/1/2 index of minimum absolute-value element */
+template<typename T>
+constexpr T MinAbsElementIndex(const FVector3<T>& Vector)
+{
+	return TMathUtil<T>::Min3Index(TMathUtil<T>::Abs(Vector.X), TMathUtil<T>::Abs(Vector.Y), TMathUtil<T>::Abs(Vector.Z));
+}
+
+template<typename T>
+constexpr FColor ToFColor(const FVector3<T>& Vector)
+{
+	return FColor(
+		FMathf::Clamp((int)((float)Vector.X * 255.0f), 0, 255),
+		FMathf::Clamp((int)((float)Vector.Y * 255.0f), 0, 255),
+		FMathf::Clamp((int)((float)Vector.Z * 255.0f), 0, 255));
+}
+
+template<typename T>
+FVector3<T> Lerp(const FVector3<T>& A, const FVector3<T>& B, T Alpha)
+{
+	T OneMinusAlpha = (T)1 - Alpha;
+	return FVector3<T>(OneMinusAlpha * A.X + Alpha * B.X,
+		OneMinusAlpha * A.Y + Alpha * B.Y,
+		OneMinusAlpha * A.Z + Alpha * B.Z);
+}
+
+template<typename T>
+FVector3<T> Blend3(const FVector3<T>& A, const FVector3<T>& B, const FVector3<T>& C, const T& WeightA, const T& WeightB, const T& WeightC)
+{
+	return FVector3<T>(
+		WeightA * A.X + WeightB * B.X + WeightC * C.X,
+		WeightA * A.Y + WeightB * B.Y + WeightC * C.Y,
+		WeightA * A.Z + WeightB * B.Z + WeightC * C.Z);
+}
 
 template <typename RealType>
 inline FVector3<RealType> operator*(RealType Scalar, const FVector3<RealType>& V)
@@ -987,74 +1024,77 @@ struct TVector4
 		return X * V2.X + Y * V2.Y + Z * V2.Z + W * V2.W;
 	}
 
-
-	constexpr bool IsNormalized()
-	{
-		return TMathUtil<T>::Abs((X * X + Y * Y + Z * Z + W * W) - 1) < TMathUtil<T>::ZeroTolerance;
-	}
-
-	T Normalize(const T Epsilon = 0)
+	constexpr bool operator==(const TVector4<T>& Other) const
 	{
-		T length = Length();
-		if (length > Epsilon)
-		{
-			T invLength = ((T)1) / length;
-			X *= invLength;
-			Y *= invLength;
-			Z *= invLength;
-			W *= invLength;
-			return length;
-		}
-		X = Y = Z = (T)0;
-		return 0;
+		return X == Other.X && Y == Other.Y && Z == Other.Z && W == Other.W;
 	}
 
-	constexpr TVector4<T> Normalized(const T Epsilon = 0) const
+	constexpr bool operator!=(const TVector4<T>& Other) const
 	{
-		T length = Length();
-		if (length > Epsilon)
-		{
-			T invLength = ((T)1) / length;
-			return TVector4<T>(X * invLength, Y * invLength, Z * invLength, W * invLength);
-		}
-		return TVector4<T>::Zero();
+		return X != Other.X || Y != Other.Y || Z != Other.Z || Z == Other.Z;
 	}
+};
 
+template <typename T>
+constexpr bool IsNormalized(const TVector4<T>& Vector, const T Tolerance = TMathUtil<T>::ZeroTolerance)
+{
+	return TMathUtil<T>::Abs((Vector.X*Vector.X + Vector.Y*Vector.Y + Vector.Z*Vector.Z + Vector.W*Vector.W) - 1) < Tolerance;
+}
 
-	constexpr FVector3<T> XYZ() const
-	{
-		return FVector3<T>(X, Y, Z);
-	}
-
-	static TVector4<T> Lerp(const TVector4<T>& A, const TVector4<T>& B, T Alpha)
-	{
-		T OneMinusAlpha = (T)1 - Alpha;
-		return TVector4<T>(OneMinusAlpha * A.X + Alpha * B.X,
-						   OneMinusAlpha * A.Y + Alpha * B.Y,
-						   OneMinusAlpha * A.Z + Alpha * B.Z,
-						   OneMinusAlpha * A.W + Alpha * B.W );
-	}
+template<typename T>
+T Normalize(TVector4<T>& Vector, const T Epsilon = 0)
+{
+	T length = Vector.Length();
+	if (length > Epsilon)
+	{
+		T invLength = ((T)1) / length;
+		Vector.X *= invLength;
+		Vector.Y *= invLength;
+		Vector.Z *= invLength;
+		Vector.W *= invLength;
+		return length;
+	}
+	Vector.X = Vector.Y = Vector.Z = Vector.W = (T)0;
+	return (T)0;
+}
 
-	static TVector4<T> Blend3(const TVector4<T>& A, const TVector4<T>& B, const TVector4<T>& C, const T& WeightA, const T& WeightB, const T& WeightC)
+template<typename T>
+TVector4<T> Normalized(const TVector4<T>& Vector, const T Epsilon = 0)
+{
+	T length = Vector.Length();
+	if (length > Epsilon)
 	{
-		return TVector4<T>(
-			WeightA*A.X + WeightB*B.X + WeightC*C.X,
-			WeightA*A.Y + WeightB*B.Y + WeightC*C.Y,
-			WeightA*A.Z + WeightB*B.Z + WeightC*C.Z,
-			WeightA*A.W + WeightB*B.W + WeightC*C.W);
+		T invLength = ((T)1) / length;
+		return TVector4<T>(Vector.X*invLength, Vector.Y*invLength, Vector.Z*invLength, Vector.W*invLength);
 	}
+	return TVector4<T>::Zero();
+}
 
+template<typename T>
+constexpr FVector3<T> GetXYZ(const TVector4<T>& V)
+{
+	return FVector3<T>(V.X, V.Y, V.Z);
+}
 
-	constexpr bool operator==(const TVector4<T>& Other) const
-	{
-		return X == Other.X && Y == Other.Y && Z == Other.Z && W == Other.W;
-	}
+template<typename T>
+TVector4<T> Lerp(const TVector4<T>& A, const TVector4<T>& B, T Alpha)
+{
+	T OneMinusAlpha = (T)1 - Alpha;
+	return TVector4<T>(OneMinusAlpha * A.X + Alpha * B.X,
+		OneMinusAlpha * A.Y + Alpha * B.Y,
+		OneMinusAlpha * A.Z + Alpha * B.Z,
+		OneMinusAlpha * A.W + Alpha * B.W);
+}
 
-	constexpr bool operator!=(const TVector4<T>& Other) const
-	{
-		return X != Other.X || Y != Other.Y || Z != Other.Z || Z == Other.Z;
-	}
-};
+template<typename T>
+TVector4<T> Blend3(const TVector4<T>& A, const TVector4<T>& B, const TVector4<T>& C, const T& WeightA, const T& WeightB, const T& WeightC)
+{
+	return TVector4<T>(
+		WeightA * A.X + WeightB * B.X + WeightC * C.X,
+		WeightA * A.Y + WeightB * B.Y + WeightC * C.Y,
+		WeightA * A.Z + WeightB * B.Z + WeightC * C.Z,
+		WeightA * A.W + WeightB * B.W + WeightC * C.W);
+}
 
 template <typename RealType>
 inline TVector4<RealType> operator*(RealType Scalar, const TVector4<RealType>& V)
diff --git a/Engine/Plugins/Experimental/GeometryProcessing/Source/GeometricObjects/Public/VectorUtil.h b/Engine/Plugins/Experimental/GeometryProcessing/Source/GeometricObjects/Public/VectorUtil.h
--- a/Engine/Plugins/Experimental/GeometryProcessing/Source/GeometricObjects/Public/VectorUtil.h
+++ b/Engine/Plugins/Experimental/GeometryProcessing/Source/GeometricObjects/Public/VectorUtil.h
@@ -68,12 +68,12 @@ namespace VectorUtil
 	{
 		FVector3<RealType> edge1(V1 - V0);
 		FVector3<RealType> edge2(V2 - V0);
-		edge1.Normalize();
-		edge2.Normalize();
+		Normalize(edge1);
+		Normalize(edge2);
 		// Unreal has Left-Hand Coordinate System so we need to reverse this cross-product to get proper triangle normal
 		FVector3<RealType> vCross(edge2.Cross(edge1));
 		//FVector3<RealType> vCross(edge1.Cross(edge2));
-		return vCross.Normalized();
+		return Normalized(vCross);
 	}
 
 	/**
@@ -135,7 +135,7 @@ namespace VectorUtil
 		// Unreal has Left-Hand Coordinate System so we need to reverse this cross-product to get proper triangle normal
 		FVector3d vCross = edge2.Cross(edge1);
 		//FVector3d vCross = edge1.Cross(edge2);
-		AreaOut = RealType(0.5) * vCross.Normalize();
+		AreaOut = RealType(0.5) * Normalize(vCross);
 		return vCross;
 	}
 
@@ -242,15 +242,15 @@ namespace VectorUtil
 	{
 		FVector3<RealType> vFrom = VFrom - VFrom.Dot(PlaneN) * PlaneN;
 		FVector3<RealType> vTo = VTo - VTo.Dot(PlaneN) * PlaneN;
-		vFrom.Normalize();
-		vTo.Normalize();
+		Normalize(vFrom);
+		Normalize(vTo);
 		FVector3<RealType> C = vFrom.Cross(vTo);
 		if (C.SquaredLength() < TMathUtil<RealType>::ZeroTolerance)
 		{ // vectors are parallel
 			return vFrom.Dot(vTo) < 0 ? (RealType)180 : (RealType)0;
 		}
 		RealType fSign = C.Dot(PlaneN) < 0 ? (RealType)-1 : (RealType)1;
-		return (RealType)(fSign * vFrom.AngleD(vTo));
+		return (RealType)(fSign * AngleD(vFrom, vTo));
 	}
 
 	/**
@@ -392,10 +392,10 @@ namespace VectorUtil
 	inline RealType OpeningAngleD(FVector3<RealType> A, FVector3<RealType> B, const FVector3<RealType>& P)
 	{
 		A -= P; 
-		A.Normalize();
+		Normalize(A);
 		B -= P;
-		B.Normalize();
-		return A.AngleD(B);
+		Normalize(B);
+		return AngleD(A, B);
 	}
 
 
diff --git a/Engine/Plugins/Experimental/GeometryProcessing/Source/GeometryAlgorithms/Private/ConvexHull3.cpp b/Engine/Plugins/Experimental/GeometryProcessing/Source/GeometryAlgorithms/Private/ConvexHull3.cpp
--- a/Engine/Plugins/Experimental/GeometryProcessing/Source/GeometryAlgorithms/Private/ConvexHull3.cpp
+++ b/Engine/Plugins/Experimental/GeometryProcessing/Source/GeometryAlgorithms/Private/ConvexHull3.cpp
@@ -97,7 +97,7 @@ struct TExtremePoints3
 
 		Origin = GetPointFunc(Extreme[0]);
 		Basis[0] = GetPointFunc(Extreme[1]) - Origin;
-		Basis[0].Normalize();
+		Normalize(Basis[0]);
 
 		// find point furthest from the line formed by the first two extreme points
 		{
@@ -130,7 +130,7 @@ struct TExtremePoints3
 		Basis[1] = GetPointFunc(Extreme[2]) - Origin;
 		// project Basis[1] to be orthogonal to Basis[0]
 		Basis[1] -= (Basis[0].Dot(Basis[1])) * Basis[0];
-		Basis[1].Normalize();
+		Normalize(Basis[1]);
 		Basis[2] = Basis[0].Cross(Basis[1]);
 
 		{
diff --git a/Engine/Plugins/Experimental/GeometryProcessing/Source/GeometryAlgorithms/Public/Arrangement2d.h b/Engine/Plugins/Experimental/GeometryProcessing/Source/GeometryAlgorithms/Public/Arrangement2d.h
--- a/Engine/Plugins/Experimental/GeometryProcessing/Source/GeometryAlgorithms/Public/Arrangement2d.h
+++ b/Engine/Plugins/Experimental/GeometryProcessing/Source/GeometryAlgorithms/Public/Arrangement2d.h
@@ -484,7 +484,7 @@ protected:
 		int num_hits = 0;
 		FVector2d x = FVector2d::Zero(), y = FVector2d::Zero();
 		FVector2d EPerp = (B - A).Perp();
-		EPerp.Normalize();
+		Normalize(EPerp);
 		for (int EID : Graph.EdgeIndices())
 		{
 			Graph.GetEdgeV(EID, x, y);
diff --git a/Engine/Plugins/Experimental/GeometryProcessing/Source/MeshConversion/Private/MeshDescriptionToDynamicMesh.cpp b/Engine/Plugins/Experimental/GeometryProcessing/Source/MeshConversion/Private/MeshDescriptionToDynamicMesh.cpp
--- a/Engine/Plugins/Experimental/GeometryProcessing/Source/MeshConversion/Private/MeshDescriptionToDynamicMesh.cpp
+++ b/Engine/Plugins/Experimental/GeometryProcessing/Source/MeshConversion/Private/MeshDescriptionToDynamicMesh.cpp
@@ -763,7 +763,8 @@ static void CopyTangents_Internal(const FMeshDescription* SourceMesh, const FDyn
 			FVector Tangent = InstanceTangents[InstanceTri[j]];
 			float BitangentSign = InstanceSigns[InstanceTri[j]];
 			FVector3<RealType> Bitangent = VectorUtil::Bitangent((FVector3<RealType>)Normal, (FVector3<RealType>)Tangent, (RealType)BitangentSign);
-			Tangent.Normalize(); Bitangent.Normalize();
+			Tangent.Normalize();
+			UE::Geometry::Normalize(Bitangent);
 			TangentsOut->SetPerTriangleTangent(TriID, j, Tangent, Bitangent);
 		}
 	}
diff --git a/Engine/Plugins/Experimental/MeshModelingToolset/Source/MeshModelingTools/Private/BakeTransformTool.cpp b/Engine/Plugins/Experimental/MeshModelingToolset/Source/MeshModelingTools/Private/BakeTransformTool.cpp
--- a/Engine/Plugins/Experimental/MeshModelingToolset/Source/MeshModelingTools/Private/BakeTransformTool.cpp
+++ b/Engine/Plugins/Experimental/MeshModelingToolset/Source/MeshModelingTools/Private/BakeTransformTool.cpp
@@ -208,7 +208,7 @@ void UBakeTransformTool::UpdateAssets()
 				RemainingUniformScale = AbsScales[BestSubIdx];
 				if (RemainingUniformScale <= FLT_MIN)
 				{
-					RemainingUniformScale = AbsScales.MaxAbsElement();
+					RemainingUniformScale = MaxAbsElement(AbsScales);
 				}
 			}
 			switch (BasicProperties->BakeScale)
diff --git a/Engine/Plugins/Experimental/MeshModelingToolset/Source/MeshModelingTools/Private/DeformMeshPolygonsTool.cpp b/Engine/Plugins/Experimental/MeshModelingToolset/Source/MeshModelingTools/Private/DeformMeshPolygonsTool.cpp
--- a/Engine/Plugins/Experimental/MeshModelingToolset/Source/MeshModelingTools/Private/DeformMeshPolygonsTool.cpp
+++ b/Engine/Plugins/Experimental/MeshModelingToolset/Source/MeshModelingTools/Private/DeformMeshPolygonsTool.cpp
@@ -1037,7 +1037,7 @@ bool UDeformMeshPolygonsTool::HitTest(const FRay& WorldRay, FHitResult& OutHit)
 	FTransform3d Transform(ComponentTarget->GetWorldTransform());
 	FRay3d LocalRay(Transform.InverseTransformPosition(WorldRay.Origin),
 	                Transform.InverseTransformVector(WorldRay.Direction));
-	LocalRay.Direction.Normalize();
+	UE::Geometry::Normalize(LocalRay.Direction);
 
 	FGroupTopologySelection Selection;
 	FVector3d LocalPosition, LocalNormal;
@@ -1083,7 +1083,7 @@ void UDeformMeshPolygonsTool::OnBeginDrag(const FRay& WorldRay)
 	FTransform3d Transform(ComponentTarget->GetWorldTransform());
 	FRay3d LocalRay(Transform.InverseTransformPosition(WorldRay.Origin),
 	                Transform.InverseTransformVector(WorldRay.Direction));
-	LocalRay.Direction.Normalize();
+	UE::Geometry::Normalize(LocalRay.Direction);
 
 	HilightSelection.Clear();
 
@@ -1298,7 +1298,7 @@ bool UDeformMeshPolygonsTool::OnUpdateHover(const FInputDeviceRay& DevicePos)
 		FTransform3d Transform(ComponentTarget->GetWorldTransform());
 		FRay3d LocalRay(Transform.InverseTransformPosition(DevicePos.WorldRay.Origin),
 		                Transform.InverseTransformVector(DevicePos.WorldRay.Direction));
-		LocalRay.Direction.Normalize();
+		UE::Geometry::Normalize(LocalRay.Direction);
 
 		HilightSelection.Clear();
 		FVector3d LocalPosition, LocalNormal;
@@ -1445,10 +1445,10 @@ void UDeformMeshPolygonsTool::ComputeUpdate_Rotate()
 	}
 
 	FVector2d RotateStartVec = RotationStartFrame.ToPlaneUV(RotationStartPointWorld, 2);
-	RotateStartVec.Normalize();
+	UE::Geometry::Normalize(RotateStartVec);
 	FVector2d RotateToVec = RotationStartFrame.ToPlaneUV(NewHitPosWorld, 2);
-	RotateToVec.Normalize();
-	double AngleRad = RotateStartVec.SignedAngleR(RotateToVec);
+	UE::Geometry::Normalize(RotateToVec);
+	double AngleRad = UE::Geometry::SignedAngleR(RotateStartVec, RotateToVec);
 	FQuaterniond Rotation(Transform.InverseTransformVectorNoScale(RotationStartFrame.Z()), AngleRad, false);
 	FVector3d LocalOrigin = Transform.InverseTransformPosition(RotationStartFrame.Origin);
 
diff --git a/Engine/Plugins/Experimental/MeshModelingToolset/Source/MeshModelingTools/Private/DrawPolyPathTool.cpp b/Engine/Plugins/Experimental/MeshModelingToolset/Source/MeshModelingTools/Private/DrawPolyPathTool.cpp
--- a/Engine/Plugins/Experimental/MeshModelingToolset/Source/MeshModelingTools/Private/DrawPolyPathTool.cpp
+++ b/Engine/Plugins/Experimental/MeshModelingToolset/Source/MeshModelingTools/Private/DrawPolyPathTool.cpp
@@ -451,8 +451,8 @@ void UDrawPolyPathTool::OnCompleteSurfacePath()
 	int NumPoints = CurPathPoints.Num();
 	// align frames
 	FVector3d PlaneNormal = DrawPlaneWorld.Z();
-	CurPathPoints[0].ConstrainedAlignAxis(0, (CurPathPoints[1].Origin - CurPathPoints[0].Origin).Normalized(), PlaneNormal);
-	CurPathPoints[NumPoints-1].ConstrainedAlignAxis(0, (CurPathPoints[NumPoints-1].Origin - CurPathPoints[NumPoints-2].Origin).Normalized(), PlaneNormal);
+	CurPathPoints[0].ConstrainedAlignAxis(0, UE::Geometry::Normalized(CurPathPoints[1].Origin - CurPathPoints[0].Origin), PlaneNormal);
+	CurPathPoints[NumPoints-1].ConstrainedAlignAxis(0, UE::Geometry::Normalized(CurPathPoints[NumPoints-1].Origin - CurPathPoints[NumPoints-2].Origin), PlaneNormal);
 	double DistOffsetDelta = 0.01;
 	OffsetScaleFactors.SetNum(NumPoints);
 	OffsetScaleFactors[0] = OffsetScaleFactors[NumPoints-1] = 1.0;
@@ -475,7 +475,7 @@ void UDrawPolyPathTool::OnCompleteSurfacePath()
 
 		if (Line1.Direction.Dot(Line2.Direction) > 0.999 )
 		{
-			CurPathPoints[j].ConstrainedAlignAxis(0, (Next-Prev).Normalized(), PlaneNormal);
+			CurPathPoints[j].ConstrainedAlignAxis(0, UE::Geometry::Normalized(Next-Prev), PlaneNormal);
 			OffsetScaleFactors[j] = 1.0;
 		}
 		else
@@ -484,7 +484,7 @@ void UDrawPolyPathTool::OnCompleteSurfacePath()
 			Distance.GetSquared();
 			FVector3d OffsetPoint = 0.5 * (Distance.Line1ClosestPoint + Distance.Line2ClosestPoint);
 			OffsetScaleFactors[j] = OffsetPoint.Distance(Cur) / DistOffsetDelta;
-			FVector3d TangentDir = (OffsetPoint - Cur).Normalized().Cross(PlaneNormal);
+			FVector3d TangentDir = UE::Geometry::Normalized(OffsetPoint - Cur).Cross(PlaneNormal);
 			CurPathPoints[j].ConstrainedAlignAxis(0, TangentDir, PlaneNormal);
 		}
 	}
diff --git a/Engine/Plugins/Experimental/MeshModelingToolset/Source/MeshModelingTools/Private/DrawPolygonTool.cpp b/Engine/Plugins/Experimental/MeshModelingToolset/Source/MeshModelingTools/Private/DrawPolygonTool.cpp
--- a/Engine/Plugins/Experimental/MeshModelingToolset/Source/MeshModelingTools/Private/DrawPolygonTool.cpp
+++ b/Engine/Plugins/Experimental/MeshModelingToolset/Source/MeshModelingTools/Private/DrawPolygonTool.cpp
@@ -250,7 +250,7 @@ void DrawEdgeTicks(FPrimitiveDrawInterface* PDI,
 	FVector3d Center = Segment.Center;
 	FVector3d X = Segment.Direction;
 	FVector3d Y = X.Cross(PlaneNormal);
-	Y.Normalize();
+	UE::Geometry::Normalize(Y);
 	FVector3d A = Center - Height * 0.25*X - Height * Y;
 	FVector3d B = Center + Height * 0.25*X + Height * Y;
 	PDI->DrawLine((FVector)A, (FVector)B, Color, DepthPriorityGroup, LineThickness, 0.0f, bIsScreenSpace);
diff --git a/Engine/Plugins/Experimental/MeshModelingToolset/Source/MeshModelingTools/Private/DynamicMeshSculptTool.cpp b/Engine/Plugins/Experimental/MeshModelingToolset/Source/MeshModelingTools/Private/DynamicMeshSculptTool.cpp
--- a/Engine/Plugins/Experimental/MeshModelingToolset/Source/MeshModelingTools/Private/DynamicMeshSculptTool.cpp
+++ b/Engine/Plugins/Experimental/MeshModelingToolset/Source/MeshModelingTools/Private/DynamicMeshSculptTool.cpp
@@ -331,7 +331,7 @@ bool UDynamicMeshSculptTool::HitTest(const FRay& Ray, FHitResult& OutHit)
 {
 	FRay3d LocalRay(CurTargetTransform.InverseTransformPosition(Ray.Origin),
 		CurTargetTransform.InverseTransformVector(Ray.Direction));
-	LocalRay.Direction.Normalize();
+	UE::Geometry::Normalize(LocalRay.Direction);
 	FDynamicMesh3* Mesh = DynamicMeshComponent->GetMesh();
 
 	int HitTID = FindHitSculptMeshTriangle(LocalRay);
@@ -666,7 +666,7 @@ bool UDynamicMeshSculptTool::ApplySmoothBrush(const FRay& WorldRay)
 		FVector3d SmoothedPos = (SculptProperties->bPreserveUVFlow) ?
 			FMeshWeights::CotanCentroidSafe(*Mesh, VertIdx, 10.0) : FMeshWeights::UniformCentroid(*Mesh, VertIdx);
 
-		FVector3d NewPos = FVector3d::Lerp(OrigPos, SmoothedPos, Falloff*SculptProperties->SmoothBrushSpeed);
+		FVector3d NewPos = UE::Geometry::Lerp(OrigPos, SmoothedPos, Falloff*SculptProperties->SmoothBrushSpeed);
 
 		ROIPositionBuffer[k] = NewPos;
 	});
@@ -795,7 +795,7 @@ bool UDynamicMeshSculptTool::ApplySculptMaxBrush(const FRay& WorldRay)
 			FVector3d DeltaPos = NewPos - BasePos;
 			if (DeltaPos.SquaredLength() > MaxOffset*MaxOffset)
 			{
-				DeltaPos.Normalize();
+				UE::Geometry::Normalize(DeltaPos);
 				NewPos = BasePos + MaxOffset * DeltaPos;
 			}
 			ROIPositionBuffer[k] = NewPos;
@@ -822,7 +822,7 @@ bool UDynamicMeshSculptTool::ApplyPinchBrush(const FRay& WorldRay)
 
 	FVector3d MotionVec = NewSmoothBrushPosLocal - LastSmoothBrushPosLocal;
 	bool bHaveMotion = (MotionVec.Length() > FMathf::ZeroTolerance);
-	MotionVec.Normalize();
+	UE::Geometry::Normalize(MotionVec);
 	FLine3d MoveLine(LastSmoothBrushPosLocal, MotionVec);
 
 	FDynamicMesh3* Mesh = DynamicMeshComponent->GetMesh();
@@ -845,7 +845,7 @@ bool UDynamicMeshSculptTool::ApplyPinchBrush(const FRay& WorldRay)
 
 		if (bHaveMotion && Falloff < 0.8f)
 		{
-			double AnglePower = 1.0 - FMathd::Abs(MoveVec.Normalized().Dot(MotionVec));
+			double AnglePower = 1.0 - FMathd::Abs(UE::Geometry::Normalized(MoveVec).Dot(MotionVec));
 			Falloff *= AnglePower;
 		}
 
@@ -874,7 +874,7 @@ FFrame3d UDynamicMeshSculptTool::ComputeROIBrushPlane(const FVector3d& BrushCent
 		AveragePos += Weight * Centroid;
 		WeightSum += Weight;
 	}
-	AverageNormal.Normalize();
+	UE::Geometry::Normalize(AverageNormal);
 	AveragePos /= WeightSum;
 
 	if (bViewAligned)
@@ -965,7 +965,7 @@ bool UDynamicMeshSculptTool::ApplyFixedPlaneBrush(const FRay& WorldRay)
 		FVector3d NewPos = OrigPos;
 		if (Dot * PlaneSign >= 0)
 		{
-			double MaxDist = Delta.Normalize();
+			double MaxDist = UE::Geometry::Normalize(Delta);
 			double Falloff = CalculateBrushFalloff(OrigPos.Distance(NewBrushPosLocal));
 			FVector3d MoveVec = Falloff * UseSpeed * Delta;
 			NewPos = (MoveVec.SquaredLength() > MaxDist* MaxDist) ?
@@ -1012,7 +1012,7 @@ bool UDynamicMeshSculptTool::ApplyFlattenBrush(const FRay& WorldRay)
 		FVector3d NewPos = OrigPos;
 		if (Dot * PlaneSign >= 0)
 		{
-			double MaxDist = Delta.Normalize();
+			double MaxDist = UE::Geometry::Normalize(Delta);
 			double Falloff = CalculateBrushFalloff(OrigPos.Distance(NewBrushPosLocal));
 			FVector3d MoveVec = Falloff * UseSpeed * Delta;
 			NewPos = (MoveVec.SquaredLength() > MaxDist*MaxDist) ?
@@ -1276,7 +1276,7 @@ bool UDynamicMeshSculptTool::UpdateBrushPositionOnTargetMesh(const FRay& WorldRa
 
 	FRay3d LocalRay(CurTargetTransform.InverseTransformPosition(WorldRay.Origin),
 		CurTargetTransform.InverseTransformVector(WorldRay.Direction));
-	LocalRay.Direction.Normalize();
+	UE::Geometry::Normalize(LocalRay.Direction);
 
 	const FDynamicMesh3* TargetMesh = BrushTargetMeshSpatial.GetMesh();
 
@@ -1310,7 +1310,7 @@ bool UDynamicMeshSculptTool::UpdateBrushPositionOnSculptMesh(const FRay& WorldRa
 {
 	FRay3d LocalRay(CurTargetTransform.InverseTransformPosition(WorldRay.Origin),
 		CurTargetTransform.InverseTransformVector(WorldRay.Direction));
-	LocalRay.Direction.Normalize();
+	UE::Geometry::Normalize(LocalRay.Direction);
 
 	int HitTID = FindHitSculptMeshTriangle(LocalRay);
 	if (HitTID != IndexConstants::InvalidID)
@@ -1515,7 +1515,7 @@ void UDynamicMeshSculptTool::OnTick(float DeltaTime)
 		FHitResult Result;
 		bool bWorldHit = ToolSceneQueriesUtil::FindNearestVisibleObjectHit(TargetWorld, Result, CursorWorldRay);
 		FRay3d LocalRay(CurTargetTransform.InverseTransformPosition(CursorWorldRay.Origin), CurTargetTransform.InverseTransformVector(CursorWorldRay.Direction));
-		LocalRay.Direction.Normalize();
+		UE::Geometry::Normalize(LocalRay.Direction);
 		bool bObjectHit = (FindHitSculptMeshTriangle(LocalRay) != IndexConstants::InvalidID);
 		if (bWorldHit &&
 			(bObjectHit == false || (CursorWorldRay.GetParameter(Result.ImpactPoint) < CursorWorldRay.GetParameter((FVector)LastBrushPosWorld))))
@@ -2138,7 +2138,7 @@ bool UDynamicMeshSculptTool::GetTargetMeshNearest(const FVector3d& Position, dou
 		Query.TriangleBaryCoords.X*BrushTargetNormals[Tri.A]
 		+ Query.TriangleBaryCoords.Y*BrushTargetNormals[Tri.B]
 		+ Query.TriangleBaryCoords.Z*BrushTargetNormals[Tri.C];
-	TargetNormalOut.Normalize();
+	UE::Geometry::Normalize(TargetNormalOut);
 	TargetPosOut = Query.ClosestTrianglePoint;
 	return true;
 }
diff --git a/Engine/Plugins/Experimental/MeshModelingToolset/Source/MeshModelingTools/Private/EditMeshPolygonsTool.cpp b/Engine/Plugins/Experimental/MeshModelingToolset/Source/MeshModelingTools/Private/EditMeshPolygonsTool.cpp
--- a/Engine/Plugins/Experimental/MeshModelingToolset/Source/MeshModelingTools/Private/EditMeshPolygonsTool.cpp
+++ b/Engine/Plugins/Experimental/MeshModelingToolset/Source/MeshModelingTools/Private/EditMeshPolygonsTool.cpp
@@ -1392,21 +1392,21 @@ void UEditMeshPolygonsTool::ApplyCutFaces()
 	FVector3d PlaneNormal;
 	if (CutProperties->Orientation == EPolyEditCutPlaneOrientation::ViewDirection)
 	{
-		FVector3d Direction0 = (Point0.Origin - CameraState.Position).Normalized();
-		FVector3d Direction1 = (Point1.Origin - CameraState.Position).Normalized();
+		FVector3d Direction0 = UE::Geometry::Normalized(Point0.Origin - CameraState.Position);
+		FVector3d Direction1 = UE::Geometry::Normalized(Point1.Origin - CameraState.Position);
 		PlaneNormal = Direction1.Cross(Direction0);
 	}
 	else
 	{
-		FVector3d LineDirection = (Point1.Origin - Point0.Origin).Normalized();
-		FVector3d UpVector = (Point0.Z() + Point1.Z()).Normalized();
+		FVector3d LineDirection = UE::Geometry::Normalized(Point1.Origin - Point0.Origin);
+		FVector3d UpVector = UE::Geometry::Normalized(Point0.Z() + Point1.Z());
 		PlaneNormal = LineDirection.Cross(UpVector);
 	}
 	FVector3d PlaneOrigin = 0.5 * (Point0.Origin + Point1.Origin);
 	// map into local space of target mesh
 	PlaneOrigin = WorldTransform.InverseTransformPosition(PlaneOrigin);
 	PlaneNormal = WorldTransform.InverseTransformNormal(PlaneNormal);
-	PlaneNormal.Normalize();
+	UE::Geometry::Normalize(PlaneNormal);
 
 	// track changes
 	FDynamicMeshChangeTracker ChangeTracker(Mesh);
@@ -1485,7 +1485,7 @@ void UEditMeshPolygonsTool::UpdateSetUVS()
 		SurfacePathMechanic->InitializePlaneSurface(PlanarFrame);
 
 		FVector3d Delta = PlanarFrame.Origin - SurfacePathMechanic->HitPath[0].Origin;
-		double Dist = Delta.Normalize();
+		double Dist = UE::Geometry::Normalize(Delta);
 		UVScale *= FMathd::Lerp(1.0, 25.0, Dist / ActiveSelectionBounds.MaxDim());
 		PlanarFrame = SurfacePathMechanic->HitPath[0];
 		PlanarFrame.ConstrainedAlignAxis(0, Delta, PlanarFrame.Z());
@@ -1514,7 +1514,7 @@ void UEditMeshPolygonsTool::ApplySetUVs()
 	FFrame3d PlanarFrame = SurfacePathMechanic->HitPath[0];
 	double UVScale = 1.0 / ActiveSelectionBounds.MaxDim();
 	FVector3d Delta = SurfacePathMechanic->HitPath[1].Origin - PlanarFrame.Origin;
-	double Dist = Delta.Normalize();
+	double Dist = UE::Geometry::Normalize(Delta);
 	UVScale *= FMathd::Lerp(1.0, 25.0, Dist / ActiveSelectionBounds.MaxDim());
 	PlanarFrame.ConstrainedAlignAxis(0, Delta, PlanarFrame.Z());
 
@@ -1914,7 +1914,7 @@ void UEditMeshPolygonsTool::ApplyStraightenEdges()
 			for (int k = 1; k < NumV-1; ++k)
 			{
 				double t = VtxArcLengths[k] / EdgeArcLen;
-				Mesh->SetVertex(EdgeVerts[k], FVector3d::Lerp(A, B, t));
+				Mesh->SetVertex(EdgeVerts[k], UE::Geometry::Lerp(A, B, t));
 			}
 		}
 	}
diff --git a/Engine/Plugins/Experimental/MeshModelingToolset/Source/MeshModelingTools/Private/EditUVIslandsTool.cpp b/Engine/Plugins/Experimental/MeshModelingToolset/Source/MeshModelingTools/Private/EditUVIslandsTool.cpp
--- a/Engine/Plugins/Experimental/MeshModelingToolset/Source/MeshModelingTools/Private/EditUVIslandsTool.cpp
+++ b/Engine/Plugins/Experimental/MeshModelingToolset/Source/MeshModelingTools/Private/EditUVIslandsTool.cpp
@@ -459,9 +459,9 @@ FFrame3d FUVGroupTopology::GetIslandFrame(int32 GroupID, FDynamicMeshAABBTree3&
 		UVOverlay->GetTriElements(TriangleID, fi, fj, fk);
 
 		FVector3d GradX = VectorUtil::TriGradient<double>(A, B, C, fi.X, fj.X, fk.X);
-		AccumX += GradX.Normalized();
+		AccumX += UE::Geometry::Normalized(GradX);
 	}
-	AccumX.Normalize();
+	UE::Geometry::Normalize(AccumX);
 	Frame.AlignAxis(0, AccumX);
 
 	return Frame;
diff --git a/Engine/Plugins/Experimental/MeshModelingToolset/Source/MeshModelingTools/Private/FindPolygonsAlgorithm.cpp b/Engine/Plugins/Experimental/MeshModelingToolset/Source/MeshModelingTools/Private/FindPolygonsAlgorithm.cpp
--- a/Engine/Plugins/Experimental/MeshModelingToolset/Source/MeshModelingTools/Private/FindPolygonsAlgorithm.cpp
+++ b/Engine/Plugins/Experimental/MeshModelingToolset/Source/MeshModelingTools/Private/FindPolygonsAlgorithm.cpp
@@ -141,7 +141,7 @@ void FFindPolygonsAlgorithm::SetGroupsFromPolygons()
 		PolygonTags[PolyIdx] = (PolyIdx + 1);
 
 		// find a normal if the average failed
-		AccumNormal.Normalize();
+		UE::Geometry::Normalize(AccumNormal);
 		int SubIdx = 0;
 		while (AccumNormal.Length() < 0.9 && SubIdx < NumTriangles)
 		{
diff --git a/Engine/Plugins/Experimental/MeshModelingToolset/Source/MeshModelingTools/Private/GroupEdgeInsertionTool.cpp b/Engine/Plugins/Experimental/MeshModelingToolset/Source/MeshModelingTools/Private/GroupEdgeInsertionTool.cpp
--- a/Engine/Plugins/Experimental/MeshModelingToolset/Source/MeshModelingTools/Private/GroupEdgeInsertionTool.cpp
+++ b/Engine/Plugins/Experimental/MeshModelingToolset/Source/MeshModelingTools/Private/GroupEdgeInsertionTool.cpp
@@ -630,8 +630,8 @@ bool UGroupEdgeInsertionTool::GetHoveredItem(const FRay& WorldRay,
 			if (EdgeSegmentID > 0)
 			{
 				// Average with previous normalized edge vector
-				PointOut.Tangent += (StartVert - CurrentMesh->GetVertex(GroupEdge.Span.Vertices[EdgeSegmentID - 1])).Normalized();
-				PointOut.Tangent.Normalize();
+				PointOut.Tangent += UE::Geometry::Normalized(StartVert - CurrentMesh->GetVertex(GroupEdge.Span.Vertices[EdgeSegmentID - 1]));
+				UE::Geometry::Normalize(PointOut.Tangent);
 			}
 		}
 		else if (abs(DistDownEdge - EdgeLength) <= Settings->VertexTolerance)
@@ -640,8 +640,8 @@ bool UGroupEdgeInsertionTool::GetHoveredItem(const FRay& WorldRay,
 			PointOut.ElementID = EndVid;
 			if (EdgeSegmentID + 2 < GroupEdge.Span.Vertices.Num())
 			{
-				PointOut.Tangent += (CurrentMesh->GetVertex(GroupEdge.Span.Vertices[EdgeSegmentID + 2]) - EndVert).Normalized();
-				PointOut.Tangent.Normalize();
+				PointOut.Tangent += UE::Geometry::Normalized(CurrentMesh->GetVertex(GroupEdge.Span.Vertices[EdgeSegmentID + 2]) - EndVert);
+				UE::Geometry::Normalize(PointOut.Tangent);
 			}
 		}
 		else
@@ -686,9 +686,9 @@ void UGroupEdgeInsertionTool::GetCornerTangent(int32 CornerID, int32 GroupID, in
 	check(AdjacentPoints.Num() == 2);
 
 	FVector3d CornerPosition = CurrentMesh->GetVertex(CornerVid);
-	TangentOut = (CornerPosition - AdjacentPoints[0]).Normalized();
-	TangentOut += (AdjacentPoints[1] - CornerPosition).Normalized();
-	TangentOut.Normalize();
+	TangentOut = UE::Geometry::Normalized(CornerPosition - AdjacentPoints[0]);
+	TangentOut += UE::Geometry::Normalized(AdjacentPoints[1] - CornerPosition);
+	UE::Geometry::Normalize(TangentOut);
 }
 
 bool GetSharedBoundary(const FGroupTopology& Topology,
diff --git a/Engine/Plugins/Experimental/MeshModelingToolset/Source/MeshModelingTools/Private/MeshGroupPaintTool.cpp b/Engine/Plugins/Experimental/MeshModelingToolset/Source/MeshModelingTools/Private/MeshGroupPaintTool.cpp
--- a/Engine/Plugins/Experimental/MeshModelingToolset/Source/MeshModelingTools/Private/MeshGroupPaintTool.cpp
+++ b/Engine/Plugins/Experimental/MeshModelingToolset/Source/MeshModelingTools/Private/MeshGroupPaintTool.cpp
@@ -654,9 +654,9 @@ bool ApproxSelfClipPolyline(TArray<FVector2f>& Polyline)
 	}
 
 
-	FVector2f StartDirOut = (Polyline[0] - Polyline[1]).Normalized();
+	FVector2f StartDirOut = UE::Geometry::Normalized(Polyline[0] - Polyline[1]);
 	FLine2f StartLine(Polyline[0], StartDirOut);
-	FVector2f EndDirOut = (Polyline[N - 1] - Polyline[N - 2]).Normalized();
+	FVector2f EndDirOut = UE::Geometry::Normalized(Polyline[N - 1] - Polyline[N - 2]);
 	FLine2f EndLine(Polyline[N - 1], EndDirOut);
 	FIntrLine2Line2f LineIntr(StartLine, EndLine);
 	bool bIntersects = false;
@@ -732,9 +732,9 @@ void UMeshGroupPaintTool::OnPolyLassoFinished(const FCameraPolyLasso& Lasso, boo
 	// outwards from the endpoints and then try to create a closed very large polygon
 	if (ApproxSelfClipPolyline(Polyline) == false)
 	{
-		FVector2f StartDirOut = (Polyline[0] - Polyline[1]).Normalized();
+		FVector2f StartDirOut = UE::Geometry::Normalized(Polyline[0] - Polyline[1]);
 		FLine2f StartLine(Polyline[0], StartDirOut);
-		FVector2f EndDirOut = (Polyline[N-1] - Polyline[N-2]).Normalized();
+		FVector2f EndDirOut = UE::Geometry::Normalized(Polyline[N-1] - Polyline[N-2]);
 		FLine2f EndLine(Polyline[N-1], EndDirOut);
 
 		// if we did not intersect, we are in ambiguous territory. Check if a segment along either end-direction
@@ -897,7 +897,7 @@ void UMeshGroupPaintTool::ApplyVisibilityFilter(const TArray<int32>& Triangles,
 		}
 		if (FilterProperties->VisibilityFilter == EMeshGroupPaintVisibilityType::Unoccluded)
 		{
-			int32 HitTID = Octree.FindNearestHitObject(FRay3d(LocalEyePosition, (Centroid - LocalEyePosition).Normalized()));
+			int32 HitTID = Octree.FindNearestHitObject(FRay3d(LocalEyePosition, UE::Geometry::Normalized(Centroid - LocalEyePosition)));
 			if (HitTID != Triangles[idx])
 			{
 				VisibilityFilterBuffer[idx] = false;
diff --git a/Engine/Plugins/Experimental/MeshModelingToolset/Source/MeshModelingTools/Private/MeshSelectionTool.cpp b/Engine/Plugins/Experimental/MeshModelingToolset/Source/MeshModelingTools/Private/MeshSelectionTool.cpp
--- a/Engine/Plugins/Experimental/MeshModelingToolset/Source/MeshModelingTools/Private/MeshSelectionTool.cpp
+++ b/Engine/Plugins/Experimental/MeshModelingToolset/Source/MeshModelingTools/Private/MeshSelectionTool.cpp
@@ -538,7 +538,7 @@ void UMeshSelectionTool::UpdateFaceSelection(const FBrushStampData& Stamp, const
 		FVector3d StartNormal = Mesh->GetTriNormal(StartROI[0]);
 		int AngleTol = SelectionProps->AngleTolerance;
 		FMeshConnectedComponents::GrowToConnectedTriangles(Mesh, StartROI, LocalROI, &TemporaryBuffer, &TemporarySet,
-			[Mesh, AngleTol, StartNormal](int t1, int t2) { return Mesh->GetTriNormal(t2).AngleD(StartNormal) < AngleTol; });
+			[Mesh, AngleTol, StartNormal](int t1, int t2) { return UE::Geometry::AngleD(Mesh->GetTriNormal(t2), StartNormal) < AngleTol; });
 
 		UseROI = &LocalROI;
 	}
@@ -550,7 +550,7 @@ void UMeshSelectionTool::UpdateFaceSelection(const FBrushStampData& Stamp, const
 		FVector3d StartNormal = Mesh->GetTriNormal(StartROI[0]);
 		int AngleTol = SelectionProps->AngleTolerance;
 		FMeshConnectedComponents::GrowToConnectedTriangles(Mesh, StartROI, LocalROI, &TemporaryBuffer, &TemporarySet,
-			[Mesh, AngleTol, StartNormal, &BrushROI](int t1, int t2) { return BrushROI.Contains(t2) && Mesh->GetTriNormal(t2).AngleD(StartNormal) < AngleTol; });
+			[Mesh, AngleTol, StartNormal, &BrushROI](int t1, int t2) { return BrushROI.Contains(t2) && UE::Geometry::AngleD(Mesh->GetTriNormal(t2), StartNormal) < AngleTol; });
 		UseROI = &LocalROI;
 	}
 	else if (SelectionProps->SelectionMode == EMeshSelectionToolPrimaryMode::Visible)
@@ -562,7 +562,7 @@ void UMeshSelectionTool::UpdateFaceSelection(const FBrushStampData& Stamp, const
 		for (int tid : TriangleROI)
 		{
 			FVector3d Centroid = Mesh->GetTriCentroid(tid);
-			int HitTID = GetOctree()->FindNearestHitObject(FRay3d(LocalEyePosition, (Centroid - LocalEyePosition).Normalized()));
+			int HitTID = GetOctree()->FindNearestHitObject(FRay3d(LocalEyePosition, UE::Geometry::Normalized(Centroid - LocalEyePosition)));
 			if (HitTID == tid)
 			{
 				LocalROI.Add(HitTID);
diff --git a/Engine/Plugins/Experimental/MeshModelingToolset/Source/MeshModelingTools/Private/MeshSpaceDeformerTool.cpp b/Engine/Plugins/Experimental/MeshModelingToolset/Source/MeshModelingTools/Private/MeshSpaceDeformerTool.cpp
--- a/Engine/Plugins/Experimental/MeshModelingToolset/Source/MeshModelingTools/Private/MeshSpaceDeformerTool.cpp
+++ b/Engine/Plugins/Experimental/MeshModelingToolset/Source/MeshModelingTools/Private/MeshSpaceDeformerTool.cpp
@@ -495,7 +495,7 @@ void UMeshSpaceDeformerTool::SetGizmoFrameFromWorldPos(const FVector& Position,
 		// the user would want the axis to lie along the object on the side they clicked, not pierce inwards.
 		// Still, it's hard to come up with a clean alternative.
 		FVector3d FrameZ(Normal);
-		FVector3d FrameY = FrameZ.Cross(FVector3d::UnitZ()).Normalized(); // orthogonal to world Z and frame Z 
+		FVector3d FrameY = UE::Geometry::Normalized(FrameZ.Cross(FVector3d::UnitZ())); // orthogonal to world Z and frame Z 
 		FVector3d FrameX = FrameY.Cross(Normal); // safe to not normalize because already orthogonal
 		GizmoFrame = FFrame3d(Position, FrameX, FrameY, Normal);
 	}
diff --git a/Engine/Plugins/Experimental/MeshModelingToolset/Source/MeshModelingTools/Private/Sculpting/KelvinletBrushOp.h b/Engine/Plugins/Experimental/MeshModelingToolset/Source/MeshModelingTools/Private/Sculpting/KelvinletBrushOp.h
--- a/Engine/Plugins/Experimental/MeshModelingToolset/Source/MeshModelingTools/Private/Sculpting/KelvinletBrushOp.h
+++ b/Engine/Plugins/Experimental/MeshModelingToolset/Source/MeshModelingTools/Private/Sculpting/KelvinletBrushOp.h
@@ -651,7 +651,7 @@ public:
 		case EKelvinletBrushMode::LaplacianTwistPullKelvinlet:
 		{
 			FVector3d TwistAxis = Properties.Direction;
-			TwistAxis.Normalize();
+			UE::Geometry::Normalize(TwistAxis);
 			TwistAxis *= Properties.Speed;
 			FTwistKelvinlet TwistKelvinlet(TwistAxis, Properties.Size, Properties.Mu, Properties.Nu);
 
diff --git a/Engine/Plugins/Experimental/MeshModelingToolset/Source/MeshModelingTools/Private/Sculpting/MeshSculptToolBase.cpp b/Engine/Plugins/Experimental/MeshModelingToolset/Source/MeshModelingTools/Private/Sculpting/MeshSculptToolBase.cpp
--- a/Engine/Plugins/Experimental/MeshModelingToolset/Source/MeshModelingTools/Private/Sculpting/MeshSculptToolBase.cpp
+++ b/Engine/Plugins/Experimental/MeshModelingToolset/Source/MeshModelingTools/Private/Sculpting/MeshSculptToolBase.cpp
@@ -492,7 +492,7 @@ FRay3d UMeshSculptToolBase::GetLocalRay(const FRay& WorldRay) const
 {
 	FRay3d LocalRay(CurTargetTransform.InverseTransformPosition(WorldRay.Origin),
 		CurTargetTransform.InverseTransformVector(WorldRay.Direction));
-	LocalRay.Direction.Normalize();
+	UE::Geometry::Normalize(LocalRay.Direction);
 	return LocalRay;
 }
 
@@ -527,7 +527,7 @@ retry_frame_update:
 	if (InStroke() && BrushProperties->Lazyness > 0)
 	{
 		double t = FMathd::Lerp(1.0, 0.1, BrushProperties->Lazyness);
-		LastBrushFrameWorld.Origin = FVector3d::Lerp(LastBrushFrameWorld.Origin, NewFrame.Origin, t);
+		LastBrushFrameWorld.Origin = UE::Geometry::Lerp(LastBrushFrameWorld.Origin, NewFrame.Origin, t);
 		LastBrushFrameWorld.Rotation = FQuaterniond(LastBrushFrameWorld.Rotation, NewFrame.Rotation, t);
 	}
 	else
@@ -754,7 +754,7 @@ FFrame3d UMeshSculptToolBase::ComputeStampRegionPlane(const FFrame3d& StampFrame
 		AveragePos += Weight * Centroid;
 		WeightSum += Weight;
 	}
-	AverageNormal.Normalize();
+	UE::Geometry::Normalize(AverageNormal);
 	AveragePos /= WeightSum;
 
 	if (bViewAligned)
diff --git a/Engine/Plugins/Experimental/MeshModelingToolset/Source/MeshModelingTools/Public/Sculpting/MeshMoveBrushOps.h b/Engine/Plugins/Experimental/MeshModelingToolset/Source/MeshModelingTools/Public/Sculpting/MeshMoveBrushOps.h
--- a/Engine/Plugins/Experimental/MeshModelingToolset/Source/MeshModelingTools/Public/Sculpting/MeshMoveBrushOps.h
+++ b/Engine/Plugins/Experimental/MeshModelingToolset/Source/MeshModelingTools/Public/Sculpting/MeshMoveBrushOps.h
@@ -57,7 +57,7 @@ public:
 		if (Props->AxisFilters.AnyAxisFiltered())
 		{
 			FVector3d WorldMoveVec = Stamp.WorldFrame.Origin - Stamp.PrevWorldFrame.Origin;
-			WorldMoveVec = MoveVec.Length() * WorldMoveVec.Normalized();		// apply local space scaling
+			WorldMoveVec = MoveVec.Length() * UE::Geometry::Normalized(WorldMoveVec);		// apply local space scaling
 			WorldMoveVec.X = (Props->AxisFilters.bAxisX) ? WorldMoveVec.X : 0;
 			WorldMoveVec.Y = (Props->AxisFilters.bAxisY) ? WorldMoveVec.Y : 0;
 			WorldMoveVec.Z = (Props->AxisFilters.bAxisZ) ? WorldMoveVec.Z : 0;
diff --git a/Engine/Plugins/Experimental/MeshModelingToolset/Source/MeshModelingTools/Public/Sculpting/MeshPinchBrushOps.h b/Engine/Plugins/Experimental/MeshModelingToolset/Source/MeshModelingTools/Public/Sculpting/MeshPinchBrushOps.h
--- a/Engine/Plugins/Experimental/MeshModelingToolset/Source/MeshModelingTools/Public/Sculpting/MeshPinchBrushOps.h
+++ b/Engine/Plugins/Experimental/MeshModelingToolset/Source/MeshModelingTools/Public/Sculpting/MeshPinchBrushOps.h
@@ -55,13 +55,13 @@ public:
 	virtual void ApplyStamp(const FDynamicMesh3* Mesh, const FSculptBrushStamp& Stamp, const TArray<int32>& Vertices, TArray<FVector3d>& NewPositionsOut) override
 	{
 		// hardcoded LazyBrush
-		FVector3d NewSmoothBrushPosLocal = FVector3d::Lerp(LastSmoothBrushPosLocal, Stamp.LocalFrame.Origin, 0.75f);
-		FVector3d NewSmoothBrushNormal = FVector3d::Lerp(LastSmoothBrushNormalLocal, Stamp.LocalFrame.Z(), 0.75f);
-		NewSmoothBrushNormal.Normalize();
+		FVector3d NewSmoothBrushPosLocal = UE::Geometry::Lerp(LastSmoothBrushPosLocal, Stamp.LocalFrame.Origin, 0.75);
+		FVector3d NewSmoothBrushNormal = UE::Geometry::Lerp(LastSmoothBrushNormalLocal, Stamp.LocalFrame.Z(), 0.75);
+		UE::Geometry::Normalize(NewSmoothBrushNormal);
 
 		FVector3d MotionVec = NewSmoothBrushPosLocal - LastSmoothBrushPosLocal;
 		bool bHaveMotion = (MotionVec.Length() > FMathd::ZeroTolerance);
-		MotionVec.Normalize();
+		UE::Geometry::Normalize(MotionVec);
 		UE::Geometry::FLine3d MoveLine(LastSmoothBrushPosLocal, MotionVec);
 
 		FVector3d DepthPosLocal = NewSmoothBrushPosLocal - (Stamp.Depth * Stamp.Radius * NewSmoothBrushNormal);
@@ -77,7 +77,7 @@ public:
 			FVector3d OrigPos = Mesh->GetVertex(VertIdx);
 
 			FVector3d Delta = DepthPosLocal - OrigPos;
-			Delta.Normalize();
+			UE::Geometry::Normalize(Delta);
 			FVector3d MoveVec = Delta;
 
 			// pinch uses 1/x falloff
@@ -89,7 +89,7 @@ public:
 
 			if (bLimitDrag && bHaveMotion && UseFalloff < 0.7f)
 			{
-				double AnglePower = 1.0 - FMathd::Abs(MoveVec.Normalized().Dot(MotionVec));
+				double AnglePower = 1.0 - FMathd::Abs( UE::Geometry::Normalized(MoveVec).Dot(MotionVec));
 				UseFalloff *= AnglePower;
 			}
 
diff --git a/Engine/Plugins/Experimental/MeshModelingToolset/Source/MeshModelingTools/Public/Sculpting/MeshPlaneBrushOps.h b/Engine/Plugins/Experimental/MeshModelingToolset/Source/MeshModelingTools/Public/Sculpting/MeshPlaneBrushOps.h
--- a/Engine/Plugins/Experimental/MeshModelingToolset/Source/MeshModelingTools/Public/Sculpting/MeshPlaneBrushOps.h
+++ b/Engine/Plugins/Experimental/MeshModelingToolset/Source/MeshModelingTools/Public/Sculpting/MeshPlaneBrushOps.h
@@ -140,7 +140,7 @@ public:
 			{
 				double Falloff = GetFalloff().Evaluate(Stamp, OrigPos);
 				FVector3d MoveVec = Falloff * UseSpeed * Delta;
-				double MaxDist = Delta.Normalize();
+				double MaxDist = UE::Geometry::Normalize(Delta);
 				NewPos = (MoveVec.SquaredLength() > MaxDist * MaxDist) ?
 					PlanePos : (OrigPos + Falloff * MoveVec);
 			}
diff --git a/Engine/Plugins/Experimental/MeshModelingToolset/Source/MeshModelingTools/Public/Sculpting/MeshSculptBrushOps.h b/Engine/Plugins/Experimental/MeshModelingToolset/Source/MeshModelingTools/Public/Sculpting/MeshSculptBrushOps.h
--- a/Engine/Plugins/Experimental/MeshModelingToolset/Source/MeshModelingTools/Public/Sculpting/MeshSculptBrushOps.h
+++ b/Engine/Plugins/Experimental/MeshModelingToolset/Source/MeshModelingTools/Public/Sculpting/MeshSculptBrushOps.h
@@ -255,7 +255,7 @@ public:
 				FVector3d DeltaPos = NewPos - BasePos;
 				if (DeltaPos.SquaredLength() > MaxOffset * MaxOffset)
 				{
-					DeltaPos.Normalize();
+					UE::Geometry::Normalize(DeltaPos);
 					NewPos = BasePos + MaxOffset * DeltaPos;
 				}
 
diff --git a/Engine/Plugins/Experimental/MeshModelingToolset/Source/MeshModelingTools/Public/Sculpting/MeshSmoothingBrushOps.h b/Engine/Plugins/Experimental/MeshModelingToolset/Source/MeshModelingTools/Public/Sculpting/MeshSmoothingBrushOps.h
--- a/Engine/Plugins/Experimental/MeshModelingToolset/Source/MeshModelingTools/Public/Sculpting/MeshSmoothingBrushOps.h
+++ b/Engine/Plugins/Experimental/MeshModelingToolset/Source/MeshModelingTools/Public/Sculpting/MeshSmoothingBrushOps.h
@@ -95,7 +95,7 @@ public:
 				SmoothedPos = UE::Geometry::FMeshWeights::UniformCentroid(*Mesh, VertIdx);
 			}
 
-			FVector3d NewPos = FVector3d::Lerp(OrigPos, SmoothedPos, Falloff * Stamp.Power);
+			FVector3d NewPos = UE::Geometry::Lerp(OrigPos, SmoothedPos, Falloff * Stamp.Power);
 
 			NewPositionsOut[k] = NewPos;
 		});
@@ -160,7 +160,7 @@ public:
 				SmoothedPos = UE::Geometry::FMeshWeights::UniformCentroid(*Mesh, VertIdx);
 			}
 
-			FVector3d NewPos = FVector3d::Lerp(OrigPos, SmoothedPos, Falloff * Stamp.Power);
+			FVector3d NewPos = UE::Geometry::Lerp(OrigPos, SmoothedPos, Falloff * Stamp.Power);
 
 			if ((NewPos - OrigPos).Dot(Normal) * Direction < 0)
 			{
@@ -236,7 +236,7 @@ public:
 			{
 				double Falloff = GetFalloff().Evaluate(Stamp, OrigPos);
 				FVector3d MoveVec = Falloff * UseSpeed * MoveDelta;
-				double MaxDist = MoveDelta.Normalize();
+				double MaxDist = UE::Geometry::Normalize(MoveDelta);
 				NewPos = (MoveVec.SquaredLength() > MaxDist * MaxDist) ?
 					PlanePos : (OrigPos + Falloff * MoveVec);
 			}
@@ -321,7 +321,7 @@ public:
 				}
 				else
 				{
-					MoveVec.Normalize();
+					UE::Geometry::Normalize(MoveVec);
 					NewPositionsOut[k] = OrigPos + MoveDist * MoveVec;
 				}
 			}
diff --git a/Engine/Plugins/Experimental/MeshModelingToolset/Source/MeshModelingToolsEditorOnly/Private/Hair/GroomCardsEditorTool.cpp b/Engine/Plugins/Experimental/MeshModelingToolset/Source/MeshModelingToolsEditorOnly/Private/Hair/GroomCardsEditorTool.cpp
--- a/Engine/Plugins/Experimental/MeshModelingToolset/Source/MeshModelingToolsEditorOnly/Private/Hair/GroomCardsEditorTool.cpp
+++ b/Engine/Plugins/Experimental/MeshModelingToolset/Source/MeshModelingToolsEditorOnly/Private/Hair/GroomCardsEditorTool.cpp
@@ -561,7 +561,7 @@ void UGroomCardsEditorTool::BeginCardEdit(int32 CardGroupID)
 	const UE::GroomQueries::FMeshCardStrip& CardStrip = EditableCardSet->CardStrips->FindStripForGroup(CardGroupID);
 
 	FTransform3d Transform(PreviewMesh->GetTransform());
-	check(FMathd::Abs(Transform.TransformVector(FVector3d(1, 1, 1).Normalized()).SquaredLength() - 1.0) < FMathd::ZeroTolerance);
+	check(FMathd::Abs(Transform.TransformVector(UE::Geometry::Normalized(FVector3d(1,1,1))).SquaredLength() - 1.0) < FMathd::ZeroTolerance);
 	FFrame3d TransformFrame(Transform.GetTranslation(), Transform.GetRotation());
 
 	TArray<FFrame3d> LocalFrames;
@@ -596,7 +596,7 @@ void UGroomCardsEditorTool::RestoreCardEdit(const FGroomCardEdit* RestoreEdit)
 {
 	// this should probably be cached once...
 	FTransform3d Transform(PreviewMesh->GetTransform());
-	check(FMathd::Abs(Transform.TransformVector(FVector3d(1, 1, 1).Normalized()).SquaredLength() - 1.0) < FMathd::ZeroTolerance);
+	check(FMathd::Abs(Transform.TransformVector(UE::Geometry::Normalized(FVector3d(1, 1, 1))).SquaredLength() - 1.0) < FMathd::ZeroTolerance);
 	FFrame3d TransformFrame(Transform.GetTranslation(), Transform.GetRotation());
 
 	ActiveCardEdit = MakePimpl<FGroomCardEdit>();
@@ -644,7 +644,7 @@ void UGroomCardsEditorTool::UpdateOnCurveEdit()
 	}
 
 	FTransform3d Transform(PreviewMesh->GetTransform());
-	check( FMathd::Abs(Transform.TransformVector(FVector3d(1,1,1).Normalized()).SquaredLength() - 1.0) < FMathd::ZeroTolerance );
+	check( FMathd::Abs(Transform.TransformVector(UE::Geometry::Normalized(FVector3d(1,1,1))).SquaredLength() - 1.0) < FMathd::ZeroTolerance );
 	FFrame3d TransformFrame(Transform.GetTranslation(), Transform.GetRotation());
 
 	TArray<FFrame3d> NewPositions;
diff --git a/Engine/Plugins/Experimental/MeshModelingToolset/Source/MeshModelingToolsEditorOnly/Private/Hair/GroomQueryUtil.cpp b/Engine/Plugins/Experimental/MeshModelingToolset/Source/MeshModelingToolsEditorOnly/Private/Hair/GroomQueryUtil.cpp
--- a/Engine/Plugins/Experimental/MeshModelingToolset/Source/MeshModelingToolsEditorOnly/Private/Hair/GroomQueryUtil.cpp
+++ b/Engine/Plugins/Experimental/MeshModelingToolset/Source/MeshModelingToolsEditorOnly/Private/Hair/GroomQueryUtil.cpp
@@ -448,24 +448,24 @@ void UE::GroomQueries::ExtractCardCurves(
 			FVector3d Center = CardStrip.CardCurve[k];
 
 			FFrame3d CardFrame(Center);
-			CardFrame.AlignAxis(0, (Right - Center).Normalized());
+			CardFrame.AlignAxis(0, UE::Geometry::Normalized(Right - Center));
 
 			FVector3d Y;
 			if (k == 0)
 			{
 				FVector3d Forward = CardStrip.CardCurve[k + 1];
-				Y = (Forward - Center).Normalized();
+				Y = UE::Geometry::Normalized(Forward - Center);
 			}
 			else if (k == NumVertices - 1)
 			{
 				FVector3d Back = CardStrip.CardCurve[k-1];
-				Y = (Center - Back).Normalized();
+				Y = UE::Geometry::Normalized(Center - Back);
 			}
 			else
 			{
 				FVector3d Forward = CardStrip.CardCurve[k + 1];
 				FVector3d Back = CardStrip.CardCurve[k - 1];
-				Y = (Forward - Back).Normalized();
+				Y = UE::Geometry::Normalized(Forward - Back);
 			}
 
 			CardFrame.ConstrainedAlignAxis(1, Y, CardFrame.X());
diff --git a/Engine/Plugins/Experimental/MeshModelingToolset/Source/MeshModelingToolsEditorOnly/Private/Hair/GroomToMeshTool.cpp b/Engine/Plugins/Experimental/MeshModelingToolset/Source/MeshModelingToolsEditorOnly/Private/Hair/GroomToMeshTool.cpp
--- a/Engine/Plugins/Experimental/MeshModelingToolset/Source/MeshModelingToolsEditorOnly/Private/Hair/GroomToMeshTool.cpp
+++ b/Engine/Plugins/Experimental/MeshModelingToolset/Source/MeshModelingToolsEditorOnly/Private/Hair/GroomToMeshTool.cpp
@@ -774,7 +774,7 @@ public:
 			if (bApplySmoothingPerPass)
 			{
 				FVector3d Centroid = (NextPos + PrevPos) * 0.5;
-				Pos = FVector3d::Lerp(Pos, Centroid, SmoothingAlpha);
+				Pos = UE::Geometry::Lerp(Pos, Centroid, SmoothingAlpha);
 			}
 
 			Pos = UpdatePositionFunc(vid, Pos);
@@ -847,13 +847,13 @@ TSharedPtr<FDynamicMesh3, ESPMode::ThreadSafe> UGroomToMeshTool::UpdateClipMesh(
 				Processor.ApplyPasses(10, [&](int32 vid, const FVector3d& Pos)
 				{
 					FVector3d NearestPos = BVTree.FindNearestPoint(EditMesh->GetVertex(vid));
-					return FVector3d::Lerp(Pos, NearestPos, ProjectionSpeed);
+					return UE::Geometry::Lerp(Pos, NearestPos, ProjectionSpeed);
 				});
 
 				Processor.ApplyPasses(10, [&](int32 vid, const FVector3d& Pos)
 				{
 					FVector3d NearestPos = BVTree.FindNearestPoint(EditMesh->GetVertex(vid));
-					return FVector3d::Lerp(Pos, NearestPos, TuneSpeed);
+					return UE::Geometry::Lerp(Pos, NearestPos, TuneSpeed);
 				});
 			}
 
diff --git a/Engine/Plugins/Experimental/MeshModelingToolset/Source/MeshModelingToolsEditorOnly/Private/MeshTangentsTool.cpp b/Engine/Plugins/Experimental/MeshModelingToolset/Source/MeshModelingToolsEditorOnly/Private/MeshTangentsTool.cpp
--- a/Engine/Plugins/Experimental/MeshModelingToolset/Source/MeshModelingToolsEditorOnly/Private/MeshTangentsTool.cpp
+++ b/Engine/Plugins/Experimental/MeshModelingToolset/Source/MeshModelingToolsEditorOnly/Private/MeshTangentsTool.cpp
@@ -313,13 +313,14 @@ void UMeshTangentsTool::OnTangentsUpdated(const TUniquePtr<FMeshTangentsd>& NewR
 				{
 					FVector3d TangentMikkt, BitangentMikkt;
 					MikktTangents->GetPerTriangleTangent(Index, j, TangentMikkt, BitangentMikkt);
-					TangentMikkt.Normalize(); BitangentMikkt.Normalize();
+					UE::Geometry::Normalize(TangentMikkt); UE::Geometry::Normalize(BitangentMikkt);
 					FVector3d TangentNew, BitangentNew;
 					NewTangents->GetPerTriangleTangent(Index, j, TangentNew, BitangentNew);
-					TangentNew.Normalize(); BitangentNew.Normalize();
-					ensure(TangentMikkt.IsNormalized() && BitangentMikkt.IsNormalized());
-					ensure(TangentNew.IsNormalized() && BitangentNew.IsNormalized());
-					double MaxAngleDeg = FMathd::Max(TangentMikkt.AngleD(TangentNew), BitangentMikkt.AngleD(BitangentNew));
+					UE::Geometry::Normalize(TangentNew); 
+					UE::Geometry::Normalize(BitangentNew);
+					ensure(UE::Geometry::IsNormalized(TangentMikkt) && UE::Geometry::IsNormalized(BitangentMikkt));
+					ensure(UE::Geometry::IsNormalized(TangentNew) && UE::Geometry::IsNormalized(BitangentNew));
+					double MaxAngleDeg = FMathd::Max(UE::Geometry::AngleD(TangentMikkt, TangentNew), UE::Geometry::AngleD(BitangentMikkt, BitangentNew));
 					if (MaxAngleDeg > 0.5)
 					{
 						FMikktDeviation Deviation{ static_cast<float>(MaxAngleDeg), Index, j, Verts[j], TangentMikkt, BitangentMikkt, TangentNew, BitangentNew };
diff --git a/Engine/Plugins/Experimental/MeshModelingToolset/Source/MeshModelingToolsEditorOnly/Private/ShapeSprayTool.cpp b/Engine/Plugins/Experimental/MeshModelingToolset/Source/MeshModelingToolsEditorOnly/Private/ShapeSprayTool.cpp
--- a/Engine/Plugins/Experimental/MeshModelingToolset/Source/MeshModelingToolsEditorOnly/Private/ShapeSprayTool.cpp
+++ b/Engine/Plugins/Experimental/MeshModelingToolset/Source/MeshModelingToolsEditorOnly/Private/ShapeSprayTool.cpp
@@ -152,7 +152,7 @@ void UShapeSprayTool::OnUpdateDrag(const FRay& Ray)
 		FVector3f PlanePt(Rad*FMathf::Cos(Angle), Rad*FMathf::Sin(Angle), 0);
 		FVector3f WorldPt = WorldFrame.PointAt(PlanePt);
 		FVector3f SampleRayDir = WorldPt - (FVector3f)Ray.Origin; 
-		SampleRayDir.Normalize();
+		UE::Geometry::Normalize(SampleRayDir);
 		FRay WorldRay((FVector)Ray.Origin, (FVector)SampleRayDir);
 		FHitResult Hit;
 		if (HitTest(WorldRay, Hit))
@@ -213,7 +213,7 @@ void UShapeSprayTool::UpdateShapeMesh()
 	for (int vid : ShapeMesh.VertexIndicesItr())
 	{
 		FVector3d Pos = ShapeMesh.GetVertex(vid);
-		Pos.Normalize();
+		UE::Geometry::Normalize(Pos);
 		ShapeMesh.SetVertexNormal(vid, (FVector3f)Pos);
 	}
 }
diff --git a/Engine/Plugins/Experimental/MeshModelingToolset/Source/ModelingComponents/Private/BaseDynamicMeshSceneProxy.h b/Engine/Plugins/Experimental/MeshModelingToolset/Source/ModelingComponents/Private/BaseDynamicMeshSceneProxy.h
--- a/Engine/Plugins/Experimental/MeshModelingToolset/Source/ModelingComponents/Private/BaseDynamicMeshSceneProxy.h
+++ b/Engine/Plugins/Experimental/MeshModelingToolset/Source/ModelingComponents/Private/BaseDynamicMeshSceneProxy.h
@@ -510,7 +510,7 @@ public:
 				}
 
 				RenderBuffers->ColorVertexBuffer.VertexColor(VertIdx) = (bHaveColors) ?
-					(FColor)Mesh->GetVertexColor(Tri[j]) : TriColor;
+					UE::Geometry::ToFColor(Mesh->GetVertexColor(Tri[j])) : TriColor;
 
 				RenderBuffers->IndexBuffer.Indices[TriIdx++] = VertIdx;		// currently TriIdx == VertIdx so we don't really need both...
 				VertIdx++;
@@ -689,7 +689,7 @@ public:
 				if (bUpdateColors)
 				{
 					RenderBuffers->ColorVertexBuffer.VertexColor(VertIdx) = (bHaveColors) ?
-						(FColor)Mesh->GetVertexColor(Tri[j]) : TriColor;
+						UE::Geometry::ToFColor(Mesh->GetVertexColor(Tri[j])) : TriColor;
 				}
 
 				VertIdx++;
diff --git a/Engine/Plugins/Experimental/MeshModelingToolset/Source/ModelingComponents/Private/Mechanics/ConstructionPlaneMechanic.cpp b/Engine/Plugins/Experimental/MeshModelingToolset/Source/ModelingComponents/Private/Mechanics/ConstructionPlaneMechanic.cpp
--- a/Engine/Plugins/Experimental/MeshModelingToolset/Source/ModelingComponents/Private/Mechanics/ConstructionPlaneMechanic.cpp
+++ b/Engine/Plugins/Experimental/MeshModelingToolset/Source/ModelingComponents/Private/Mechanics/ConstructionPlaneMechanic.cpp
@@ -95,7 +95,7 @@ void UConstructionPlaneMechanic::SetDrawPlaneFromWorldPos(const FVector3d& Posit
 		{
 			// Otherwise, let's place one of the other axes into the XY plane so that the frame is more
 			// useful for translation. We somewhat arbitrarily choose Y for this. 
-			FVector3d FrameY = Normal.Cross(FVector3d::UnitZ()).Normalized(); // orthogonal to world Z and frame Z 
+			FVector3d FrameY = UE::Geometry::Normalized( Normal.Cross(FVector3d::UnitZ()) ); // orthogonal to world Z and frame Z 
 			FVector3d FrameX = FrameY.Cross(Normal); // safe to not normalize because already orthogonal
 			Plane = FFrame3d(Position, FrameX, FrameY, Normal);
 		}
diff --git a/Engine/Plugins/Experimental/MeshModelingToolset/Source/ModelingComponents/Private/Mechanics/DragAlignmentMechanic.cpp b/Engine/Plugins/Experimental/MeshModelingToolset/Source/ModelingComponents/Private/Mechanics/DragAlignmentMechanic.cpp
--- a/Engine/Plugins/Experimental/MeshModelingToolset/Source/ModelingComponents/Private/Mechanics/DragAlignmentMechanic.cpp
+++ b/Engine/Plugins/Experimental/MeshModelingToolset/Source/ModelingComponents/Private/Mechanics/DragAlignmentMechanic.cpp
@@ -201,7 +201,7 @@ void UDragAlignmentMechanic::InitializeDeformedMeshRayCast(TFunction<FDynamicMes
 		// Transform ray into local space
 		FRay3d LocalRay(TargetTransform.InverseTransformPosition(WorldRay.Origin),
 			TargetTransform.InverseTransformVector(WorldRay.Direction));
-		LocalRay.Direction.Normalize();
+		UE::Geometry::Normalize(LocalRay.Direction);
 
 		FDynamicMeshAABBTree3* Spatial = GetSpatialIn();
 		const FDynamicMesh3* Mesh = Spatial->GetMesh();
diff --git a/Engine/Plugins/Experimental/MeshModelingToolset/Source/ModelingComponents/Private/Mechanics/PlaneDistanceFromHitMechanic.cpp b/Engine/Plugins/Experimental/MeshModelingToolset/Source/ModelingComponents/Private/Mechanics/PlaneDistanceFromHitMechanic.cpp
--- a/Engine/Plugins/Experimental/MeshModelingToolset/Source/ModelingComponents/Private/Mechanics/PlaneDistanceFromHitMechanic.cpp
+++ b/Engine/Plugins/Experimental/MeshModelingToolset/Source/ModelingComponents/Private/Mechanics/PlaneDistanceFromHitMechanic.cpp
@@ -98,7 +98,7 @@ void UPlaneDistanceFromHitMechanic::UpdateCurrentDistance(const FRay& WorldRay)
 		FDistLine3Ray3d Distance(FLine3d(PreviewHeightFrame.Origin, PreviewHeightFrame.Z()), FRay3d(WorldRay));
 		float DistSqr = Distance.GetSquared();
 		float WorldHitDist = WorldRay.GetParameter( (FVector)Distance.RayClosestPoint );
-		NearestHitFrameWorld = FFrame3d(Distance.RayClosestPoint, (Distance.RayClosestPoint - Distance.LineClosestPoint).Normalized());
+		NearestHitFrameWorld = FFrame3d(Distance.RayClosestPoint, UE::Geometry::Normalized(Distance.RayClosestPoint - Distance.LineClosestPoint) );
 		FVector3d HitPosLocal = PreviewHeightFrame.ToFramePoint(Distance.RayClosestPoint);
 		NearestHitHeight = HitPosLocal.Z;
 		NearestHitDist = WorldHitDist;
diff --git a/Engine/Plugins/Experimental/MeshModelingToolset/Source/ModelingComponents/Private/Mechanics/SpaceCurveDeformationMechanic.cpp b/Engine/Plugins/Experimental/MeshModelingToolset/Source/ModelingComponents/Private/Mechanics/SpaceCurveDeformationMechanic.cpp
--- a/Engine/Plugins/Experimental/MeshModelingToolset/Source/ModelingComponents/Private/Mechanics/SpaceCurveDeformationMechanic.cpp
+++ b/Engine/Plugins/Experimental/MeshModelingToolset/Source/ModelingComponents/Private/Mechanics/SpaceCurveDeformationMechanic.cpp
@@ -375,7 +375,7 @@ void USpaceCurveDeformationMechanic::GizmoTransformChanged(UTransformProxy* Prox
 				t = 1.0 - t;		// linear
 			}
 
-			NewFrame.Origin = FVector3d::Lerp(InitialFrame.Origin, NewFrame.Origin, t);
+			NewFrame.Origin = UE::Geometry::Lerp(InitialFrame.Origin, NewFrame.Origin, t);
 			NewFrame.Rotation.SetToSlerp(InitialFrame.Rotation, NewFrame.Rotation, t);
 			//NewFrame.Rotation = SafeSlerp(InitialFrame.Rotation, NewFrame.Rotation, t);
 		}
diff --git a/Engine/Plugins/Experimental/MeshModelingToolset/Source/ModelingComponents/Private/PreviewMesh.cpp b/Engine/Plugins/Experimental/MeshModelingToolset/Source/ModelingComponents/Private/PreviewMesh.cpp
--- a/Engine/Plugins/Experimental/MeshModelingToolset/Source/ModelingComponents/Private/PreviewMesh.cpp
+++ b/Engine/Plugins/Experimental/MeshModelingToolset/Source/ModelingComponents/Private/PreviewMesh.cpp
@@ -319,7 +319,7 @@ bool UPreviewMesh::FindRayIntersection(const FRay3d& WorldRay, FHitResult& HitOu
 		FTransform3d Transform(TemporaryParentActor->GetActorTransform());
 		FRay3d LocalRay(Transform.InverseTransformPosition(WorldRay.Origin),
 			Transform.InverseTransformVector(WorldRay.Direction));
-		LocalRay.Direction.Normalize();
+		UE::Geometry::Normalize(LocalRay.Direction);
 		int HitTriID = MeshAABBTree.FindNearestHitTriangle(LocalRay);
 		if (HitTriID != FDynamicMesh3::InvalidID)
 		{
diff --git a/Engine/Plugins/Experimental/MeshModelingToolset/Source/ModelingComponents/Private/Selection/GroupTopologySelector.cpp b/Engine/Plugins/Experimental/MeshModelingToolset/Source/ModelingComponents/Private/Selection/GroupTopologySelector.cpp
--- a/Engine/Plugins/Experimental/MeshModelingToolset/Source/ModelingComponents/Private/Selection/GroupTopologySelector.cpp
+++ b/Engine/Plugins/Experimental/MeshModelingToolset/Source/ModelingComponents/Private/Selection/GroupTopologySelector.cpp
@@ -131,7 +131,7 @@ bool FGroupTopologySelector::FindSelectedElement(const FSelectionSettings& Setti
 		if (DoCornerBasedSelection(Settings, Ray, Spatial, TopoSpatial, CornerResults, CornerPosition, &CornerSegmentEdgeID))
 		{
 			bHaveCornerHit = true;
-			CornerAngle = Ray.Direction.AngleD((CornerPosition - Ray.Origin).Normalized());
+			CornerAngle = UE::Geometry::AngleD(Ray.Direction, UE::Geometry::Normalized(CornerPosition - Ray.Origin) );
 		}
 	}
 
@@ -146,7 +146,7 @@ bool FGroupTopologySelector::FindSelectedElement(const FSelectionSettings& Setti
 		if (DoEdgeBasedSelection(Settings, Ray, Spatial, TopoSpatial, EdgeResults, EdgePosition, &EdgeSegmentEdgeID))
 		{
 			bHaveEdgeHit = true;
-			EdgeAngle = Ray.Direction.AngleD((EdgePosition - Ray.Origin).Normalized());
+			EdgeAngle = UE::Geometry::AngleD(Ray.Direction, UE::Geometry::Normalized(EdgePosition - Ray.Origin) );
 		}
 	}
 
@@ -276,7 +276,7 @@ bool FGroupTopologySelector::DoCornerBasedSelection(const FSelectionSettings& Se
 			}
 
 			// Make sure that closest corner to current element is parallel with view ray
-			FVector3d ClosestTowardElement = (Element.NearestGeoPoint - ClosestElement->NearestGeoPoint).Normalized();
+			FVector3d ClosestTowardElement = UE::Geometry::Normalized(Element.NearestGeoPoint - ClosestElement->NearestGeoPoint);
 			// There would usually be one more abs() in here, but we know that other elements are down ray direction
 			if (abs(ClosestTowardElement.Dot(Ray.Direction) - 1.0) < KINDA_SMALL_NUMBER)
 			{
@@ -298,7 +298,7 @@ bool FGroupTopologySelector::DoCornerBasedSelection(const FSelectionSettings& Se
 		{
 			FDynamicMesh3::FEdge Edge = Mesh->GetEdge(Eid);
 			int32 OtherVid = (Edge.Vert.A == ClosestVid) ? Edge.Vert.B : Edge.Vert.A;
-			FVector3d EdgeVector = (Mesh->GetVertex(OtherVid) - Mesh->GetVertex(ClosestVid)).Normalized();
+			FVector3d EdgeVector = UE::Geometry::Normalized(Mesh->GetVertex(OtherVid) - Mesh->GetVertex(ClosestVid));
 			if (abs(EdgeVector.Dot(Ray.Direction) - 1.0) < KINDA_SMALL_NUMBER)
 			{
 				int TopologyEdgeIndex = Topology->FindGroupEdgeID(Eid);
@@ -326,7 +326,7 @@ bool FGroupTopologySelector::DoCornerBasedSelection(const FSelectionSettings& Se
 				for (int32 Eid : Mesh->VtxEdgesItr(Topology->GetCornerVertexID(DownRayElements[i])))
 				{
 					FDynamicMesh3::FEdge Edge = Mesh->GetEdge(Eid);
-					FVector3d EdgeVector = (Mesh->GetVertex(Edge.Vert.A) - Mesh->GetVertex(Edge.Vert.B)).Normalized();
+					FVector3d EdgeVector = UE::Geometry::Normalized(Mesh->GetVertex(Edge.Vert.A) - Mesh->GetVertex(Edge.Vert.B));
 
 					// Compare absolute value of dot product to 1. We already made sure that one of the vertices is in line
 					// with the closest vertex earlier.
@@ -578,7 +578,7 @@ bool IsOccluded(const FGeometrySet3::FNearest& ClosestElement, const FVector3d&
 bool IsOccluded(const FVector3d& Point, const FVector3d& ViewOrigin, const FDynamicMeshAABBTree3* Spatial)
 {
 	// Shoot ray backwards to see if we hit something. 
-	FRay3d ToEyeRay(Point, (ViewOrigin - Point).Normalized(), true);
+	FRay3d ToEyeRay(Point, UE::Geometry::Normalized(ViewOrigin - Point), true);
 	ToEyeRay.Origin += (double)(100 * FMathf::ZeroTolerance) * ToEyeRay.Direction;
 	if (Spatial->FindNearestHitTriangle(ToEyeRay) >= 0)
 	{
diff --git a/Engine/Plugins/Experimental/MeshModelingToolset/Source/ModelingComponents/Private/Selection/PolygonSelectionMechanic.cpp b/Engine/Plugins/Experimental/MeshModelingToolset/Source/ModelingComponents/Private/Selection/PolygonSelectionMechanic.cpp
--- a/Engine/Plugins/Experimental/MeshModelingToolset/Source/ModelingComponents/Private/Selection/PolygonSelectionMechanic.cpp
+++ b/Engine/Plugins/Experimental/MeshModelingToolset/Source/ModelingComponents/Private/Selection/PolygonSelectionMechanic.cpp
@@ -277,7 +277,7 @@ bool UPolygonSelectionMechanic::TopologyHitTest(const FRay& WorldRay, FHitResult
 {
 	FRay3d LocalRay(TargetTransform.InverseTransformPosition(WorldRay.Origin),
 		TargetTransform.InverseTransformVector(WorldRay.Direction));
-	LocalRay.Direction.Normalize();
+	UE::Geometry::Normalize(LocalRay.Direction);
 
 	FVector3d LocalPosition, LocalNormal;
 	int32 EdgeSegmentId; // Only used if hit is an edge
@@ -361,7 +361,7 @@ bool UPolygonSelectionMechanic::UpdateHighlight(const FRay& WorldRay)
 
 	FRay3d LocalRay(TargetTransform.InverseTransformPosition(WorldRay.Origin),
 		TargetTransform.InverseTransformVector(WorldRay.Direction));
-	LocalRay.Direction.Normalize();
+	UE::Geometry::Normalize(LocalRay.Direction);
 
 	HilightSelection.Clear();
 	FVector3d LocalPosition, LocalNormal;
@@ -414,9 +414,9 @@ bool UPolygonSelectionMechanic::UpdateHighlight(const FRay& WorldRay)
 
 					FIndex3i VertIndices = Mesh->GetTriangle(Tid);
 					DrawnTriangleSetComponent->AddTriangle(FRenderableTriangle(HighlightedFaceMaterial,
-						FRenderableTriangleVertex((FVector)Mesh->GetVertex(VertIndices.A), (FVector2D)Mesh->GetVertexUV(VertIndices.A), (FVector)TriangleNormal, (FColor)Mesh->GetVertexColor(VertIndices.A)),
-						FRenderableTriangleVertex((FVector)Mesh->GetVertex(VertIndices.B), (FVector2D)Mesh->GetVertexUV(VertIndices.B), (FVector)TriangleNormal, (FColor)Mesh->GetVertexColor(VertIndices.B)),
-						FRenderableTriangleVertex((FVector)Mesh->GetVertex(VertIndices.C), (FVector2D)Mesh->GetVertexUV(VertIndices.C), (FVector)TriangleNormal, (FColor)Mesh->GetVertexColor(VertIndices.C))));
+						FRenderableTriangleVertex((FVector)Mesh->GetVertex(VertIndices.A), (FVector2D)Mesh->GetVertexUV(VertIndices.A), (FVector)TriangleNormal, UE::Geometry::ToFColor(Mesh->GetVertexColor(VertIndices.A))),
+						FRenderableTriangleVertex((FVector)Mesh->GetVertex(VertIndices.B), (FVector2D)Mesh->GetVertexUV(VertIndices.B), (FVector)TriangleNormal, UE::Geometry::ToFColor(Mesh->GetVertexColor(VertIndices.B))),
+						FRenderableTriangleVertex((FVector)Mesh->GetVertex(VertIndices.C), (FVector2D)Mesh->GetVertexUV(VertIndices.C), (FVector)TriangleNormal, UE::Geometry::ToFColor(Mesh->GetVertexColor(VertIndices.C))) ));
 				}
 
 				CurrentlyHighlightedGroups.Add(Gid);
@@ -440,7 +440,7 @@ bool UPolygonSelectionMechanic::UpdateSelection(const FRay& WorldRay, FVector3d&
 {
 	FRay3d LocalRay(TargetTransform.InverseTransformPosition(WorldRay.Origin),
 		TargetTransform.InverseTransformVector(WorldRay.Direction));
-	LocalRay.Direction.Normalize();
+	UE::Geometry::Normalize(LocalRay.Direction);
 
 	bool bSelectionModified = false;
 	FVector3d LocalPosition, LocalNormal;
diff --git a/Engine/Plugins/Experimental/MeshModelingToolset/Source/ModelingComponents/Private/Snapping/PointPlanarSnapSolver.cpp b/Engine/Plugins/Experimental/MeshModelingToolset/Source/ModelingComponents/Private/Snapping/PointPlanarSnapSolver.cpp
--- a/Engine/Plugins/Experimental/MeshModelingToolset/Source/ModelingComponents/Private/Snapping/PointPlanarSnapSolver.cpp
+++ b/Engine/Plugins/Experimental/MeshModelingToolset/Source/ModelingComponents/Private/Snapping/PointPlanarSnapSolver.cpp
@@ -168,7 +168,8 @@ void FPointPlanarSnapSolver::GenerateLineIntersectionTargets()
 			FVector3d DirectionInPlane = Plane.ToFrameVector(Line.Line.Direction);
 			if (DirectionInPlane.Z < IN_PLANE_THRESHOLD)
 			{
-				LinesInPlane.Add(FLine2d(OriginInPlane.XY(), DirectionInPlane.XY()));
+				LinesInPlane.Add(FLine2d(
+					FVector2d(OriginInPlane.X, OriginInPlane.Y), FVector2d(DirectionInPlane.X, DirectionInPlane.Y)));
 				OriginalLinePointers.Add(&Line);
 			}
 		}
@@ -181,7 +182,8 @@ void FPointPlanarSnapSolver::GenerateLineIntersectionTargets()
 			FVector3d DirectionInPlane = Plane.ToFrameVector(Line.Line.Direction);
 			if (DirectionInPlane.Z < IN_PLANE_THRESHOLD)
 			{
-				LinesInPlane.Add(FLine2d(OriginInPlane.XY(), DirectionInPlane.XY()));
+				LinesInPlane.Add(FLine2d(
+					FVector2d(OriginInPlane.X, OriginInPlane.Y), FVector2d(DirectionInPlane.X, DirectionInPlane.Y)));
 				OriginalLinePointers.Add(&Line);
 			}
 		}
diff --git a/Engine/Plugins/Experimental/MeshModelingToolset/Source/ModelingComponents/Private/Transforms/QuickAxisTranslater.cpp b/Engine/Plugins/Experimental/MeshModelingToolset/Source/ModelingComponents/Private/Transforms/QuickAxisTranslater.cpp
--- a/Engine/Plugins/Experimental/MeshModelingToolset/Source/ModelingComponents/Private/Transforms/QuickAxisTranslater.cpp
+++ b/Engine/Plugins/Experimental/MeshModelingToolset/Source/ModelingComponents/Private/Transforms/QuickAxisTranslater.cpp
@@ -86,7 +86,7 @@ void FQuickAxisTranslater::UpdateActiveFrameOrigin(const FVector3d& NewOrigin)
 bool ProjectToCameraPlanePos(const FVector3d& ScenePos, const FFrame3d& CameraPlaneFrame, const FVector3d& CameraPos, FVector2d& CamPlanePos)
 {
 	FVector3d HitPos;
-	if (CameraPlaneFrame.RayPlaneIntersection(CameraPos, (ScenePos - CameraPos).Normalized(), 2, HitPos) == false)
+	if (CameraPlaneFrame.RayPlaneIntersection(CameraPos, UE::Geometry::Normalized(ScenePos - CameraPos), 2, HitPos) == false)
 	{
 		return false;
 	}
@@ -115,7 +115,7 @@ void FQuickAxisTranslater::UpdateSnapAxes()
 	bool IgnoreAxis[3] = { false, false, false };
 
 	// LocalForward is vector from eye to target position, use this as "forward"
-	FVector3d LocalForward = (AxisFrameWorld.Origin - CameraPosition).Normalized();
+	FVector3d LocalForward = UE::Geometry::Normalized(AxisFrameWorld.Origin - CameraPosition);
 	double DirectionDotX = AxisX.Dot(LocalForward);
 	double DirectionDotY = AxisY.Dot(LocalForward);
 	double DirectionDotZ = AxisZ.Dot(LocalForward);
@@ -146,9 +146,9 @@ void FQuickAxisTranslater::UpdateSnapAxes()
 	{
 		return;		// abort for now...
 	}
-	FVector2d PlaneAxisX = (CamPosAxisX - CamPosCenter).Normalized();
-	FVector2d PlaneAxisY = (CamPosAxisY - CamPosCenter).Normalized();
-	FVector2d PlaneAxisZ = (CamPosAxisZ - CamPosCenter).Normalized();
+	FVector2d PlaneAxisX = UE::Geometry::Normalized(CamPosAxisX - CamPosCenter);
+	FVector2d PlaneAxisY = UE::Geometry::Normalized(CamPosAxisY - CamPosCenter);
+	FVector2d PlaneAxisZ = UE::Geometry::Normalized(CamPosAxisZ - CamPosCenter);
 
 	// if angle between these 2D projections of axes is too small, then axes are visually overlapping
 	// from this view position, and snapping will be unstable. So only keep the one pointing "towards" camera.
diff --git a/Engine/Plugins/Experimental/MeshModelingToolset/Source/ModelingOperators/Private/MeshSpaceDeformerOps/BendMeshOp.cpp b/Engine/Plugins/Experimental/MeshModelingToolset/Source/ModelingOperators/Private/MeshSpaceDeformerOps/BendMeshOp.cpp
--- a/Engine/Plugins/Experimental/MeshModelingToolset/Source/ModelingOperators/Private/MeshSpaceDeformerOps/BendMeshOp.cpp
+++ b/Engine/Plugins/Experimental/MeshModelingToolset/Source/ModelingOperators/Private/MeshSpaceDeformerOps/BendMeshOp.cpp
@@ -175,7 +175,7 @@ void FBendMeshOp::CalculateResult(FProgressCancel* Progress)
 
 			}
 
-			Normals->SetElement(ElID, FVector3f(RotatedNormal.Normalized()));
+			Normals->SetElement(ElID, FVector3f(Normalized(RotatedNormal)));
 		});
 	}
 
diff --git a/Engine/Plugins/Experimental/MeshModelingToolset/Source/ModelingOperators/Private/MeshSpaceDeformerOps/FlareMeshOp.cpp b/Engine/Plugins/Experimental/MeshModelingToolset/Source/ModelingOperators/Private/MeshSpaceDeformerOps/FlareMeshOp.cpp
--- a/Engine/Plugins/Experimental/MeshModelingToolset/Source/ModelingOperators/Private/MeshSpaceDeformerOps/FlareMeshOp.cpp
+++ b/Engine/Plugins/Experimental/MeshModelingToolset/Source/ModelingOperators/Private/MeshSpaceDeformerOps/FlareMeshOp.cpp
@@ -151,7 +151,7 @@ void FFlareMeshOp::CalculateResult(FProgressCancel* Progress)
 
 			}
 
-			Normals->SetElement(ElID, FVector3f(RotatedNormal.Normalized()));
+			Normals->SetElement(ElID, FVector3f(Normalized(RotatedNormal)));
 		});
 	}
 
diff --git a/Engine/Plugins/Experimental/MeshModelingToolset/Source/ModelingOperators/Private/MeshSpaceDeformerOps/TwistMeshOp.cpp b/Engine/Plugins/Experimental/MeshModelingToolset/Source/ModelingOperators/Private/MeshSpaceDeformerOps/TwistMeshOp.cpp
--- a/Engine/Plugins/Experimental/MeshModelingToolset/Source/ModelingOperators/Private/MeshSpaceDeformerOps/TwistMeshOp.cpp
+++ b/Engine/Plugins/Experimental/MeshModelingToolset/Source/ModelingOperators/Private/MeshSpaceDeformerOps/TwistMeshOp.cpp
@@ -123,7 +123,7 @@ void FTwistMeshOp::CalculateResult(FProgressCancel* Progress)
 				
 			}
 
-			Normals->SetElement(ElID, FVector3f(RotatedNormal.Normalized()));
+			Normals->SetElement(ElID, FVector3f(Normalized(RotatedNormal)));
 		});
 	}
 
diff --git a/Engine/Plugins/Experimental/MeshModelingToolset/Source/ModelingOperators/Private/SmoothingOps/IterativeSmoothingOp.cpp b/Engine/Plugins/Experimental/MeshModelingToolset/Source/ModelingOperators/Private/SmoothingOps/IterativeSmoothingOp.cpp
--- a/Engine/Plugins/Experimental/MeshModelingToolset/Source/ModelingOperators/Private/SmoothingOps/IterativeSmoothingOp.cpp
+++ b/Engine/Plugins/Experimental/MeshModelingToolset/Source/ModelingOperators/Private/SmoothingOps/IterativeSmoothingOp.cpp
@@ -135,7 +135,7 @@ void FIterativeSmoothingOp::Smooth_Forward(bool bUniform)
 			}
 
 			double UseAlpha = GetSmoothAlpha(vid, false);
-			SmoothedBuffer[vid] = FVector3d::Lerp(PositionBuffer[vid], Centroid, UseAlpha);
+			SmoothedBuffer[vid] = Lerp(PositionBuffer[vid], Centroid, UseAlpha);
 
 		} /*, EParallelForFlags::ForceSingleThread*/ );
 
@@ -149,7 +149,7 @@ void FIterativeSmoothingOp::Smooth_Forward(bool bUniform)
 					[&](int32 nbrvid) { return PositionBuffer[nbrvid]; },
 					[&](int32 nbrvid) { return bIsBoundary[nbrvid]; });
 				double UseAlpha = GetSmoothAlpha(vid, true);
-				SmoothedBuffer[vid] = FVector3d::Lerp(PositionBuffer[vid], Centroid, UseAlpha);
+				SmoothedBuffer[vid] = Lerp(PositionBuffer[vid], Centroid, UseAlpha);
 			});
 		}
 
diff --git a/Engine/Plugins/Experimental/MeshModelingToolset/Source/ModelingOperators/Public/DeformationOps/MeshOffsetOps.h b/Engine/Plugins/Experimental/MeshModelingToolset/Source/ModelingOperators/Public/DeformationOps/MeshOffsetOps.h
--- a/Engine/Plugins/Experimental/MeshModelingToolset/Source/ModelingOperators/Public/DeformationOps/MeshOffsetOps.h
+++ b/Engine/Plugins/Experimental/MeshModelingToolset/Source/ModelingOperators/Public/DeformationOps/MeshOffsetOps.h
@@ -263,7 +263,7 @@ public:
 						return;
 					}
 					FVector3d Centroid = FMeshWeights::UniformCentroid(*ResultMesh, vid, [&](int32 nbrvid) { return PositionBuffer[nbrvid]; });
-					SmoothedBuffer[vid] = FVector3d::Lerp(PositionBuffer[vid], Centroid, SmoothAlpha);
+					SmoothedBuffer[vid] = UE::Geometry::Lerp(PositionBuffer[vid], Centroid, SmoothAlpha);
 					if (bReprojectSmooth)
 					{
 						SmoothedBuffer[vid] = Spatial.FindNearestPoint(SmoothedBuffer[vid]);
@@ -278,7 +278,7 @@ public:
 						FVector3d Centroid = FMeshWeights::FilteredUniformCentroid(*ResultMesh, vid,
 							[&](int32 nbrvid) { return PositionBuffer[nbrvid]; },
 							[&](int32 nbrvid) { return BoundaryCache.bIsBoundary[nbrvid]; });
-						SmoothedBuffer[vid] = FVector3d::Lerp(PositionBuffer[vid], Centroid, BoundaryAlpha);
+						SmoothedBuffer[vid] = UE::Geometry::Lerp(PositionBuffer[vid], Centroid, BoundaryAlpha);
 					});
 				}
 
diff --git a/Engine/Plugins/Experimental/PlanarCutPlugin/Source/PlanarCut/Private/PlanarCut.cpp b/Engine/Plugins/Experimental/PlanarCutPlugin/Source/PlanarCut/Private/PlanarCut.cpp
--- a/Engine/Plugins/Experimental/PlanarCutPlugin/Source/PlanarCut/Private/PlanarCut.cpp
+++ b/Engine/Plugins/Experimental/PlanarCutPlugin/Source/PlanarCut/Private/PlanarCut.cpp
@@ -302,7 +302,7 @@ namespace UE
 								for (int DivJ = 0; DivJ < DivisionsAcross; DivJ++)
 								{
 									double AlongAcross = (DivJ + 1) * FactorAcross;
-									FVector3d Bary = FVector3d::Lerp(E1Bary, E2Bary, AlongAcross);
+									FVector3d Bary = UE::Geometry::Lerp(E1Bary, E2Bary, AlongAcross);
 									FVector3d SamplePos = Triangle.BarycentricPoint(Bary);
 									if (!KnownSamples[ComponentIdx].IsCellEmptyUnsafe(SamplePos)) // fast early out; def. have pt within radius
 									{
@@ -1011,7 +1011,7 @@ private:
 					FVector3d OrigPos, NoisePos;
 					OriginalPosns->GetValue(VID, OrigPos);
 					NoisePos = Mesh.GetVertex(VID);
-					return FVector3d::Lerp(OrigPos, NoisePos, t);
+					return UE::Geometry::Lerp(OrigPos, NoisePos, t);
 				};
 				auto InterpTri = [&Mesh, &InterpVert](int TID, double t)
 				{
@@ -1067,7 +1067,7 @@ private:
 								static_cast<TDynamicMeshVertexAttribute<double, 3>*>(OtherMesh.Attributes()->GetAttachedAttribute(OriginalPositionAttribute));
 							FVector3d OrigPos;
 							OtherOriginalPosns->GetValue(IDs.B, OrigPos);
-							NewPos = FVector3d::Lerp(OrigPos, OldPos, tsafe);
+							NewPos = UE::Geometry::Lerp(OrigPos, OldPos, tsafe);
 							OtherMesh.SetVertex(IDs.B, NewPos);
 						}
 						MeshesVertices.UpdatePoint(IDs, OldPos, NewPos);
diff --git a/Engine/Plugins/Experimental/Text3D/Source/Text3D/Private/Arrangement.cpp b/Engine/Plugins/Experimental/Text3D/Source/Text3D/Private/Arrangement.cpp
--- a/Engine/Plugins/Experimental/Text3D/Source/Text3D/Private/Arrangement.cpp
+++ b/Engine/Plugins/Experimental/Text3D/Source/Text3D/Private/Arrangement.cpp
@@ -263,7 +263,7 @@ bool FArrangement::find_intersecting_edges(FVector2d A, FVector2d B, TArray<FInt
     int num_hits = 0;
     FVector2d x = FVector2d::Zero(), y = FVector2d::Zero();
     FVector2d EPerp = (B - A).Perp();
-    EPerp.Normalize();
+    UE::Geometry::Normalize(EPerp);
     for (int EID : Graph.EdgeIndices())
     {
         Graph.GetEdgeV(EID, x, y);
diff --git a/Engine/Plugins/MeshPainting/Source/MeshPaintingToolset/Private/BaseMeshPaintComponentAdapter.cpp b/Engine/Plugins/MeshPainting/Source/MeshPaintingToolset/Private/BaseMeshPaintComponentAdapter.cpp
--- a/Engine/Plugins/MeshPainting/Source/MeshPaintingToolset/Private/BaseMeshPaintComponentAdapter.cpp
+++ b/Engine/Plugins/MeshPainting/Source/MeshPaintingToolset/Private/BaseMeshPaintComponentAdapter.cpp
@@ -187,7 +187,7 @@ bool FBaseMeshPaintComponentAdapter::RayIntersectAdapter(UE::Geometry::FIndex3i&
 	int32 HitTriangleID;
 	double NearestT;
 	UE::Geometry::FVector3d Direction((UE::Geometry::FVector3d)End - (UE::Geometry::FVector3d)Start);
-	Direction.Normalize();
+	UE::Geometry::Normalize(Direction);
 	UE::Geometry::FRay3d Ray((UE::Geometry::FVector3d)Start, Direction);
 	bool bHit = AABBTree->FindNearestHitTriangle(Ray, NearestT, HitTriangleID);
 	if (bHit)
diff --git a/Engine/Plugins/MeshPainting/Source/MeshPaintingToolset/Private/MeshPaintStaticMeshAdapter.cpp b/Engine/Plugins/MeshPainting/Source/MeshPaintingToolset/Private/MeshPaintStaticMeshAdapter.cpp
--- a/Engine/Plugins/MeshPainting/Source/MeshPaintingToolset/Private/MeshPaintStaticMeshAdapter.cpp
+++ b/Engine/Plugins/MeshPainting/Source/MeshPaintingToolset/Private/MeshPaintStaticMeshAdapter.cpp
@@ -218,7 +218,8 @@ bool FMeshPaintStaticMeshComponentAdapter::LineTraceComponent(struct FHitResult&
 		//check collinearity of A,B,C
 		if (TriNormal.SquaredLength() > (double)SMALL_NUMBER)
 		{
-			UE::Geometry::FRay3d LocalRay((UE::Geometry::FVector3d)LocalStart, ((UE::Geometry::FVector3d)LocalEnd - (UE::Geometry::FVector3d)LocalStart).Normalized());
+			UE::Geometry::FVector3d RayDirection = ((UE::Geometry::FVector3d)LocalEnd - (UE::Geometry::FVector3d)LocalStart);
+			UE::Geometry::FRay3d LocalRay((UE::Geometry::FVector3d)LocalStart, UE::Geometry::Normalized(RayDirection));
 			UE::Geometry::FIntrRay3Triangle3d RayTriIntersection(LocalRay, Triangle);
 			if (RayTriIntersection.Find())
 			{
