GeometryProcessing: refactor FVector2 member functions to UE::Geometry:: namespace functions
#rnx

diff --git a/Engine/Plugins/Experimental/GeometryProcessing/Source/GeometricObjects/Private/Util/RevolveUtil.cpp b/Engine/Plugins/Experimental/GeometryProcessing/Source/GeometricObjects/Private/Util/RevolveUtil.cpp
--- a/Engine/Plugins/Experimental/GeometryProcessing/Source/GeometricObjects/Private/Util/RevolveUtil.cpp
+++ b/Engine/Plugins/Experimental/GeometryProcessing/Source/GeometricObjects/Private/Util/RevolveUtil.cpp
@@ -103,7 +103,7 @@ bool RevolveUtil::ProfileIsCCWRelativeRevolve(TArray<FVector3d>& ProfileCurve, c
 	double NextDistance = NextWrtAxis.Distance(NextAlongAxis * RevolutionAxisDirection);
 	FVector2d Next(NextDistance, NextAlongAxis);
 	
-	return (Previous - Point).Cross(Next - Point) > 0;
+	return DotPerp((Previous - Point), (Next - Point)) > 0;
 }
 
 void RevolveUtil::MakeProfileCurveMidpointOfFirstStep(TArray<FVector3d>& ProfileCurve, double DegreesPerStep,
diff --git a/Engine/Plugins/Experimental/GeometryProcessing/Source/GeometricObjects/Public/Generators/SweepGenerator.h b/Engine/Plugins/Experimental/GeometryProcessing/Source/GeometricObjects/Public/Generators/SweepGenerator.h
--- a/Engine/Plugins/Experimental/GeometryProcessing/Source/GeometricObjects/Public/Generators/SweepGenerator.h
+++ b/Engine/Plugins/Experimental/GeometryProcessing/Source/GeometricObjects/Public/Generators/SweepGenerator.h
@@ -422,7 +422,8 @@ public:
 		TArray<FVector2d> NormalSides; NormalSides.SetNum(NumX - 1);
 		for (int XIdx = 0; XIdx+1 < NumX; XIdx++)
 		{
-			NormalSides[XIdx] = Normalized( (FVector2d(Radii[XIdx+1], Heights[XIdx+1]) - FVector2d(Radii[XIdx], Heights[XIdx])).Perp() );
+			FVector2d Vec = FVector2d(Radii[XIdx + 1], Heights[XIdx + 1]) - FVector2d(Radii[XIdx], Heights[XIdx]);
+			NormalSides[XIdx] = Normalized(PerpCW(Vec));
 		}
 		TArray<FVector2d> SmoothedNormalSides; SmoothedNormalSides.SetNum(NumX);
 		// smooth internal normals
diff --git a/Engine/Plugins/Experimental/GeometryProcessing/Source/GeometricObjects/Public/Intersection/IntrLine2Line2.h b/Engine/Plugins/Experimental/GeometryProcessing/Source/GeometricObjects/Public/Intersection/IntrLine2Line2.h
--- a/Engine/Plugins/Experimental/GeometryProcessing/Source/GeometricObjects/Public/Intersection/IntrLine2Line2.h
+++ b/Engine/Plugins/Experimental/GeometryProcessing/Source/GeometricObjects/Public/Intersection/IntrLine2Line2.h
@@ -162,20 +162,20 @@ public:
 		//   s1 = Q.Dot(Perp(D0))/D0.Dot(Perp(D1))
 
 		FVector2<RealType> diff = P1 - P0;
-		RealType D0DotPerpD1 = D0.DotPerp(D1);
+		RealType D0DotPerpD1 = DotPerp(D0, D1);
 		if (FMath::Abs(D0DotPerpD1) > DotThreshold) 
 		{
 			// Lines intersect in a single point.
 			RealType invD0DotPerpD1 = 1.0 / D0DotPerpD1;
-			RealType diffDotPerpD0 = diff.DotPerp(D0);
-			RealType diffDotPerpD1 = diff.DotPerp(D1);
+			RealType diffDotPerpD0 = DotPerp(diff, D0);
+			RealType diffDotPerpD1 = DotPerp(diff, D1);
 			s[0] = diffDotPerpD1 * invD0DotPerpD1;
 			s[1] = diffDotPerpD0 * invD0DotPerpD1;
 			return EIntersectionType::Point;
 		}
 
 		// Lines are parallel; check if they are within DistThresh apart
-		RealType diffNDotPerpD1 = diff.DotPerp(D1);
+		RealType diffNDotPerpD1 = DotPerp(diff, D1);
 		if (FMath::Abs(diffNDotPerpD1) <= DistThreshold) 
 		{
 			// Lines are collinear.
diff --git a/Engine/Plugins/Experimental/GeometryProcessing/Source/GeometricObjects/Public/Intersection/IntrLine2Triangle2.h b/Engine/Plugins/Experimental/GeometryProcessing/Source/GeometricObjects/Public/Intersection/IntrLine2Triangle2.h
--- a/Engine/Plugins/Experimental/GeometryProcessing/Source/GeometricObjects/Public/Intersection/IntrLine2Triangle2.h
+++ b/Engine/Plugins/Experimental/GeometryProcessing/Source/GeometricObjects/Public/Intersection/IntrLine2Triangle2.h
@@ -152,7 +152,7 @@ public:
 		for (int i = 0; i < 3; ++i)
 		{
 			FVector2<Real> diff = Tri.V[i] - Origin;
-			Dist[i] = diff.DotPerp(Direction);
+			Dist[i] = DotPerp(diff, Direction);
 			if (Dist[i] > Tolerance)
 			{
 				Sign[i] = 1;
diff --git a/Engine/Plugins/Experimental/GeometryProcessing/Source/GeometricObjects/Public/Intersection/IntrTriangle3Triangle3.h b/Engine/Plugins/Experimental/GeometryProcessing/Source/GeometricObjects/Public/Intersection/IntrTriangle3Triangle3.h
--- a/Engine/Plugins/Experimental/GeometryProcessing/Source/GeometricObjects/Public/Intersection/IntrTriangle3Triangle3.h
+++ b/Engine/Plugins/Experimental/GeometryProcessing/Source/GeometricObjects/Public/Intersection/IntrTriangle3Triangle3.h
@@ -755,7 +755,7 @@ protected:
 		FVector2<Real> save;
 		FVector2<Real> edge0 = projTri0.V[1] - projTri0.V[0];
 		FVector2<Real> edge1 = projTri0.V[2] - projTri0.V[0];
-		if (edge0.DotPerp(edge1) < (double)0)
+		if (DotPerp(edge0, edge1) < (double)0)
 		{
 			// Triangle is clockwise, reorder it.
 			save = projTri0.V[1];
@@ -765,7 +765,7 @@ protected:
 
 		edge0 = projTri1.V[1] - projTri1.V[0];
 		edge1 = projTri1.V[2] - projTri1.V[0];
-		if (edge0.DotPerp(edge1) < (double)0)
+		if (DotPerp(edge0, edge1) < (double)0)
 		{
 			// Triangle is clockwise, reorder it.
 			save = projTri1.V[1];
diff --git a/Engine/Plugins/Experimental/GeometryProcessing/Source/GeometricObjects/Public/LineTypes.h b/Engine/Plugins/Experimental/GeometryProcessing/Source/GeometricObjects/Public/LineTypes.h
--- a/Engine/Plugins/Experimental/GeometryProcessing/Source/GeometricObjects/Public/LineTypes.h
+++ b/Engine/Plugins/Experimental/GeometryProcessing/Source/GeometricObjects/Public/LineTypes.h
@@ -114,11 +114,11 @@ struct TLine2
 	{
 		// see IntrTLine2TLine2 for more detailed explanation
 		FVector2<T> diff = OtherLine.Origin - Origin;
-		T D0DotPerpD1 = Direction.DotPerp(OtherLine.Direction);
+		T D0DotPerpD1 = DotPerp(Direction, OtherLine.Direction);
 		if (TMathUtil<T>::Abs(D0DotPerpD1) > ParallelDotTolerance)                     // TLines intersect in a single point.
 		{
 			T invD0DotPerpD1 = ((T)1) / D0DotPerpD1;
-			T diffDotPerpD1 = diff.DotPerp(OtherLine.Direction);
+			T diffDotPerpD1 = DotPerp(diff, OtherLine.Direction);
 			T s = diffDotPerpD1 * invD0DotPerpD1;
 			IntersectionPointOut = Origin + s * Direction;
 			return true;
diff --git a/Engine/Plugins/Experimental/GeometryProcessing/Source/GeometricObjects/Public/Polygon2.h b/Engine/Plugins/Experimental/GeometryProcessing/Source/GeometricObjects/Public/Polygon2.h
--- a/Engine/Plugins/Experimental/GeometryProcessing/Source/GeometricObjects/Public/Polygon2.h
+++ b/Engine/Plugins/Experimental/GeometryProcessing/Source/GeometricObjects/Public/Polygon2.h
@@ -192,7 +192,7 @@ public:
 	 */
 	FVector2<T> GetNormal(int VertexIndex) const
 	{
-		return GetTangent(VertexIndex).Perp();
+		return PerpCW(GetTangent(VertexIndex));
 	}
 
 
@@ -208,7 +208,7 @@ public:
 		next -= Vertices[VertexIndex]; Normalize(next);
 		prev -= Vertices[VertexIndex]; Normalize(prev);
 
-		FVector2<T> n = (next.Perp() - prev.Perp());
+		FVector2<T> n = (PerpCW(next) - PerpCW(prev));
 		T len = Normalize(n);
 		if (len == 0) 
 		{
@@ -420,7 +420,7 @@ public:
 			{
 				if (b.Y > QueryPoint.Y)									// an upward crossing
 				{
-					if (FVector2<T>::Orient(a, b, QueryPoint) > 0)  // P left of edge
+					if (Orient(a, b, QueryPoint) > 0)  // P left of edge
 						++nWindingNumber;                       // have a valid up intersect
 				}
 			}
@@ -428,7 +428,7 @@ public:
 			{
 				if (b.Y <= QueryPoint.Y)									// a downward crossing
 				{
-					if (FVector2<T>::Orient(a, b, QueryPoint) < 0)  // P right of edge
+					if (Orient(a, b, QueryPoint) < 0)  // P right of edge
 					{
 						--nWindingNumber;						// have a valid down intersect
 					}
@@ -829,8 +829,8 @@ public:
 			FVector2<T> prev = Vertices[k == 0 ? Vertices.Num() - 1 : k - 1];
 			FVector2<T> dn = Normalized(next - v);
 			FVector2<T> dp = Normalized(prev - v);
-			TLine2<T> ln(v + OffsetDistance * dn.Perp(), dn);
-			TLine2<T> lp(v - OffsetDistance * dp.Perp(), dp);
+			TLine2<T> ln(v + OffsetDistance * PerpCW(dn), dn);
+			TLine2<T> lp(v - OffsetDistance * PerpCW(dp), dp);
 
 			bool bIntersectsAtPoint = ln.IntersectionPoint(lp, NewVertices[k]);
 			if (!bIntersectsAtPoint)  // lines were parallel
@@ -1045,7 +1045,7 @@ public:
 			T angle = AngleD(cp, cn);
 			// TODO document what this means sign-wise
 			// TODO re-test post Unreal port that this DotPerp is doing the right thing
-			T sign = cn.DotPerp(cp);
+			T sign = DotPerp(cn, cp);
 			bool bConcave = (sign > 0);
 
 			T thresh = (bConcave) ? MinConcaveAngleDeg : MinConvexAngleDeg;
diff --git a/Engine/Plugins/Experimental/GeometryProcessing/Source/GeometricObjects/Public/SegmentTypes.h b/Engine/Plugins/Experimental/GeometryProcessing/Source/GeometricObjects/Public/SegmentTypes.h
--- a/Engine/Plugins/Experimental/GeometryProcessing/Source/GeometricObjects/Public/SegmentTypes.h
+++ b/Engine/Plugins/Experimental/GeometryProcessing/Source/GeometricObjects/Public/SegmentTypes.h
@@ -199,7 +199,7 @@ public:
 		// [TODO] subtract Center from test?
 		FVector2<T> EndPt = Center + Extent * Direction;
 		FVector2<T> StartPt = Center - Extent * Direction;
-		T det = -FVector2<T>::Orient(EndPt, StartPt, QueryPoint);
+		T det = -Orient(EndPt, StartPt, QueryPoint);
 		return (det > Tolerance ? +1 : (det < -Tolerance ? -1 : 0));
 	}
 
@@ -216,12 +216,12 @@ public:
 		// see IntrLine2Line2 and IntrSegment2Segment2 for details on this code
 
 		FVector2<T> diff = OtherSegment.Center - Center;
-		T D0DotPerpD1 = Direction.DotPerp(OtherSegment.Direction);
+		T D0DotPerpD1 = DotPerp(Direction, OtherSegment.Direction);
 		if (TMathUtil<T>::Abs(D0DotPerpD1) > DotThresh)      // Lines intersect in a single point.
 		{
 			T invD0DotPerpD1 = ((T)1) / D0DotPerpD1;
-			T diffDotPerpD0 = diff.DotPerp(Direction);
-			T diffDotPerpD1 = diff.DotPerp(OtherSegment.Direction);
+			T diffDotPerpD0 = DotPerp(diff, Direction);
+			T diffDotPerpD1 = DotPerp(diff, OtherSegment.Direction);
 			T s = diffDotPerpD1 * invD0DotPerpD1;
 			T s2 = diffDotPerpD0 * invD0DotPerpD1;
 			return TMathUtil<T>::Abs(s) <= (Extent + IntervalThresh)
@@ -230,7 +230,7 @@ public:
 
 		// Lines are parallel.
 		Normalize(diff);
-		T diffNDotPerpD1 = diff.DotPerp(OtherSegment.Direction);
+		T diffNDotPerpD1 = DotPerp(diff, OtherSegment.Direction);
 		if (TMathUtil<T>::Abs(diffNDotPerpD1) <= DotThresh)
 		{
 			// Compute the location of OtherSegment endpoints relative to our Segment
@@ -297,7 +297,7 @@ public:
 	 */
 	static int WhichSide(const FVector2<T>& StartPt, const FVector2<T>& EndPt, const FVector2<T>& QueryPt, T Tolerance = (T)0)
 	{
-		T det = -FVector2<T>::Orient(StartPt, EndPt, QueryPt);
+		T det = -Orient(StartPt, EndPt, QueryPt);
 		return (det > Tolerance ? +1 : (det < -Tolerance ? -1 : 0));
 	}
 
diff --git a/Engine/Plugins/Experimental/GeometryProcessing/Source/GeometricObjects/Public/TriangleTypes.h b/Engine/Plugins/Experimental/GeometryProcessing/Source/GeometricObjects/Public/TriangleTypes.h
--- a/Engine/Plugins/Experimental/GeometryProcessing/Source/GeometricObjects/Public/TriangleTypes.h
+++ b/Engine/Plugins/Experimental/GeometryProcessing/Source/GeometricObjects/Public/TriangleTypes.h
@@ -93,9 +93,9 @@ struct TTriangle2
 	 */
 	static bool IsInside(const FVector2<RealType>& A, const FVector2<RealType>& B, const FVector2<RealType>& C, const FVector2<RealType>& QueryPoint)
 	{
-		RealType Sign1 = FVector2<RealType>::Orient(A, B, QueryPoint);
-		RealType Sign2 = FVector2<RealType>::Orient(B, C, QueryPoint);
-		RealType Sign3 = FVector2<RealType>::Orient(C, A, QueryPoint);
+		RealType Sign1 = Orient(A, B, QueryPoint);
+		RealType Sign2 = Orient(B, C, QueryPoint);
+		RealType Sign3 = Orient(C, A, QueryPoint);
 		return (Sign1*Sign2 > 0) && (Sign2*Sign3 > 0) && (Sign3*Sign1 > 0);
 	}
 
@@ -109,9 +109,9 @@ struct TTriangle2
 	/** @return true if QueryPoint is inside triangle or on edge */
 	static bool IsInsideOrOn(const FVector2<RealType>& A, const FVector2<RealType>& B, const FVector2<RealType>& C, const FVector2<RealType>& QueryPoint)
 	{
-		RealType Sign1 = FVector2<RealType>::Orient(A, B, QueryPoint);
-		RealType Sign2 = FVector2<RealType>::Orient(B, C, QueryPoint);
-		RealType Sign3 = FVector2<RealType>::Orient(C, A, QueryPoint);
+		RealType Sign1 = Orient(A, B, QueryPoint);
+		RealType Sign2 = Orient(B, C, QueryPoint);
+		RealType Sign3 = Orient(C, A, QueryPoint);
 		return (Sign1*Sign2 >= 0) && (Sign2*Sign3 >= 0) && (Sign3*Sign1 >= 0);
 	}
 
@@ -139,21 +139,21 @@ struct TTriangle2
 	 */
 	static int IsInsideOrOn_Oriented(const FVector2<RealType>& A, const FVector2<RealType>& B, const FVector2<RealType>& C, const FVector2<RealType>& QueryPoint)
 	{
-		checkSlow(FVector2<RealType>::Orient(A, B, C) <= 0); // TODO: remove this checkSlow; it's just to make sure the orientation is as expected
+		checkSlow(Orient(A, B, C) <= 0); // TODO: remove this checkSlow; it's just to make sure the orientation is as expected
 
-		RealType Sign1 = FVector2<RealType>::Orient(A, B, QueryPoint);
+		RealType Sign1 = Orient(A, B, QueryPoint);
 		if (Sign1 > 0)
 		{
 			return 1;
 		}
 		
-		RealType Sign2 = FVector2<RealType>::Orient(B, C, QueryPoint);
+		RealType Sign2 = Orient(B, C, QueryPoint);
 		if (Sign2 > 0)
 		{
 			return 1;
 		}
 		
-		RealType Sign3 = FVector2<RealType>::Orient(A, C, QueryPoint);
+		RealType Sign3 = Orient(A, C, QueryPoint);
 		if (Sign3 < 0) // note this edge is queried backwards so the sign test is also backwards
 		{
 			return 1;
diff --git a/Engine/Plugins/Experimental/GeometryProcessing/Source/GeometricObjects/Public/VectorTypes.h b/Engine/Plugins/Experimental/GeometryProcessing/Source/GeometricObjects/Public/VectorTypes.h
--- a/Engine/Plugins/Experimental/GeometryProcessing/Source/GeometricObjects/Public/VectorTypes.h
+++ b/Engine/Plugins/Experimental/GeometryProcessing/Source/GeometricObjects/Public/VectorTypes.h
@@ -129,21 +129,6 @@ struct FVector2
 		return X * V2.X + Y * V2.Y;
 	}
 
-	// Note: DotPerp and FVector2's version of Cross are the same function
-	constexpr T DotPerp(const FVector2<T>& V2) const
-	{
-		return X * V2.Y - Y * V2.X;
-	}
-	constexpr T Cross(const FVector2<T>& V2) const
-	{
-		return X * V2.Y - Y * V2.X;
-	}
-
-	constexpr FVector2<T> Perp() const
-	{
-		return FVector2<T>(Y, -X);
-	}
-
 	constexpr  FVector2 operator-() const
 	{
 		return FVector2(-X, -Y);
@@ -233,15 +218,29 @@ struct FVector2
 		return X != Other.X || Y != Other.Y;
 	}
 
-
-	// Note: This was called "IsLeft" in the Geometry3Sharp code (where it also went through Math.Sign)
-	// Returns >0 if C is to the left of the A->B Line, <0 if to the right, 0 if on the line
-	static T Orient(const FVector2<T>& A, const FVector2<T>& B, const FVector2<T>& C)
-	{
-		return (B - A).DotPerp(C - A);
-	}
 };
 
+/** @return dot product of V1 with PerpCW(V2), ie V2 rotated 90 degrees clockwise */
+template <typename T>
+constexpr T DotPerp(const FVector2<T>& V1, const FVector2<T>& V2)
+{
+	return V1.X * V2.Y - V1.Y * V2.X;
+}
+
+/** @return right-Perpendicular vector to V, ie V rotated 90 degrees clockwise */
+template <typename T>
+constexpr FVector2<T> PerpCW(const FVector2<T>& V)
+{
+	return FVector2<T>(V.Y, -V.X);
+}
+
+/** @return > 0 if C is to the left of the line from A to B, < 0 if to the right, 0 if on the line */
+template<typename T>
+T Orient(const FVector2<T>& A, const FVector2<T>& B, const FVector2<T>& C)
+{
+	return DotPerp((B - A), (C - A));
+}
+
 
 template <typename T>
 constexpr bool IsNormalized(const FVector2<T>& Vector, const T Tolerance = TMathUtil<T>::ZeroTolerance)
@@ -301,7 +300,7 @@ T SignedAngleR(const FVector2<T>& V1, const FVector2<T>& V2)
 {
 	T DotVal = V1.Dot(V2);
 	T ClampedDot = (DotVal < (T)-1) ? (T)-1 : ((DotVal > (T)1) ? (T)1 : DotVal);
-	T Direction = V1.Cross(V2);
+	T Direction = DotPerp(V1, V2);
 	if (Direction * Direction < TMathUtil<T>::ZeroTolerance)
 	{
 		return (DotVal < 0) ? TMathUtil<T>::Pi : (T)0;
diff --git a/Engine/Plugins/Experimental/GeometryProcessing/Source/GeometricObjects/Public/VectorUtil.h b/Engine/Plugins/Experimental/GeometryProcessing/Source/GeometricObjects/Public/VectorUtil.h
--- a/Engine/Plugins/Experimental/GeometryProcessing/Source/GeometricObjects/Public/VectorUtil.h
+++ b/Engine/Plugins/Experimental/GeometryProcessing/Source/GeometricObjects/Public/VectorUtil.h
@@ -375,7 +375,7 @@ namespace VectorUtil
 		FVector2<RealType> vj(Vj.Dot(Perp0), Vj.Dot(Perp1));
 		FVector2<RealType> vk(Vk.Dot(Perp0), Vk.Dot(Perp1));
 		// calculate gradient
-		FVector2<RealType> GradX = (fj-fi)*(vi-vk).Perp() + (fk-fi)*(vj-vi).Perp();
+		FVector2<RealType> GradX = (fj-fi)*PerpCW(vi-vk) + (fk-fi)*PerpCW(vj-vi);
 		// map back to 3D vector in triangle plane
 		RealType AreaScale = (RealType)1 / ((RealType)2 * VectorUtil::Area<RealType>(Vi, Vj, Vk));
 		return AreaScale * (GradX.X * Perp0 + GradX.Y * Perp1);
diff --git a/Engine/Plugins/Experimental/GeometryProcessing/Source/GeometryAlgorithms/Public/Arrangement2d.h b/Engine/Plugins/Experimental/GeometryProcessing/Source/GeometryAlgorithms/Public/Arrangement2d.h
--- a/Engine/Plugins/Experimental/GeometryProcessing/Source/GeometryAlgorithms/Public/Arrangement2d.h
+++ b/Engine/Plugins/Experimental/GeometryProcessing/Source/GeometryAlgorithms/Public/Arrangement2d.h
@@ -483,7 +483,7 @@ protected:
 	{
 		int num_hits = 0;
 		FVector2d x = FVector2d::Zero(), y = FVector2d::Zero();
-		FVector2d EPerp = (B - A).Perp();
+		FVector2d EPerp = PerpCW(B - A);
 		Normalize(EPerp);
 		for (int EID : Graph.EdgeIndices())
 		{
diff --git a/Engine/Plugins/Experimental/MeshModelingToolset/Source/MeshModelingTools/Private/DrawPolygonTool.cpp b/Engine/Plugins/Experimental/MeshModelingToolset/Source/MeshModelingTools/Private/DrawPolygonTool.cpp
--- a/Engine/Plugins/Experimental/MeshModelingToolset/Source/MeshModelingTools/Private/DrawPolygonTool.cpp
+++ b/Engine/Plugins/Experimental/MeshModelingToolset/Source/MeshModelingTools/Private/DrawPolygonTool.cpp
@@ -830,7 +830,7 @@ void UDrawPolygonTool::GetPolygonParametersFromFixedPoints(const TArray<FVector3
 
 	double Radius = Delta.Length();
 	FVector2d AxisX = Delta / Radius;
-	FVector2d AxisY = -AxisX.Perp();
+	FVector2d AxisY = -UE::Geometry::PerpCW(AxisX);
 	FVector2d HeightPt = DrawFrame.ToPlaneUV((FixedPoints.Num() == 3) ? FixedPoints[2] : FixedPoints[1], 2);
 	FVector2d HeightDelta = HeightPt - FirstReferencePt;
 	YSign = FMathd::Sign(HeightDelta.Dot(AxisY));
diff --git a/Engine/Plugins/Experimental/MeshModelingToolset/Source/MeshModelingTools/Private/MeshGroupPaintTool.cpp b/Engine/Plugins/Experimental/MeshModelingToolset/Source/MeshModelingTools/Private/MeshGroupPaintTool.cpp
--- a/Engine/Plugins/Experimental/MeshModelingToolset/Source/MeshModelingTools/Private/MeshGroupPaintTool.cpp
+++ b/Engine/Plugins/Experimental/MeshModelingToolset/Source/MeshModelingTools/Private/MeshGroupPaintTool.cpp
@@ -744,10 +744,10 @@ void UMeshGroupPaintTool::OnPolyLassoFinished(const FCameraPolyLasso& Lasso, boo
 		float PerpSign = Polygon.IsClockwise() ? -1.0 : 1.0;
 
 		Polyline.Insert(StartLine.PointAt(10000.0f), 0);
-		Polyline.Insert(Polyline[0] + 1000 * PerpSign * StartDirOut.Perp(), 0);
+		Polyline.Insert(Polyline[0] + 1000 * PerpSign * UE::Geometry::PerpCW(StartDirOut), 0);
 
 		Polyline.Add(EndLine.PointAt(10000.0f));
-		Polyline.Add(Polyline.Last() + 1000 * PerpSign * EndDirOut.Perp());
+		Polyline.Add(Polyline.Last() + 1000 * PerpSign * UE::Geometry::PerpCW(EndDirOut));
 		FVector2f StartPos = Polyline[0];
 		Polyline.Add(StartPos);		// close polyline (cannot use Polyline[0] in case Add resizes!)
 	}
diff --git a/Engine/Plugins/Experimental/Text3D/Source/Text3D/Private/Arrangement.cpp b/Engine/Plugins/Experimental/Text3D/Source/Text3D/Private/Arrangement.cpp
--- a/Engine/Plugins/Experimental/Text3D/Source/Text3D/Private/Arrangement.cpp
+++ b/Engine/Plugins/Experimental/Text3D/Source/Text3D/Private/Arrangement.cpp
@@ -262,7 +262,7 @@ bool FArrangement::find_intersecting_edges(FVector2d A, FVector2d B, TArray<FInt
 {
     int num_hits = 0;
     FVector2d x = FVector2d::Zero(), y = FVector2d::Zero();
-    FVector2d EPerp = (B - A).Perp();
+    FVector2d EPerp = UE::Geometry::PerpCW(B - A);
     UE::Geometry::Normalize(EPerp);
     for (int EID : Graph.EdgeIndices())
     {
