GeometryProcessing: handle invalid element-triangles in TDynamicMeshOverlay::SplitBowties()
#rnx

diff --git a/Engine/Plugins/Experimental/GeometryProcessing/Source/DynamicMesh/Private/DynamicMeshOverlay.cpp b/Engine/Plugins/Experimental/GeometryProcessing/Source/DynamicMesh/Private/DynamicMeshOverlay.cpp
--- a/Engine/Plugins/Experimental/GeometryProcessing/Source/DynamicMesh/Private/DynamicMeshOverlay.cpp
+++ b/Engine/Plugins/Experimental/GeometryProcessing/Source/DynamicMesh/Private/DynamicMeshOverlay.cpp
@@ -323,6 +323,10 @@ void TDynamicMeshOverlay<RealType, ElementSize>::SplitBowties()
 			for (int SubID = 0; SubID < SubGroupElementIDs.Num(); SubID++)
 			{
 				int ElementID = SubGroupElementIDs[SubID];
+				if (ElementID < 0)
+				{
+					continue;		// skip if this is an invalid ElementID (eg from an invalid triangle)
+				}
 				// split needed the *second* time we see a sub-group using a given ElementID
 				if (ElementIDSeen.Contains(ElementID))
 				{
