AutoUV and UVPacking robustness fix cleanup based on feedback

#rnx

diff --git a/Engine/Plugins/Experimental/GeometryProcessing/Source/GeometricObjects/Private/Parameterization/UVPacking.cpp b/Engine/Plugins/Experimental/GeometryProcessing/Source/GeometricObjects/Private/Parameterization/UVPacking.cpp
--- a/Engine/Plugins/Experimental/GeometryProcessing/Source/GeometricObjects/Private/Parameterization/UVPacking.cpp
+++ b/Engine/Plugins/Experimental/GeometryProcessing/Source/GeometricObjects/Private/Parameterization/UVPacking.cpp
@@ -827,49 +827,10 @@ bool FUVPacker::StandardPack(IUVMeshView* Mesh, int NumIslands, TFunctionRef<voi
 		Chart.UVArea = 0.0f;
 		double UVLengthSum = 0, WorldLengthSum = 0;
 
-		for (int32 tid : Chart.Triangles)
-		{
-			FIndex3i Triangle3D = Packer.Mesh->GetTriangle(tid);
-			FIndex3i TriangleUV = Packer.Mesh->GetUVTriangle(tid);
-			
-			FVector3d Positions[3];
-			FVector2d UVs[3];
-
-			for (int k = 0; k < 3; k++)
-			{
-				Positions[k] = Packer.Mesh->GetVertex(Triangle3D[k]);
-				UVs[k] = (FVector2d)Packer.Mesh->GetUV(TriangleUV[k]);
-
-				Chart.MinUV.X = FMathd::Min(Chart.MinUV.X, UVs[k].X);
-				Chart.MinUV.Y = FMathd::Min(Chart.MinUV.Y, UVs[k].Y);
-				Chart.MaxUV.X = FMathd::Max(Chart.MaxUV.X, UVs[k].X);
-				Chart.MaxUV.Y = FMathd::Max(Chart.MaxUV.Y, UVs[k].Y);
-			}
-
-			FVector3d Edge1 = Positions[1] - Positions[0];
-			FVector3d Edge2 = Positions[2] - Positions[0];
-			FVector3d Edge3 = Positions[2] - Positions[1];
-			double WorldLength = Edge1.Length() + Edge2.Length() + Edge3.Length();
-
-			FVector2d EdgeUV1 = UVs[1] - UVs[0];
-			FVector2d EdgeUV2 = UVs[2] - UVs[0];
-			FVector2d EdgeUV3 = UVs[2] - UVs[1];
-			double UVLength = EdgeUV1.Length() + EdgeUV2.Length() + EdgeUV3.Length();
-			double UVArea = 0.5f * FMathd::Abs(EdgeUV1.X * EdgeUV2.Y - EdgeUV1.Y * EdgeUV2.X);
-
-			UVLengthSum += UVLength;
-			WorldLengthSum += WorldLength;
-			Chart.UVArea += UVArea;
-		}
-
-		if (!bScaleIslandsByWorldSpaceTexelRatio || UVLengthSum < FMathd::ZeroTolerance)
-		{
-			Chart.ScaleToWorld = 1;
-		}
-		else
-		{
-			Chart.ScaleToWorld = WorldLengthSum / UVLengthSum;
-		}
+		FAxisAlignedBox2d IslandBounds;
+		GetIslandStats(Mesh, Chart.Triangles, IslandBounds, Chart.ScaleToWorld, Chart.UVArea);
+		Chart.MinUV = IslandBounds.Min;
+		Chart.MaxUV = IslandBounds.Max;
 	}
 
 
@@ -922,27 +883,21 @@ bool FUVPacker::StackPack(IUVMeshView* Mesh, int NumIslands, TFunctionRef<void(i
 	// figure out maximum width and height of existing charts
 	TArray<FAxisAlignedBox2d> AllIslandBounds;
 	AllIslandBounds.SetNum(NumCharts);
+	TArray<double> AllIslandScaleFactors;
+	AllIslandScaleFactors.SetNum(NumCharts);
 	double MaxWidth = 0, MaxHeight = 0;
 	for (int32 ci = 0; ci < NumCharts; ++ci)
 	{
-		FAxisAlignedBox2d IslandBounds = FAxisAlignedBox2d::Empty();
 		TArray<int32> Island;
 		CopyIsland(ci, Island);
-		for (int32 tid : Island)
-		{
-			FVector2f UVTri[3];
-			FIndex3i UVTriInds = Mesh->GetUVTriangle(tid);
-			UVTri[0] = Mesh->GetUV(UVTriInds.A);
-			UVTri[1] = Mesh->GetUV(UVTriInds.B);
-			UVTri[2] = Mesh->GetUV(UVTriInds.C);
-			IslandBounds.Contain((FVector2d)UVTri[0]);
-			IslandBounds.Contain((FVector2d)UVTri[1]);
-			IslandBounds.Contain((FVector2d)UVTri[2]);
-			AllIslandBounds[ci] = IslandBounds;
-		}
-
-		MaxWidth = FMathd::Max(IslandBounds.Width(), MaxWidth);
-		MaxHeight = FMathd::Max(IslandBounds.Height(), MaxHeight);
+		FAxisAlignedBox2d IslandBounds;
+		double IslandScaleFactor, UVArea;
+		GetIslandStats(Mesh, Island, IslandBounds, IslandScaleFactor, UVArea);
+		AllIslandBounds[ci] = IslandBounds;
+		AllIslandScaleFactors[ci] = IslandScaleFactor;
+
+		MaxWidth = FMathd::Max(IslandBounds.Width() * IslandScaleFactor, MaxWidth);
+		MaxHeight = FMathd::Max(IslandBounds.Height() * IslandScaleFactor, MaxHeight);
 	}
 
 	// figure out uniform scale that will make them all fit
@@ -969,13 +924,67 @@ bool FUVPacker::StackPack(IUVMeshView* Mesh, int NumIslands, TFunctionRef<void(i
 			IslandElements.Add(UVTri[2]);
 		}
 
+		double ScaleFactor = UseUniformScale * AllIslandScaleFactors[ci];
 		for (int32 elemid : IslandElements)
 		{
 			FVector2d CurUV = (FVector2d)Mesh->GetUV(elemid);
-			FVector2d NewUV = (CurUV - IslandBounds.Min) * UseUniformScale;
+			FVector2d NewUV = (CurUV - IslandBounds.Min) * ScaleFactor;
 			Mesh->SetUV(elemid, (FVector2f)NewUV);
 		}
 	}
 
 	return true;
 }
+
+
+void FUVPacker::GetIslandStats(IUVMeshView* Mesh, const TArray<int32>& Island, FAxisAlignedBox2d& IslandBoundsOut, double& IslandScaleFactorOut, double& UVAreaOut)
+{
+	IslandBoundsOut = FAxisAlignedBox2d::Empty();
+	UVAreaOut = 0.0f;
+	double UVLengthSum = 0, WorldLengthSum = 0;
+
+	for (int32 tid : Island)
+	{
+		FIndex3i Triangle3D = Mesh->GetTriangle(tid);
+		FIndex3i TriangleUV = Mesh->GetUVTriangle(tid);
+
+		FVector3d Positions[3];
+		FVector2d UVs[3];
+
+		for (int k = 0; k < 3; k++)
+		{
+			Positions[k] = Mesh->GetVertex(Triangle3D[k]);
+			UVs[k] = (FVector2d)Mesh->GetUV(TriangleUV[k]);
+
+			IslandBoundsOut.Contain(UVs[k]);
+		}
+
+		FVector3d Edge1 = Positions[1] - Positions[0];
+		FVector3d Edge2 = Positions[2] - Positions[0];
+		FVector3d Edge3 = Positions[2] - Positions[1];
+		double WorldLength = Edge1.Length() + Edge2.Length() + Edge3.Length();
+
+		FVector2d EdgeUV1 = UVs[1] - UVs[0];
+		FVector2d EdgeUV2 = UVs[2] - UVs[0];
+		FVector2d EdgeUV3 = UVs[2] - UVs[1];
+		double UVLength = EdgeUV1.Length() + EdgeUV2.Length() + EdgeUV3.Length();
+		double UVArea = 0.5f * FMathd::Abs(EdgeUV1.X * EdgeUV2.Y - EdgeUV1.Y * EdgeUV2.X);
+
+		UVLengthSum += UVLength;
+		WorldLengthSum += WorldLength;
+		UVAreaOut += UVArea;
+	}
+
+	if (!bScaleIslandsByWorldSpaceTexelRatio || UVLengthSum < FMathd::ZeroTolerance)
+	{
+		IslandScaleFactorOut = 1;
+	}
+	else
+	{
+		// Use ratio of edge lengths instead of areas for robustness
+		//  (e.g. to avoid over-scaling a ~zero area UV island, from bad projections or disconnected sliver tris)
+		IslandScaleFactorOut = WorldLengthSum / UVLengthSum;
+	}
+}
+
+
diff --git a/Engine/Plugins/Experimental/GeometryProcessing/Source/GeometricObjects/Public/Parameterization/UVPacking.h b/Engine/Plugins/Experimental/GeometryProcessing/Source/GeometricObjects/Public/Parameterization/UVPacking.h
--- a/Engine/Plugins/Experimental/GeometryProcessing/Source/GeometricObjects/Public/Parameterization/UVPacking.h
+++ b/Engine/Plugins/Experimental/GeometryProcessing/Source/GeometricObjects/Public/Parameterization/UVPacking.h
@@ -4,6 +4,7 @@
 
 #include "VectorTypes.h"
 #include "IndexTypes.h"
+#include "BoxTypes.h"
 #include "DisjointSet.h"
 
 namespace UE
@@ -41,7 +42,7 @@ public:
 	/** If true, islands can be flipped in addition to rotate/translate/scale */
 	bool bAllowFlips = false;
 
-	/** Rescale islands to match texel-to-world-space ratio across islands */
+	/** Attempt to rescale islands to match texel-to-world-space ratio across islands, based on ratio of World- and UV-space edge lengths */
 	bool bScaleIslandsByWorldSpaceTexelRatio = false;
 
 	/**
@@ -76,6 +77,13 @@ public:
 			});
 	}
 
+protected:
+
+	/**
+	 * Compute common stats used by the packing algorithms to transform UV islands
+	 */
+	void GetIslandStats(IUVMeshView* Mesh, const TArray<int32>& Island, FAxisAlignedBox2d& IslandBoundsOut, double& IslandScaleFactorOut, double& UVAreaOut);
+
 };
 
 
diff --git a/Engine/Plugins/Experimental/PlanarCutPlugin/Source/PlanarCut/Private/FractureAutoUV.cpp b/Engine/Plugins/Experimental/PlanarCutPlugin/Source/PlanarCut/Private/FractureAutoUV.cpp
--- a/Engine/Plugins/Experimental/PlanarCutPlugin/Source/PlanarCut/Private/FractureAutoUV.cpp
+++ b/Engine/Plugins/Experimental/PlanarCutPlugin/Source/PlanarCut/Private/FractureAutoUV.cpp
@@ -242,7 +242,8 @@ bool UVLayout(
 	int32 UVRes,
 	float GutterSize,
 	bool bOnlyOddMaterials,
-	TArrayView<int32> WhichMaterials
+	TArrayView<int32> WhichMaterials,
+	bool bRecreateUVsForDegenerateIslands
 )
 {
 	TArray<bool> ActiveTriangles;
@@ -252,8 +253,6 @@ bool UVLayout(
 	TArray<TArray<int32>> UVIslands;
 	UE::UVPacking::CreateUVIslandsFromMeshTopology<FGeomMesh>(UVMesh, UVIslands);
 
-	// detect and fix islands that don't have proper UVs (UVs all zero or otherwise collapsed to a point)
-	bool bRecreateUVsForDegenerateIslands = true;
 	if (bRecreateUVsForDegenerateIslands)
 	{
 		TArray<int> IslandVert; IslandVert.SetNumUninitialized(UVMesh.MaxVertexID());
@@ -276,7 +275,8 @@ bool UVLayout(
 				FVector3d Normal = VectorUtil::NormalArea(VA, VB, VC, Area);
 				AvgNormal += Area * Normal;
 			}
-			Normalize(AvgNormal);
+			double AvgNormalLen = Normalize(AvgNormal);
+			bool bHasValidNormal = AvgNormalLen > 0;
 			bool bDoProjection = FMathd::Abs(UVArea) < FMathd::ZeroTolerance;
 			if (bDoProjection)
 			{
@@ -311,8 +311,8 @@ bool UVLayout(
 				{
 					Tris.Add(TID);
 				}
-				bool bExpMapFailed = !UVEditor.SetTriangleUVsFromExpMap(Tris);
-				if (!bExpMapFailed)
+				bool bExpMapSuccess = UVEditor.SetTriangleUVsFromExpMap(Tris);
+				if (bExpMapSuccess)
 				{
 					// transfer UVs from the overlay
 					FDynamicMeshUVOverlay* UVOverlay = IslandMesh.Attributes()->PrimaryUV();
@@ -323,7 +323,7 @@ bool UVLayout(
 						if (!ensure(UVTri != FIndex3i::Invalid()))
 						{
 							// we shouldn't have unset tris like this if the ExpMap returned success
-							bExpMapFailed = true;
+							bExpMapSuccess = false;
 							break;
 						}
 						for (int SubIdx = 0; SubIdx < 3; SubIdx++)
@@ -334,7 +334,7 @@ bool UVLayout(
 						}
 					}
 				}
-				if (bExpMapFailed)
+				if (!bExpMapSuccess && bHasValidNormal)
 				{
 					// expmap failed; fall back to projecting UVs
 					FFrame3d Frame(AnyPt, AvgNormal);
diff --git a/Engine/Plugins/Experimental/PlanarCutPlugin/Source/PlanarCut/Public/FractureAutoUV.h b/Engine/Plugins/Experimental/PlanarCutPlugin/Source/PlanarCut/Public/FractureAutoUV.h
--- a/Engine/Plugins/Experimental/PlanarCutPlugin/Source/PlanarCut/Public/FractureAutoUV.h
+++ b/Engine/Plugins/Experimental/PlanarCutPlugin/Source/PlanarCut/Public/FractureAutoUV.h
@@ -16,13 +16,15 @@ namespace UE { namespace PlanarCut {
  * @param GutterSize		Space to leave between UV islands, in pixels at the target resolution
  * @param bOnlyOddMaterials	If true, restrict UV island layout to odd-numbered material IDs
  * @param WhichMaterials	If non-empty, restrict UV island layout to only the listed material IDs
+ * @param bRecreateUVsForDegenerateIslands If true, detect and fix islands that don't have proper UVs (i.e. UVs all zero or otherwise collapsed to a point)
  */
 bool PLANARCUT_API UVLayout(
 	FGeometryCollection& Collection,
 	int32 UVRes = 1024,
 	float GutterSize = 1,
 	bool bOnlyOddMaterials = true,
-	TArrayView<int32> WhichMaterials = TArrayView<int32>()
+	TArrayView<int32> WhichMaterials = TArrayView<int32>(),
+	bool bRecreateUVsForDegenerateIslands = true
 );
 
 
