change VectorUtils triangle area formula to a cross product-based version to avoid -NaN values & be consistent with NormalArea

#rnx

diff --git a/Engine/Plugins/Experimental/GeometryProcessing/Source/GeometricObjects/Public/VectorUtil.h b/Engine/Plugins/Experimental/GeometryProcessing/Source/GeometricObjects/Public/VectorUtil.h
--- a/Engine/Plugins/Experimental/GeometryProcessing/Source/GeometricObjects/Public/VectorUtil.h
+++ b/Engine/Plugins/Experimental/GeometryProcessing/Source/GeometricObjects/Public/VectorUtil.h
@@ -91,10 +91,10 @@ namespace VectorUtil
 	template <typename RealType>
 	inline RealType Area(const FVector3<RealType>& V0, const FVector3<RealType>& V1, const FVector3<RealType>& V2)
 	{
-		FVector3<RealType> edge1(V1 - V0);
-		FVector3<RealType> edge2(V2 - V0);
-		RealType Dot = edge1.Dot(edge2);
-		return (RealType)0.5 * TMathUtil<RealType>::Sqrt(edge1.SquaredLength() * edge2.SquaredLength() - Dot * Dot);
+		FVector3<RealType> Edge1(V1 - V0);
+		FVector3<RealType> Edge2(V2 - V0);
+		FVector3<RealType> Cross = Edge2.Cross(Edge1);
+		return (RealType)0.5 * Cross.Length();
 	}
 
 	/**
@@ -103,10 +103,10 @@ namespace VectorUtil
 	template <typename RealType>
 	inline RealType Area(const FVector2<RealType>& V0, const FVector2<RealType>& V1, const FVector2<RealType>& V2)
 	{
-		FVector2<RealType> edge1(V1 - V0);
-		FVector2<RealType> edge2(V2 - V0);
-		RealType Dot = edge1.Dot(edge2);
-		return (RealType)0.5 * TMathUtil<RealType>::Sqrt(edge1.SquaredLength() * edge2.SquaredLength() - Dot * Dot);
+		FVector2<RealType> Edge1(V1 - V0);
+		FVector2<RealType> Edge2(V2 - V0);
+		RealType DotPerp = Edge1.DotPerp(Edge2);
+		return (RealType)0.5 * TMathUtil<RealType>::Abs(DotPerp);
 	}
 
 	/**
