cleanup EmbedSurfacePath code: remove unused lambda function

#rnx

diff --git a/Engine/Plugins/Experimental/GeometryProcessing/Source/DynamicMesh/Private/Operations/EmbedSurfacePath.cpp b/Engine/Plugins/Experimental/GeometryProcessing/Source/DynamicMesh/Private/Operations/EmbedSurfacePath.cpp
--- a/Engine/Plugins/Experimental/GeometryProcessing/Source/DynamicMesh/Private/Operations/EmbedSurfacePath.cpp
+++ b/Engine/Plugins/Experimental/GeometryProcessing/Source/DynamicMesh/Private/Operations/EmbedSurfacePath.cpp
@@ -143,13 +143,6 @@ bool WalkMeshPlanar(const FDynamicMesh3* Mesh, int StartTri, FVector3d StartPt,
 		Tri.V[2] = VertexToPosnFn(Mesh, TriVertIDs.C);
 	};
 
-	auto PtTriPlaneSignedDist = [](FVector3d Pt, const FTriangle3d& Tri)
-	{
-		FVector3d Normal = Tri.Normal();
-		Normalize(Normal,FMathd::Epsilon);
-		return (Pt - Tri.V[0]).Dot(Normal);
-	};
-
 	auto PtInsideTri = [](const FVector3d& BaryCoord, double BaryThreshold = FMathd::ZeroTolerance)
 	{
 		return BaryCoord[0] >= -BaryThreshold && BaryCoord[1] >= -BaryThreshold && BaryCoord[2] >= -BaryThreshold;
