CIS Fix: dotperp is not a member function on the geometry vector anymore

#rnx

diff --git a/Engine/Plugins/Experimental/GeometryProcessing/Source/GeometricObjects/Public/VectorUtil.h b/Engine/Plugins/Experimental/GeometryProcessing/Source/GeometricObjects/Public/VectorUtil.h
--- a/Engine/Plugins/Experimental/GeometryProcessing/Source/GeometricObjects/Public/VectorUtil.h
+++ b/Engine/Plugins/Experimental/GeometryProcessing/Source/GeometricObjects/Public/VectorUtil.h
@@ -105,8 +105,8 @@ namespace VectorUtil
 	{
 		FVector2<RealType> Edge1(V1 - V0);
 		FVector2<RealType> Edge2(V2 - V0);
-		RealType DotPerp = Edge1.DotPerp(Edge2);
-		return (RealType)0.5 * TMathUtil<RealType>::Abs(DotPerp);
+		RealType CrossZ = DotPerp(Edge1, Edge2);
+		return (RealType)0.5 * TMathUtil<RealType>::Abs(CrossZ);
 	}
 
 	/**
