GeometryProcessing: define FVector2i/3i, FInterval1i, and FAxisAlignedBox2i/3i as explicit classes rather than template instantiations of real-vector types. Update usage. This is necessary to avoid conflict with engine LWC vector template types.
#rnx

diff --git a/Engine/Plugins/Experimental/GeometryProcessing/Source/DynamicMesh/Public/Operations/FFDLattice.h b/Engine/Plugins/Experimental/GeometryProcessing/Source/DynamicMesh/Public/Operations/FFDLattice.h
--- a/Engine/Plugins/Experimental/GeometryProcessing/Source/DynamicMesh/Public/Operations/FFDLattice.h
+++ b/Engine/Plugins/Experimental/GeometryProcessing/Source/DynamicMesh/Public/Operations/FFDLattice.h
@@ -3,6 +3,7 @@
 #pragma once
 
 #include "VectorTypes.h"
+#include "IntVectorTypes.h"
 #include "BoxTypes.h"
 #include "Containers/StaticArray.h"
 #include "DynamicMeshAttributeSet.h"
diff --git a/Engine/Plugins/Experimental/GeometryProcessing/Source/DynamicMesh/Public/Operations/MeshConvexHull.h b/Engine/Plugins/Experimental/GeometryProcessing/Source/DynamicMesh/Public/Operations/MeshConvexHull.h
--- a/Engine/Plugins/Experimental/GeometryProcessing/Source/DynamicMesh/Public/Operations/MeshConvexHull.h
+++ b/Engine/Plugins/Experimental/GeometryProcessing/Source/DynamicMesh/Public/Operations/MeshConvexHull.h
@@ -4,6 +4,7 @@
 
 #include "MathUtil.h"
 #include "VectorTypes.h"
+#include "IntVectorTypes.h"
 #include "GeometryTypes.h"
 #include "DynamicMesh3.h"
 
diff --git a/Engine/Plugins/Experimental/GeometryProcessing/Source/DynamicMesh/Public/Selections/MeshConnectedComponents.h b/Engine/Plugins/Experimental/GeometryProcessing/Source/DynamicMesh/Public/Selections/MeshConnectedComponents.h
--- a/Engine/Plugins/Experimental/GeometryProcessing/Source/DynamicMesh/Public/Selections/MeshConnectedComponents.h
+++ b/Engine/Plugins/Experimental/GeometryProcessing/Source/DynamicMesh/Public/Selections/MeshConnectedComponents.h
@@ -5,6 +5,7 @@
 #pragma once
 
 #include "CoreMinimal.h"
+#include "IntBoxTypes.h"
 #include "DynamicMesh3.h"
 #include "Containers/IndirectArray.h"
 
diff --git a/Engine/Plugins/Experimental/GeometryProcessing/Source/GeometricObjects/Public/BoxTypes.h b/Engine/Plugins/Experimental/GeometryProcessing/Source/GeometricObjects/Public/BoxTypes.h
--- a/Engine/Plugins/Experimental/GeometryProcessing/Source/GeometricObjects/Public/BoxTypes.h
+++ b/Engine/Plugins/Experimental/GeometryProcessing/Source/GeometricObjects/Public/BoxTypes.h
@@ -206,8 +206,6 @@ struct TInterval1
 
 typedef TInterval1<float> FInterval1f;
 typedef TInterval1<double> FInterval1d;
-typedef TInterval1<int> FInterval1i;
-
 
 
 template <typename RealType>
@@ -713,10 +711,8 @@ struct TAxisAlignedBox2
 
 typedef TAxisAlignedBox2<float> FAxisAlignedBox2f;
 typedef TAxisAlignedBox2<double> FAxisAlignedBox2d;
-typedef TAxisAlignedBox2<int> FAxisAlignedBox2i;
 typedef TAxisAlignedBox3<float> FAxisAlignedBox3f;
 typedef TAxisAlignedBox3<double> FAxisAlignedBox3d;
-typedef TAxisAlignedBox3<int> FAxisAlignedBox3i;
 
 } // end namespace UE::Geometry
 } // end namespace UE
\ No newline at end of file
diff --git a/Engine/Plugins/Experimental/GeometryProcessing/Source/GeometricObjects/Public/ExplicitUseGeometryMathTypes.h b/Engine/Plugins/Experimental/GeometryProcessing/Source/GeometricObjects/Public/ExplicitUseGeometryMathTypes.h
--- a/Engine/Plugins/Experimental/GeometryProcessing/Source/GeometricObjects/Public/ExplicitUseGeometryMathTypes.h
+++ b/Engine/Plugins/Experimental/GeometryProcessing/Source/GeometricObjects/Public/ExplicitUseGeometryMathTypes.h
@@ -29,9 +29,11 @@
 
 #include "IndexTypes.h"
 #include "VectorTypes.h"
+#include "IntVectorTypes.h"
 #include "MatrixTypes.h"
 #include "Quaternion.h"
 #include "BoxTypes.h"
+#include "IntBoxTypes.h"
 #include "PlaneTypes.h"
 #include "HalfspaceTypes.h"
 #include "LineTypes.h"
@@ -57,7 +59,6 @@ using UE::Geometry::FVector2i;
 using UE::Geometry::FVector3i;
 //using UE::Geometry::FVector4f;
 //using UE::Geometry::FVector4d;
-using UE::Geometry::FVector4i;
 
 using UE::Geometry::FMatrix3f;
 using UE::Geometry::FMatrix3d;
diff --git a/Engine/Plugins/Experimental/GeometryProcessing/Source/GeometricObjects/Public/Generators/MarchingCubes.h b/Engine/Plugins/Experimental/GeometryProcessing/Source/GeometricObjects/Public/Generators/MarchingCubes.h
--- a/Engine/Plugins/Experimental/GeometryProcessing/Source/GeometricObjects/Public/Generators/MarchingCubes.h
+++ b/Engine/Plugins/Experimental/GeometryProcessing/Source/GeometricObjects/Public/Generators/MarchingCubes.h
@@ -114,7 +114,7 @@ public:
 		}
 
 		SetDimensions();
-		GridBounds = TAxisAlignedBox3<int>(FVector3i::Zero(), CellDimensions - FVector3i(1,1,1)); // grid bounds are inclusive
+		GridBounds = FAxisAlignedBox3i(FVector3i::Zero(), CellDimensions - FVector3i(1,1,1)); // grid bounds are inclusive
 
 		corner_values_grid = FDenseGrid3f(CellDimensions.X+1, CellDimensions.Y+1, CellDimensions.Z+1, FMathf::MaxReal);
 		edge_vertices.Reset();
@@ -138,7 +138,7 @@ public:
 		}
 
 		SetDimensions();
-		GridBounds = TAxisAlignedBox3<int>(FVector3i::Zero(), CellDimensions - FVector3i(1,1,1)); // grid bounds are inclusive
+		GridBounds = FAxisAlignedBox3i(FVector3i::Zero(), CellDimensions - FVector3i(1,1,1)); // grid bounds are inclusive
 
 		if (LastGridBounds != GridBounds)
 		{
diff --git a/Engine/Plugins/Experimental/GeometryProcessing/Source/GeometricObjects/Public/Generators/MeshShapeGenerator.h b/Engine/Plugins/Experimental/GeometryProcessing/Source/GeometricObjects/Public/Generators/MeshShapeGenerator.h
--- a/Engine/Plugins/Experimental/GeometryProcessing/Source/GeometricObjects/Public/Generators/MeshShapeGenerator.h
+++ b/Engine/Plugins/Experimental/GeometryProcessing/Source/GeometricObjects/Public/Generators/MeshShapeGenerator.h
@@ -3,6 +3,7 @@
 #pragma once
 
 #include "VectorTypes.h"
+#include "IntVectorTypes.h"
 #include "IndexTypes.h"
 
 namespace UE
diff --git a/Engine/Plugins/Experimental/GeometryProcessing/Source/GeometricObjects/Public/Image/ImageDimensions.h b/Engine/Plugins/Experimental/GeometryProcessing/Source/GeometricObjects/Public/Image/ImageDimensions.h
--- a/Engine/Plugins/Experimental/GeometryProcessing/Source/GeometricObjects/Public/Image/ImageDimensions.h
+++ b/Engine/Plugins/Experimental/GeometryProcessing/Source/GeometricObjects/Public/Image/ImageDimensions.h
@@ -3,6 +3,7 @@
 #pragma once
 
 #include "VectorTypes.h"
+#include "IntVectorTypes.h"
 
 namespace UE
 {
diff --git a/Engine/Plugins/Experimental/GeometryProcessing/Source/GeometricObjects/Public/IntBoxTypes.h b/Engine/Plugins/Experimental/GeometryProcessing/Source/GeometricObjects/Public/IntBoxTypes.h
new file mode 100644
--- /dev/null
+++ b/Engine/Plugins/Experimental/GeometryProcessing/Source/GeometricObjects/Public/IntBoxTypes.h
@@ -0,0 +1,255 @@
+// Copyright Epic Games, Inc. All Rights Reserved.
+
+#pragma once
+
+#include "IntVectorTypes.h"
+
+namespace UE
+{
+namespace Geometry
+{
+
+struct FInterval1i
+{
+	int32 Min;
+	int32 Max;
+
+	FInterval1i() :
+		FInterval1i(Empty())
+	{
+	}
+
+	FInterval1i(const int32& Min, const int32& Max)
+	{
+		this->Min = Min;
+		this->Max = Max;
+	}
+
+	static FInterval1i Empty()
+	{
+		return FInterval1i(TNumericLimits<int32>::Max(), -TNumericLimits<int32>::Max());
+	}
+
+	int32 Center() const
+	{
+		return (Min + Max) / 2;
+	}
+
+	int32 Extent() const
+	{
+		return (Max - Min) / 2;
+	}
+	int32 Length() const
+	{
+		return Max - Min;
+	}
+
+	inline bool IsEmpty() const
+	{
+		return Max < Min;
+	}
+
+	void Expand(int32 Radius)
+	{
+		Max += Radius;
+		Min -= Radius;
+	}
+
+	void Contain(int32 V)
+	{
+		if (V < Min)
+		{
+			Min = V;
+		}
+		if (V > Max)
+		{
+			Max = V;
+		}
+	}
+};
+
+
+
+
+struct FAxisAlignedBox2i
+{
+	FVector2i Min;
+	FVector2i Max;
+
+	FAxisAlignedBox2i() :
+		FAxisAlignedBox2i(Empty())
+	{
+	}
+
+	FAxisAlignedBox2i(const FVector2i& Min, const FVector2i& Max)
+		: Min(Min), Max(Max)
+	{
+	}
+
+	FAxisAlignedBox2i(const FAxisAlignedBox2i& OtherBox) = default;
+
+	FAxisAlignedBox2i(int32 SquareSize)
+		: Min((int32)0, (int32)0), Max(SquareSize, SquareSize)
+	{
+	}
+	FAxisAlignedBox2i(int32 Width, int32 Height)
+		: Min((int32)0, (int32)0), Max(Width, Height)
+	{
+	}
+
+	static FAxisAlignedBox2i Empty()
+	{
+		return FAxisAlignedBox2i(
+			FVector2i(TNumericLimits<int32>::Max(), TNumericLimits<int32>::Max()),
+			FVector2i(-TNumericLimits<int32>::Max(), -TNumericLimits<int32>::Max()));
+	}
+
+	/**
+	 * Corners are ordered to follow the perimeter of the bounding rectangle, starting from the (Min.X, Min.Y) corner and ending at (Min.X, Max.Y)
+	 * @param Index which corner to return, must be in range [0,3]
+	 * @return Corner of the bounding rectangle
+	 */
+	FVector2i GetCorner(int Index) const
+	{
+		check(Index >= 0 && Index <= 3);
+		int32 X = ((Index % 3) == 0) ? (Min.X) : (Max.X);
+		int32 Y = ((Index & 2) == 0) ? (Min.Y) : (Max.Y);
+		return FVector2i(X, Y);
+	}
+
+	inline bool IsEmpty() const
+	{
+		return Max.X < Min.X || Max.Y < Min.Y;
+	}
+
+	void Expand(int32 Radius)
+	{
+		Max.X += Radius;
+		Max.Y += Radius;
+		Min.X -= Radius;
+		Min.Y -= Radius;
+	}
+};
+
+
+
+
+
+
+struct FAxisAlignedBox3i
+{
+	FVector3i Min;
+	FVector3i Max;
+
+	FAxisAlignedBox3i() : 
+		FAxisAlignedBox3i(FAxisAlignedBox3i::Empty())
+	{
+	}
+
+	FAxisAlignedBox3i(const FVector3i& Min, const FVector3i& Max)
+	{
+		this->Min = Min;
+		this->Max = Max;
+	}
+
+	FAxisAlignedBox3i(const FAxisAlignedBox3i& OtherBox) = default;
+
+	bool operator==(const FAxisAlignedBox3i& Other) const
+	{
+		return Max == Other.Max && Min == Other.Min;
+	}
+	bool operator!=(const FAxisAlignedBox3i& Other) const
+	{
+		return Max != Other.Max || Min != Other.Min;
+	}
+
+	int32 Width() const
+	{
+		return TMathUtil<int32>::Max(Max.X - Min.X, 0);
+	}
+
+	int32 Height() const
+	{
+		return TMathUtil<int32>::Max(Max.Y - Min.Y, 0);
+	}
+
+	int32 Depth() const
+	{
+		return TMathUtil<int32>::Max(Max.Z - Min.Z, 0);
+	}
+
+	int32 Volume() const
+	{
+		return Width() * Height() * Depth();
+	}
+
+	FVector3i Diagonal() const
+	{
+		return FVector3i(Max.X - Min.X, Max.Y - Min.Y, Max.Z - Min.Z);
+	}
+
+	bool Contains(const FVector3i& V) const
+	{
+		return (Min.X <= V.X) && (Min.Y <= V.Y) && (Min.Z <= V.Z) && (Max.X >= V.X) && (Max.Y >= V.Y) && (Max.Z >= V.Z);
+	}
+
+	void Contain(const FVector3i& V)
+	{
+		if (V.X < Min.X)
+		{
+			Min.X = V.X;
+		}
+		if (V.X > Max.X)
+		{
+			Max.X = V.X;
+		}
+		if (V.Y < Min.Y)
+		{
+			Min.Y = V.Y;
+		}
+		if (V.Y > Max.Y)
+		{
+			Max.Y = V.Y;
+		}
+		if (V.Z < Min.Z)
+		{
+			Min.Z = V.Z;
+		}
+		if (V.Z > Max.Z)
+		{
+			Max.Z = V.Z;
+		}
+	}
+
+	/**
+	* @param Index corner index in range 0-7
+	* @return Corner point on the box identified by the given index. See diagram in OrientedBoxTypes.h for index/corner mapping.
+	*/
+	FVector3i GetCorner(int Index) const
+	{
+		check(Index >= 0 && Index <= 7);
+		int32 X = (((Index & 1) != 0) ^ ((Index & 2) != 0)) ? (Max.X) : (Min.X);
+		int32 Y = ((Index / 2) % 2 == 0) ? (Min.Y) : (Max.Y);
+		int32 Z = (Index < 4) ? (Min.Z) : (Max.Z);
+		return FVector3i(X, Y, Z);
+	}
+
+	static FAxisAlignedBox3i Empty()
+	{
+		return FAxisAlignedBox3i(
+			FVector3i(TNumericLimits<int32>::Max(), TNumericLimits<int32>::Max(), TNumericLimits<int32>::Max()),
+			FVector3i(-TNumericLimits<int32>::Max(), -TNumericLimits<int32>::Max(), -TNumericLimits<int32>::Max()));
+	}
+
+	static FAxisAlignedBox3i Infinite()
+	{
+		return FAxisAlignedBox3i(
+			FVector3i(-TNumericLimits<int32>::Max(), -TNumericLimits<int32>::Max(), -TNumericLimits<int32>::Max()),
+			FVector3i(TNumericLimits<int32>::Max(), TNumericLimits<int32>::Max(), TNumericLimits<int32>::Max()) );
+	}
+};
+
+
+
+} // end namespace UE::Geometry
+} // end namespace UE
\ No newline at end of file
diff --git a/Engine/Plugins/Experimental/GeometryProcessing/Source/GeometricObjects/Public/IntVectorTypes.h b/Engine/Plugins/Experimental/GeometryProcessing/Source/GeometricObjects/Public/IntVectorTypes.h
new file mode 100644
--- /dev/null
+++ b/Engine/Plugins/Experimental/GeometryProcessing/Source/GeometricObjects/Public/IntVectorTypes.h
@@ -0,0 +1,464 @@
+// Copyright Epic Games, Inc. All Rights Reserved.
+
+#pragma once
+
+#include "Math/Vector.h"
+#include "MathUtil.h"
+#include "VectorTypes.h"
+#include <sstream>
+
+
+namespace UE {
+namespace Geometry {
+
+
+/**
+* 2D 32-bit integer Vector
+*/
+struct FVector2i
+{
+	int32 X{}, Y{};
+
+	constexpr FVector2i()
+		: X(0), Y(0)
+	{
+	}
+
+	constexpr FVector2i(int32 ValX, int32 ValY)
+		: X(ValX), Y(ValY)
+	{
+	}
+
+	constexpr FVector2i(const int32* Data)
+		: X(Data[0]), Y(Data[1])
+	{
+	}
+
+	constexpr FVector2i(const FVector2i& Vec) = default;
+	constexpr FVector2i& operator=(const FVector2i& V2) = default;
+
+	explicit constexpr operator const int32*() const
+	{
+		return &X;
+	};
+
+	explicit constexpr operator int32*()
+	{
+		return &X;
+	}
+
+	explicit operator FVector2D() const
+	{
+		return FVector2D((float)X, (float)Y);
+	}
+	explicit operator FVector2f() const
+	{
+		return FVector2f((float)X, (float)Y);
+	}
+	explicit operator FVector2d() const
+	{
+		return FVector2d((double)X, (double)Y);
+	}
+
+	explicit FVector2i(const FVector& Vec)
+		: X((int32)Vec.X), Y((int32)Vec.Y)
+	{
+	}
+
+	explicit FVector2i(const FVector2f& Vec)
+		: X((int32)Vec.X), Y((int32)Vec.Y)
+	{
+	}
+
+	explicit FVector2i(const FVector2d& Vec)
+		: X((int32)Vec.X), Y((int32)Vec.Y)
+	{
+	}
+
+	constexpr static FVector2i Zero()
+	{
+		return FVector2i(0, 0);
+	}
+	constexpr static FVector2i One()
+	{
+		return FVector2i(1, 1);
+	}
+	constexpr static FVector2i UnitX()
+	{
+		return FVector2i(1, 0);
+	}
+	constexpr static FVector2i UnitY()
+	{
+		return FVector2i(0, 1);
+	}
+
+	constexpr int32& operator[](int Idx)
+	{
+		return (&X)[Idx];
+	}
+	constexpr const int32& operator[](int Idx) const
+	{
+		return (&X)[Idx];
+	}
+
+	constexpr int32 SquaredLength() const
+	{
+		return X * X + Y * Y;
+	}
+
+	constexpr int32 DistanceSquared(const FVector2i& V2) const
+	{
+		int32 dx = V2.X - X;
+		int32 dy = V2.Y - Y;
+		return dx * dx + dy * dy;
+	}
+
+	constexpr int32 Dot(const FVector2i& V2) const
+	{
+		return X * V2.X + Y * V2.Y;
+	}
+
+	constexpr FVector2i operator-() const
+	{
+		return FVector2i(-X, -Y);
+	}
+
+	constexpr FVector2i operator+(const FVector2i& V2) const
+	{
+		return FVector2i(X + V2.X, Y + V2.Y);
+	}
+
+	constexpr FVector2i operator-(const FVector2i& V2) const
+	{
+		return FVector2i(X - V2.X, Y - V2.Y);
+	}
+
+	constexpr FVector2i operator+(const int32& Scalar) const
+	{
+		return FVector2i(X + Scalar, Y + Scalar);
+	}
+
+	constexpr FVector2i operator-(const int32& Scalar) const
+	{
+		return FVector2i(X - Scalar, Y - Scalar);
+	}
+
+	constexpr FVector2i operator*(const int32& Scalar) const
+	{
+		return FVector2i(X * Scalar, Y * Scalar);
+	}
+
+	constexpr FVector2i operator*(const FVector2i& V2) const // component-wise
+	{
+		return FVector2i(X * V2.X, Y * V2.Y);
+	}
+
+	constexpr FVector2i operator/(const int32& Scalar) const
+	{
+		return FVector2i(X / Scalar, Y / Scalar);
+	}
+
+	constexpr FVector2i operator/(const FVector2i& V2) const // component-wise
+	{
+		return FVector2i(X / V2.X, Y / V2.Y);
+	}
+
+	constexpr FVector2i& operator+=(const FVector2i& V2)
+	{
+		X += V2.X;
+		Y += V2.Y;
+		return *this;
+	}
+
+	constexpr FVector2i& operator-=(const FVector2i& V2)
+	{
+		X -= V2.X;
+		Y -= V2.Y;
+		return *this;
+	}
+
+	constexpr FVector2i& operator*=(const int32& Scalar)
+	{
+		X *= Scalar;
+		Y *= Scalar;
+		return *this;
+	}
+
+	constexpr FVector2i& operator/=(const int32& Scalar)
+	{
+		X /= Scalar;
+		Y /= Scalar;
+		return *this;
+	}
+
+	constexpr bool operator==(const FVector2i& Other) const
+	{
+		return X == Other.X && Y == Other.Y;
+	}
+
+	constexpr bool operator!=(const FVector2i& Other) const
+	{
+		return X != Other.X || Y != Other.Y;
+	}
+};
+
+inline FVector2i operator*(int32 Scalar, const FVector2i& V)
+{
+	return FVector2i(Scalar * V.X, Scalar * V.Y);
+}
+
+inline std::ostream& operator<<(std::ostream& os, const FVector2i& Vec)
+{
+	os << Vec.X << " " << Vec.Y;
+	return os;
+}
+
+FORCEINLINE uint32 GetTypeHash(const FVector2i& Vector)
+{
+	return FCrc::MemCrc_DEPRECATED(&Vector, sizeof(FVector2i));
+}
+
+
+
+/**
+* 3D 32-bit integer Vector
+*/
+struct FVector3i
+{
+	int32 X{}, Y{}, Z{};
+
+	constexpr FVector3i()
+		: X(0), Y(0), Z(0)
+	{
+	}
+
+	constexpr FVector3i(int32 ValX, int32 ValY, int32 ValZ)
+		: X(ValX), Y(ValY), Z(ValZ)
+	{
+	}
+
+	constexpr FVector3i(const int32* Data)
+		: X(Data[0]), Y(Data[1]), Z(Data[2])
+	{
+	}
+
+	constexpr FVector3i(const FVector3i& Vec) = default;
+	FVector3i& operator=(const FVector3i& V2) = default;
+
+	explicit constexpr operator const int32*() const
+	{
+		return &X;
+	};
+	explicit constexpr operator int32*()
+	{
+		return &X;
+	}
+
+	explicit operator FVector() const
+	{
+		return FVector((float)X, (float)Y, (float)Z);
+	}
+	explicit operator FVector3f() const
+	{
+		return FVector3f((float)X, (float)Y, (float)Z);
+	}
+	explicit operator FVector3d() const
+	{
+		return FVector3d((double)X, (double)Y, (double)Z);
+	}
+
+	explicit FVector3i(const FVector& Vec)
+		: X((int32)Vec.X), Y((int32)Vec.Y), Z((int32)Vec.Z)
+	{
+	}
+
+	explicit FVector3i(const FVector3f& Vec)
+		: X((int32)Vec.X), Y((int32)Vec.Y), Z((int32)Vec.Z)
+	{
+	}
+
+	explicit FVector3i(const FVector3d& Vec)
+		: X((int32)Vec.X), Y((int32)Vec.Y), Z((int32)Vec.Z)
+	{
+	}
+
+	static FVector3i Zero()
+	{
+		return FVector3i(0, 0, 0);
+	}
+	static FVector3i One()
+	{
+		return FVector3i(1, 1, 1);
+	}
+	static FVector3i UnitX()
+	{
+		return FVector3i(1, 0, 0);
+	}
+	static FVector3i UnitY()
+	{
+		return FVector3i(0, 1, 0);
+	}
+	static FVector3i UnitZ()
+	{
+		return FVector3i(0, 0, 1);
+	}
+	static FVector3i MaxVector()
+	{
+		return FVector3i(TNumericLimits<int32>::Max(), TNumericLimits<int32>::Max(), TNumericLimits<int32>::Max());
+	}
+	
+	int32& operator[](int Idx)
+	{
+		return (&X)[Idx];
+	}
+	const int32& operator[](int Idx) const
+	{
+		return (&X)[Idx];
+	}
+
+	constexpr int32 SquaredLength() const
+	{
+		return X * X + Y * Y + Z * Z;
+	}
+
+	constexpr int32 DistanceSquared(const FVector3i& V2) const
+	{
+		int32 dx = V2.X - X;
+		int32 dy = V2.Y - Y;
+		int32 dz = V2.Z - Z;
+		return dx * dx + dy * dy + dz * dz;
+	}
+
+	constexpr FVector3i operator-() const
+	{
+		return FVector3i(-X, -Y, -Z);
+	}
+
+	constexpr FVector3i operator+(const FVector3i& V2) const
+	{
+		return FVector3i(X + V2.X, Y + V2.Y, Z + V2.Z);
+	}
+
+	constexpr FVector3i operator-(const FVector3i& V2) const
+	{
+		return FVector3i(X - V2.X, Y - V2.Y, Z - V2.Z);
+	}
+
+	constexpr FVector3i operator+(const int32& Scalar) const
+	{
+		return FVector3i(X + Scalar, Y + Scalar, Z + Scalar);
+	}
+
+	constexpr FVector3i operator-(const int32& Scalar) const
+	{
+		return FVector3i(X - Scalar, Y - Scalar, Z - Scalar);
+	}
+
+	constexpr FVector3i operator*(const int32& Scalar) const
+	{
+		return FVector3i(X * Scalar, Y * Scalar, Z * Scalar);
+	}
+
+	constexpr FVector3i operator*(const FVector3i& V2) const // component-wise
+	{
+		return FVector3i(X * V2.X, Y * V2.Y, Z * V2.Z);
+	}
+
+	constexpr FVector3i operator/(const int32& Scalar) const
+	{
+		return FVector3i(X / Scalar, Y / Scalar, Z / Scalar);
+	}
+
+	constexpr FVector3i operator/(const FVector3i& V2) const // component-wise
+	{
+		return FVector3i(X / V2.X, Y / V2.Y, Z / V2.Z);
+	}
+
+	constexpr FVector3i& operator+=(const FVector3i& V2)
+	{
+		X += V2.X;
+		Y += V2.Y;
+		Z += V2.Z;
+		return *this;
+	}
+
+	constexpr FVector3i& operator-=(const FVector3i& V2)
+	{
+		X -= V2.X;
+		Y -= V2.Y;
+		Z -= V2.Z;
+		return *this;
+	}
+
+	constexpr FVector3i& operator*=(const int32& Scalar)
+	{
+		X *= Scalar;
+		Y *= Scalar;
+		Z *= Scalar;
+		return *this;
+	}
+
+	constexpr FVector3i& operator/=(const int32& Scalar)
+	{
+		X /= Scalar;
+		Y /= Scalar;
+		Z /= Scalar;
+		return *this;
+	}
+
+	constexpr int32 Dot(const FVector3i& V2) const
+	{
+		return X * V2.X + Y * V2.Y + Z * V2.Z;
+	}
+
+	constexpr bool operator==(const FVector3i& Other) const
+	{
+		return X == Other.X && Y == Other.Y && Z == Other.Z;
+	}
+
+	constexpr bool operator!=(const FVector3i& Other) const
+	{
+		return X != Other.X || Y != Other.Y || Z != Other.Z;
+	}
+};
+
+
+inline FVector3i operator*(int32 Scalar, const FVector3i& V)
+{
+	return FVector3i(Scalar * V.X, Scalar * V.Y, Scalar * V.Z);
+}
+
+
+inline FVector3i Min(const FVector3i& V0, const FVector3i& V1)
+{
+	return FVector3i(TMathUtil<int32>::Min(V0.X, V1.X),
+		TMathUtil<int32>::Min(V0.Y, V1.Y),
+		TMathUtil<int32>::Min(V0.Z, V1.Z));
+}
+
+inline FVector3i Max(const FVector3i& V0, const FVector3i& V1)
+{
+	return FVector3i(TMathUtil<int32>::Max(V0.X, V1.X),
+		TMathUtil<int32>::Max(V0.Y, V1.Y),
+		TMathUtil<int32>::Max(V0.Z, V1.Z));
+}
+
+
+
+inline std::ostream& operator<<(std::ostream& os, const FVector3i& Vec)
+{
+	os << Vec.X << " " << Vec.Y << " " << Vec.Z;
+	return os;
+}
+
+FORCEINLINE uint32 GetTypeHash(const FVector3i& Vector)
+{
+	return FCrc::MemCrc_DEPRECATED(&Vector, sizeof(FVector3i));
+}
+
+
+
+} // end namespace UE::Geometry
+} // end namespace UE
+
+
diff --git a/Engine/Plugins/Experimental/GeometryProcessing/Source/GeometricObjects/Public/Intersection/IntrLine2Triangle2.h b/Engine/Plugins/Experimental/GeometryProcessing/Source/GeometricObjects/Public/Intersection/IntrLine2Triangle2.h
--- a/Engine/Plugins/Experimental/GeometryProcessing/Source/GeometricObjects/Public/Intersection/IntrLine2Triangle2.h
+++ b/Engine/Plugins/Experimental/GeometryProcessing/Source/GeometricObjects/Public/Intersection/IntrLine2Triangle2.h
@@ -5,6 +5,7 @@
 #pragma once
 
 #include "VectorTypes.h"
+#include "IntVectorTypes.h"
 #include "LineTypes.h"
 #include "TriangleTypes.h"
 #include "VectorUtil.h"
diff --git a/Engine/Plugins/Experimental/GeometryProcessing/Source/GeometricObjects/Public/Intersection/IntrSegment2Triangle2.h b/Engine/Plugins/Experimental/GeometryProcessing/Source/GeometricObjects/Public/Intersection/IntrSegment2Triangle2.h
--- a/Engine/Plugins/Experimental/GeometryProcessing/Source/GeometricObjects/Public/Intersection/IntrSegment2Triangle2.h
+++ b/Engine/Plugins/Experimental/GeometryProcessing/Source/GeometricObjects/Public/Intersection/IntrSegment2Triangle2.h
@@ -5,6 +5,7 @@
 #pragma once
 
 #include "VectorTypes.h"
+#include "IntVectorTypes.h"
 #include "SegmentTypes.h"
 #include "TriangleTypes.h"
 #include "VectorUtil.h"
diff --git a/Engine/Plugins/Experimental/GeometryProcessing/Source/GeometricObjects/Public/Spatial/DenseGrid2.h b/Engine/Plugins/Experimental/GeometryProcessing/Source/GeometricObjects/Public/Spatial/DenseGrid2.h
--- a/Engine/Plugins/Experimental/GeometryProcessing/Source/GeometricObjects/Public/Spatial/DenseGrid2.h
+++ b/Engine/Plugins/Experimental/GeometryProcessing/Source/GeometricObjects/Public/Spatial/DenseGrid2.h
@@ -6,6 +6,7 @@
 
 #include "CoreMinimal.h"
 #include "BoxTypes.h"
+#include "IntBoxTypes.h"
 
 #include "HAL/PlatformAtomics.h"
 
diff --git a/Engine/Plugins/Experimental/GeometryProcessing/Source/GeometricObjects/Public/Spatial/DenseGrid3.h b/Engine/Plugins/Experimental/GeometryProcessing/Source/GeometricObjects/Public/Spatial/DenseGrid3.h
--- a/Engine/Plugins/Experimental/GeometryProcessing/Source/GeometricObjects/Public/Spatial/DenseGrid3.h
+++ b/Engine/Plugins/Experimental/GeometryProcessing/Source/GeometricObjects/Public/Spatial/DenseGrid3.h
@@ -6,6 +6,7 @@
 
 #include "CoreMinimal.h"
 #include "BoxTypes.h"
+#include "IntBoxTypes.h"
 
 #include "HAL/PlatformAtomics.h"
 
diff --git a/Engine/Plugins/Experimental/GeometryProcessing/Source/GeometricObjects/Public/Spatial/SparseDynamicOctree3.h b/Engine/Plugins/Experimental/GeometryProcessing/Source/GeometricObjects/Public/Spatial/SparseDynamicOctree3.h
--- a/Engine/Plugins/Experimental/GeometryProcessing/Source/GeometricObjects/Public/Spatial/SparseDynamicOctree3.h
+++ b/Engine/Plugins/Experimental/GeometryProcessing/Source/GeometricObjects/Public/Spatial/SparseDynamicOctree3.h
@@ -5,6 +5,7 @@
 
 #include "CoreMinimal.h"
 #include "Containers/StaticArray.h"
+#include "IntVectorTypes.h"
 #include "BoxTypes.h"
 #include "GeometryTypes.h"
 #include "Util/DynamicVector.h"
diff --git a/Engine/Plugins/Experimental/GeometryProcessing/Source/GeometricObjects/Public/Spatial/SparseGrid3.h b/Engine/Plugins/Experimental/GeometryProcessing/Source/GeometricObjects/Public/Spatial/SparseGrid3.h
--- a/Engine/Plugins/Experimental/GeometryProcessing/Source/GeometricObjects/Public/Spatial/SparseGrid3.h
+++ b/Engine/Plugins/Experimental/GeometryProcessing/Source/GeometricObjects/Public/Spatial/SparseGrid3.h
@@ -6,6 +6,7 @@
 
 #include "CoreMinimal.h"
 #include "BoxTypes.h"
+#include "IntBoxTypes.h"
 
 namespace UE
 {
diff --git a/Engine/Plugins/Experimental/GeometryProcessing/Source/GeometricObjects/Public/Util/GridIndexing2.h b/Engine/Plugins/Experimental/GeometryProcessing/Source/GeometricObjects/Public/Util/GridIndexing2.h
--- a/Engine/Plugins/Experimental/GeometryProcessing/Source/GeometricObjects/Public/Util/GridIndexing2.h
+++ b/Engine/Plugins/Experimental/GeometryProcessing/Source/GeometricObjects/Public/Util/GridIndexing2.h
@@ -5,6 +5,7 @@
 #pragma once
 
 #include "VectorTypes.h"
+#include "IntVectorTypes.h"
 
 namespace UE
 {
diff --git a/Engine/Plugins/Experimental/GeometryProcessing/Source/GeometricObjects/Public/Util/GridIndexing3.h b/Engine/Plugins/Experimental/GeometryProcessing/Source/GeometricObjects/Public/Util/GridIndexing3.h
--- a/Engine/Plugins/Experimental/GeometryProcessing/Source/GeometricObjects/Public/Util/GridIndexing3.h
+++ b/Engine/Plugins/Experimental/GeometryProcessing/Source/GeometricObjects/Public/Util/GridIndexing3.h
@@ -5,6 +5,7 @@
 #pragma once
 
 #include "VectorTypes.h"
+#include "IntVectorTypes.h"
 #include "BoxTypes.h"
 
 namespace UE
diff --git a/Engine/Plugins/Experimental/GeometryProcessing/Source/GeometricObjects/Public/Util/IndexUtil.h b/Engine/Plugins/Experimental/GeometryProcessing/Source/GeometricObjects/Public/Util/IndexUtil.h
--- a/Engine/Plugins/Experimental/GeometryProcessing/Source/GeometricObjects/Public/Util/IndexUtil.h
+++ b/Engine/Plugins/Experimental/GeometryProcessing/Source/GeometricObjects/Public/Util/IndexUtil.h
@@ -7,6 +7,7 @@
 #include "Util/DynamicVector.h"
 #include "IndexTypes.h"
 #include "VectorTypes.h"
+#include "IntVectorTypes.h"
 
 
 namespace IndexUtil
diff --git a/Engine/Plugins/Experimental/GeometryProcessing/Source/GeometricObjects/Public/VectorTypes.h b/Engine/Plugins/Experimental/GeometryProcessing/Source/GeometricObjects/Public/VectorTypes.h
--- a/Engine/Plugins/Experimental/GeometryProcessing/Source/GeometricObjects/Public/VectorTypes.h
+++ b/Engine/Plugins/Experimental/GeometryProcessing/Source/GeometricObjects/Public/VectorTypes.h
@@ -16,7 +16,7 @@ namespace Geometry {
 * maintaining compatibility with existing g3Sharp code. Has an API
 * similar to WildMagic, GTEngine, Eigen, etc.
 *
-* Convenience typedefs for FVector2f/FVector2d/FVector2i are defined, and
+* Convenience typedefs for FVector2f/FVector2d are defined, and
 * should be preferentially used over the base template type
 *
 * @todo Possibly can be replaced/merged with Chaos TVector<T,N>
@@ -326,7 +326,7 @@ FVector2<T> Lerp(const FVector2<T>& A, const FVector2<T>& B, T Alpha)
  * maintaining compatibility with existing g3Sharp code. Has an API
  * similar to WildMagic, GTEngine, Eigen, etc. 
  * 
- * Convenience typedefs for FVector3f/FVector3d/FVector3i are defined, and
+ * Convenience typedefs for FVector3f/FVector3d are defined, and
  * should be preferentially used over the base template type
  * 
  * @todo Possibly can be replaced/merged with Chaos TVector<T,N>
@@ -780,8 +780,6 @@ std::ostream& operator<<(std::ostream& os, const FVector3<RealType>& Vec)
 	return os;
 }
 
-typedef FVector3<int> FVector3i;
-
 template <typename T>
 FORCEINLINE uint32 GetTypeHash(const FVector3<T>& Vector)
 {
@@ -822,10 +820,6 @@ std::ostream& operator<<(std::ostream& os, const FVector2<RealType>& Vec)
 	return os;
 }
 
-typedef FVector2<int> FVector2i;
-
-
-
 
 
 
@@ -1111,7 +1105,6 @@ std::ostream& operator<<(std::ostream& os, const TVector4<RealType>& Vec)
 	return os;
 }
 
-typedef TVector4<int> FVector4i;
 
 template <typename T>
 FORCEINLINE uint32 GetTypeHash(const TVector4<T>& Vector)
diff --git a/Engine/Plugins/Experimental/MeshModelingToolset/Source/ModelingComponents/Public/Mechanics/LatticeControlPointsMechanic.h b/Engine/Plugins/Experimental/MeshModelingToolset/Source/ModelingComponents/Public/Mechanics/LatticeControlPointsMechanic.h
--- a/Engine/Plugins/Experimental/MeshModelingToolset/Source/ModelingComponents/Public/Mechanics/LatticeControlPointsMechanic.h
+++ b/Engine/Plugins/Experimental/MeshModelingToolset/Source/ModelingComponents/Public/Mechanics/LatticeControlPointsMechanic.h
@@ -12,6 +12,7 @@
 #include "Spatial/GeometrySet3.h"
 #include "ToolContextInterfaces.h" //FViewCameraState
 #include "VectorTypes.h"
+#include "IntVectorTypes.h"
 #include "VectorUtil.h"
 #include "ToolDataVisualizer.h"
 #include "SceneView.h"
