GeometryProcessing: make local version of FAllocator2D to remove dependency and allow improvements without breaking Lightmap UV computation.
#rnx

diff --git a/Engine/Plugins/Experimental/GeometryProcessing/Source/GeometricObjects/GeometricObjects.Build.cs b/Engine/Plugins/Experimental/GeometryProcessing/Source/GeometricObjects/GeometricObjects.Build.cs
--- a/Engine/Plugins/Experimental/GeometryProcessing/Source/GeometricObjects/GeometricObjects.Build.cs
+++ b/Engine/Plugins/Experimental/GeometryProcessing/Source/GeometricObjects/GeometricObjects.Build.cs
@@ -3,16 +3,14 @@
 using UnrealBuildTool;
 
 public class GeometricObjects : ModuleRules
-{	
+{
 	public GeometricObjects(ReadOnlyTargetRules Target) : base(Target)
 	{
         PCHUsage = ModuleRules.PCHUsageMode.UseExplicitOrSharedPCHs;
 
-        PublicDependencyModuleNames.AddRange(
+		PublicDependencyModuleNames.AddRange(
 			new string[] {
-				"Core",
-
-				"MeshUtilitiesCommon"		// currently required for FAllocator2D used in FDynamicMeshUVPacker
+				"Core"
 			}
 			);
 	}
diff --git a/Engine/Plugins/Experimental/GeometryProcessing/Source/GeometricObjects/Private/Parameterization/UVPacking.cpp b/Engine/Plugins/Experimental/GeometryProcessing/Source/GeometricObjects/Private/Parameterization/UVPacking.cpp
--- a/Engine/Plugins/Experimental/GeometryProcessing/Source/GeometricObjects/Private/Parameterization/UVPacking.cpp
+++ b/Engine/Plugins/Experimental/GeometryProcessing/Source/GeometricObjects/Private/Parameterization/UVPacking.cpp
@@ -2,12 +2,12 @@
 
 
 #include "Parameterization/UVPacking.h"
+#include "Parameterization/UVSpaceAllocator.h"
 
 #include "Async/Future.h"
 #include "Async/Async.h"
 
 #include "Misc/SecureHash.h"
-#include "Allocator2D.h"
 
 #include "BoxTypes.h"
 
@@ -24,6 +24,8 @@ inline uint32 GetTypeHash(const FMD5Hash& Hash)
 
 namespace UE { namespace InternalUVPacking {
 
+	using namespace UE::Geometry;
+
 //
 // local representation of a UV island as a set of triangle indices
 //
@@ -84,7 +86,7 @@ protected:
 	void ScaleCharts(TArray<FUVIsland>& Charts, double UVScale);
 	bool PackCharts(TArray<FUVIsland>& Charts, double UVScale, double& OutEfficiency, TAtomic<bool>& bAbort);
 	void OrientChart(FUVIsland& Chart, int32 Orientation);
-	void RasterizeChart(const FUVIsland& Chart, uint32 RectW, uint32 RectH, FAllocator2D& OutChartRaster);
+	void RasterizeChart(const FUVIsland& Chart, uint32 RectW, uint32 RectH, FUVSpaceAllocator& OutChartRaster);
 };
 
 
@@ -398,9 +400,9 @@ bool FStandardChartPacker::PackCharts(TArray<FUVIsland>& Charts, double UVScale,
 {
 	ScaleCharts(Charts, UVScale);
 
-	FAllocator2D BestChartRaster(FAllocator2D::EMode::UsedSegments, TextureResolution, TextureResolution, ELightmapUVVersion::Latest);
-	FAllocator2D ChartRaster(FAllocator2D::EMode::UsedSegments, TextureResolution, TextureResolution, ELightmapUVVersion::Latest);
-	FAllocator2D LayoutRaster(FAllocator2D::EMode::FreeSegments, TextureResolution, TextureResolution, ELightmapUVVersion::Latest);
+	FUVSpaceAllocator BestChartRaster(FUVSpaceAllocator::EMode::UsedSegments, TextureResolution, TextureResolution);
+	FUVSpaceAllocator ChartRaster(FUVSpaceAllocator::EMode::UsedSegments, TextureResolution, TextureResolution);
+	FUVSpaceAllocator LayoutRaster(FUVSpaceAllocator::EMode::FreeSegments, TextureResolution, TextureResolution);
 
 	uint64 RasterizeCycles = 0;
 	uint64 FindCycles = 0;
@@ -427,19 +429,19 @@ bool FStandardChartPacker::PackCharts(TArray<FUVIsland>& Charts, double UVScale,
 
 			// Try different orientations and pick best
 			int32				BestOrientation = -1;
-			FAllocator2D::FRect	BestRect = { ~0u, ~0u, ~0u, ~0u };
+			FUVSpaceAllocator::FRect	BestRect = { ~0u, ~0u, ~0u, ~0u };
 
 			// This version focus on minimal surface area giving fairness to both horizontal and vertical chart placement
 			// instead of only taking the pixel offset of the lower left corner into account.
-			TFunction<bool(const FAllocator2D::FRect&)> IsBestRect =
-				[&BestRect](const FAllocator2D::FRect& Rect)
+			TFunction<bool(const FUVSpaceAllocator::FRect&)> IsBestRect =
+				[&BestRect](const FUVSpaceAllocator::FRect& Rect)
 			{
 				return ((Rect.X + Rect.W) + (Rect.Y + Rect.H)) <((BestRect.X + BestRect.W) + (BestRect.Y + BestRect.H));
 			};
 
 			// simpler thing?
-			//TFunction<bool(const FAllocator2D::FRect&)> IsBestRect =
-			//	[this, &BestRect](const FAllocator2D::FRect& Rect)
+			//TFunction<bool(const FUVSpaceAllocator::FRect&)> IsBestRect =
+			//	[this, &BestRect](const FUVSpaceAllocator::FRect& Rect)
 			//{
 			//	return Rect.X + Rect.Y * TextureResolution <BestRect.X + BestRect.Y * TextureResolution;
 			//};
@@ -455,7 +457,7 @@ bool FStandardChartPacker::PackCharts(TArray<FUVIsland>& Charts, double UVScale,
 				ChartSize = ChartSize.X * Chart.PackingScaleU + ChartSize.Y * Chart.PackingScaleV;
 
 				// Only need half pixel dilate for rects
-				FAllocator2D::FRect	Rect;
+				FUVSpaceAllocator::FRect	Rect;
 				Rect.X = 0;
 				Rect.Y = 0;
 				Rect.W = FMath::CeilToInt( (float)FMathd::Abs(ChartSize.X) + 1.0f);
@@ -500,7 +502,7 @@ bool FStandardChartPacker::PackCharts(TArray<FUVIsland>& Charts, double UVScale,
 					bool bFound = false;
 
 					// Use the real raster size for optimal placement
-					FAllocator2D::FRect RasterRect = Rect;
+					FUVSpaceAllocator::FRect RasterRect = Rect;
 					RasterRect.W = ChartRaster.GetRasterWidth();
 					RasterRect.H = ChartRaster.GetRasterHeight();
 
@@ -525,7 +527,6 @@ bool FStandardChartPacker::PackCharts(TArray<FUVIsland>& Charts, double UVScale,
 							RasterRect.Y = StartPos->Y;
 						}
 
-						LayoutRaster.ResetStats();
 						bFound = LayoutRaster.FindWithSegments(RasterRect, ChartRaster, IsBestRect);
 						if (bFound)
 						{
@@ -557,8 +558,6 @@ bool FStandardChartPacker::PackCharts(TArray<FUVIsland>& Charts, double UVScale,
 							Rect.Y = RasterRect.Y;
 
 						}
-
-						LayoutRaster.PublishStats(ChartIndex, Orientation, bFound, Rect, BestRect, RasterMD5, IsBestRect);
 					}
 
 				
@@ -692,7 +691,7 @@ void FStandardChartPacker::OrientChart(FUVIsland& Chart, int32 Orientation)
 // Dilate in 28.4 fixed point. Half pixel dilation is conservative rasterization.
 // Dilation same as Minkowski sum of triangle and square.
 template<int32 Dilate>
-void InternalRasterizeTriangle(FAllocator2D& Shader, const FVector2f Points[3], int32 ScissorWidth, int32 ScissorHeight)
+void InternalRasterizeTriangle(FUVSpaceAllocator& Shader, const FVector2f Points[3], int32 ScissorWidth, int32 ScissorHeight)
 {
 	const FVector2f HalfPixel(0.5f, 0.5f);
 	FVector2f p0 = Points[0] - HalfPixel;
@@ -772,7 +771,7 @@ void InternalRasterizeTriangle(FAllocator2D& Shader, const FVector2f Points[3],
 
 
 
-void FStandardChartPacker::RasterizeChart(const FUVIsland& Chart, uint32 RectW, uint32 RectH, FAllocator2D& OutChartRaster)
+void FStandardChartPacker::RasterizeChart(const FUVIsland& Chart, uint32 RectW, uint32 RectH, FUVSpaceAllocator& OutChartRaster)
 {
 	// Bilinear footprint is -1 to 1 pixels. If packed geometrically, only a half pixel dilation
 	// would be needed to guarantee all charts were at least 1 pixel away, safe for bilinear filtering.
diff --git a/Engine/Plugins/Experimental/GeometryProcessing/Source/GeometricObjects/Private/Parameterization/UVSpaceAllocator.cpp b/Engine/Plugins/Experimental/GeometryProcessing/Source/GeometricObjects/Private/Parameterization/UVSpaceAllocator.cpp
new file mode 100644
--- /dev/null
+++ b/Engine/Plugins/Experimental/GeometryProcessing/Source/GeometricObjects/Private/Parameterization/UVSpaceAllocator.cpp
@@ -0,0 +1,775 @@
+// Copyright Epic Games, Inc. All Rights Reserved.
+
+#include "Parameterization/UVSpaceAllocator.h"
+
+#include "Misc/SecureHash.h"
+#include "ProfilingDebugging/MiscTrace.h"
+
+using namespace UE::Geometry;
+
+FUVSpaceAllocator::FUVSpaceAllocator( FUVSpaceAllocator::EMode InMode, uint32 InWidth, uint32 InHeight)
+	: Mode( InMode )
+	, Width( InWidth )
+	, Height( InHeight )
+	, Pitch( ( InWidth + 63 ) / 64 )
+	, RasterWidth(0)
+	, RasterHeight(0)
+{
+	check(Width  <= MAX_uint16);
+	check(Height <= MAX_uint16);
+
+	Bits.SetNumZeroed(Pitch * Height + 1); // alloc +1 to avoid buffer overrun
+
+	Rows.SetNum( Height );
+	SortedRowsIndex.SetNum(Height);
+	for (uint32 Index = 0; Index < Height; ++Index)
+	{
+		SortedRowsIndex[Index] = (uint16)Index;
+	}
+
+	Columns.SetNum( Width );
+	SortedColumnsIndex.SetNum(Width);
+	for (uint32 Index = 0; Index < Width; ++Index)
+	{
+		SortedColumnsIndex[Index] = (uint16)Index;
+	}
+
+	Clear();
+}
+
+void FUVSpaceAllocator::Clear()
+{
+	InitSegments();
+
+	// Only clear the section we used to reduce memory traffic
+	FPlatformMemory::Memzero(Bits.GetData(), Pitch * RasterHeight * Bits.GetTypeSize());
+
+	RasterWidth  = 0;
+	RasterHeight = 0;
+}
+
+void FUVSpaceAllocator::CopyRuns(TArray<FRun>& Runs, const TArray<FRun>& OtherRuns, int32 MaxSize)
+{
+	Runs.SetNum(OtherRuns.Num(), false);
+
+	MaxSize = FMath::Min(Runs.Num(), MaxSize);
+
+	FRun* RunsPtr = Runs.GetData();
+	const FRun* OtherRunsPtr = OtherRuns.GetData();
+	for (int32 Index = 0; Index < MaxSize; ++Index, ++RunsPtr, ++OtherRunsPtr)
+	{
+		*RunsPtr = *OtherRunsPtr;
+	}
+}
+
+// Optimized version to alloc and copy as less data possible.
+FUVSpaceAllocator& FUVSpaceAllocator::operator = (const FUVSpaceAllocator & Other)
+{
+	if (&Other == this)
+	{
+		return *this;
+	}
+
+	check(Mode == Other.Mode);
+	check(Width == Other.Width);
+	check(Height == Other.Height);
+	check(Pitch == Other.Pitch);
+
+	FPlatformMemory::Memcpy(Bits.GetData(), Other.Bits.GetData(), Pitch * Other.RasterHeight * Bits.GetTypeSize());
+
+	// Zero out anything that was used and not overwritten by the copy
+	if (RasterHeight > Other.RasterHeight)
+	{
+		FPlatformMemory::Memzero(Bits.GetData() + Pitch * Other.RasterHeight, Pitch * (RasterHeight - Other.RasterHeight) * Bits.GetTypeSize());
+	}
+
+	RasterWidth  = Other.RasterWidth;
+	RasterHeight = Other.RasterHeight;
+
+	CopyRuns(Rows, Other.Rows, Other.RasterHeight);
+	SortedRowsIndex = Other.SortedRowsIndex;
+
+	CopyRuns(Columns, Other.Columns, Other.RasterWidth);
+	SortedColumnsIndex = Other.SortedColumnsIndex;
+
+	return *this;
+}
+
+FMD5Hash FUVSpaceAllocator::GetRasterMD5() const
+{
+	FMD5 MD5;
+
+	const FRun* RowPtr = Rows.GetData();
+	const FRun* RowPtrEnd = RowPtr + RasterHeight;
+
+	for (; RowPtr < RowPtrEnd; ++RowPtr)
+	{
+		uint32 SegmentsNum = RowPtr->Segments.Num();
+		
+		// It is important for empty rows to also have a footprint in the MD5
+		MD5.Update((const uint8*)&SegmentsNum, sizeof(SegmentsNum));
+		MD5.Update((const uint8*)RowPtr->Segments.GetData(), SegmentsNum * RowPtr->Segments.GetTypeSize() );
+	}
+
+	FMD5Hash MD5Hash;
+	MD5Hash.Set(MD5);
+	return MD5Hash;
+}
+
+bool FUVSpaceAllocator::Find( FRect& Rect )
+{
+	FRect TestRect = Rect;
+	for( TestRect.X = 0; TestRect.X <= Width - TestRect.W; TestRect.X++ )
+	{
+		for( TestRect.Y = 0; TestRect.Y <= Height - TestRect.H; TestRect.Y++ )
+		{
+			if( Test( TestRect ) )
+			{
+				Rect = TestRect;
+				return true;
+			}
+		}
+	}
+
+	return false;
+}
+
+FString FUVSpaceAllocator::ToString() const
+{
+	TArray<TCHAR> Text;
+	Text.SetNum(Width+2);
+
+	// Truncate output to last meaningful row
+	int32 lastRow = 1;
+	for (const FRun& Row : Rows)
+	{
+		if ((Mode == EMode::FreeSegments && Row.Segments.Num() > 1) ||
+			(Mode == EMode::UsedSegments && Row.Segments.Num() != 0))
+		{
+			lastRow++;
+		}
+	}
+
+	FString Output;
+	Output.Append(TEXT("BEGIN -----------------------\n"));
+
+	const TCHAR FillChar = Mode == EMode::FreeSegments ? TEXT('-') : TEXT(' ');
+	const TCHAR UsedChar = Mode == EMode::FreeSegments ? TEXT(' ') : TEXT('+');
+
+	for (const FRun& Row : Rows)
+	{
+		for (uint32 x = 0; x < Width; ++x)
+		{
+			Text[x] = FillChar;
+		}
+
+		for (const FSegment& Segment : Row.Segments)
+		{
+			for (uint32 x = Segment.StartPos; x < Segment.StartPos + Segment.Length; ++x)
+			{
+				Text[x] = UsedChar;
+			}
+		}
+
+		Text[Text.Num()-2] = TEXT('\n');
+		Text[Text.Num()-1] = TEXT('\0');
+
+		Output.Append(Text.GetData());
+
+		if (--lastRow == 0)
+		{
+			break;
+		}
+	}
+
+	Output.Appendf(TEXT("MD5 %s\n"), *LexToString(GetRasterMD5()));
+	Output.Append(TEXT("END -----------------------\n"));
+
+	return Output;
+}
+
+bool FUVSpaceAllocator::FindBitByBit( FRect& Rect, const FUVSpaceAllocator& Other )
+{
+	FRect TestRect = Rect;
+	for( TestRect.X = 0; TestRect.X <= Width - TestRect.W; TestRect.X++ )
+	{
+		for( TestRect.Y = 0; TestRect.Y <= Height - TestRect.H; TestRect.Y++ )
+		{
+			if( Test( TestRect, Other ) )
+			{
+				Rect = TestRect;
+				return true;
+			}
+		}
+	}
+
+	return false;
+}
+
+bool FUVSpaceAllocator::FindWithSegments( FRect& Rect, const FUVSpaceAllocator& Other, TFunctionRef<bool (const FRect&)> IsBestRect) const
+{
+	FRect TestRect = Rect;
+
+	const uint32 MaxWidth    = Width - TestRect.W;
+	const uint32 MaxHeight   = Height - TestRect.H;
+
+	// For charts that are longer on the Y axis, its a lot faster to use
+	// another grid optimized for searching vertically because it will reduce
+	// iteration count tremendously by improving strides we make for each iteration.
+	// We need a version check to activate this feature because it will favor
+	// placement along the Y axis leading to (still efficient) but different results.
+	if (TestRect.H > TestRect.W)
+	{
+		// Perfect fit algorithm will try every position to fit our chart
+		for( ; TestRect.X <= MaxWidth; ++TestRect.X)
+		{
+			// This represents the length we know for sure we won't fit in so we can skip over
+			uint32 FailedLength = 1;
+			for( TestRect.Y = 0; TestRect.Y <= MaxHeight; TestRect.Y += FailedLength)
+			{
+				// Do not bother continue if another orientation had a better result
+				if( !IsBestRect(Rect) )
+				{
+					return false;
+				}
+
+				if( TestAllColumns( TestRect, Other, FailedLength ) )
+				{
+					Rect = TestRect;
+					return true;
+				}
+			}
+		}
+	}
+	else
+	{
+		// Perfect fit algorithm will try every position to fit our chart
+		for( ; TestRect.Y <= MaxHeight; ++TestRect.Y)
+		{
+			// This represents the length we know for sure we won't fit in so we can skip over
+			uint32 FailedLength = 1;
+			for( TestRect.X = 0; TestRect.X <= MaxWidth; TestRect.X += FailedLength)
+			{
+				// Do not bother continue if another orientation had a better result
+				if( !IsBestRect(Rect) )
+				{
+					return false;
+				}
+
+				if( TestAllRows( TestRect, Other, FailedLength ) )
+				{
+					Rect = TestRect;
+					return true;
+				}
+			}
+		}
+	}
+
+	return false;
+}
+
+uint32 FUVSpaceAllocator::GetUsedTexels() const
+{
+	uint32 Texels = 0;
+	for (const FRun& Row : Rows)
+	{
+		for (const FSegment& Segment : Row.Segments)
+		{
+			Texels += Segment.Length;
+		}
+	}
+
+	if (Mode == EMode::FreeSegments)
+	{
+		Texels = Width * Height - Texels;
+	}
+
+	return Texels;
+}
+
+void FUVSpaceAllocator::Alloc( FRect Rect )
+{
+	for( uint32 y = Rect.Y; y < Rect.Y + Rect.H; y++ )
+	{
+		for( uint32 x = Rect.X; x < Rect.X + Rect.W; x++ )
+		{
+			SetBit( x, y );
+		}
+	}
+}
+
+void FUVSpaceAllocator::Alloc( FRect Rect, const FUVSpaceAllocator& Other )
+{
+	for( uint32 y = 0; y < Rect.H; y++ )
+	{
+		for( uint32 x = 0; x < Rect.W; x++ )
+		{
+			if( Other.GetBit( x, y ) )
+			{
+				SetBit( x + Rect.X, y + Rect.Y );
+			}
+		}
+	}
+
+	MergeSegments( Rect, Other );
+}
+
+bool FUVSpaceAllocator::TestAllRows( const FRect& Rect, const FUVSpaceAllocator& Other, uint32& OutFailedLength ) const
+{
+	for ( uint32 Index = 0; Index < Other.RasterHeight; ++Index)
+	{
+		// Longest rows first have a higher chance of giving us a big stride when it fails
+		const uint16 y = Other.SortedRowsIndex.GetData()[Index];
+
+		const FRun& ThisRow  = Rows.GetData()[ Rect.Y + y ];
+		const FRun& OtherRow = Other.Rows.GetData()[ y ];
+
+		if (!TestOneRun(ThisRow, OtherRow, Rect.X, Rect.W, Width, OutFailedLength))
+		{
+			return false;
+		}
+	}
+
+	return true;
+}
+
+bool FUVSpaceAllocator::TestAllColumns( const FRect& Rect, const FUVSpaceAllocator& Other, uint32& OutFailedLength ) const
+{
+	for ( uint32 Index = 0; Index < Other.RasterWidth; ++Index)
+	{
+		// Longest columns first have a higher chance of giving us a big stride when it fails
+		const uint16 x = Other.SortedColumnsIndex.GetData()[Index];
+
+		const FRun& ThisColumn  = Columns.GetData()[ Rect.X + x ];
+		const FRun& OtherColumn = Other.Columns.GetData()[ x ];
+
+		if (!TestOneRun(ThisColumn, OtherColumn, Rect.Y, Rect.H, Height, OutFailedLength))
+		{
+			return false;
+		}
+	}
+
+	return true;
+}
+
+bool FUVSpaceAllocator::TestOneRun(const FRun& ThisRun, const FRun& OtherRun, uint32 RectOffset, uint32 RectLength, uint32 PrimaryResolution, uint32& OutFailedLength) const
+{
+	// Early out if we can't fit the longest used segment into our longest free segment
+	if ( ThisRun.LongestSegment < OtherRun.LongestSegment )
+	{
+		OutFailedLength = PrimaryResolution;
+		return false;
+	}
+
+	// Threshold to perform a linear scan until the cache miss induced by the lookup is worth it
+	const int32     LoopkupThreshold  = 16;
+	const FSegment* FreeSegments      = ThisRun.Segments.GetData();
+	const FSegment* FreeSegmentsEnd   = ThisRun.Segments.GetData() + ThisRun.Segments.Num();
+	const uint16*   FreeSegmentLookup = ThisRun.Segments.Num() > LoopkupThreshold ? ThisRun.FreeSegmentsLookup.GetData() : nullptr;
+	const FSegment* UsedSegments      = OtherRun.Segments.GetData();
+	const FSegment* UsedSegmentsEnd   = OtherRun.Segments.GetData() + OtherRun.Segments.Num();
+
+	// Running pointer for used segments
+	const FSegment* UsedSegment = UsedSegments;
+	for ( ; UsedSegment < UsedSegmentsEnd; ++UsedSegment)
+	{
+		// #TODO maintain this for backward compatibility but not sure if this could ever happen
+		// and if the result is valid if we break instead of returning false when a used segment doesn't fit
+		if ( UsedSegment->StartPos >= RectLength )
+		{
+			break;
+		}
+
+		const uint32 StartPos = RectOffset + UsedSegment->StartPos;
+		const uint32 EndPos   = RectOffset + FMath::Min( UsedSegment->StartPos + UsedSegment->Length, RectLength );
+
+		// Running pointer for free segments
+		const FSegment* FreeSegment = FreeSegments;
+
+		// If lookup exist, advance pointer to the exact spot instead of scanning
+		if (FreeSegmentLookup)
+		{
+			FreeSegment += FreeSegmentLookup[StartPos];
+		}
+
+		// Scan the search range to see if the current used segment fits in
+		// we use a for loop to support scanning from the beginning if no lookup exists
+		bool bFoundSpaceForSegment = false;
+		for (; FreeSegment < FreeSegmentsEnd && StartPos >= FreeSegment->StartPos; ++FreeSegment)
+		{
+			// Check if there's a free segment that can fit the used segment
+			if (EndPos <= FreeSegment->StartPos + FreeSegment->Length)
+			{
+				bFoundSpaceForSegment = true;
+				break;
+			}
+		}
+
+		// If nothing was found in the search range, scan other free segments to hint our
+		// caller of how much farther it should advance for next test
+		if (!bFoundSpaceForSegment)
+		{
+			for (; FreeSegment < FreeSegmentsEnd; ++FreeSegment)
+			{
+				if (UsedSegment->Length <= FreeSegment->Length)
+				{
+					OutFailedLength = FreeSegment->StartPos - StartPos;
+
+					return false;
+				}
+			}
+
+			// We couldn't find any free segment big enough
+			// set failed length so the whole run will be skipped
+			OutFailedLength = PrimaryResolution;
+
+			return false;
+		}
+	}
+
+	return true;
+}
+
+void FUVSpaceAllocator::FlipX( const FRect& Rect )
+{
+	// If we have empty padding around the Rect, keep it there
+	uint32 MaxX = 0;
+	uint32 MaxY = RasterHeight - 1;
+
+	MaxX = RasterWidth - 1;
+
+	for ( uint32 Y = 0; Y <= MaxY; ++Y )
+	{
+		for ( uint32 LowX = 0; LowX < ( MaxX + 1 ) / 2; ++LowX )
+		{
+			uint32 HighX = MaxX - LowX;
+
+			const uint64 BitLow = GetBit( LowX, Y );
+			const uint64 BitHigh = GetBit( HighX, Y );
+
+			if ( BitLow != 0ull )
+			{
+				SetBit( HighX, Y );
+			}
+			else
+			{
+				ClearBit( HighX, Y );
+			}
+
+			if ( BitHigh != 0ull )
+			{
+				SetBit( LowX, Y );
+			}
+			else
+			{
+				ClearBit( LowX, Y );
+			}
+		}
+	}
+
+	CreateUsedSegments();
+}
+
+void FUVSpaceAllocator::FlipY( const FRect& Rect )
+{
+	uint32 MinY = 0;
+	uint32 MaxY = RasterHeight - 1;
+
+	for ( uint32 LowY = 0; LowY < ( MaxY + 1 ) / 2; ++LowY )
+	{
+		for ( uint32 X = 0; X < RasterWidth; ++X )
+		{
+			uint32 HighY = MaxY - LowY;
+
+			const uint64 BitLow = GetBit( X, LowY );
+			const uint64 BitHigh = GetBit( X, HighY );
+
+			if ( BitLow != 0ull )
+			{
+				SetBit( X, HighY );
+			}
+			else
+			{
+				ClearBit( X, HighY );
+			}
+
+			if ( BitHigh != 0ull )
+			{
+				SetBit( X, LowY );
+			}
+			else
+			{
+				ClearBit( X, LowY );
+			}
+		}
+	}
+
+	CreateUsedSegments();
+}
+
+void FUVSpaceAllocator::InitRuns(TArray<FRun>& Runs, uint32 PrimaryResolution, uint32 PerpendicularRasterSize)
+{
+	if (Mode == EMode::FreeSegments)
+	{
+		FSegment FreeSegment;
+		FreeSegment.StartPos = 0;
+		FreeSegment.Length = PrimaryResolution;
+
+		for ( FRun& Run : Runs )
+		{
+			Run.Segments.Reset();
+			Run.Segments.Add( FreeSegment );
+			Run.LongestSegment = PrimaryResolution;
+			Run.FreeSegmentsLookup.Reset();
+		}
+	}
+	else
+	{
+		// This is called inside a hot loop, get rid of RangeCheck.
+		FRun* Run = Runs.GetData();
+		for ( uint32 Index = 0; Index < PerpendicularRasterSize; ++Index )
+		{
+			Run[Index].Segments.Reset();
+			Run[Index].LongestSegment = 0;
+		}
+	}
+}
+
+void FUVSpaceAllocator::InitSegments()
+{
+	InitRuns(Rows, Width, RasterHeight);
+
+	InitRuns(Columns, Height, RasterWidth);
+}
+
+void FUVSpaceAllocator::CreateUsedSegments()
+{
+	check(Mode == EMode::UsedSegments);
+
+	SortedRowsIndex.SetNum(RasterHeight, false);
+
+	uint64* BitsData = Bits.GetData();
+
+	// Create segments along the X axis for each rows
+	for ( uint32 y = 0; y < RasterHeight; ++y )
+	{
+		SortedRowsIndex[y] = y;
+
+		FRun& CurrentRow = Rows[y];
+		CurrentRow.LongestSegment = 0;
+		CurrentRow.Segments.Reset();
+
+		int32 FirstUsedX = -1;
+		for ( uint32 k = 0; k < Pitch; ++k, ++BitsData )
+		{
+			const uint32 x = k * 64;
+			const uint64 BitsValue = *BitsData;
+
+			// If all bits are set
+			if ( BitsValue == ~0ull )
+			{
+				if ( FirstUsedX < 0 )
+				{
+					FirstUsedX = x;
+				}
+
+				if ( k == Pitch - 1 )
+				{
+					AddUsedSegment( CurrentRow, FirstUsedX, x + 64 - FirstUsedX );
+					FirstUsedX = -1;
+				}
+			}
+			// No bits are set
+			else if ( BitsValue == 0ull )
+			{
+				if ( FirstUsedX >= 0 )
+				{
+					AddUsedSegment( CurrentRow, FirstUsedX, x - FirstUsedX );
+					FirstUsedX = -1;
+				}
+			}
+			// Some bits are set
+			else
+			{
+				for ( uint32 i = 0; i < 64; ++i )
+				{
+					const uint32 SubX = x + i;
+
+					if ( BitsValue & ( 1ull << i ) )
+					{
+						if ( FirstUsedX < 0 )
+						{
+							FirstUsedX = SubX;
+						}
+
+						if ( SubX == Width - 1 )
+						{
+							AddUsedSegment( CurrentRow, FirstUsedX, SubX + 1 - FirstUsedX );
+							FirstUsedX = -1;
+						}
+					}
+					else if ( FirstUsedX >= 0 )
+					{
+						AddUsedSegment( CurrentRow, FirstUsedX, SubX - FirstUsedX );
+						FirstUsedX = -1;
+					}
+				}
+			}
+		}
+	}
+
+	SortedRowsIndex.Sort([this](uint32 a, uint32 b) { return Rows[b].LongestSegment < Rows[a].LongestSegment;});
+
+	// Create segments along the Y axis for each columns
+	SortedColumnsIndex.SetNum(RasterWidth, false);
+
+	for ( uint32 x = 0; x < RasterWidth; ++x )
+	{
+		SortedColumnsIndex[x] = x;
+
+		FRun& CurrentColumn = Columns[x];
+		CurrentColumn.LongestSegment = 0;
+		CurrentColumn.Segments.Reset();
+
+		int32 FirstUsedY = -1;
+		for ( uint32 y = 0; y < RasterHeight; ++y )
+		{
+			if (GetBit(x, y))
+			{
+				if ( FirstUsedY < 0 )
+				{
+					FirstUsedY = y;
+				}
+
+				if ( y == RasterHeight - 1 )
+				{
+					AddUsedSegment( CurrentColumn, FirstUsedY, y + 1 - FirstUsedY );
+					FirstUsedY = -1;
+				}
+			}
+			else if ( FirstUsedY >= 0 )
+			{
+				AddUsedSegment( CurrentColumn, FirstUsedY, y - FirstUsedY );
+				FirstUsedY = -1;
+			}
+		}
+	}
+
+	SortedColumnsIndex.Sort([this](uint32 a, uint32 b) { return Columns[b].LongestSegment < Columns[a].LongestSegment;});
+}
+
+void FUVSpaceAllocator::AddUsedSegment( FRun& Row, uint32 StartPos, uint32 Length )
+{
+	FSegment UsedSegment;
+	UsedSegment.StartPos = StartPos;
+	UsedSegment.Length = Length;
+	Row.Segments.Add(UsedSegment);
+
+	if ( UsedSegment.Length > Row.LongestSegment )
+	{
+		Row.LongestSegment = UsedSegment.Length;
+	}
+}
+
+void FUVSpaceAllocator::MergeRun(FRun& ThisRun, const FRun& OtherRun, uint32 RectOffset, uint32 RectLength, uint32 PrimaryResolution, uint32 PerpendicularResolution)
+{
+	// #TODO
+	// Perform a linear scan once for both array at the same time zipping them together
+	// Ensure we keep the order of the segments while merging them to avoid costly sort when number of free segments is high
+	for ( const FSegment& OtherUsedSegment : OtherRun.Segments )
+	{
+		for ( int32 Index = 0; Index < ThisRun.Segments.Num(); ++Index )
+		{
+			FSegment& ThisFreeSegment = ThisRun.Segments[Index];
+
+			uint32 StartPos = RectOffset + OtherUsedSegment.StartPos;
+
+			if ( StartPos >= ThisFreeSegment.StartPos &&
+				StartPos < ThisFreeSegment.StartPos + ThisFreeSegment.Length )
+			{
+				if ( ThisFreeSegment.Length == 1 )
+				{
+					ThisRun.Segments.RemoveAtSwap(Index);
+					break;
+				}
+
+				FSegment FirstSegment;
+				FirstSegment.StartPos = ThisFreeSegment.StartPos;
+				FirstSegment.Length = StartPos - ThisFreeSegment.StartPos;
+
+				uint32 EndPos = RectOffset + FMath::Min( OtherUsedSegment.StartPos + OtherUsedSegment.Length, RectLength ) - 1;
+
+				FSegment SecondSegment;
+				SecondSegment.StartPos = EndPos + 1;
+				SecondSegment.Length = ThisFreeSegment.StartPos + ThisFreeSegment.Length - SecondSegment.StartPos;
+
+				ThisRun.Segments.RemoveAtSwap(Index);
+
+				if ( FirstSegment.Length > 0 )
+				{
+					ThisRun.Segments.Add( FirstSegment );
+				}
+
+				if ( SecondSegment.Length > 0 )
+				{
+					ThisRun.Segments.Add( SecondSegment );
+				}
+
+				break;
+			}
+		}
+	}
+
+	ThisRun.Segments.Sort();
+
+	// When we're dealing with an Axis, the lookup table is the size of the perpendicular Axis
+	ThisRun.FreeSegmentsLookup.SetNum(PerpendicularResolution);
+	ThisRun.LongestSegment = 0;
+
+	uint32 LastIndex = 0;
+
+	// Avoid RangeCheck by indexing the pointer directly since we're in a hot loop
+	FSegment* SegmentsPtr           = ThisRun.Segments.GetData();
+	uint16*   FreeSegmentsLookupPtr = ThisRun.FreeSegmentsLookup.GetData();
+
+	int32 SegmentIndex = 0;
+	for ( ; SegmentIndex < ThisRun.Segments.Num(); ++SegmentIndex )
+	{
+		FSegment& Segment = SegmentsPtr[SegmentIndex];
+
+		if (Segment.Length > ThisRun.LongestSegment)
+		{
+			ThisRun.LongestSegment = Segment.Length;
+		}
+
+		uint32 StopIndex  = Segment.StartPos + Segment.Length;
+		for (uint32 LookupIndex = LastIndex; LookupIndex < StopIndex; ++LookupIndex)
+		{
+			FreeSegmentsLookupPtr[LookupIndex] = SegmentIndex;
+		}
+
+		LastIndex = StopIndex;
+	}
+
+	for (uint32 Index = LastIndex; Index < PerpendicularResolution; ++Index)
+	{
+		FreeSegmentsLookupPtr[Index] = SegmentIndex;
+	}
+}
+
+void FUVSpaceAllocator::MergeSegments( const FRect& Rect, const FUVSpaceAllocator& Other )
+{
+	check(Mode       == EMode::FreeSegments);
+	check(Other.Mode == EMode::UsedSegments);
+
+	for ( uint32 y = 0; y < Other.RasterHeight; ++y )
+	{
+		MergeRun(Rows[Rect.Y + y], Other.Rows[y], Rect.X, Rect.W, Width, Height);
+	}
+
+	for ( uint32 x = 0; x < Other.RasterWidth; ++x )
+	{
+		MergeRun(Columns[Rect.X + x], Other.Columns[x], Rect.Y, Rect.H, Height, Width);
+	}
+}
+
diff --git a/Engine/Plugins/Experimental/GeometryProcessing/Source/GeometricObjects/Private/Parameterization/UVSpaceAllocator.h b/Engine/Plugins/Experimental/GeometryProcessing/Source/GeometricObjects/Private/Parameterization/UVSpaceAllocator.h
new file mode 100644
--- /dev/null
+++ b/Engine/Plugins/Experimental/GeometryProcessing/Source/GeometricObjects/Private/Parameterization/UVSpaceAllocator.h
@@ -0,0 +1,219 @@
+// Copyright Epic Games, Inc. All Rights Reserved.
+
+#pragma once
+
+#include "CoreMinimal.h"
+
+struct FRect;
+struct Rect;
+struct FMD5Hash;
+
+namespace UE
+{
+namespace Geometry
+{
+
+//
+// GeometryProcessing port of MeshUtilitiesCommon FAllocator2D
+//
+class FUVSpaceAllocator
+{
+public:
+	enum class EMode
+	{
+		// In this mode, segments represents free space
+		// Used for the layout merging usedsegments
+		FreeSegments,
+		// In this mode, segments represents used space
+		// Used for the rasterization of charts.
+		UsedSegments
+	};
+
+	struct FRect
+	{
+		uint32 X;
+		uint32 Y;
+		uint32 W;
+		uint32 H;
+	};
+
+	struct FSegment
+	{
+		uint32 StartPos;
+		uint32 Length;
+
+		bool operator<( const FSegment& Other ) const { return StartPos < Other.StartPos; }
+	};
+
+	struct FRun
+	{
+		uint32 LongestSegment;
+		TArray< FSegment > Segments;
+
+		// Contains mapping from pixel position to first segment index in search range.
+		// Only computed when we're in FreeSegments mode to help TestOneRun find
+		// the proper segment in O(c) at the expense of a likely cache miss.
+		// We'll use a threshold to use this method when the number of iterations
+		// saved is worth the cache miss.
+		// Obviously, using a uint16 here will reduce cache misses but impose
+		// an hopefully enough limitation of texture size 65536x65536 (4GB).
+		TArray< uint16 > FreeSegmentsLookup;
+	};
+
+public:
+	FUVSpaceAllocator( EMode Mode, uint32 Width, uint32 Height );
+
+	// Must clear before using
+	void       Clear();
+
+	bool       Find( FRect& Rect );
+	bool       Test( FRect Rect );
+	void       Alloc( FRect Rect );
+
+	bool       FindBitByBit( FRect& Rect, const FUVSpaceAllocator& Other );
+	bool       FindWithSegments( FRect& Rect, const FUVSpaceAllocator& Other, TFunctionRef<bool (const FUVSpaceAllocator::FRect&)> IsBestRect ) const;
+	bool       Test( FRect Rect, const FUVSpaceAllocator& Other );
+	void       Alloc( FRect Rect, const FUVSpaceAllocator& Other );
+
+	uint64     GetBit( uint32 x, uint32 y ) const;
+	void       SetBit( uint32 x, uint32 y );
+	void       ClearBit( uint32 x, uint32 y );
+
+	void       CreateUsedSegments();
+	void       MergeRun( FRun& Run, const FRun& OtherRun, uint32 RectOffset, uint32 RectLength, uint32 PrimaryResolution /* Resolution along the axis the run belongs to */, uint32 PerpendicularResolution );
+	void       MergeSegments( const FRect& Rect, const FUVSpaceAllocator& Other );
+
+	void       FlipX( const FRect& Rect );
+	void       FlipY( const FRect& Rect );
+
+	uint32     GetUsedTexels() const;
+
+	void       CopyRuns( TArray<FRun>& Runs, const TArray<FRun>& OtherRuns, int32 MaxSize );
+
+	// Take control of the copy to reduce the amount of data movement to the strict minimum
+	FUVSpaceAllocator& operator = (const FUVSpaceAllocator& Other);
+
+	// Allow to visualize the content in ascii for debugging purpose. (i.e Watch or Immediate window).
+	FString    ToString() const;
+
+	// Get the MD5 hash of the rasterized content
+	FMD5Hash   GetRasterMD5() const;
+
+	uint32     GetRasterWidth()  const { return RasterWidth; }
+	uint32     GetRasterHeight() const { return RasterHeight; }
+
+protected:
+	bool       TestOneRun( const FRun& Run, const FRun& OtherRun, uint32 RectOffset, uint32 RectLength, uint32 PrimaryResolution, uint32& OutFailedLength ) const;
+	bool       TestAllRows( const FRect& Rect, const FUVSpaceAllocator& Other, uint32& FailedLength ) const;
+	bool       TestAllColumns( const FRect& Rect, const FUVSpaceAllocator& Other, uint32& FailedLength ) const;
+	void       InitRuns( TArray<FRun>& Runs, uint32 PrimaryResolution, uint32 PerpendicularRasterSize);
+	void       InitSegments();
+	void       AddUsedSegment( FRun& Run, uint32 StartPos, uint32 Length );
+
+	// Enforce that those cannot be changed in flight
+	const EMode              Mode;
+	const uint32             Width;
+	const uint32             Height;
+	const uint32             Pitch;
+
+	uint32                   RasterWidth;
+	uint32                   RasterHeight;
+	TArray< FRun >           Rows;        // Represent rows in the grid
+	TArray< FRun >           Columns;     // Represent columns in the grid (used when version >= Segments2D).
+	TArray< uint64 >         Bits;
+
+	// Index inside rows that will be sorted by rows with longest used segment first
+	TArray< uint16 >         SortedRowsIndex;
+	// Index inside columns that will be sorted by columns with longest used segment first
+	TArray< uint16 >         SortedColumnsIndex;
+};
+
+// Returns non-zero if set
+FORCEINLINE uint64 FUVSpaceAllocator::GetBit( uint32 x, uint32 y ) const
+{
+	return Bits[ (x >> 6) + y * Pitch ] & ( 1ull << ( x & 63 ) );
+}
+
+FORCEINLINE void FUVSpaceAllocator::SetBit( uint32 x, uint32 y )
+{
+	Bits[ (x >> 6) + y * Pitch ] |= ( 1ull << ( x & 63 ) );
+
+	// Keep track of the rasterized dimension to optimize operations on that area only
+	if (y >= RasterHeight)
+	{
+		RasterHeight = y + 1;
+	}
+
+	if (x >= RasterWidth)
+	{
+		RasterWidth = x + 1;
+	}
+}
+
+FORCEINLINE void FUVSpaceAllocator::ClearBit( uint32 x, uint32 y )
+{
+	Bits[ (x >> 6) + y * Pitch ] &= ~( 1ull << ( x & 63 ) );
+}
+
+inline bool FUVSpaceAllocator::Test( FRect Rect )
+{
+	for( uint32 y = Rect.Y; y < Rect.Y + Rect.H; y++ )
+	{
+		for( uint32 x = Rect.X; x < Rect.X + Rect.W; x++ )
+		{
+			if( GetBit( x, y ) )
+			{
+				return false;
+			}
+		}
+	}
+	
+	return true;
+}
+
+inline bool FUVSpaceAllocator::Test( FRect Rect, const FUVSpaceAllocator& Other )
+{
+	const uint32 LowShift = Rect.X & 63;
+	const uint32 HighShift = 64 - LowShift;
+
+	for( uint32 y = 0; y < Rect.H; y++ )
+	{
+#if 1
+		uint32 ThisIndex = (Rect.X >> 6) + (y + Rect.Y) * Pitch;
+		uint32 OtherIndex = y * Pitch;
+
+		// Test a uint64 at a time
+		for( uint32 x = 0; x < Rect.W; x += 64 )
+		{
+			// no need to zero out HighInt on wrap around because Other will always be zero outside Rect.
+			uint64 LowInt  = Bits[ ThisIndex ];
+			uint64 HighInt = Bits[ ThisIndex + 1 ];
+
+			uint64 ThisInt = (HighInt << HighShift) | (LowInt >> LowShift);
+			uint64 OtherInt = Other.Bits[ OtherIndex ];
+
+			if( ThisInt & OtherInt )
+			{
+				return false;
+			}
+
+			ThisIndex++;
+			OtherIndex++;
+		}
+#else
+		for( uint32 x = 0; x < Rect.W; x++ )
+		{
+			if( Other.GetBit( x, y ) && GetBit( x + Rect.X, y + Rect.Y ) )
+			{
+				return false;
+			}
+		}
+#endif
+	}
+	
+	return true;
+}
+
+
+} // end namespace UE::Geometry
+} // end namespace UE
\ No newline at end of file
diff --git a/Engine/Plugins/Experimental/GeometryProcessing/Source/GeometricObjects/Public/Parameterization/UVPacking.h b/Engine/Plugins/Experimental/GeometryProcessing/Source/GeometricObjects/Public/Parameterization/UVPacking.h
--- a/Engine/Plugins/Experimental/GeometryProcessing/Source/GeometricObjects/Public/Parameterization/UVPacking.h
+++ b/Engine/Plugins/Experimental/GeometryProcessing/Source/GeometricObjects/Public/Parameterization/UVPacking.h
@@ -5,7 +5,6 @@
 #include "VectorTypes.h"
 #include "IndexTypes.h"
 #include "BoxTypes.h"
-#include "DisjointSet.h"
 
 namespace UE
 {
@@ -90,67 +89,3 @@ protected:
 } // end namespace UE::Geometry
 } // end namespace UE
 
-namespace UE
-{
-namespace UVPacking
-{
-	using namespace UE::Geometry;
-/**
- * Create UV islands from a triangle mesh connectivity.
- * Assumes the triangles are already split at UV seams, but topologically connected otherwise.
- * Not recommended for meshes that already have built-in edge connectivity data.
- *
- * @param Mesh The mesh to create islands for
- * @param IslandsOut The triangle IDs for each island
- * @param IncludeTri Optional function to filter which tris are assigned to islands
- */
-template <typename TriangleMeshType>
-inline void CreateUVIslandsFromMeshTopology(TriangleMeshType& Mesh,
-	TArray<TArray<int>>& IslandsOut,
-	TFunctionRef<bool(int32)> IncludeTri = [](int32)
-	{
-		return true;
-	})
-{
-	FDisjointSet VertComponents(Mesh.MaxVertexID());
-	// Add Source vertices to hash & disjoint sets
-	for (int32 TID = 0; TID < Mesh.MaxTriangleID(); TID++)
-	{
-		if (!Mesh.IsTriangle(TID) || !IncludeTri(TID))
-		{
-			continue;
-		}
-
-		FIndex3i Tri = Mesh.GetTriangle(TID);
-		for (int32 First = 2, Second = 0; Second < 3; First = Second++)
-		{
-			VertComponents.Union(Tri[First], Tri[Second]);
-		}
-	}
-
-	IslandsOut.Reset();
-	TMap<uint32, int32> IslandIDToIdx;
-	for (int32 TID = 0; TID < Mesh.MaxTriangleID(); TID++)
-	{
-		if (!Mesh.IsTriangle(TID) || !IncludeTri(TID))
-		{
-			continue;
-		}
-
-		FIndex3i Tri = Mesh.GetTriangle(TID);
-		uint32 IslandID = VertComponents.Find(Tri.A);
-		int32 Idx = -1;
-		int32* FoundIdx = IslandIDToIdx.Find(IslandID);
-		if (!FoundIdx)
-		{
-			Idx = IslandsOut.Emplace();
-			IslandIDToIdx.Add(IslandID, Idx);
-		}
-		else
-		{
-			Idx = *FoundIdx;
-		}
-		IslandsOut[Idx].Add(TID);
-	}
-}
-}} // namespace UE::UVPacking
diff --git a/Engine/Plugins/Experimental/PlanarCutPlugin/Source/PlanarCut/Private/FractureAutoUV.cpp b/Engine/Plugins/Experimental/PlanarCutPlugin/Source/PlanarCut/Private/FractureAutoUV.cpp
--- a/Engine/Plugins/Experimental/PlanarCutPlugin/Source/PlanarCut/Private/FractureAutoUV.cpp
+++ b/Engine/Plugins/Experimental/PlanarCutPlugin/Source/PlanarCut/Private/FractureAutoUV.cpp
@@ -10,6 +10,7 @@
 #include "Parameterization/DynamicMeshUVEditor.h"
 
 #include "Parameterization/UVPacking.h"
+#include "DisjointSet.h"
 
 #include "Image/ImageOccupancyMap.h"
 #include "VectorUtil.h"
@@ -22,6 +23,76 @@
 #include "ExplicitUseGeometryMathTypes.h"		// using UE::Geometry::(math types)
 using namespace UE::Geometry;
 
+
+
+namespace UE
+{
+	namespace UVPacking
+	{
+		using namespace UE::Geometry;
+		/**
+		 * Create UV islands from a triangle mesh connectivity.
+		 * Assumes the triangles are already split at UV seams, but topologically connected otherwise.
+		 * Not recommended for meshes that already have built-in edge connectivity data.
+		 *
+		 * @param Mesh The mesh to create islands for
+		 * @param IslandsOut The triangle IDs for each island
+		 * @param IncludeTri Optional function to filter which tris are assigned to islands
+		 */
+		template <typename TriangleMeshType>
+		inline void CreateUVIslandsFromMeshTopology(TriangleMeshType& Mesh,
+			TArray<TArray<int>>& IslandsOut,
+			TFunctionRef<bool(int32)> IncludeTri = [](int32)
+		{
+			return true;
+		})
+		{
+			FDisjointSet VertComponents(Mesh.MaxVertexID());
+			// Add Source vertices to hash & disjoint sets
+			for (int32 TID = 0; TID < Mesh.MaxTriangleID(); TID++)
+			{
+				if (!Mesh.IsTriangle(TID) || !IncludeTri(TID))
+				{
+					continue;
+				}
+
+				FIndex3i Tri = Mesh.GetTriangle(TID);
+				for (int32 First = 2, Second = 0; Second < 3; First = Second++)
+				{
+					VertComponents.Union(Tri[First], Tri[Second]);
+				}
+			}
+
+			IslandsOut.Reset();
+			TMap<uint32, int32> IslandIDToIdx;
+			for (int32 TID = 0; TID < Mesh.MaxTriangleID(); TID++)
+			{
+				if (!Mesh.IsTriangle(TID) || !IncludeTri(TID))
+				{
+					continue;
+				}
+
+				FIndex3i Tri = Mesh.GetTriangle(TID);
+				uint32 IslandID = VertComponents.Find(Tri.A);
+				int32 Idx = -1;
+				int32* FoundIdx = IslandIDToIdx.Find(IslandID);
+				if (!FoundIdx)
+				{
+					Idx = IslandsOut.Emplace();
+					IslandIDToIdx.Add(IslandID, Idx);
+				}
+				else
+				{
+					Idx = *FoundIdx;
+				}
+				IslandsOut[Idx].Add(TID);
+			}
+		}
+	}
+} // namespace UE::UVPacking
+
+
+
 namespace UE { namespace PlanarCut {
 
 /**
@@ -237,6 +308,9 @@ struct FGeomFlatUVMesh
 };
 
 
+
+
+
 bool UVLayout(
 	FGeometryCollection& Collection,
 	int32 UVRes,
