Optimize computation of normal deformation in FFD lattice

#rnx

diff --git a/Engine/Plugins/Experimental/GeometryProcessing/Source/DynamicMesh/Private/Operations/FFDLattice.cpp b/Engine/Plugins/Experimental/GeometryProcessing/Source/DynamicMesh/Private/Operations/FFDLattice.cpp
--- a/Engine/Plugins/Experimental/GeometryProcessing/Source/DynamicMesh/Private/Operations/FFDLattice.cpp
+++ b/Engine/Plugins/Experimental/GeometryProcessing/Source/DynamicMesh/Private/Operations/FFDLattice.cpp
@@ -425,28 +425,43 @@ FVector3d FFFDLattice::InterpolatedPosition(const FEmbedding& VertexEmbedding, c
 
 FMatrix3d FFFDLattice::LinearInterpolationJacobian(const FEmbedding& VertexEmbedding, const TArray<FVector3d>& LatticeControlPoints) const
 {
-	// TODO: This was written for clarity and correctness. Could definitely be faster.
+	int X0 = VertexEmbedding.LatticeCell.X;
+	int Y0 = VertexEmbedding.LatticeCell.Y;
+	int Y1 = Y0 + 1;
+	int Z0 = VertexEmbedding.LatticeCell.Z;
+	int Z1 = Z0 + 1;
+
+	FVector3d FV000, FV100;
+	GetValuePair(X0, Y0, Z0, FV000, FV100, LatticeControlPoints);
+	FVector3d FV001, FV101;
+	GetValuePair(X0, Y0, Z1, FV001, FV101, LatticeControlPoints);
+	FVector3d FV010, FV110;
+	GetValuePair(X0, Y1, Z0, FV010, FV110, LatticeControlPoints);
+	FVector3d FV011, FV111;
+	GetValuePair(X0, Y1, Z1, FV011, FV111, LatticeControlPoints);
+
+	double AlphaX = VertexEmbedding.CellWeighting.X;
+	double AlphaY = VertexEmbedding.CellWeighting.Y;
+	double AlphaZ = VertexEmbedding.CellWeighting.Z;
+	double OneMinusAlphaX = 1.0 - AlphaX;
 
 	// Partial wrt x
-	FEmbedding XFloor = VertexEmbedding;
-	XFloor.CellWeighting.X = 0.0;
-	FEmbedding XCeil = VertexEmbedding;
-	XCeil.CellWeighting.X = 1.0;
-	FVector3d PartialX = InterpolatedPosition(XCeil, LatticeControlPoints) - InterpolatedPosition(XFloor, LatticeControlPoints);
+	FVector3d PartialX = (FV100 - FV000) * (1 - AlphaY) * (1 - AlphaZ) +
+		(FV101 - FV001) * (1 - AlphaY) * (AlphaZ)+
+		(FV110 - FV010) * (AlphaY) * (1 - AlphaZ) +
+		(FV111 - FV011) * (AlphaY) * (AlphaZ);
+
+	// common terms for partialy and partialz
+	FVector3d T0 = OneMinusAlphaX * FV000 + AlphaX * FV100;
+	FVector3d T1 = OneMinusAlphaX * FV001 + AlphaX * FV101;
+	FVector3d T2 = OneMinusAlphaX * FV010 + AlphaX * FV110;
+	FVector3d T3 = OneMinusAlphaX * FV011 + AlphaX * FV111;
 
 	// Partial wrt y
-	FEmbedding YFloor = VertexEmbedding;
-	YFloor.CellWeighting.Y = 0.0;
-	FEmbedding YCeil = VertexEmbedding;
-	YCeil.CellWeighting.Y = 1.0;
-	FVector3d PartialY = InterpolatedPosition(YCeil, LatticeControlPoints) - InterpolatedPosition(YFloor, LatticeControlPoints);
+	FVector3d PartialY = T0 * (AlphaZ - 1) - T1 * AlphaZ + T2 * (1 - AlphaZ) + T3 * AlphaZ;
 
 	// Partial wrt z
-	FEmbedding ZFloor = VertexEmbedding;
-	ZFloor.CellWeighting.Z = 0.0;
-	FEmbedding ZCeil = VertexEmbedding;
-	ZCeil.CellWeighting.Z = 1.0;
-	FVector3d PartialZ = InterpolatedPosition(ZCeil, LatticeControlPoints) - InterpolatedPosition(ZFloor, LatticeControlPoints);
+	FVector3d PartialZ = T0 * (AlphaY - 1) + T1 * (1 - AlphaY) - T2 * AlphaY + T3 * AlphaY;
 
 	FMatrix3d Mat(PartialX, PartialY, PartialZ, false);
 
diff --git a/Engine/Plugins/Experimental/GeometryProcessing/Source/DynamicMesh/Public/Operations/FFDLattice.h b/Engine/Plugins/Experimental/GeometryProcessing/Source/DynamicMesh/Public/Operations/FFDLattice.h
--- a/Engine/Plugins/Experimental/GeometryProcessing/Source/DynamicMesh/Public/Operations/FFDLattice.h
+++ b/Engine/Plugins/Experimental/GeometryProcessing/Source/DynamicMesh/Public/Operations/FFDLattice.h
@@ -74,7 +74,10 @@ public:
 									 FLatticeExecutionInfo ExecutionInfo,
 									 FProgressCancel* Progress) const;
 
+	/// Return the set of lattice corner positions in undeformed state
 	void GenerateInitialLatticePositions(TArray<FVector3d>& OutLatticePositions) const;
+
+	/// Return the lattice edges by index. The indices refer to the array of lattice positions returned by GenerateInitialLatticePositions
 	void GenerateLatticeEdges(TArray<FVector2i>& OutLatticeEdges) const;
 
 protected:
@@ -92,8 +95,6 @@ protected:
 	TArray<FEmbedding> VertexEmbeddings;
 
 
-	void InitializeLatticePositions();
-
 	/// Get the index into the flat TArray of positions given the (i,j,k) coordinates in the lattice
 	int ControlPointIndexFromCoordinates(int i, int j, int k) const
 	{
@@ -101,6 +102,7 @@ protected:
 		return Idx;
 	}
 
+	/// Get the index into the flat TArray of positions given the (i,j,k) coordinates in the lattice
 	int ControlPointIndexFromCoordinates(const FVector3i& Index) const
 	{
 		return ControlPointIndexFromCoordinates(Index.X, Index.Y, Index.Z);
