Fix missing includes in non-unity build

diff --git a/Engine/Plugins/Experimental/GeometryProcessing/Source/GeometricObjects/Public/Util/ProgressCancel.h b/Engine/Plugins/Experimental/GeometryProcessing/Source/GeometricObjects/Public/Util/ProgressCancel.h
--- a/Engine/Plugins/Experimental/GeometryProcessing/Source/GeometricObjects/Public/Util/ProgressCancel.h
+++ b/Engine/Plugins/Experimental/GeometryProcessing/Source/GeometricObjects/Public/Util/ProgressCancel.h
@@ -5,6 +5,9 @@
 #pragma once
 
 #include "Templates/Function.h"
+#include "Containers/Array.h"
+#include "Internationalization/Text.h"
+#include "Misc/DateTime.h"
 
 /**
  * FProgressCancel is an obejct that is intended to be passed to long-running
