set bitangentsign of tangent space in mesh description builder
#rnx

diff --git a/Engine/Plugins/Experimental/GeometryProcessing/Source/MeshConversion/Private/MeshDescriptionBuilder.cpp b/Engine/Plugins/Experimental/GeometryProcessing/Source/MeshConversion/Private/MeshDescriptionBuilder.cpp
--- a/Engine/Plugins/Experimental/GeometryProcessing/Source/MeshConversion/Private/MeshDescriptionBuilder.cpp
+++ b/Engine/Plugins/Experimental/GeometryProcessing/Source/MeshConversion/Private/MeshDescriptionBuilder.cpp
@@ -172,12 +172,11 @@ void FMeshDescriptionBuilder::SetInstanceTangentSpace(const FVertexInstanceID& I
 
 	if (InstanceTangents.IsValid())
 	{
-		FVector4 TangentExtended;  
-		TangentExtended.X = Tangent.X; 
-		TangentExtended.Y = Tangent.Y; 
-		TangentExtended.Z = Tangent.Z; 
-		TangentExtended.W =  Sign;
-		InstanceTangents.Set(InstanceID, TangentExtended);
+		InstanceTangents.Set(InstanceID, Tangent);
+	}
+	if (InstanceBiTangentSign.IsValid())
+	{
+		InstanceBiTangentSign.Set(InstanceID, Sign);
 	}
 
 
diff --git a/Engine/Plugins/Experimental/GeometryProcessing/Source/MeshConversion/Public/MeshDescriptionBuilder.h b/Engine/Plugins/Experimental/GeometryProcessing/Source/MeshConversion/Public/MeshDescriptionBuilder.h
--- a/Engine/Plugins/Experimental/GeometryProcessing/Source/MeshConversion/Public/MeshDescriptionBuilder.h
+++ b/Engine/Plugins/Experimental/GeometryProcessing/Source/MeshConversion/Public/MeshDescriptionBuilder.h
@@ -64,7 +64,7 @@ public:
 	/** Set the Normal of a vertex instance*/
 	void SetInstanceNormal(const FVertexInstanceID& InstanceID, const FVector& Normal);
 
-	/** Set the full tangent space of a vertex instance, this is stored as a vec3 normal and a vec4 tangent */
+	/** Set the full tangent space of a vertex instance, this is stored as a vec3 normal, vec3 tangent, and a bitangent sign */
 	void SetInstanceTangentSpace(const FVertexInstanceID& InstanceID, const FVector& Normal, const FVector& Tangent, float Sign);
 
 	/** Set the Color of a vertex instance*/
