MeshModelingTools: Changing edge loop insertion and group edge insertion to use undo like PolyEdit does, both to fix invisibility bug and to prep for eventual movement into PolyEdit. Also made loop insertion highlight non-quad groups when a loop is stopped.

#rnx

diff --git a/Engine/Plugins/Experimental/GeometryProcessing/Source/DynamicMesh/Private/DynamicMesh3.cpp b/Engine/Plugins/Experimental/GeometryProcessing/Source/DynamicMesh/Private/DynamicMesh3.cpp
--- a/Engine/Plugins/Experimental/GeometryProcessing/Source/DynamicMesh/Private/DynamicMesh3.cpp
+++ b/Engine/Plugins/Experimental/GeometryProcessing/Source/DynamicMesh/Private/DynamicMesh3.cpp
@@ -225,12 +225,19 @@ void FDynamicMesh3::Copy(const FDynamicMesh3& copy, bool bNormals, bool bColors,
 	Edges         = copy.Edges;
 	EdgeRefCounts = copy.EdgeRefCounts;
 
-	DiscardAttributes();
+	// Note that we populate our existing AttributeSet when possible rather than building a
+	// new one. For the common case of updating a mesh from its preview copy, it seems reasonable
+	// for a client to be able to hold on to the AttributeSet pointer without expecting it to
+	// get reset by the Copy() call.
 	if (bAttributes && copy.HasAttributes())
 	{
-		EnableAttributes();
+		EnableAttributes(); // does nothing if already enabled
 		AttributeSet->Copy(*copy.AttributeSet);
 	}
+	else
+	{
+		DiscardAttributes();
+	}
 
 	Timestamp = FMath::Max(Timestamp + 1, copy.Timestamp);
 	ShapeTimestamp = TopologyTimestamp = Timestamp;
diff --git a/Engine/Plugins/Experimental/GeometryProcessing/Source/DynamicMesh/Private/Operations/GroupEdgeInserter.cpp b/Engine/Plugins/Experimental/GeometryProcessing/Source/DynamicMesh/Private/Operations/GroupEdgeInserter.cpp
--- a/Engine/Plugins/Experimental/GeometryProcessing/Source/DynamicMesh/Private/Operations/GroupEdgeInserter.cpp
+++ b/Engine/Plugins/Experimental/GeometryProcessing/Source/DynamicMesh/Private/Operations/GroupEdgeInserter.cpp
@@ -26,15 +26,18 @@ using namespace UE::Geometry;
 namespace GroupEdgeInserterLocals {
 
 bool GetEdgeLoopOpposingEdgeAndCorner(const FGroupTopology& Topology, int32 GroupID, int32 GroupEdgeIDIn,
-	int32 CornerIDIn, int32& GroupEdgeIDOut, int32& CornerIDOut, int32& BoundaryIndexOut);
+	int32 CornerIDIn, int32& GroupEdgeIDOut, int32& CornerIDOut, int32& BoundaryIndexOut, 
+	FGroupEdgeInserter::FOptionalOutputParams& OptionalOut);
 bool InsertEdgeLoopEdgesInDirection(const FGroupEdgeInserter::FEdgeLoopInsertionParams& Params,
 	const TArray<FGroupEdgeInserter::FGroupEdgeSplitPoint>& StartEndpoints,
 	int32 NextGroupID, int32 NextEdgeID, int32 NextCornerID, int32 NextBoundaryID,
-	TSet<int32>& AlteredGroups, int32& NumInserted, TSet<int32>* NewEidsOut, FProgressCancel* Progress);
+	TSet<int32>& AlteredGroups, int32& NumInserted, FGroupEdgeInserter::FOptionalOutputParams& OptionalOut,
+	FProgressCancel* Progress);
 void InsertNewVertexEndpoints(
 	const FGroupEdgeInserter::FEdgeLoopInsertionParams& Params,
 	int32 GroupEdgeID, int32 StartCornerID,
-	TArray<FGroupEdgeInserter::FGroupEdgeSplitPoint>& EndPointsOut);
+	TArray<FGroupEdgeInserter::FGroupEdgeSplitPoint>& EndPointsOut, 
+	FGroupEdgeInserter::FOptionalOutputParams& OptionalOut);
 void ConvertProportionsToArcLengths(
 	const FGroupTopology& Topology, int32 GroupEdgeID,
 	const TArray<double>& ProportionsIn,
@@ -44,16 +47,20 @@ bool ConnectEndpoints(
 	const FGroupTopology::FGroupBoundary& GroupBoundary,
 	const TArray<FGroupEdgeInserter::FGroupEdgeSplitPoint>& StartPoints,
 	const TArray<FGroupEdgeInserter::FGroupEdgeSplitPoint>& EndPoints,
-	TSet<int32>* NewEidsOut, int32& NumGroupsCreated, FProgressCancel* Progress);
+	int32& NumGroupsCreated, FGroupEdgeInserter::FOptionalOutputParams& OptionalOut, 
+	FProgressCancel* Progress);
 bool ConnectMultipleUsingRetriangulation(
 	FDynamicMesh3& Mesh, const FGroupTopology& Topology, int32 GroupID,
 	const FGroupTopology::FGroupBoundary& GroupBoundary,
 	const TArray<FGroupEdgeInserter::FGroupEdgeSplitPoint>& StartPoints,
 	const TArray <FGroupEdgeInserter::FGroupEdgeSplitPoint>& EndPoints,
-	TSet<int32>* ConnectionEidsOut, int32& NumGroupsCreated, FProgressCancel* Progress);
+	int32& NumGroupsCreated, FGroupEdgeInserter::FOptionalOutputParams& OptionalOut,
+	FProgressCancel* Progress);
 bool DeleteGroupTrianglesAndGetLoop(FDynamicMesh3& Mesh, const FGroupTopology& Topology, int32 GroupID,
 	const FGroupTopology::FGroupBoundary& GroupBoundary, TArray<int32>& BoundaryVerticesOut,
-	TArray<FMeshRegionBoundaryLoops::VidOverlayMap<FVector2f>>& BoundaryVidUVMapsOut, FProgressCancel* Progress);
+	TArray<FMeshRegionBoundaryLoops::VidOverlayMap<FVector2f>>& BoundaryVidUVMapsOut, 
+	FGroupEdgeInserter::FOptionalOutputParams& OptionalOut,
+	FProgressCancel* Progress);
 void AppendInclusiveRangeWrapAround(const TArray<int32>& InputArray, TArray<int32>& OutputArray,
 	int32 StartIndex, int32 InclusiveEndIndex);
 bool RetriangulateLoop(FDynamicMesh3& Mesh, const TArray<int32>& LoopVertices,
@@ -64,13 +71,15 @@ bool ConnectMultipleUsingPlaneCut(FDynamicMesh3& Mesh,
 	const FGroupTopology::FGroupBoundary& GroupBoundary,
 	const TArray<FGroupEdgeInserter::FGroupEdgeSplitPoint>& StartPoints,
 	const TArray <FGroupEdgeInserter::FGroupEdgeSplitPoint>& EndPoints,
-	double VertexTolerance, TSet<int32>* ConnectionEidsOut,
-	int32& NumGroupsCreated, FProgressCancel* Progress);
+	double VertexTolerance, int32& NumGroupsCreated, 
+	FGroupEdgeInserter::FOptionalOutputParams& OptionalOut, FProgressCancel* Progress);
 bool EmbedPlaneCutPath(FDynamicMesh3& Mesh, int32 GroupID,
 	const FGroupEdgeInserter::FGroupEdgeSplitPoint& StartPoint,
 	const FGroupEdgeInserter::FGroupEdgeSplitPoint& EndPoint,
-	double VertexTolerance, TSet<int32>& PathEidsOut, FProgressCancel* Progress);
-bool CreateNewGroups(FDynamicMesh3& Mesh, TSet<int32>& PathEids, int32 OriginalGroupID, int32& NumGroupsCreated, FProgressCancel* Progress);
+	double VertexTolerance, TSet<int32>& PathEidsOut, 
+	TSet<int32>* ChangedTrisOut, FProgressCancel* Progress);
+bool CreateNewGroups(FDynamicMesh3& Mesh, TSet<int32>& PathEids, int32 OriginalGroupID, int32& NumGroupsCreated, 
+	FGroupEdgeInserter::FOptionalOutputParams& OptionalOut, FProgressCancel* Progress);
 bool GetPlaneCutPath(const FDynamicMesh3& Mesh, int32 GroupID,
 	const FGroupEdgeInserter::FGroupEdgeSplitPoint& StartPoint, const FGroupEdgeInserter::FGroupEdgeSplitPoint& EndPoint,
 	TArray<TPair<FMeshSurfacePoint, int>>& OutputPath, double VertexCutTolerance, FProgressCancel* Progress);
@@ -79,12 +88,12 @@ bool InsertSingleWithRetriangulation(FDynamicMesh3& Mesh, FGroupTopology& Topolo
 	int32 GroupID, int32 BoundaryIndex,
 	const FGroupEdgeInserter::FGroupEdgeSplitPoint& StartPoint,
 	const FGroupEdgeInserter::FGroupEdgeSplitPoint& EndPoint,
-	TSet<int32>* NewEidsOut, FProgressCancel* Progress);
+	FGroupEdgeInserter::FOptionalOutputParams& OptionalOut, FProgressCancel* Progress);
 }
 using namespace GroupEdgeInserterLocals;
 
 /** Inserts an edge loop into a mesh, where an edge loop is a sequence of (group) edges across quads. */
-bool FGroupEdgeInserter::InsertEdgeLoops(const FEdgeLoopInsertionParams& Params, TSet<int32>* NewEidsOut, FProgressCancel* Progress)
+bool FGroupEdgeInserter::InsertEdgeLoops(const FEdgeLoopInsertionParams& Params, FOptionalOutputParams OptionalOut, FProgressCancel* Progress)
 {
 	if (Progress && Progress->Cancelled())
 	{
@@ -105,12 +114,12 @@ bool FGroupEdgeInserter::InsertEdgeLoops(const FEdgeLoopInsertionParams& Params,
 	int32 ForwardGroupID = GroupEdge.Groups.A;
 	int32 ForwardEdgeID, ForwardCornerID, ForwardBoundaryIndex;
 	bool bHaveForwardEdge = GetEdgeLoopOpposingEdgeAndCorner(*Params.Topology, ForwardGroupID,
-		Params.GroupEdgeID, Params.StartCornerID, ForwardEdgeID, ForwardCornerID, ForwardBoundaryIndex);
+		Params.GroupEdgeID, Params.StartCornerID, ForwardEdgeID, ForwardCornerID, ForwardBoundaryIndex, OptionalOut);
 
 	int32 BackwardGroupID = GroupEdge.Groups.B;
 	int32 BackwardEdgeID, BackwardCornerID, BackwardBoundaryIndex;
 	bool bHaveBackwardEdge = GetEdgeLoopOpposingEdgeAndCorner(*Params.Topology, BackwardGroupID,
-		Params.GroupEdgeID, Params.StartCornerID, BackwardEdgeID, BackwardCornerID, BackwardBoundaryIndex);
+		Params.GroupEdgeID, Params.StartCornerID, BackwardEdgeID, BackwardCornerID, BackwardBoundaryIndex, OptionalOut);
 
 	if (!bHaveForwardEdge && !bHaveBackwardEdge)
 	{
@@ -130,7 +139,7 @@ bool FGroupEdgeInserter::InsertEdgeLoops(const FEdgeLoopInsertionParams& Params,
 	// Although we have code that can insert new edges at edge endpoints, it is cleaner to do splits for all the loops
 	// down an edge ahead of time to make vertex endpoints, in part because if we don't, a split can change the eid
 	// of the next endpoint.
-	InsertNewVertexEndpoints(Params, Params.GroupEdgeID, Params.StartCornerID, StartEndpoints);
+	InsertNewVertexEndpoints(Params, Params.GroupEdgeID, Params.StartCornerID, StartEndpoints, OptionalOut);
 
 	if (StartEndpoints.Num() == 0 || (Progress && Progress->Cancelled()))
 	{
@@ -144,13 +153,13 @@ bool FGroupEdgeInserter::InsertEdgeLoops(const FEdgeLoopInsertionParams& Params,
 	if (bHaveForwardEdge)
 	{
 		bSuccess = InsertEdgeLoopEdgesInDirection(Params, StartEndpoints, ForwardGroupID, ForwardEdgeID,
-			ForwardCornerID, ForwardBoundaryIndex, AlteredGroups, TotalNumInserted, NewEidsOut, Progress);
+			ForwardCornerID, ForwardBoundaryIndex, AlteredGroups, TotalNumInserted, OptionalOut, Progress);
 	}
 	if (bHaveBackwardEdge)
 	{
 		int32 NumInserted = 0;
 		bSuccess = InsertEdgeLoopEdgesInDirection(Params, StartEndpoints, BackwardGroupID, BackwardEdgeID,
-			BackwardCornerID, BackwardBoundaryIndex, AlteredGroups, NumInserted, NewEidsOut, Progress) && bSuccess;
+			BackwardCornerID, BackwardBoundaryIndex, AlteredGroups, NumInserted, OptionalOut, Progress) && bSuccess;
 		TotalNumInserted += NumInserted;
 	}
 
@@ -169,10 +178,14 @@ namespace GroupEdgeInserterLocals {
  * finds the id of the opposite (i.e., destination) group edge. Additionally, gives the corner ID  attached
  * to the provided one. Safe to call with an FDynamicMesh3::InvalidID parameters (will return false)
  *
+ * If OptionalOut.ProblemGroupEdgeIDsOut is not null, will add the boundaries of a non-quad component to
+ * that set if the component turns out not to be valid due to not being quad-like.
+ * 
  * @returns true if a satisfactory edge was found.
  */
 bool GetEdgeLoopOpposingEdgeAndCorner(const FGroupTopology& Topology, int32 GroupID, int32 GroupEdgeIDIn,
-	int32 CornerIDIn, int32& GroupEdgeIDOut, int32& CornerIDOut, int32& BoundaryIndexOut)
+	int32 CornerIDIn, int32& GroupEdgeIDOut, int32& CornerIDOut, int32& BoundaryIndexOut, 
+	FGroupEdgeInserter::FOptionalOutputParams& OptionalOut)
 {
 	GroupEdgeIDOut = FDynamicMesh3::InvalidID;
 	CornerIDOut = FDynamicMesh3::InvalidID;
@@ -193,6 +206,10 @@ bool GetEdgeLoopOpposingEdgeAndCorner(const FGroupTopology& Topology, int32 Grou
 		{
 			if (Boundary.GroupEdges.Num() != 4)
 			{
+				if (OptionalOut.ProblemGroupEdgeIDsOut)
+				{
+					OptionalOut.ProblemGroupEdgeIDsOut->Append(Boundary.GroupEdges);
+				}
 				return false;
 			}
 
@@ -235,7 +252,8 @@ bool GetEdgeLoopOpposingEdgeAndCorner(const FGroupTopology& Topology, int32 Grou
 bool InsertEdgeLoopEdgesInDirection(const FGroupEdgeInserter::FEdgeLoopInsertionParams& Params,
 	const TArray<FGroupEdgeInserter::FGroupEdgeSplitPoint>& StartEndpoints,
 	int32 NextGroupID, int32 NextEdgeID, int32 NextCornerID, int32 NextBoundaryIndex,
-	TSet<int32>& AlteredGroups, int32& NumInserted, TSet<int32>* NewEidsOut, FProgressCancel* Progress)
+	TSet<int32>& AlteredGroups, int32& NumInserted, FGroupEdgeInserter::FOptionalOutputParams& OptionalOut,
+	FProgressCancel* Progress)
 {
 	NumInserted = 0;
 	if (AlteredGroups.Contains(NextGroupID) || StartEndpoints.Num() == 0)
@@ -267,14 +285,14 @@ bool InsertEdgeLoopEdgesInDirection(const FGroupEdgeInserter::FEdgeLoopInsertion
 		{
 			int32 NumGroupsCreated;
 			bSuccess = ConnectEndpoints(Params, NextGroupID, Boundary, *CurrentEndpoints, StartEndpoints, 
-				NewEidsOut, NumGroupsCreated, Progress);
+				NumGroupsCreated, OptionalOut, Progress);
 			AlteredGroups.Add(NextGroupID);
 			NumInserted += (NumGroupsCreated > 1 ? 1 : 0);
 			break;
 		}
 
 		// Otherwise, create next endpoints
-		InsertNewVertexEndpoints(Params, NextEdgeID, NextCornerID, *NextEndpoints);
+		InsertNewVertexEndpoints(Params, NextEdgeID, NextCornerID, *NextEndpoints, OptionalOut);
 
 		if (NextEndpoints->Num() == 0)
 		{
@@ -286,7 +304,7 @@ bool InsertEdgeLoopEdgesInDirection(const FGroupEdgeInserter::FEdgeLoopInsertion
 		// Connect up the endpoints
 		int32 NumGroupsCreated;
 		bSuccess = ConnectEndpoints(Params, NextGroupID, Boundary, *CurrentEndpoints, *NextEndpoints, 
-			NewEidsOut, NumGroupsCreated, Progress);
+			NumGroupsCreated, OptionalOut, Progress);
 		AlteredGroups.Add(NextGroupID);
 		NumInserted += (NumGroupsCreated > 1 ? 1 : 0);
 
@@ -302,7 +320,7 @@ bool InsertEdgeLoopEdgesInDirection(const FGroupEdgeInserter::FEdgeLoopInsertion
 		}
 		NextGroupID = Params.Topology->Edges[NextEdgeID].OtherGroupID(NextGroupID);
 		bHaveNextGroup = GetEdgeLoopOpposingEdgeAndCorner(*Params.Topology, NextGroupID, NextEdgeID, NextCornerID,
-			NextEdgeID, NextCornerID, NextBoundaryIndex); // outputs
+			NextEdgeID, NextCornerID, NextBoundaryIndex, OptionalOut); // outputs
 
 		Swap(CurrentEndpoints, NextEndpoints);
 	}
@@ -323,7 +341,8 @@ bool InsertEdgeLoopEdgesInDirection(const FGroupEdgeInserter::FEdgeLoopInsertion
 void InsertNewVertexEndpoints(
 	const FGroupEdgeInserter::FEdgeLoopInsertionParams& Params,
 	int32 GroupEdgeID, int32 StartCornerID,
-	TArray<FGroupEdgeInserter::FGroupEdgeSplitPoint>& EndPointsOut)
+	TArray<FGroupEdgeInserter::FGroupEdgeSplitPoint>& EndPointsOut,
+	FGroupEdgeInserter::FOptionalOutputParams& OptionalOut)
 {
 	EndPointsOut.Reset();
 	if (Params.SortedInputLengths->Num() == 0)
@@ -446,6 +465,16 @@ void InsertNewVertexEndpoints(
 			FDynamicMesh3::FEdgeSplitInfo EdgeSplitInfo;
 			Params.Mesh->SplitEdge(CurrentEid, EdgeSplitInfo, SplitT);
 
+			// Record the changed triangles
+			if (OptionalOut.ChangedTidsOut)
+			{
+				OptionalOut.ChangedTidsOut->Add(EdgeSplitInfo.OriginalTriangles.A);
+				if (EdgeSplitInfo.OriginalTriangles.B != FDynamicMesh3::InvalidID)
+				{
+					OptionalOut.ChangedTidsOut->Add(EdgeSplitInfo.OriginalTriangles.B);
+				}
+			}
+
 			// Update our position given that we're at a new vertex.
 			CurrentVid = EdgeSplitInfo.NewVertex;
 			CurrentArcLength = TargetLength;
@@ -485,18 +514,19 @@ bool ConnectEndpoints(
 	const FGroupTopology::FGroupBoundary& GroupBoundary,
 	const TArray<FGroupEdgeInserter::FGroupEdgeSplitPoint>& StartPoints,
 	const TArray<FGroupEdgeInserter::FGroupEdgeSplitPoint>& EndPoints,
-	TSet<int32>* NewEidsOut, int32& NumGroupsCreated, FProgressCancel* Progress)
+	int32& NumGroupsCreated, FGroupEdgeInserter::FOptionalOutputParams& OptionalOut, 
+	FProgressCancel* Progress)
 {
 	if (Params.Mode == FGroupEdgeInserter::EInsertionMode::Retriangulate)
 	{
 		return ConnectMultipleUsingRetriangulation(*Params.Mesh, *Params.Topology, GroupID, 
-			GroupBoundary, StartPoints, EndPoints, NewEidsOut, NumGroupsCreated, Progress);
+			GroupBoundary, StartPoints, EndPoints, NumGroupsCreated, OptionalOut, Progress);
 	}
 	else if (Params.Mode == FGroupEdgeInserter::EInsertionMode::PlaneCut)
 	{
 		return ConnectMultipleUsingPlaneCut(*Params.Mesh, *Params.Topology, GroupID,
-			GroupBoundary, StartPoints, EndPoints, Params.VertexTolerance, NewEidsOut, 
-			NumGroupsCreated, Progress);
+			GroupBoundary, StartPoints, EndPoints, Params.VertexTolerance,
+			NumGroupsCreated, OptionalOut, Progress);
 	}
 	else
 	{
@@ -529,7 +559,8 @@ bool ConnectMultipleUsingRetriangulation(
 	const FGroupTopology::FGroupBoundary& GroupBoundary,
 	const TArray<FGroupEdgeInserter::FGroupEdgeSplitPoint>& StartPoints,
 	const TArray <FGroupEdgeInserter::FGroupEdgeSplitPoint>& EndPoints,
-	TSet<int32>* ConnectionEidsOut, int32& NumGroupsCreated, FProgressCancel* Progress)
+	int32& NumGroupsCreated, FGroupEdgeInserter::FOptionalOutputParams& OptionalOut,
+	FProgressCancel* Progress)
 {
 	NumGroupsCreated = 0;
 	if (Progress && Progress->Cancelled())
@@ -546,7 +577,7 @@ bool ConnectMultipleUsingRetriangulation(
 	TArray<int32> BoundaryVertices;
 	TArray<FMeshRegionBoundaryLoops::VidOverlayMap<FVector2f>> VidUVMaps;
 	bool bSuccess = DeleteGroupTrianglesAndGetLoop(Mesh, Topology, GroupID, GroupBoundary, 
-		BoundaryVertices, VidUVMaps, Progress);
+		BoundaryVertices, VidUVMaps, OptionalOut, Progress);
 
 	if (!bSuccess || (Progress && Progress->Cancelled()))
 	{
@@ -658,11 +689,11 @@ bool ConnectMultipleUsingRetriangulation(
 		NumGroupsCreated += (bSuccess ? 1 : 0);
 	}
 
-	if (ConnectionEidsOut)
+	if (OptionalOut.NewEidsOut)
 	{
 		for (int32 i = 0; i < NumNewEdges; ++i)
 		{
-			ConnectionEidsOut->Add(Mesh.FindEdge(StartPoints[i].ElementID, EndPoints[i].ElementID));
+			OptionalOut.NewEidsOut->Add(Mesh.FindEdge(StartPoints[i].ElementID, EndPoints[i].ElementID));
 		}
 	}
 
@@ -675,7 +706,8 @@ bool ConnectMultipleUsingRetriangulation(
  */
 bool DeleteGroupTrianglesAndGetLoop(FDynamicMesh3& Mesh, const FGroupTopology& Topology, int32 GroupID,
 	const FGroupTopology::FGroupBoundary& GroupBoundary, TArray<int32>& BoundaryVerticesOut,
-	TArray<FMeshRegionBoundaryLoops::VidOverlayMap<FVector2f>>& BoundaryVidUVMapsOut,	FProgressCancel* Progress)
+	TArray<FMeshRegionBoundaryLoops::VidOverlayMap<FVector2f>>& BoundaryVidUVMapsOut, 
+	FGroupEdgeInserter::FOptionalOutputParams& OptionalOut, FProgressCancel* Progress)
 {
 	if (Progress && Progress->Cancelled())
 	{
@@ -739,6 +771,11 @@ bool DeleteGroupTrianglesAndGetLoop(FDynamicMesh3& Mesh, const FGroupTopology& T
 		return false;
 	}
 
+	if (OptionalOut.ChangedTidsOut)
+	{
+		OptionalOut.ChangedTidsOut->Append(Component.Indices);
+	}
+
 	// When deleting, we don't we don't want to remove isolated verts on the boundary,
 	// but we do want to remove isolated verts on the interior of the component. We 
 	// could finish the retriangulation and look for isolated verts afterwards, but
@@ -823,8 +860,9 @@ bool ConnectMultipleUsingPlaneCut(FDynamicMesh3& Mesh,
 	const FGroupTopology::FGroupBoundary& GroupBoundary,
 	const TArray<FGroupEdgeInserter::FGroupEdgeSplitPoint>& StartPoints,
 	const TArray <FGroupEdgeInserter::FGroupEdgeSplitPoint>& EndPoints,
-	double VertexTolerance, TSet<int32>* ConnectionEidsOut, 
-	int32& NumGroupsCreated, FProgressCancel* Progress)
+	double VertexTolerance, int32& NumGroupsCreated, 
+	FGroupEdgeInserter::FOptionalOutputParams& OptionalOut, 
+	FProgressCancel* Progress)
 {
 	NumGroupsCreated = 0;
 	if (Progress && Progress->Cancelled())
@@ -837,7 +875,7 @@ bool ConnectMultipleUsingPlaneCut(FDynamicMesh3& Mesh,
 	for (int32 i = 0; i < NumEdgesToInsert; ++i)
 	{
 		bool bSuccess = EmbedPlaneCutPath(Mesh, GroupID, StartPoints[i], EndPoints[i], 
-			VertexTolerance, PathsEids, Progress);
+			VertexTolerance, PathsEids, OptionalOut.ChangedTidsOut, Progress);
 
 		if (!bSuccess || (Progress && Progress->Cancelled()))
 		{
@@ -845,12 +883,12 @@ bool ConnectMultipleUsingPlaneCut(FDynamicMesh3& Mesh,
 		}
 	}
 
-	if (ConnectionEidsOut)
+	if (OptionalOut.NewEidsOut)
 	{
-		*ConnectionEidsOut = ConnectionEidsOut->Union(PathsEids);
+		*OptionalOut.NewEidsOut = OptionalOut.NewEidsOut->Union(PathsEids);
 	}
 
-	bool bSuccess = CreateNewGroups(Mesh, PathsEids, GroupID, NumGroupsCreated, Progress);
+	bool bSuccess = CreateNewGroups(Mesh, PathsEids, GroupID, NumGroupsCreated, OptionalOut, Progress);
 	return bSuccess;
 }
 
@@ -861,7 +899,8 @@ bool ConnectMultipleUsingPlaneCut(FDynamicMesh3& Mesh,
 bool EmbedPlaneCutPath(FDynamicMesh3& Mesh, int32 GroupID,
 	const FGroupEdgeInserter::FGroupEdgeSplitPoint& StartPoint,
 	const FGroupEdgeInserter::FGroupEdgeSplitPoint& EndPoint,
-	double VertexTolerance, TSet<int32>& PathEidsOut, FProgressCancel* Progress)
+	double VertexTolerance, TSet<int32>& PathEidsOut, 
+	TSet<int32>* ChangedTrisOut, FProgressCancel* Progress)
 {
 	if (Progress && Progress->Cancelled())
 	{
@@ -878,6 +917,24 @@ bool EmbedPlaneCutPath(FDynamicMesh3& Mesh, int32 GroupID,
 	}
 	check(CutPath.Num() >= 2);
 
+	// Save triangles that will change
+	if (ChangedTrisOut)
+	{
+		for (int32 i = 0; i < CutPath.Num(); ++i)
+		{
+			FMeshSurfacePoint& Point = CutPath[i].Key;
+			if (Point.PointType == ESurfacePointType::Edge)
+			{
+				FIndex2i EdgeTris = Mesh.GetEdgeT(Point.ElementID);
+				ChangedTrisOut->Add(EdgeTris.A);
+				if (EdgeTris.B != FDynamicMesh3::InvalidID)
+				{
+					ChangedTrisOut->Add(EdgeTris.B);
+				}
+			}
+		}
+	}
+
 	// Embed the path.
 	FMeshSurfacePath PathEmbedder(&Mesh);
 	PathEmbedder.Path = CutPath;
@@ -900,7 +957,8 @@ bool EmbedPlaneCutPath(FDynamicMesh3& Mesh, int32 GroupID,
 /**
  * Uses the given path edge ID's to split a group into new groups.
  */
-bool CreateNewGroups(FDynamicMesh3& Mesh, TSet<int32>& PathEids, int32 OriginalGroup, int32& NumGroupsCreated, FProgressCancel* Progress)
+bool CreateNewGroups(FDynamicMesh3& Mesh, TSet<int32>& PathEids, int32 OriginalGroup, int32& NumGroupsCreated, 
+	FGroupEdgeInserter::FOptionalOutputParams& OptionalOut, FProgressCancel* Progress)
 {
 	NumGroupsCreated = 0;
 
@@ -949,6 +1007,11 @@ bool CreateNewGroups(FDynamicMesh3& Mesh, TSet<int32>& PathEids, int32 OriginalG
 	{
 		FMeshConnectedComponents::FComponent& Component = ConnectedComponents.GetComponent(i);
 
+		if (OptionalOut.ChangedTidsOut)
+		{
+			OptionalOut.ChangedTidsOut->Append(Component.Indices);
+		}
+
 		int32 NewGroupID = Mesh.AllocateTriangleGroup();
 		for (int Tid : Component.Indices)
 		{
@@ -963,7 +1026,7 @@ bool CreateNewGroups(FDynamicMesh3& Mesh, TSet<int32>& PathEids, int32 OriginalG
 }//end namespace GroupEdgeInserterLocals
 
 /** Inserts a group edge into a given group. */
-bool FGroupEdgeInserter::InsertGroupEdge(FGroupEdgeInsertionParams& Params, TSet<int32>* NewEidsOut, FProgressCancel* Progress)
+bool FGroupEdgeInserter::InsertGroupEdge(FGroupEdgeInsertionParams& Params, FGroupEdgeInserter::FOptionalOutputParams OptionalOut, FProgressCancel* Progress)
 {
 	if (Progress && Progress->Cancelled())
 	{
@@ -987,17 +1050,17 @@ bool FGroupEdgeInserter::InsertGroupEdge(FGroupEdgeInsertionParams& Params, TSet
 	if (Params.Mode == EInsertionMode::PlaneCut)
 	{
 		TSet<int32> TempNewEids;
-		TSet<int32>* NewEids = NewEidsOut ? NewEidsOut : &TempNewEids;
+		TSet<int32>* NewEids = OptionalOut.NewEidsOut ? OptionalOut.NewEidsOut : &TempNewEids;
 
-		bool bSuccess = EmbedPlaneCutPath(*Params.Mesh, Params.GroupID, 
-			Params.StartPoint, Params.EndPoint, Params.VertexTolerance, *NewEids, Progress);
+		bool bSuccess = EmbedPlaneCutPath(*Params.Mesh, Params.GroupID, Params.StartPoint,
+			Params.EndPoint, Params.VertexTolerance, *NewEids, OptionalOut.ChangedTidsOut, Progress);
 		if (!bSuccess || (Progress && Progress->Cancelled()))
 		{
 			return false;
 		}
 		
 		int32 NumGroupsCreated;
-		bSuccess = CreateNewGroups(*Params.Mesh, *NewEids, Params.GroupID, NumGroupsCreated, Progress);
+		bSuccess = CreateNewGroups(*Params.Mesh, *NewEids, Params.GroupID, NumGroupsCreated, OptionalOut, Progress);
 		if (!bSuccess || (Progress && Progress->Cancelled()))
 		{
 			return false;
@@ -1007,7 +1070,7 @@ bool FGroupEdgeInserter::InsertGroupEdge(FGroupEdgeInsertionParams& Params, TSet
 	{
 		bool bSuccess = InsertSingleWithRetriangulation(*Params.Mesh, *Params.Topology, 
 			Params.GroupID, Params.GroupBoundaryIndex, Params.StartPoint, Params.EndPoint, 
-			NewEidsOut, Progress);
+			OptionalOut, Progress);
 		if (!bSuccess || (Progress && Progress->Cancelled()))
 		{
 			return false;
@@ -1030,7 +1093,7 @@ bool InsertSingleWithRetriangulation(FDynamicMesh3& Mesh, FGroupTopology& Topolo
 	int32 GroupID, int32 BoundaryIndex,
 	const FGroupEdgeInserter::FGroupEdgeSplitPoint& StartPoint,
 	const FGroupEdgeInserter::FGroupEdgeSplitPoint& EndPoint, 
-	TSet<int32>* NewEidsOut, FProgressCancel* Progress)
+	FGroupEdgeInserter::FOptionalOutputParams& OptionalOut, FProgressCancel* Progress)
 {
 	if (Progress && Progress->Cancelled())
 	{
@@ -1065,7 +1128,7 @@ bool InsertSingleWithRetriangulation(FDynamicMesh3& Mesh, FGroupTopology& Topolo
 	TArray<int32> BoundaryVertices;
 	TArray<FMeshRegionBoundaryLoops::VidOverlayMap<FVector2f>> VidUVMaps;
 	bool bSuccess = DeleteGroupTrianglesAndGetLoop(Mesh, Topology, GroupID, Group->Boundaries[BoundaryIndex], 
-		BoundaryVertices, VidUVMaps, Progress);
+		BoundaryVertices, VidUVMaps, OptionalOut, Progress);
 	if (!bSuccess || (Progress && Progress->Cancelled()))
 	{
 		return false;
@@ -1097,9 +1160,9 @@ bool InsertSingleWithRetriangulation(FDynamicMesh3& Mesh, FGroupTopology& Topolo
 	}
 	bSuccess = RetriangulateLoop(Mesh, LoopVids, Mesh.AllocateTriangleGroup(), VidUVMaps);
 
-	if (NewEidsOut)
+	if (OptionalOut.NewEidsOut)
 	{
-		NewEidsOut->Add(Mesh.FindEdge(StartVid, EndVid));
+		OptionalOut.NewEidsOut->Add(Mesh.FindEdge(StartVid, EndVid));
 	}
 
 	return bSuccess;
diff --git a/Engine/Plugins/Experimental/GeometryProcessing/Source/DynamicMesh/Public/DynamicMesh3.h b/Engine/Plugins/Experimental/GeometryProcessing/Source/DynamicMesh/Public/DynamicMesh3.h
--- a/Engine/Plugins/Experimental/GeometryProcessing/Source/DynamicMesh/Public/DynamicMesh3.h
+++ b/Engine/Plugins/Experimental/GeometryProcessing/Source/DynamicMesh/Public/DynamicMesh3.h
@@ -202,7 +202,11 @@ public:
 	/** Initialize mesh from the output of a MeshShapeGenerator (assumes Generate() was already called) */
 	void Copy(const FMeshShapeGenerator* Generator);
 
-	/** Copy input mesh while compacting, i.e. removing unused vertices/triangles/edges */
+	/** 
+	 * Copy input mesh while compacting, i.e. removing unused vertices/triangles/edges. 
+	 * Note that it is currently unsafe to hold on to the attribute set pointer (returned from Attributes())
+	 * across the CompactCopy call, as the latter may rebuild an entirely new attribute set.
+	 */
 	void CompactCopy(const FDynamicMesh3& CopyMesh, bool bNormals = true, bool bColors = true, bool bUVs = true,
 	                 bool bAttributes = true, FCompactMaps* CompactInfo = nullptr);
 
diff --git a/Engine/Plugins/Experimental/GeometryProcessing/Source/DynamicMesh/Public/Operations/GroupEdgeInserter.h b/Engine/Plugins/Experimental/GeometryProcessing/Source/DynamicMesh/Public/Operations/GroupEdgeInserter.h
--- a/Engine/Plugins/Experimental/GeometryProcessing/Source/DynamicMesh/Public/Operations/GroupEdgeInserter.h
+++ b/Engine/Plugins/Experimental/GeometryProcessing/Source/DynamicMesh/Public/Operations/GroupEdgeInserter.h
@@ -66,7 +66,31 @@ public:
 		EInsertionMode Mode = EInsertionMode::Retriangulate;
 	};
 
-	bool InsertEdgeLoops(const FEdgeLoopInsertionParams& Params, TSet<int32>* NewEidsOut = nullptr, FProgressCancel* Progress = nullptr);
+	struct FOptionalOutputParams
+	{
+		// Some compilers (clang, at least) have issues with using a nested class with default initializers
+		// and a default constructor as default argument in a function, hence this constructor.
+		FOptionalOutputParams(){};
+
+		/**
+		 * Edge IDs of the edges composing the newly inserted group edges.
+		 */
+		TSet<int32>* NewEidsOut = nullptr;
+		
+		/** 
+		 * Any triangle IDs whose triangles were deleted or changed by the operation (but not newly
+		 * created tids). Useful for setting up undo. 
+		 */
+		TSet<int32>* ChangedTidsOut = nullptr;
+
+		/**
+		 * In loop insertion, the group edge IDs in the original topology that surround non-quad-like
+		 * groups that stopped the loop.
+		 */
+		TSet<int32>* ProblemGroupEdgeIDsOut = nullptr;
+	};
+
+	bool InsertEdgeLoops(const FEdgeLoopInsertionParams& Params, FOptionalOutputParams OptionalOut = FOptionalOutputParams(), FProgressCancel* Progress = nullptr);
 
 
 	/** Point along a group edge that is used as a start/endpoint for an inserted group edge. */
@@ -117,7 +141,7 @@ public:
 		EInsertionMode Mode = EInsertionMode::Retriangulate;
 	};
 
-	bool InsertGroupEdge(FGroupEdgeInsertionParams& Params, TSet<int32>* NewEidsOut = nullptr, FProgressCancel* Progress = nullptr);
+	bool InsertGroupEdge(FGroupEdgeInsertionParams& Params, FOptionalOutputParams OptionalOut = FOptionalOutputParams(), FProgressCancel* Progress = nullptr);
 
 };
 
diff --git a/Engine/Plugins/Experimental/MeshModelingToolset/Source/MeshModelingTools/Private/EdgeLoopInsertionTool.cpp b/Engine/Plugins/Experimental/MeshModelingToolset/Source/MeshModelingTools/Private/EdgeLoopInsertionTool.cpp
--- a/Engine/Plugins/Experimental/MeshModelingToolset/Source/MeshModelingTools/Private/EdgeLoopInsertionTool.cpp
+++ b/Engine/Plugins/Experimental/MeshModelingToolset/Source/MeshModelingTools/Private/EdgeLoopInsertionTool.cpp
@@ -6,6 +6,7 @@
 #include "BaseBehaviors/MouseHoverBehavior.h"
 #include "CuttingOps/EdgeLoopInsertionOp.h"
 #include "DynamicMeshToMeshDescription.h"
+#include "DynamicMeshChangeTracker.h"
 #include "InteractiveToolManager.h"
 #include "MeshDescriptionToDynamicMesh.h"
 #include "Operations/GroupEdgeInserter.h"
@@ -102,10 +103,10 @@ void UEdgeLoopInsertionTool::Setup()
 		EToolMessageLevel::UserNotification);
 
 	// Initialize the mesh that we'll be operating on
-	CurrentMesh = MakeShared<FDynamicMesh3>();
+	CurrentMesh = MakeShared<FDynamicMesh3, ESPMode::ThreadSafe>();
 	FMeshDescriptionToDynamicMesh Converter;
 	Converter.Convert(Cast<IMeshDescriptionProvider>(Target)->GetMeshDescription(), *CurrentMesh);
-	CurrentTopology = MakeShared<FGroupTopology>(CurrentMesh.Get(), true);
+	CurrentTopology = MakeShared<FGroupTopology, ESPMode::ThreadSafe>(CurrentMesh.Get(), true);
 	MeshSpatial.SetMesh(CurrentMesh.Get(), true);
 
 	// Set up properties
@@ -123,12 +124,19 @@ void UEdgeLoopInsertionTool::Setup()
 
 	SetupPreview();
 
-	// These draw the group edges and the loops to be inserted
+	// Draws the old group topology
 	ExistingEdgesRenderer.LineColor = FLinearColor::Red;
 	ExistingEdgesRenderer.LineThickness = 2.0;
+
+	// Draws the new group edges that are added
 	PreviewEdgeRenderer.LineColor = FLinearColor::Green;
 	PreviewEdgeRenderer.LineThickness = 2.0;
 
+	// Highlights non-quad groups that stop the loop;
+	ProblemTopologyRenderer.LineColor = FLinearColor::Red;
+	ProblemTopologyRenderer.LineThickness = 3.0;
+	ProblemTopologyRenderer.DepthBias = 1.0;
+
 	// Set up the topology selector, which we use to select the edges where we insert the loops
 	TopologySelector.Initialize(CurrentMesh.Get(), CurrentTopology.Get());
 	TopologySelector.SetSpatialSource([this]() {return &MeshSpatial; });
@@ -156,20 +164,43 @@ void UEdgeLoopInsertionTool::SetupPreview()
 	Preview->ConfigureMaterials(MaterialSet.Materials, ToolSetupUtil::GetDefaultWorkingMaterial(GetToolManager()));
 
 	// Whenever we get a new result from the op, we need to extract the preview edges so that
-	// we can draw them if we want to.
+	// we can draw them if we want to, and the additional outputs we need (changed triangles and
+	// topology).
 	Preview->OnOpCompleted.AddLambda([this](const FDynamicMeshOperator* UncastOp) {
 		const FEdgeLoopInsertionOp* Op = static_cast<const FEdgeLoopInsertionOp*>(UncastOp);
 
 		bLastComputeSucceeded = Op->bSucceeded;
 		LatestOpTopologyResult.Reset();
 		PreviewEdges.Reset();
+		LatestOpChangedTids.Reset();
 
 		if (bLastComputeSucceeded)
 		{
 			Op->GetLoopEdgeLocations(PreviewEdges);
 			LatestOpTopologyResult = Op->ResultTopology;
+			LatestOpChangedTids = Op->ChangedTids;
+		}
+
+		// Regardless of success, extract things for highlighting any non-quads that stopped our loop.
+		ProblemTopologyEdges.Reset();
+		ProblemTopologyVerts.Reset();
+		for (int32 GroupEdgeID : Op->ProblemGroupEdgeIDs)
+		{
+			for (int32 Eid : CurrentTopology->GetGroupEdgeEdges(GroupEdgeID))
+			{
+				TPair<FVector3d, FVector3d> Endpoints;
+				CurrentMesh->GetEdgeV(Eid, Endpoints.Key, Endpoints.Value);
+				ProblemTopologyEdges.Add(MoveTemp(Endpoints));
+			}
+			FGroupTopology::FGroupEdge& GroupEdge = CurrentTopology->Edges[GroupEdgeID];
+			if (GroupEdge.EndpointCorners.A != FDynamicMesh3::InvalidID)
+			{
+				ProblemTopologyVerts.AddUnique(CurrentMesh->GetVertex(CurrentTopology->Corners[GroupEdge.EndpointCorners.A].VertexID));
+				ProblemTopologyVerts.AddUnique(CurrentMesh->GetVertex(CurrentTopology->Corners[GroupEdge.EndpointCorners.B].VertexID));
+			}
 		}
 		});
+
 	// In case of failure, we want to hide the broken preview, since we wouldn't accept it on
 	// a click. Note that this can't be fired OnOpCompleted because the preview is updated
 	// with the op result after that callback, which would undo the reset. The preview edge
@@ -194,9 +225,22 @@ void UEdgeLoopInsertionTool::SetupPreview()
 
 void UEdgeLoopInsertionTool::Shutdown(EToolShutdownType ShutdownType)
 {
+	// Set visibility before committing so that it doesn't get saved as false.
+	Cast<IPrimitiveComponentBackedTarget>(Target)->SetOwnerVisibility(true);
+
+	if (ShutdownType == EToolShutdownType::Accept)
+	{
+		GetToolManager()->BeginUndoTransaction(LOCTEXT("EdgeLoopInsertionToolTransactionName", "Edge Loop Tool"));
+		Cast<IMeshDescriptionCommitter>(Target)->CommitMeshDescription([this](const IMeshDescriptionCommitter::FCommitterParams& CommitParams)
+		{
+			FDynamicMeshToMeshDescription Converter;
+			Converter.Convert(CurrentMesh.Get(), *CommitParams.MeshDescriptionOut);
+		});
+		GetToolManager()->EndUndoTransaction();
+	}
+
 	Settings->SaveProperties(this);
 	Preview->Shutdown();
-	Cast<IPrimitiveComponentBackedTarget>(Target)->SetOwnerVisibility(true);
 	CurrentMesh.Reset();
 	CurrentTopology.Reset();
 	ExpireChanges();
@@ -212,18 +256,9 @@ void UEdgeLoopInsertionTool::OnTick(float DeltaTime)
 		{
 			if (bLastComputeSucceeded)
 			{
-				// Apply the insertion
-				GetToolManager()->BeginUndoTransaction(LOCTEXT("EdgeLoopInsertionTransactionName", "Edge Loop Insertion"));
-
-				GetToolManager()->EmitObjectChange(this, MakeUnique<FEdgeLoopInsertionChangeBookend>(CurrentChangeStamp, true), LOCTEXT("EdgeLoopInsertion", "Edge Loop Insertion"));
-				Cast<IMeshDescriptionCommitter>(Target)->CommitMeshDescription([this](const IMeshDescriptionCommitter::FCommitterParams& CommitParams)
-					{
-						FDynamicMeshToMeshDescription Converter;
-						Converter.Convert(Preview->PreviewMesh->GetMesh(), *CommitParams.MeshDescriptionOut);
-					});
-				GetToolManager()->EmitObjectChange(this, MakeUnique<FEdgeLoopInsertionChangeBookend>(CurrentChangeStamp, false), LOCTEXT("EdgeLoopInsertion", "Edge Loop Insertion"));
-
-				GetToolManager()->EndUndoTransaction();
+				FDynamicMeshChangeTracker ChangeTracker(CurrentMesh.Get());
+				ChangeTracker.BeginChange();
+				ChangeTracker.SaveTriangles(*LatestOpChangedTids, true /*bSaveVertices*/);
 
 				// Update current mesh and topology
 				CurrentMesh->Copy(*Preview->PreviewMesh->GetMesh(), true, true, true, true);
@@ -231,9 +266,18 @@ void UEdgeLoopInsertionTool::OnTick(float DeltaTime)
 				CurrentTopology->RetargetOnClonedMesh(CurrentMesh.Get());
 				MeshSpatial.Build();
 				TopologySelector.Invalidate(true, true);
+
+				// Emit transaction
+				GetToolManager()->BeginUndoTransaction(LOCTEXT("EdgeLoopInsertionTransactionName", "Edge Loop Insertion"));
+				GetToolManager()->EmitObjectChange(this, MakeUnique<FEdgeLoopInsertionChange>(ChangeTracker.EndChange(), CurrentChangeStamp), 
+					LOCTEXT("EdgeLoopInsertion", "Edge Loop Insertion"));
+				GetToolManager()->EndUndoTransaction();
 			}
 
 			PreviewEdges.Reset();
+			ProblemTopologyEdges.Reset();
+			ProblemTopologyVerts.Reset();
+
 			bWaitingForInsertionCompletion = false;
 		}
 	}
@@ -267,6 +311,22 @@ void UEdgeLoopInsertionTool::Render(IToolsContextRenderAPI* RenderAPI)
 		PreviewEdgeRenderer.DrawLine(EdgeVerts.Key, EdgeVerts.Value);
 	}
 	PreviewEdgeRenderer.EndFrame();
+
+	if (Settings->bHighlightProblemGroups)
+	{
+		// Highlight any non-quad groups that stopped the loop.
+		ProblemTopologyRenderer.BeginFrame(RenderAPI, RenderCameraState);
+		ProblemTopologyRenderer.SetTransform(Preview->PreviewMesh->GetTransform());
+		for (TPair<FVector3d, FVector3d>& EdgeVerts : ProblemTopologyEdges)
+		{
+			ProblemTopologyRenderer.DrawLine(EdgeVerts.Key, EdgeVerts.Value);
+		}
+		for (FVector3d& Vert : ProblemTopologyVerts)
+		{
+			ProblemTopologyRenderer.DrawViewFacingX(Vert, ProblemVertTickWidth);
+		}
+		ProblemTopologyRenderer.EndFrame();
+	}
 }
 
 bool UEdgeLoopInsertionTool::CanAccept() const
@@ -487,34 +547,24 @@ void UEdgeLoopInsertionTool::OnClicked(const FInputDeviceRay& ClickPos)
 
 // Undo/redo support
 
-void FEdgeLoopInsertionChangeBookend::Apply(UObject* Object)
+void FEdgeLoopInsertionChange::Apply(UObject* Object)
 {
-	if (!bBeforeChange)
-	{
-		UEdgeLoopInsertionTool* Tool = Cast<UEdgeLoopInsertionTool>(Object);
-		Tool->CurrentMesh->Clear();
-		FMeshDescriptionToDynamicMesh Converter;
-		Converter.Convert(Cast<IMeshDescriptionProvider>(Tool->Target)->GetMeshDescription(), *Tool->CurrentMesh);
-		Tool->CurrentTopology->RebuildTopology();
-		Tool->MeshSpatial.Build();
-		Tool->TopologySelector.Invalidate(true, true);
-		Tool->ClearPreview();
-	}
+	UEdgeLoopInsertionTool* Tool = Cast<UEdgeLoopInsertionTool>(Object);
+	MeshChange->Apply(Tool->CurrentMesh.Get(), false);
+	Tool->MeshSpatial.Build();
+	Tool->TopologySelector.Invalidate(true, true);
+	Tool->CurrentTopology->RebuildTopology();
+	Tool->ClearPreview();
 }
 
-void FEdgeLoopInsertionChangeBookend::Revert(UObject* Object)
+void FEdgeLoopInsertionChange::Revert(UObject* Object)
 {
-	if (bBeforeChange)
-	{
-		UEdgeLoopInsertionTool* Tool = Cast<UEdgeLoopInsertionTool>(Object);
-		Tool->CurrentMesh->Clear();
-		FMeshDescriptionToDynamicMesh Converter;
-		Converter.Convert(Cast<IMeshDescriptionProvider>(Tool->Target)->GetMeshDescription(), *Tool->CurrentMesh);
-		Tool->CurrentTopology->RebuildTopology();
-		Tool->MeshSpatial.Build();
-		Tool->TopologySelector.Invalidate(true, true);
-		Tool->ClearPreview();
-	}
+	UEdgeLoopInsertionTool* Tool = Cast<UEdgeLoopInsertionTool>(Object);
+	MeshChange->Apply(Tool->CurrentMesh.Get(), true);
+	Tool->MeshSpatial.Build();
+	Tool->TopologySelector.Invalidate(true, true);
+	Tool->CurrentTopology->RebuildTopology();
+	Tool->ClearPreview();
 }
 
 #undef LOCTEXT_NAMESPACE
\ No newline at end of file
diff --git a/Engine/Plugins/Experimental/MeshModelingToolset/Source/MeshModelingTools/Private/GroupEdgeInsertionTool.cpp b/Engine/Plugins/Experimental/MeshModelingToolset/Source/MeshModelingTools/Private/GroupEdgeInsertionTool.cpp
--- a/Engine/Plugins/Experimental/MeshModelingToolset/Source/MeshModelingTools/Private/GroupEdgeInsertionTool.cpp
+++ b/Engine/Plugins/Experimental/MeshModelingToolset/Source/MeshModelingTools/Private/GroupEdgeInsertionTool.cpp
@@ -5,6 +5,7 @@
 #include "BaseBehaviors/SingleClickBehavior.h"
 #include "BaseBehaviors/MouseHoverBehavior.h"
 #include "CuttingOps/GroupEdgeInsertionOp.h"
+#include "DynamicMeshChangeTracker.h"
 #include "DynamicMeshToMeshDescription.h"
 #include "InteractiveToolManager.h"
 #include "MeshDescriptionToDynamicMesh.h"
@@ -145,12 +146,14 @@ void UGroupEdgeInsertionTool::SetupPreview()
 
 		bLastComputeSucceeded = Op->bSucceeded;
 		LatestOpTopologyResult.Reset();
+		LatestOpChangedTids.Reset();
 
 		PreviewEdges.Reset();
 		if (bLastComputeSucceeded)
 		{
 			Op->GetEdgeLocations(PreviewEdges);
 			LatestOpTopologyResult = Op->ResultTopology;
+			LatestOpChangedTids = Op->ChangedTids;
 		}
 		else
 		{
@@ -180,9 +183,22 @@ void UGroupEdgeInsertionTool::SetupPreview()
 
 void UGroupEdgeInsertionTool::Shutdown(EToolShutdownType ShutdownType)
 {
+	// Set visibility before committing so that it doesn't get saved as false.
+	Cast<IPrimitiveComponentBackedTarget>(Target)->SetOwnerVisibility(true);
+
+	if (ShutdownType == EToolShutdownType::Accept)
+	{
+		GetToolManager()->BeginUndoTransaction(LOCTEXT("GroupEdgeInsertionToolTransactionName", "Group Edge Insert Tool"));
+		Cast<IMeshDescriptionCommitter>(Target)->CommitMeshDescription([this](const IMeshDescriptionCommitter::FCommitterParams& CommitParams)
+		{
+			FDynamicMeshToMeshDescription Converter;
+			Converter.Convert(CurrentMesh.Get(), *CommitParams.MeshDescriptionOut);
+		});
+		GetToolManager()->EndUndoTransaction();
+	}
+
 	Settings->SaveProperties(this);
 	Preview->Shutdown();
-	Cast<IPrimitiveComponentBackedTarget>(Target)->SetOwnerVisibility(true);
 	CurrentMesh.Reset();
 	CurrentTopology.Reset();
 	ExpireChanges();
@@ -198,18 +214,9 @@ void UGroupEdgeInsertionTool::OnTick(float DeltaTime)
 		{
 			if (bLastComputeSucceeded)
 			{
-				// Apply the insertion
-				GetToolManager()->BeginUndoTransaction(LOCTEXT("GroupEdgeInsertionTransactionName", "Group Edge Insertion"));
-
-				GetToolManager()->EmitObjectChange(this, MakeUnique<FGroupEdgeInsertionChangeBookend>(CurrentChangeStamp, true), LOCTEXT("GroupEdgeInsertion", "Group Edge Insertion"));
-				Cast<IMeshDescriptionCommitter>(Target)->CommitMeshDescription([this](const IMeshDescriptionCommitter::FCommitterParams& CommitParams)
-					{
-						FDynamicMeshToMeshDescription Converter;
-						Converter.Convert(Preview->PreviewMesh->GetMesh(), *CommitParams.MeshDescriptionOut);
-					});
-				GetToolManager()->EmitObjectChange(this, MakeUnique<FGroupEdgeInsertionChangeBookend>(CurrentChangeStamp, false), LOCTEXT("GroupEdgeInsertion", "Group Edge Insertion"));
-
-				GetToolManager()->EndUndoTransaction();
+				FDynamicMeshChangeTracker ChangeTracker(CurrentMesh.Get());
+				ChangeTracker.BeginChange();
+				ChangeTracker.SaveTriangles(*LatestOpChangedTids, true /*bSaveVertices*/);
 
 				// Update current mesh and topology
 				CurrentMesh->Copy(*Preview->PreviewMesh->GetMesh(), true, true, true, true);
@@ -218,6 +225,12 @@ void UGroupEdgeInsertionTool::OnTick(float DeltaTime)
 				MeshSpatial.Build();
 				TopologySelector.Invalidate(true, true);
 
+				// Emit transaction
+				GetToolManager()->BeginUndoTransaction(LOCTEXT("GroupEdgeInsertionTransactionName", "Group Edge Insertion"));
+				GetToolManager()->EmitObjectChange(this, MakeUnique<FGroupEdgeInsertionChange>(ChangeTracker.EndChange(), CurrentChangeStamp),
+					LOCTEXT("GroupEdgeInsertion", "GroupEdge Insertion"));
+				GetToolManager()->EndUndoTransaction();
+
 				ToolState = EToolState::GettingStart;
 			}
 			else
@@ -816,47 +829,24 @@ void FGroupEdgeInsertionFirstPointChange::Revert(UObject* Object)
 	bHaveDoneUndo = true;
 }
 
-void FGroupEdgeInsertionChangeBookend::Revert(UObject* Object)
+void FGroupEdgeInsertionChange::Apply(UObject* Object)
 {
-	if (bBeforeChange)
-	{
-		// Load from the component, which has been updated
-		UGroupEdgeInsertionTool* Tool = Cast<UGroupEdgeInsertionTool>(Object);
-		Tool->CurrentMesh->Clear();
-		FMeshDescriptionToDynamicMesh Converter;
-		Converter.Convert(Cast<IMeshDescriptionProvider>(Tool->Target)->GetMeshDescription(), *Tool->CurrentMesh);
-		Tool->CurrentTopology->RebuildTopology();
-		Tool->MeshSpatial.Build();
-		Tool->TopologySelector.Invalidate(true, true);
-
-		Tool->ClearPreview(false, true);
-
-		check(Tool->ToolState == UGroupEdgeInsertionTool::EToolState::GettingStart);
-		// If we were doing full undo/redo of the start point insertions, instead of just
-		// letting the user back out of their latest one, then we would set the state
-		// to GettingEnd here. Instead we go all the way back to GettingStart.
-	}
+	UGroupEdgeInsertionTool* Tool = Cast<UGroupEdgeInsertionTool>(Object);
+	MeshChange->Apply(Tool->CurrentMesh.Get(), false);
+	Tool->MeshSpatial.Build();
+	Tool->TopologySelector.Invalidate(true, true);
+	Tool->CurrentTopology->RebuildTopology();
+	Tool->ClearPreview(true, true);
 }
 
-void FGroupEdgeInsertionChangeBookend::Apply(UObject* Object)
+void FGroupEdgeInsertionChange::Revert(UObject* Object)
 {
-	if (!bBeforeChange)
-	{
-		// Load from the component, which has been updated
-		UGroupEdgeInsertionTool* Tool = Cast<UGroupEdgeInsertionTool>(Object);
-		Tool->CurrentMesh->Clear();
-		FMeshDescriptionToDynamicMesh Converter;
-		Converter.Convert(Cast<IMeshDescriptionProvider>(Tool->Target)->GetMeshDescription(), *Tool->CurrentMesh);
-		Tool->CurrentTopology->RebuildTopology();
-		Tool->MeshSpatial.Build();
-		Tool->TopologySelector.Invalidate(true, true);
-
-		Tool->ClearPreview();
-
-		// Since we always go all the way back to GettingStart on revert, we would expect
-		// this to be the state on redo.
-		check(Tool->ToolState == UGroupEdgeInsertionTool::EToolState::GettingStart);
-	}
+	UGroupEdgeInsertionTool* Tool = Cast<UGroupEdgeInsertionTool>(Object);
+	MeshChange->Apply(Tool->CurrentMesh.Get(), true);
+	Tool->MeshSpatial.Build();
+	Tool->TopologySelector.Invalidate(true, true);
+	Tool->CurrentTopology->RebuildTopology();
+	Tool->ClearPreview(true, true);
 }
 
 #undef LOCTEXT_NAMESPACE
\ No newline at end of file
diff --git a/Engine/Plugins/Experimental/MeshModelingToolset/Source/MeshModelingTools/Public/EdgeLoopInsertionTool.h b/Engine/Plugins/Experimental/MeshModelingToolset/Source/MeshModelingTools/Public/EdgeLoopInsertionTool.h
--- a/Engine/Plugins/Experimental/MeshModelingToolset/Source/MeshModelingTools/Public/EdgeLoopInsertionTool.h
+++ b/Engine/Plugins/Experimental/MeshModelingToolset/Source/MeshModelingTools/Public/EdgeLoopInsertionTool.h
@@ -17,7 +17,7 @@
 
 #include "EdgeLoopInsertionTool.generated.h"
 
-class UDynamicMeshReplacementChangeTarget;
+PREDECLARE_USE_GEOMETRY_CLASS(FDynamicMeshChange);
 
 UCLASS()
 class MESHMODELINGTOOLS_API UEdgeLoopInsertionToolBuilder : public USingleSelectionMeshEditingToolBuilder
@@ -93,6 +93,10 @@ public:
 	UPROPERTY(EditAnywhere, Category = InsertEdgeLoop)
 	bool bWireframe = true;
 
+	/** When true, non-quad-like groups that stop the loop will be highlighted, with X's marking the corners. */
+	UPROPERTY(EditAnywhere, Category = InsertEdgeLoop)
+	bool bHighlightProblemGroups = true;
+
 	/** How close a new loop edge needs to pass next to an existing vertex to use that vertex rather than creating a new one. */
 	UPROPERTY(EditAnywhere, Category = InsertEdgeLoop, AdvancedDisplay)
 	double VertexTolerance = 0.001;
@@ -119,7 +123,7 @@ class MESHMODELINGTOOLS_API UEdgeLoopInsertionTool : public USingleSelectionMesh
 public:
 
 	friend class UEdgeLoopInsertionOperatorFactory;
-	friend class FEdgeLoopInsertionChangeBookend;
+	friend class FEdgeLoopInsertionChange;
 
 	UEdgeLoopInsertionTool() {};
 
@@ -150,13 +154,17 @@ protected:
 	UPROPERTY()
 	UEdgeLoopInsertionProperties* Settings = nullptr;
 
-	TSharedPtr<FDynamicMesh3> CurrentMesh;
-	TSharedPtr<FGroupTopology> CurrentTopology;
+	TSharedPtr<FDynamicMesh3, ESPMode::ThreadSafe> CurrentMesh;
+	TSharedPtr<FGroupTopology, ESPMode::ThreadSafe> CurrentTopology;
 	UE::Geometry::FDynamicMeshAABBTree3 MeshSpatial;
 	FGroupTopologySelector TopologySelector;
 
 	TArray<TPair<FVector3d, FVector3d>> PreviewEdges;
 
+	// Used to highlight problematic topology (non-quad groups) when it stops a loop.
+	TArray<TPair<FVector3d, FVector3d>> ProblemTopologyEdges;
+	TArray<FVector3d> ProblemTopologyVerts;
+
 	FViewCameraState CameraState;
 
 	UPROPERTY()
@@ -164,7 +172,9 @@ protected:
 
 	FToolDataVisualizer ExistingEdgesRenderer;
 	FToolDataVisualizer PreviewEdgeRenderer;
+	FToolDataVisualizer ProblemTopologyRenderer;
 	FGroupTopologySelector::FSelectionSettings TopologySelectorSettings;
+	float ProblemVertTickWidth = 8;
 
 	void SetupPreview();
 
@@ -187,7 +197,8 @@ protected:
 
 	// Copied over on op completion
 	bool bLastComputeSucceeded = false;
-	TSharedPtr<FGroupTopology> LatestOpTopologyResult;
+	TSharedPtr<FGroupTopology, ESPMode::ThreadSafe> LatestOpTopologyResult;
+	TSharedPtr<TSet<int32>, ESPMode::ThreadSafe> LatestOpChangedTids;
 
 	// Used to expire undo/redo changes on op shutdown.
 	int32 CurrentChangeStamp = 0;
@@ -205,16 +216,15 @@ protected:
 };
 
 /**
- * This change object is a bit of a hack: if it is emitted on both sides of the associated
- * ComponentTarget change, it will reload the current mesh and topology from the target
- * on Undo/Redo, thereby propagating it to the tool.
+ * Wraps a FDynamicMeshChange so that it can be expired and so that other data
+ * structures in the tool can be updated.
  */
-class MESHMODELINGTOOLS_API FEdgeLoopInsertionChangeBookend : public FToolCommandChange
+class MESHMODELINGTOOLS_API FEdgeLoopInsertionChange : public FToolCommandChange
 {
 public:
-	FEdgeLoopInsertionChangeBookend(int32 CurrentChangeStamp, bool bBeforeChangeIn)
-		: ChangeStamp(CurrentChangeStamp) 
-		, bBeforeChange(bBeforeChangeIn)
+	FEdgeLoopInsertionChange(TUniquePtr<UE::Geometry::FDynamicMeshChange> MeshChangeIn, int32 CurrentChangeStamp)
+		: MeshChange(MoveTemp(MeshChangeIn))
+		, ChangeStamp(CurrentChangeStamp)
 	{};
 
 	virtual void Apply(UObject* Object) override;
@@ -225,10 +235,10 @@ public:
 	}
 	virtual FString ToString() const override
 	{
-		return TEXT("FEdgeLoopInsertionChangeBookend");
+		return TEXT("FEdgeLoopInsertionChange");
 	}
 
 protected:
+	TUniquePtr<UE::Geometry::FDynamicMeshChange> MeshChange;
 	int32 ChangeStamp;
-	bool bBeforeChange;
 };
\ No newline at end of file
diff --git a/Engine/Plugins/Experimental/MeshModelingToolset/Source/MeshModelingTools/Public/GroupEdgeInsertionTool.h b/Engine/Plugins/Experimental/MeshModelingToolset/Source/MeshModelingTools/Public/GroupEdgeInsertionTool.h
--- a/Engine/Plugins/Experimental/MeshModelingToolset/Source/MeshModelingTools/Public/GroupEdgeInsertionTool.h
+++ b/Engine/Plugins/Experimental/MeshModelingToolset/Source/MeshModelingTools/Public/GroupEdgeInsertionTool.h
@@ -18,7 +18,7 @@
 
 #include "GroupEdgeInsertionTool.generated.h"
 
-class UDynamicMeshReplacementChangeTarget;
+PREDECLARE_USE_GEOMETRY_CLASS(FDynamicMeshChange);
 
 using UE::Geometry::FGroupEdgeInserter;
 using UE::Geometry::FDynamicMesh3;
@@ -92,7 +92,7 @@ public:
 
 	friend class UGroupEdgeInsertionOperatorFactory;
 	friend class FGroupEdgeInsertionFirstPointChange;
-	friend class FGroupEdgeInsertionChangeBookend;
+	friend class FGroupEdgeInsertionChange;
 
 	UGroupEdgeInsertionTool() {};
 
@@ -155,9 +155,11 @@ protected:
 
 	// State control:
 	EToolState ToolState = EToolState::GettingStart;
+
 	bool bShowingBaseMesh = false;
 	bool bLastComputeSucceeded = false;
 	TSharedPtr<FGroupTopology, ESPMode::ThreadSafe> LatestOpTopologyResult;
+	TSharedPtr<TSet<int32>, ESPMode::ThreadSafe> LatestOpChangedTids;
 
 	int32 CurrentChangeStamp = 0;
 
@@ -216,17 +218,15 @@ protected:
 };
 
 /**
- * This should get emitted on either side of the ComponentTarget that occurs when a second
- * point is successfully picked so that the tool can reload the current mesh from the changed
- * target.
- * TODO: This is a hack similar to FEdgeLoopInsertionChangeBookend but it works. Is there a cleaner way?
+ * Wraps a FDynamicMeshChange so that it can be expired and so that other data
+ * structures in the tool can be updated.
  */
-class MESHMODELINGTOOLS_API FGroupEdgeInsertionChangeBookend : public FToolCommandChange
+class MESHMODELINGTOOLS_API FGroupEdgeInsertionChange : public FToolCommandChange
 {
 public:
-	FGroupEdgeInsertionChangeBookend(int32 CurrentChangeStamp, bool bBeforeChangeIn)
-		: ChangeStamp(CurrentChangeStamp)
-		, bBeforeChange(bBeforeChangeIn)
+	FGroupEdgeInsertionChange(TUniquePtr<FDynamicMeshChange> MeshChangeIn, int32 CurrentChangeStamp)
+		: MeshChange(MoveTemp(MeshChangeIn))
+		, ChangeStamp(CurrentChangeStamp)
 	{};
 
 	virtual void Apply(UObject* Object) override;
@@ -237,11 +237,10 @@ public:
 	}
 	virtual FString ToString() const override
 	{
-		return TEXT("FGroupEdgeInsertionChangeBookend");
+		return TEXT("FGroupEdgeInsertionChange");
 	}
 
 protected:
+	TUniquePtr<FDynamicMeshChange> MeshChange;
 	int32 ChangeStamp;
-	bool bBeforeChange;
 };
-
diff --git a/Engine/Plugins/Experimental/MeshModelingToolset/Source/ModelingOperators/Private/CuttingOps/EdgeLoopInsertionOp.cpp b/Engine/Plugins/Experimental/MeshModelingToolset/Source/ModelingOperators/Private/CuttingOps/EdgeLoopInsertionOp.cpp
--- a/Engine/Plugins/Experimental/MeshModelingToolset/Source/ModelingOperators/Private/CuttingOps/EdgeLoopInsertionOp.cpp
+++ b/Engine/Plugins/Experimental/MeshModelingToolset/Source/ModelingOperators/Private/CuttingOps/EdgeLoopInsertionOp.cpp
@@ -32,9 +32,12 @@ void FEdgeLoopInsertionOp::CalculateResult(FProgressCancel* Progress)
 	}
 
 	ResultMesh->Copy(*OriginalMesh, true, true, true, true);
-	ResultTopology = MakeShared<FGroupTopology>();
+	ResultTopology = MakeShared<FGroupTopology, ESPMode::ThreadSafe>();
 	*ResultTopology = *OriginalTopology;
 	ResultTopology->RetargetOnClonedMesh(ResultMesh.Get());
+	ChangedTids = MakeShared<TSet<int32>, ESPMode::ThreadSafe>();
+	LoopEids.Reset();
+	ProblemGroupEdgeIDs.Reset();
 
 	if (Progress && Progress->Cancelled())
 	{
@@ -54,6 +57,10 @@ void FEdgeLoopInsertionOp::CalculateResult(FProgressCancel* Progress)
 		Params.VertexTolerance = VertexTolerance;
 
 		FGroupEdgeInserter Inserter;
-		bSucceeded = Inserter.InsertEdgeLoops(Params, &LoopEids, Progress);
+		FGroupEdgeInserter::FOptionalOutputParams OptionalOut;
+		OptionalOut.NewEidsOut = &LoopEids;
+		OptionalOut.ChangedTidsOut = ChangedTids.Get();
+		OptionalOut.ProblemGroupEdgeIDsOut = &ProblemGroupEdgeIDs;
+		bSucceeded = Inserter.InsertEdgeLoops(Params, OptionalOut, Progress);
 	}
 }
diff --git a/Engine/Plugins/Experimental/MeshModelingToolset/Source/ModelingOperators/Private/CuttingOps/GroupEdgeInsertionOp.cpp b/Engine/Plugins/Experimental/MeshModelingToolset/Source/ModelingOperators/Private/CuttingOps/GroupEdgeInsertionOp.cpp
--- a/Engine/Plugins/Experimental/MeshModelingToolset/Source/ModelingOperators/Private/CuttingOps/GroupEdgeInsertionOp.cpp
+++ b/Engine/Plugins/Experimental/MeshModelingToolset/Source/ModelingOperators/Private/CuttingOps/GroupEdgeInsertionOp.cpp
@@ -36,6 +36,7 @@ void FGroupEdgeInsertionOp::CalculateResult(FProgressCancel* Progress)
 	ResultTopology = MakeShared<FGroupTopology, ESPMode::ThreadSafe>();
 	*ResultTopology = *OriginalTopology;
 	ResultTopology->RetargetOnClonedMesh(ResultMesh.Get());
+	ChangedTids = MakeShared<TSet<int32>, ESPMode::ThreadSafe>();
 
 	if (bShowingBaseMesh || (Progress && Progress->Cancelled()))
 	{
@@ -53,5 +54,9 @@ void FGroupEdgeInsertionOp::CalculateResult(FProgressCancel* Progress)
 	Params.GroupID = CommonGroupID;
 	Params.GroupBoundaryIndex = CommonBoundaryIndex;
 
-	bSucceeded = Inserter.InsertGroupEdge(Params, &Eids, Progress);
+	FGroupEdgeInserter::FOptionalOutputParams OptionalOut;
+	OptionalOut.NewEidsOut = &Eids;
+	OptionalOut.ChangedTidsOut = ChangedTids.Get();
+
+	bSucceeded = Inserter.InsertGroupEdge(Params, OptionalOut, Progress);
 }
diff --git a/Engine/Plugins/Experimental/MeshModelingToolset/Source/ModelingOperators/Public/CuttingOps/EdgeLoopInsertionOp.h b/Engine/Plugins/Experimental/MeshModelingToolset/Source/ModelingOperators/Public/CuttingOps/EdgeLoopInsertionOp.h
--- a/Engine/Plugins/Experimental/MeshModelingToolset/Source/ModelingOperators/Public/CuttingOps/EdgeLoopInsertionOp.h
+++ b/Engine/Plugins/Experimental/MeshModelingToolset/Source/ModelingOperators/Public/CuttingOps/EdgeLoopInsertionOp.h
@@ -19,8 +19,8 @@ public:
 	virtual ~FEdgeLoopInsertionOp() {}
 
 	// Inputs:
-	TSharedPtr<const FDynamicMesh3> OriginalMesh;
-	TSharedPtr<const FGroupTopology> OriginalTopology;
+	TSharedPtr<const FDynamicMesh3, ESPMode::ThreadSafe> OriginalMesh;
+	TSharedPtr<const FGroupTopology, ESPMode::ThreadSafe> OriginalTopology;
 	FGroupEdgeInserter::EInsertionMode Mode;
 	double VertexTolerance; // TODO: Add some defaults
 	int32 GroupEdgeID = FDynamicMesh3::InvalidID;
@@ -29,8 +29,17 @@ public:
 	int32 StartCornerID = FDynamicMesh3::InvalidID;
 
 	// Outputs:
-	TSet<int32> LoopEids; // These are edge ID's in the ResultMesh.
-	TSharedPtr<FGroupTopology> ResultTopology;
+	// Edge IDs in the ResultMesh corresponding to the loop.
+	TSet<int32> LoopEids; 
+
+	// IDs of triangles in the OriginalMesh that were changed or deleted.
+	TSharedPtr<TSet<int32>, ESPMode::ThreadSafe> ChangedTids;
+
+	// IDs of group edges in OriginalTopology that surround non-quad-like regions
+	// that stopped the loop.
+	TSet<int32> ProblemGroupEdgeIDs;
+
+	TSharedPtr<FGroupTopology, ESPMode::ThreadSafe> ResultTopology;
 	bool bSucceeded = false;
 
 	void SetTransform(const FTransform& Transform);
diff --git a/Engine/Plugins/Experimental/MeshModelingToolset/Source/ModelingOperators/Public/CuttingOps/GroupEdgeInsertionOp.h b/Engine/Plugins/Experimental/MeshModelingToolset/Source/ModelingOperators/Public/CuttingOps/GroupEdgeInsertionOp.h
--- a/Engine/Plugins/Experimental/MeshModelingToolset/Source/ModelingOperators/Public/CuttingOps/GroupEdgeInsertionOp.h
+++ b/Engine/Plugins/Experimental/MeshModelingToolset/Source/ModelingOperators/Public/CuttingOps/GroupEdgeInsertionOp.h
@@ -39,6 +39,9 @@ public:
 	// Edge ID's in the ResultMesh that together make up the new group edge.
 	TSet<int32> Eids; 
 
+	// Triangle ID's in the OriginalMesh that got changed or deleted in the operation.
+	TSharedPtr<TSet<int32>, ESPMode::ThreadSafe> ChangedTids;
+
 	TSharedPtr<FGroupTopology, ESPMode::ThreadSafe> ResultTopology;
 	bool bSucceeded = false;
 
diff --git a/Engine/Source/Runtime/Experimental/InteractiveToolsFramework/Private/ToolDataVisualizer.cpp b/Engine/Source/Runtime/Experimental/InteractiveToolsFramework/Private/ToolDataVisualizer.cpp
--- a/Engine/Source/Runtime/Experimental/InteractiveToolsFramework/Private/ToolDataVisualizer.cpp
+++ b/Engine/Source/Runtime/Experimental/InteractiveToolsFramework/Private/ToolDataVisualizer.cpp
@@ -194,3 +194,22 @@ void FToolDataVisualizer::InternalDrawViewFacingCircle(const FVector& Position,
 		(bDepthTestedIn) ? SDPG_World : SDPG_Foreground,
 		LineThicknessIn * PDISizeScale, DepthBias, true);
 }
+
+void FToolDataVisualizer::InternalDrawViewFacingX(const FVector& Position, float Width, const FLinearColor& Color, float LineThicknessIn, bool bDepthTestedIn)
+{
+	checkf(bHaveCameraState, TEXT("To call this function, you must first call the version of BeginFrame that takes the CameraState"));
+	
+	FVector WorldPosition = TransformP(Position);
+
+	FVector UpOffset = CameraState.Up() * Width / 2;
+	FVector RightOffset = CameraState.Right() * Width / 2;
+	InternalDrawTransformedLine(
+		WorldPosition - UpOffset - RightOffset, 
+		WorldPosition + UpOffset + RightOffset, 
+		Color, LineThicknessIn, bDepthTestedIn);
+
+	InternalDrawTransformedLine(
+		WorldPosition + UpOffset - RightOffset,
+		WorldPosition - UpOffset + RightOffset,
+		Color, LineThicknessIn, bDepthTestedIn);
+}
\ No newline at end of file
diff --git a/Engine/Source/Runtime/Experimental/InteractiveToolsFramework/Public/ToolDataVisualizer.h b/Engine/Source/Runtime/Experimental/InteractiveToolsFramework/Public/ToolDataVisualizer.h
--- a/Engine/Source/Runtime/Experimental/InteractiveToolsFramework/Public/ToolDataVisualizer.h
+++ b/Engine/Source/Runtime/Experimental/InteractiveToolsFramework/Public/ToolDataVisualizer.h
@@ -178,6 +178,20 @@ public:
 		InternalDrawViewFacingCircle((FVector)Position, Radius, Steps, Color, LineThicknessIn, bDepthTestedIn);
 	}
 
+	/** Draw a world-space X facing the viewer at the given position. */
+	template<typename PointType>
+	void DrawViewFacingX(const PointType& Position, float Width)
+	{
+		InternalDrawViewFacingX((FVector)Position, Width, LineColor, LineThickness, bDepthTested);
+	}
+
+	/** Draw a world-space X facing the viewer at the given position. */
+	template<typename PointType>
+	void DrawViewFacingX(const PointType& Position, float Width, const FLinearColor& Color, float LineThicknessIn, bool bDepthTestedIn)
+	{
+		InternalDrawViewFacingX((FVector)Position, Width, Color, LineThicknessIn, bDepthTestedIn);
+	}
+
 	/** Draw a 3D cylinder, parameterized the same as the 3D circle but extruded by Height */
 	template<typename PointType>
 	void DrawWireCylinder(const PointType& Position, const PointType& Normal, float Radius, float Height, int Steps, const FLinearColor& Color, float LineThicknessIn, bool bDepthTestedIn)
@@ -247,6 +261,7 @@ protected:
 	virtual void InternalDrawWireBox(const FBox& Box, const FLinearColor& Color, float LineThickness, bool bDepthTested);
 	// actually does the square drawing
 	virtual void InternalDrawSquare(const FVector& Center, const FVector& SideA, const FVector& SideB, const FLinearColor& Color, float LineThickness, bool bDepthTested);
-
+	// actually does the X drawing
+	virtual void InternalDrawViewFacingX(const FVector& Position, float Width, const FLinearColor& Color, float LineThickness, bool bDepthTested);
 
 };
\ No newline at end of file
