GeometryProcessing: replace UE::Geometry::FVector3d/FVector3f with UE Core versions based on TVector<T>. FVector3<T> is now an extension of TVector<T>. Update all affected call sites.
#rnx

diff --git a/Engine/Plugins/Experimental/Enterprise/DataprepGeometryOperations/Source/Private/DataprepGeometryOperations.cpp b/Engine/Plugins/Experimental/Enterprise/DataprepGeometryOperations/Source/Private/DataprepGeometryOperations.cpp
--- a/Engine/Plugins/Experimental/Enterprise/DataprepGeometryOperations/Source/Private/DataprepGeometryOperations.cpp
+++ b/Engine/Plugins/Experimental/Enterprise/DataprepGeometryOperations/Source/Private/DataprepGeometryOperations.cpp
@@ -788,15 +788,15 @@ TUniquePtr<UE::Geometry::FDynamicMeshOperator> UDataprepPlaneCutOperation::MakeN
 	FTransform WorldToLocal = LocalToWorld.Inverse();
 	FVector LocalOrigin = WorldToLocal.TransformPosition(CutPlaneOrigin);
 	UE::Geometry::FTransform3d W2LForNormal(WorldToLocal);
-	FVector LocalNormal = (FVector)W2LForNormal.TransformNormal(WorldNormal);
+	FVector LocalNormal = (FVector)W2LForNormal.TransformNormal((FVector3d)WorldNormal);
 	FVector BackTransformed = LocalToWorld.TransformVector(LocalNormal);
 	float NormalScaleFactor = FVector::DotProduct(BackTransformed, WorldNormal);
 	if (NormalScaleFactor >= FLT_MIN)
 	{
 		NormalScaleFactor = 1.0 / NormalScaleFactor;
 	}
-	CutOp->LocalPlaneOrigin = LocalOrigin;
-	CutOp->LocalPlaneNormal = LocalNormal;
+	CutOp->LocalPlaneOrigin = (FVector3d)LocalOrigin;
+	CutOp->LocalPlaneNormal = (FVector3d)LocalNormal;
 	CutOp->OriginalMesh = InOriginalMesh;
 	CutOp->bKeepBothHalves = (CutPlaneKeepSide == EPlaneCutKeepSide::Both);
 	CutOp->CutPlaneLocalThickness = SpacingBetweenHalves * NormalScaleFactor;
diff --git a/Engine/Plugins/Experimental/GeometryProcessing/Source/DynamicMesh/Private/DynamicMesh3.cpp b/Engine/Plugins/Experimental/GeometryProcessing/Source/DynamicMesh/Private/DynamicMesh3.cpp
--- a/Engine/Plugins/Experimental/GeometryProcessing/Source/DynamicMesh/Private/DynamicMesh3.cpp
+++ b/Engine/Plugins/Experimental/GeometryProcessing/Source/DynamicMesh/Private/DynamicMesh3.cpp
@@ -12,7 +12,7 @@ using namespace UE::Geometry;
 constexpr int       FDynamicMesh3::InvalidID;
 constexpr int       FDynamicMesh3::NonManifoldID;
 constexpr int       FDynamicMesh3::DuplicateTriangleID;
-constexpr FVector3d FDynamicMesh3::InvalidVertex;
+const FVector3d FDynamicMesh3::InvalidVertex = FVector3d(TNumericLimits<double>::Max(), 0.0, 0.0);
 constexpr FIndex3i  FDynamicMesh3::InvalidTriangle;
 constexpr FIndex2i  FDynamicMesh3::InvalidEdge;
 
diff --git a/Engine/Plugins/Experimental/GeometryProcessing/Source/DynamicMesh/Private/DynamicMesh3_Queries.cpp b/Engine/Plugins/Experimental/GeometryProcessing/Source/DynamicMesh/Private/DynamicMesh3_Queries.cpp
--- a/Engine/Plugins/Experimental/GeometryProcessing/Source/DynamicMesh/Private/DynamicMesh3_Queries.cpp
+++ b/Engine/Plugins/Experimental/GeometryProcessing/Source/DynamicMesh/Private/DynamicMesh3_Queries.cpp
@@ -859,7 +859,7 @@ FFrame3d FDynamicMesh3::GetTriFrame(int tID, int nEdge) const
 	FVector3d normal = edge2.Cross(edge1);
 	Normalize(normal);
 	FVector3d other  = normal.Cross(edge1);
-	FVector3d center = (TriVerts[0] + TriVerts[1] + TriVerts[2]) / 3;
+	FVector3d center = (TriVerts[0] + TriVerts[1] + TriVerts[2]) / 3.0;
 
 	return FFrame3d(center, edge1, other, normal);
 }
diff --git a/Engine/Plugins/Experimental/GeometryProcessing/Source/DynamicMesh/Private/DynamicMeshEditor.cpp b/Engine/Plugins/Experimental/GeometryProcessing/Source/DynamicMesh/Private/DynamicMeshEditor.cpp
--- a/Engine/Plugins/Experimental/GeometryProcessing/Source/DynamicMesh/Private/DynamicMeshEditor.cpp
+++ b/Engine/Plugins/Experimental/GeometryProcessing/Source/DynamicMesh/Private/DynamicMeshEditor.cpp
@@ -210,7 +210,7 @@ bool FDynamicMeshEditor::StitchSparselyCorrespondedVertexLoops(const TArray<int>
 			{
 				IndNext = (Ind + 1) % VertexIDs.Num();
 				FVector3d VNext = M->GetVertex(VertexIDs[IndNext]);
-				LenTotal += V.Distance(VNext);
+				LenTotal += Distance(V, VNext);
 				Ind = IndNext;
 				V = VNext;
 			}
@@ -238,7 +238,7 @@ bool FDynamicMeshEditor::StitchSparselyCorrespondedVertexLoops(const TArray<int>
 				
 				Tri.C = VertexIDs1[Walks[0]];
 				FVector3d NextV = Mesh->GetVertex(Tri.C);
-				LenAlong[0] += NextV.Distance(Vertex[0]);
+				LenAlong[0] += Distance(NextV, Vertex[0]);
 				Vertex[0] = NextV;
 			}
 			else
@@ -246,7 +246,7 @@ bool FDynamicMeshEditor::StitchSparselyCorrespondedVertexLoops(const TArray<int>
 				Walks[1] = (Walks[1] + 1) % VertexIDs2.Num();
 				Tri.C = VertexIDs2[Walks[1]];
 				FVector3d NextV = Mesh->GetVertex(Tri.C);
-				LenAlong[1] += NextV.Distance(Vertex[1]);
+				LenAlong[1] += Distance(NextV, Vertex[1]);
 				Vertex[1] = NextV;
 			}
 			if (bReverseOrientation)
@@ -1182,7 +1182,7 @@ void FDynamicMeshEditor::RescaleAttributeUVs(float UVScale, bool bWorldSpace, in
 			for (int j = 2, i = 0; i < 3; j = i++)
 			{
 				TotalEdgeUVLen += TriUVs[j].Distance(TriUVs[i]);
-				TotalEdgeLen += TriVs[j].Distance(TriVs[i]);
+				TotalEdgeLen += Distance(TriVs[j], TriVs[i]);
 			}
 		}
 		if (TotalEdgeUVLen > KINDA_SMALL_NUMBER)
diff --git a/Engine/Plugins/Experimental/GeometryProcessing/Source/DynamicMesh/Private/DynamicMeshOverlay.cpp b/Engine/Plugins/Experimental/GeometryProcessing/Source/DynamicMesh/Private/DynamicMeshOverlay.cpp
--- a/Engine/Plugins/Experimental/GeometryProcessing/Source/DynamicMesh/Private/DynamicMeshOverlay.cpp
+++ b/Engine/Plugins/Experimental/GeometryProcessing/Source/DynamicMesh/Private/DynamicMeshOverlay.cpp
@@ -840,7 +840,7 @@ void TDynamicMeshOverlay<RealType, ElementSize>::OnSplitEdge(const FDynamicMesh3
 
 		// create new element at lerp position
 		NewElemID = AppendElement((RealType)0);
-		SetElementFromLerp(NewElemID, Triangle0[idx_base_a1], Triangle0[idx_base_b1], (RealType)splitInfo.SplitT);
+		SetElementFromLerp(NewElemID, Triangle0[idx_base_a1], Triangle0[idx_base_b1], (double)splitInfo.SplitT);
 
 		// rewrite triangle 0
 		ElementTriangles[3 * orig_t0 + idx_base_b1] = NewElemID;
@@ -876,7 +876,7 @@ void TDynamicMeshOverlay<RealType, ElementSize>::OnSplitEdge(const FDynamicMesh3
 		{
 			// create new element at lerp position
 			OtherNewElemID = AppendElement((RealType)0);
-			SetElementFromLerp(OtherNewElemID, Triangle1[idx_base_a2], Triangle1[idx_base_b2], (RealType)splitInfo.SplitT);
+			SetElementFromLerp(OtherNewElemID, Triangle1[idx_base_a2], Triangle1[idx_base_b2], (double)splitInfo.SplitT);
 		}
 
 		// rewrite triangle 1
@@ -1138,7 +1138,7 @@ void TDynamicMeshOverlay<RealType, ElementSize>::OnCollapseEdge(const FDynamicMe
 			}
 
 			// update value of kept element
-			SetElementFromLerp(kept_elemid[i], kept_elemid[i], removed_elemid[i], collapseInfo.CollapseT);
+			SetElementFromLerp(kept_elemid[i], kept_elemid[i], removed_elemid[i], (double)collapseInfo.CollapseT);
 		}
 	}
 
@@ -1194,7 +1194,7 @@ void TDynamicMeshOverlay<RealType, ElementSize>::OnPokeTriangle(const FDynamicMe
 
 	// create new element at barycentric position
 	int CenterElemID = AppendElement((RealType)0);
-	FVector3<RealType> BaryCoords((RealType)PokeInfo.BaryCoords.X, (RealType)PokeInfo.BaryCoords.Y, (RealType)PokeInfo.BaryCoords.Z);
+	FVector3<double> BaryCoords((double)PokeInfo.BaryCoords.X, (double)PokeInfo.BaryCoords.Y, (double)PokeInfo.BaryCoords.Z);
 	SetElementFromBary(CenterElemID, Triangle[0], Triangle[1], Triangle[2], BaryCoords);
 
 	// update orig triangle and two new ones. Winding orders here mirror FDynamicMesh3::PokeTriangle
@@ -1280,20 +1280,21 @@ void TDynamicMeshOverlay<RealType, ElementSize>::OnSplitVertex(const DynamicMesh
 
 
 template<typename RealType, int ElementSize>
-void TDynamicMeshOverlay<RealType, ElementSize>::SetElementFromLerp(int SetElement, int ElementA, int ElementB, RealType Alpha)
+void TDynamicMeshOverlay<RealType, ElementSize>::SetElementFromLerp(int SetElement, int ElementA, int ElementB, double Alpha)
 {
 	int IndexSet = ElementSize * SetElement;
 	int IndexA = ElementSize * ElementA;
 	int IndexB = ElementSize * ElementB;
-	RealType Beta = ((RealType)1 - Alpha);
+	double Beta = ((double)1 - Alpha);
 	for (int i = 0; i < ElementSize; ++i)
 	{
-		Elements[IndexSet+i] = Beta*Elements[IndexA+i] + Alpha*Elements[IndexB+i];
+		double LerpValue =  Beta*(double)Elements[IndexA+i] + Alpha*(double)Elements[IndexB+i];
+		Elements[IndexSet+i] = (RealType)LerpValue;
 	}
 }
 
 template<typename RealType, int ElementSize>
-void TDynamicMeshOverlay<RealType, ElementSize>::SetElementFromBary(int SetElement, int ElementA, int ElementB, int ElementC, const FVector3<RealType>& BaryCoords)
+void TDynamicMeshOverlay<RealType, ElementSize>::SetElementFromBary(int SetElement, int ElementA, int ElementB, int ElementC, const FVector3<double>& BaryCoords)
 {
 	int IndexSet = ElementSize * SetElement;
 	int IndexA = ElementSize * ElementA;
@@ -1301,8 +1302,8 @@ void TDynamicMeshOverlay<RealType, ElementSize>::SetElementFromBary(int SetEleme
 	int IndexC = ElementSize * ElementC;
 	for (int i = 0; i < ElementSize; ++i)
 	{
-		Elements[IndexSet + i] = 
-			BaryCoords.X*Elements[IndexA+i] + BaryCoords.Y*Elements[IndexB+i] + BaryCoords.Z*Elements[IndexC+i];
+		double BaryValue = BaryCoords.X*(double)Elements[IndexA+i] + BaryCoords.Y*(double)Elements[IndexB+i] + BaryCoords.Z*(double)Elements[IndexC+i];
+		Elements[IndexSet + i] = (RealType)BaryValue;			
 	}
 }
 
diff --git a/Engine/Plugins/Experimental/GeometryProcessing/Source/DynamicMesh/Private/EdgeLoop.cpp b/Engine/Plugins/Experimental/GeometryProcessing/Source/DynamicMesh/Private/EdgeLoop.cpp
--- a/Engine/Plugins/Experimental/GeometryProcessing/Source/DynamicMesh/Private/EdgeLoop.cpp
+++ b/Engine/Plugins/Experimental/GeometryProcessing/Source/DynamicMesh/Private/EdgeLoop.cpp
@@ -93,16 +93,6 @@ FAxisAlignedBox3d FEdgeLoop::GetBounds() const
 }
 
 
-void FEdgeLoop::GetVertices(TArray<FVector3d>& VerticesIn) const
-{
-	int NumV = Vertices.Num();
-	for (int i = 0; i < NumV; ++i)
-	{
-		VerticesIn.Add(Mesh->GetVertex(Vertices[i]));
-	}
-}
-
-
 bool FEdgeLoop::SetCorrectOrientation()
 {
 	int NumEdges = Edges.Num();
@@ -180,7 +170,7 @@ int FEdgeLoop::FindNearestVertexIndex(const FVector3d& QueryPoint) const
 	for (int i = 0; i < N; ++i)
 	{
 		FVector3d lv = Mesh->GetVertex(Vertices[i]);
-		double d2 = QueryPoint.DistanceSquared(lv);
+		double d2 = DistanceSquared(QueryPoint, lv);
 		if (d2 < fNearSqr)
 		{
 			fNearSqr = d2;
diff --git a/Engine/Plugins/Experimental/GeometryProcessing/Source/DynamicMesh/Private/EdgeSpan.cpp b/Engine/Plugins/Experimental/GeometryProcessing/Source/DynamicMesh/Private/EdgeSpan.cpp
--- a/Engine/Plugins/Experimental/GeometryProcessing/Source/DynamicMesh/Private/EdgeSpan.cpp
+++ b/Engine/Plugins/Experimental/GeometryProcessing/Source/DynamicMesh/Private/EdgeSpan.cpp
@@ -195,7 +195,7 @@ int FEdgeSpan::FindNearestVertexIndex(const FVector3d& QueryPoint) const
 	for (int i = 0; i < N; ++i)
 	{
 		FVector3d lv = Mesh->GetVertex(Vertices[i]);
-		double d2 = QueryPoint.DistanceSquared(lv);
+		double d2 = DistanceSquared(QueryPoint, lv);
 		if (d2 < fNearSqr)
 		{
 			fNearSqr = d2;
diff --git a/Engine/Plugins/Experimental/GeometryProcessing/Source/DynamicMesh/Private/GroupTopology.cpp b/Engine/Plugins/Experimental/GeometryProcessing/Source/DynamicMesh/Private/GroupTopology.cpp
--- a/Engine/Plugins/Experimental/GeometryProcessing/Source/DynamicMesh/Private/GroupTopology.cpp
+++ b/Engine/Plugins/Experimental/GeometryProcessing/Source/DynamicMesh/Private/GroupTopology.cpp
@@ -314,7 +314,7 @@ double FGroupTopology::GetEdgeArcLength(int32 GroupEdgeID, TArray<double>* PerVe
 	double AccumLength = 0;
 	for (int32 k = 1; k < NumV; ++k)
 	{
-		AccumLength += Mesh->GetVertex(Vertices[k]).Distance(Mesh->GetVertex(Vertices[k-1]));
+		AccumLength += Distance(Mesh->GetVertex(Vertices[k]), Mesh->GetVertex(Vertices[k-1]));
 		if (PerVertexLengthsOut != nullptr)
 		{
 			(*PerVertexLengthsOut)[k] = AccumLength;
@@ -336,13 +336,13 @@ FVector3d FGroupTopology::GetEdgeMidpoint(int32 GroupEdgeID, double* ArcLengthOu
 		FVector3d A(Mesh->GetVertex(Vertices[0])), B(Mesh->GetVertex(Vertices[1]));
 		if (ArcLengthOut)
 		{
-			*ArcLengthOut = A.Distance(B);
+			*ArcLengthOut = Distance(A, B);
 		}
 		if (PerVertexLengthsOut)
 		{
 			(*PerVertexLengthsOut).SetNum(2);
 			(*PerVertexLengthsOut)[0] = 0;
-			(*PerVertexLengthsOut)[1] = A.Distance(B);
+			(*PerVertexLengthsOut)[1] = Distance(A, B);
 		}
 		return (A + B) * 0.5;
 	}
@@ -378,7 +378,7 @@ FVector3d FGroupTopology::GetEdgeMidpoint(int32 GroupEdgeID, double* ArcLengthOu
 	double AccumLength = 0;
 	for (int32 k = 1; k < NumV; ++k)
 	{
-		double NewLen = AccumLength + Mesh->GetVertex(Vertices[k]).Distance(Mesh->GetVertex(Vertices[k-1]));
+		double NewLen = AccumLength + Distance(Mesh->GetVertex(Vertices[k]), Mesh->GetVertex(Vertices[k-1]));
 		if ( NewLen > Len )
 		{
 			double t = (Len - AccumLength) / (NewLen - AccumLength);
@@ -625,7 +625,7 @@ bool FGroupTopology::GetGroupEdgeTangent(int GroupEdgeID, FVector3d& TangentOut)
 	const FGroupEdge& Edge = Edges[GroupEdgeID];
 	FVector3d StartPos = Mesh->GetVertex(Edge.Span.Vertices[0]);
 	FVector3d EndPos = Mesh->GetVertex(Edge.Span.Vertices[Edge.Span.Vertices.Num()-1]);
-	if (StartPos.DistanceSquared(EndPos) > 100 * FMathd::ZeroTolerance)
+	if (DistanceSquared(StartPos, EndPos) > 100 * FMathd::ZeroTolerance)
 	{
 		TangentOut = Normalized(EndPos - StartPos);
 		return true;
diff --git a/Engine/Plugins/Experimental/GeometryProcessing/Source/DynamicMesh/Private/MeshBoundaryLoops.cpp b/Engine/Plugins/Experimental/GeometryProcessing/Source/DynamicMesh/Private/MeshBoundaryLoops.cpp
--- a/Engine/Plugins/Experimental/GeometryProcessing/Source/DynamicMesh/Private/MeshBoundaryLoops.cpp
+++ b/Engine/Plugins/Experimental/GeometryProcessing/Source/DynamicMesh/Private/MeshBoundaryLoops.cpp
@@ -32,7 +32,7 @@ int FMeshBoundaryLoops::GetLongestLoopIndex() const
 		for (int32 k = 0; k < LoopNum; ++k)
 		{
 			FVector3d NextPos = Mesh->GetVertex(Loops[i].Vertices[(k+1)%LoopNum]);
-			LoopLen += Mesh->GetVertex(Loops[i].Vertices[k]).Distance(NextPos);
+			LoopLen += Distance(Mesh->GetVertex(Loops[i].Vertices[k]), NextPos);
 		}
 		if (LoopLen > LongestLoopLen)
 		{
diff --git a/Engine/Plugins/Experimental/GeometryProcessing/Source/DynamicMesh/Private/MeshSimplification.cpp b/Engine/Plugins/Experimental/GeometryProcessing/Source/DynamicMesh/Private/MeshSimplification.cpp
--- a/Engine/Plugins/Experimental/GeometryProcessing/Source/DynamicMesh/Private/MeshSimplification.cpp
+++ b/Engine/Plugins/Experimental/GeometryProcessing/Source/DynamicMesh/Private/MeshSimplification.cpp
@@ -92,14 +92,14 @@ void TMeshSimplification<QuadricErrorType>::InitializeSeamQuadrics()
 		FVector3d nA = Mesh->GetTriNormal(edge.Tri.A);
 
 		// this constrains the point to a plane aligned with the edge and normal to the face
-		FSeamQuadricType& seamQuadric = seamQuadrics.Add(eid, CreateSeamQuadric(p0, p1, nA));
+		FSeamQuadricType& seamQuadric = seamQuadrics.Add(eid, CreateSeamQuadric<double>(p0, p1, nA));
 
 
 		// add the other side - this constrains the point to the line where the two planes intersect.
 		if (edge.Tri.B != FDynamicMesh3::InvalidID)
 		{
 			FVector3d nB = Mesh->GetTriNormal(edge.Tri.B);
-			seamQuadric.Add(CreateSeamQuadric(p0, p1, nB));
+			seamQuadric.Add(CreateSeamQuadric<double>(p0, p1, nB));
 		}
 
 		seamQuadric.Scale(EdgeWeight);
@@ -441,12 +441,12 @@ void TMeshSimplification<QuadricErrorType>::UpdateNeighborhood(const FDynamicMes
 				FVector3d nA = Mesh->GetTriNormal(ne.Tri.A);
 
 				// this constrains the point to a plane aligned with the edge and normal to the face
-				*seamQuadric = CreateSeamQuadric(p0, p1, nA);
+				*seamQuadric = CreateSeamQuadric<double>(p0, p1, nA);
 				// add the other side - this constrains the point to the line where the two planes intersect.
 				if (ne.Tri.B != FDynamicMesh3::InvalidID)
 				{
 					FVector3d nB = Mesh->GetTriNormal(ne.Tri.B);
-					seamQuadric->Add(CreateSeamQuadric(p0, p1, nB));
+					seamQuadric->Add(CreateSeamQuadric<double>(p0, p1, nB));
 				}
 
 				seamQuadric->Scale(EdgeWeight);
@@ -1083,7 +1083,7 @@ void TMeshSimplification<QuadricErrorType>::FastCollapsePass(double fMinEdgeLeng
 			}
 
 			Mesh->GetEdgeV(eid, va, vb);
-			if (va.DistanceSquared(vb) > min_sqr)
+			if (DistanceSquared(va, vb) > min_sqr)
 			{
 				continue;
 			}
@@ -1328,8 +1328,8 @@ ESimplificationResult TMeshSimplification<QuadricErrorType>::CollapseEdge(int ed
 		
 		// [TODO] maybe skip Projection call when !bModeAllowsVertMovment.
 		vNewPos = GetProjectedCollapsePosition(iKeep, vNewPos);
-		double div = vA.Distance(vB);
-		collapse_t = (div < FMathd::ZeroTolerance) ? 0.5 : (vNewPos.Distance(Mesh->GetVertex(iKeep))) / div;
+		double div = Distance(vA, vB);
+		collapse_t = (div < FMathd::ZeroTolerance) ? 0.5 : (Distance(vNewPos, Mesh->GetVertex(iKeep))) / div;
 		collapse_t = VectorUtil::Clamp(collapse_t, 0.0, 1.0);
 
 		if (bMinimalVertexMode)
@@ -1564,10 +1564,10 @@ void TMeshSimplification<FAttrBasedQuadricErrord>::OnEdgeCollapse(int edgeID, in
 	FAttrBasedQuadricErrord& Quadric = EdgeQuadrics[edgeID].q;
 	FVector3d collapse_pt = EdgeQuadrics[edgeID].collapse_pt;
 
-	FVector3d UpdatedNormald;
+	FVector3<double> UpdatedNormald;
 	Quadric.ComputeAttributes(collapse_pt, UpdatedNormald);
 
-	FVector3f UpdatedNormal(UpdatedNormald.X, UpdatedNormald.Y, UpdatedNormald.Z);
+	FVector3f UpdatedNormal((float)UpdatedNormald.X, (float)UpdatedNormald.Y, (float)UpdatedNormald.Z);
 	Normalize(UpdatedNormal);
 
 	if (NormalOverlay != nullptr)
diff --git a/Engine/Plugins/Experimental/GeometryProcessing/Source/DynamicMesh/Private/NormalFlowRemesher.cpp b/Engine/Plugins/Experimental/GeometryProcessing/Source/DynamicMesh/Private/NormalFlowRemesher.cpp
--- a/Engine/Plugins/Experimental/GeometryProcessing/Source/DynamicMesh/Private/NormalFlowRemesher.cpp
+++ b/Engine/Plugins/Experimental/GeometryProcessing/Source/DynamicMesh/Private/NormalFlowRemesher.cpp
@@ -183,7 +183,7 @@ void FNormalFlowRemesher::TrackedFaceProjectionPass(double& MaxDistanceMoved)
 			continue;
 		}
 
-		MaxDistanceMoved = FMath::Max(MaxDistanceMoved, CurrentPosition.Distance(ProjectedPosition));
+		MaxDistanceMoved = FMath::Max(MaxDistanceMoved, Distance(CurrentPosition, ProjectedPosition));
 
 		TempFlagBuffer[VertexID] = true;
 		TempPosBuffer[VertexID] = ProjectedPosition;
@@ -194,7 +194,7 @@ void FNormalFlowRemesher::TrackedFaceProjectionPass(double& MaxDistanceMoved)
 			int OtherVertexID = (EdgeVertices.A == VertexID) ? EdgeVertices.B : EdgeVertices.A;
 			FVector3d OtherVertexPosition = Mesh->GetVertex(OtherVertexID);
 
-			double NewEdgeLength = ProjectedPosition.Distance(OtherVertexPosition);
+			double NewEdgeLength = Distance(ProjectedPosition, OtherVertexPosition);
 			if (NewEdgeLength < MinEdgeLength || NewEdgeLength > MaxEdgeLength)
 			{
 				QueueEdge(EdgeID);
diff --git a/Engine/Plugins/Experimental/GeometryProcessing/Source/DynamicMesh/Private/Operations/EmbedSurfacePath.cpp b/Engine/Plugins/Experimental/GeometryProcessing/Source/DynamicMesh/Private/Operations/EmbedSurfacePath.cpp
--- a/Engine/Plugins/Experimental/GeometryProcessing/Source/DynamicMesh/Private/Operations/EmbedSurfacePath.cpp
+++ b/Engine/Plugins/Experimental/GeometryProcessing/Source/DynamicMesh/Private/Operations/EmbedSurfacePath.cpp
@@ -49,7 +49,7 @@ void RefineSurfacePtFromTriangleToSubElement(const FDynamicMesh3* Mesh, FVector3
 	double BestElementDistSq = 0;
 	for (int VertSubIdx = 0; VertSubIdx < 3; VertSubIdx++)
 	{
-		double DistSq = Pos.DistanceSquared(Mesh->GetVertex(TriVertIDs[VertSubIdx]));
+		double DistSq = DistanceSquared(Pos, Mesh->GetVertex(TriVertIDs[VertSubIdx]));
 		if (DistSq <= SnapElementThresholdSq && (BestSubIdx == -1 || DistSq < BestElementDistSq))
 		{
 			BestSubIdx = VertSubIdx;
@@ -361,7 +361,7 @@ bool WalkMeshPlanar(const FDynamicMesh3* Mesh, int StartTri, FVector3d StartPt,
 
 		for (int32 NewComputedPtIdx = InitialComputedPointsNum; NewComputedPtIdx < ComputedPointsAndSources.Num(); NewComputedPtIdx++)
 		{
-			double DistSq = EndPt.DistanceSquared(ComputedPointsAndSources[NewComputedPtIdx].Value.Position);
+			double DistSq = DistanceSquared(EndPt, ComputedPointsAndSources[NewComputedPtIdx].Value.Position);
 			// TODO: reject cases that move us backwards if !bAllowBackwardsSearch
 			bool bIsForward = ForwardsDirection.Dot(ComputedPointsAndSources[NewComputedPtIdx].Value.Position - StartPt) >= -BackwardsTolerance;
 			ensure(bAllowBackwardsSearch || bIsForward);
diff --git a/Engine/Plugins/Experimental/GeometryProcessing/Source/DynamicMesh/Private/Operations/ExtrudeMesh.cpp b/Engine/Plugins/Experimental/GeometryProcessing/Source/DynamicMesh/Private/Operations/ExtrudeMesh.cpp
--- a/Engine/Plugins/Experimental/GeometryProcessing/Source/DynamicMesh/Private/Operations/ExtrudeMesh.cpp
+++ b/Engine/Plugins/Experimental/GeometryProcessing/Source/DynamicMesh/Private/Operations/ExtrudeMesh.cpp
@@ -172,7 +172,7 @@ bool FExtrudeMesh::ApplyExtrude(FExtrusionInfo& Region, FMeshNormals* UseNormals
 
 					if (k > 0)
 					{
-						AccumUVTranslation += Mesh->GetVertex(BaseLoop.Vertices[k]).Distance(Mesh->GetVertex(BaseLoop.Vertices[k - 1]));
+						AccumUVTranslation += Distance(Mesh->GetVertex(BaseLoop.Vertices[k]), Mesh->GetVertex(BaseLoop.Vertices[k - 1]));
 					}
 
 					// translate horizontally such that vertical spans are adjacent in UV space (so textures tile/wrap properly)
diff --git a/Engine/Plugins/Experimental/GeometryProcessing/Source/DynamicMesh/Private/Operations/GroupTopologyDeformer.cpp b/Engine/Plugins/Experimental/GeometryProcessing/Source/DynamicMesh/Private/Operations/GroupTopologyDeformer.cpp
--- a/Engine/Plugins/Experimental/GeometryProcessing/Source/DynamicMesh/Private/Operations/GroupTopologyDeformer.cpp
+++ b/Engine/Plugins/Experimental/GeometryProcessing/Source/DynamicMesh/Private/Operations/GroupTopologyDeformer.cpp
@@ -280,7 +280,7 @@ void FGroupTopologyDeformer::ComputeEncoding()
 				VtxEncoding.Weights[k] /= WeightSum;
 				Reconstruct += VtxEncoding.Weights[k] * (Mesh->GetVertex(Face.BoundaryVerts[k]) + VtxEncoding.Deltas[k]);
 			}
-			checkSlow(Reconstruct.Distance(Pos) < 0.0001);   // sanity check
+			checkSlow(Distance(Reconstruct, Pos) < 0.0001);   // sanity check
 		}
 	}
 }
diff --git a/Engine/Plugins/Experimental/GeometryProcessing/Source/DynamicMesh/Private/Operations/InsetMeshRegion.cpp b/Engine/Plugins/Experimental/GeometryProcessing/Source/DynamicMesh/Private/Operations/InsetMeshRegion.cpp
--- a/Engine/Plugins/Experimental/GeometryProcessing/Source/DynamicMesh/Private/Operations/InsetMeshRegion.cpp
+++ b/Engine/Plugins/Experimental/GeometryProcessing/Source/DynamicMesh/Private/Operations/InsetMeshRegion.cpp
@@ -354,7 +354,7 @@ bool FInsetMeshRegion::ApplyInset(FInsetInfo& Region, FMeshNormals* UseNormals)
 
 				if (k > 0)
 				{
-					AccumUVTranslation += Mesh->GetVertex(BaseLoopV[k]).Distance(Mesh->GetVertex(BaseLoopV[k - 1]));
+					AccumUVTranslation += Distance(Mesh->GetVertex(BaseLoopV[k]), Mesh->GetVertex(BaseLoopV[k - 1]));
 				}
 
 				// translate horizontally such that vertical spans are adjacent in UV space (so textures tile/wrap properly)
diff --git a/Engine/Plugins/Experimental/GeometryProcessing/Source/DynamicMesh/Private/Operations/JoinMeshLoops.cpp b/Engine/Plugins/Experimental/GeometryProcessing/Source/DynamicMesh/Private/Operations/JoinMeshLoops.cpp
--- a/Engine/Plugins/Experimental/GeometryProcessing/Source/DynamicMesh/Private/Operations/JoinMeshLoops.cpp
+++ b/Engine/Plugins/Experimental/GeometryProcessing/Source/DynamicMesh/Private/Operations/JoinMeshLoops.cpp
@@ -52,7 +52,7 @@ static double CalculateAverageUVScale(const FDynamicMeshUVOverlay* UVOverlay, in
 			for (int32 j = 0; j < 3; ++j)
 			{
 				int32 jn = (j + 1) % 3;
-				double Len3D = Pos[j].Distance(Pos[jn]);
+				double Len3D = Distance(Pos[j], Pos[jn]);
 				double LenUV = UV[j].Distance(UV[jn]);
 				double Scale = LenUV / FMathd::Max(FMathf::ZeroTolerance, Len3D);
 				AverageScale += Scale;
@@ -128,11 +128,11 @@ bool FJoinMeshLoops::Apply()
 		{
 			int32 j = (i + 1) % NV;
 			AccumArcLen[i] = TotalArcLen;
-			RowWidth[i] = Mesh->GetVertex(LoopA[i]).Distance(Mesh->GetVertex(LoopB[i]));
+			RowWidth[i] = Distance(Mesh->GetVertex(LoopA[i]), Mesh->GetVertex(LoopB[i]));
 			AverageWidth += RowWidth[i];
 			FVector3d MidCur = (Mesh->GetVertex(LoopA[i]) + Mesh->GetVertex(LoopB[i])) * 0.5;
 			FVector3d MidNext = (Mesh->GetVertex(LoopA[j]) + Mesh->GetVertex(LoopB[j])) * 0.5;
-			TotalArcLen += MidCur.Distance(MidNext);
+			TotalArcLen += Distance(MidCur, MidNext);
 		}
 		AccumArcLen[NV] = TotalArcLen;
 		RowWidth[NV] = RowWidth[0];
diff --git a/Engine/Plugins/Experimental/GeometryProcessing/Source/DynamicMesh/Private/Operations/MeshBoolean.cpp b/Engine/Plugins/Experimental/GeometryProcessing/Source/DynamicMesh/Private/Operations/MeshBoolean.cpp
--- a/Engine/Plugins/Experimental/GeometryProcessing/Source/DynamicMesh/Private/Operations/MeshBoolean.cpp
+++ b/Engine/Plugins/Experimental/GeometryProcessing/Source/DynamicMesh/Private/Operations/MeshBoolean.cpp
@@ -289,7 +289,7 @@ bool FMeshBoolean::Compute()
 				for (int ChainSubIdx = ChainIdx + 1; ChainSubIdx + 1 < ChainEnd; ChainSubIdx++)
 				{
 					int VID[2]{ Cut.VertexChains[MeshIdx][ChainSubIdx], Cut.VertexChains[MeshIdx][ChainSubIdx + 1] };
-					if (CutMesh[MeshIdx]->GetVertex(VID[0]).DistanceSquared(CutMesh[MeshIdx]->GetVertex(VID[1])) < DegenerateEdgeTolSq)
+					if (DistanceSquared(CutMesh[MeshIdx]->GetVertex(VID[0]), CutMesh[MeshIdx]->GetVertex(VID[1])) < DegenerateEdgeTolSq)
 					{
 						EIDs.Add(CutMesh[MeshIdx]->FindEdge(VID[0], VID[1]));
 					}
@@ -306,7 +306,7 @@ bool FMeshBoolean::Compute()
 				}
 				FVector3d A, B;
 				CutMesh[MeshIdx]->GetEdgeV(EID, A, B);
-				if (A.DistanceSquared(B) > DegenerateEdgeTolSq)
+				if (DistanceSquared(A,B) > DegenerateEdgeTolSq)
 				{
 					continue;
 				}
@@ -561,7 +561,7 @@ bool FMeshBoolean::Compute()
 				FVector3d Pos = CutMesh[MeshIdx]->GetVertex(BoundaryVID);
 				TPair<int, double> VIDDist = OtherMeshPointHash.FindNearestInRadius(Pos, SnapTolerance, [&Pos, &OtherMesh](int VID)
 					{
-						return Pos.DistanceSquared(OtherMesh.GetVertex(VID));
+						return DistanceSquared(Pos, OtherMesh.GetVertex(VID));
 					});
 				int NearestVID = VIDDist.Key; // ID of nearest vertex on other mesh
 				double DSq = VIDDist.Value;   // square distance to that vertex
@@ -572,7 +572,7 @@ bool FMeshBoolean::Compute()
 					int* Match = FoundMatches.Find(NearestVID);
 					if (Match)
 					{
-						double OldDSq = CutMesh[MeshIdx]->GetVertex(*Match).DistanceSquared(OtherMesh.GetVertex(NearestVID));
+						double OldDSq = DistanceSquared(CutMesh[MeshIdx]->GetVertex(*Match), OtherMesh.GetVertex(NearestVID));
 						if (DSq < OldDSq) // new vertex is a better match than the old one
 						{
 							int OldVID = *Match; // copy old VID out of match before updating the TMap
@@ -606,7 +606,7 @@ bool FMeshBoolean::Compute()
 						FVector3d EdgePts[2];
 						OtherMesh.GetEdgeV(OtherEID, EdgePts[0], EdgePts[1]);
 						// only accept the match if it's not going to create a degenerate edge -- TODO: filter already-matched edges from the FindNearestEdge query!
-						if (EdgePts[0].DistanceSquared(Pos) > SnapToleranceSq&& EdgePts[1].DistanceSquared(Pos) > SnapToleranceSq)
+						if (DistanceSquared(EdgePts[0], Pos) > SnapToleranceSq&& DistanceSquared(EdgePts[1], Pos) > SnapToleranceSq)
 						{
 							FSegment3d Seg(EdgePts[0], EdgePts[1]);
 							double Along = Seg.ProjectUnitRange(Pos);
@@ -876,7 +876,7 @@ bool FMeshBoolean::CollapseWouldChangeShapeOrUVs(
 					FVector2f OtherUV = Mesh.GetVertexUV(OtherV);
 					FVector2f RemoveUV = Mesh.GetVertexUV(RemoveV);
 					FVector2f KeepUV = Mesh.GetVertexUV(KeepV);
-					if (Lerp(OtherUV, KeepUV, LerpT).DistanceSquared(RemoveUV) > UVEqualThresholdSq)
+					if ( Lerp(OtherUV, KeepUV, LerpT).DistanceSquared(RemoveUV) > UVEqualThresholdSq)
 					{
 						bHasBadEdge = true;
 						return;
@@ -1243,7 +1243,7 @@ bool FMeshBoolean::MergeEdges(const FMeshIndexMappings& IndexMaps, FDynamicMesh3
 				}
 				FVector3d A, B;
 				Result->GetEdgeV(EID, A, B);
-				if (OA.DistanceSquared(A) < SnapToleranceSq && OB.DistanceSquared(B) < SnapToleranceSq)
+				if (DistanceSquared(OA, A) < SnapToleranceSq && DistanceSquared(OB, B) < SnapToleranceSq)
 				{
 					FDynamicMesh3::FMergeEdgesInfo MergeInfo;
 					EMeshResult EdgeMergeResult = Result->MergeEdges(EID, OtherEID, MergeInfo);
diff --git a/Engine/Plugins/Experimental/GeometryProcessing/Source/DynamicMesh/Private/Operations/MeshMeshCut.cpp b/Engine/Plugins/Experimental/GeometryProcessing/Source/DynamicMesh/Private/Operations/MeshMeshCut.cpp
--- a/Engine/Plugins/Experimental/GeometryProcessing/Source/DynamicMesh/Private/Operations/MeshMeshCut.cpp
+++ b/Engine/Plugins/Experimental/GeometryProcessing/Source/DynamicMesh/Private/Operations/MeshMeshCut.cpp
@@ -387,7 +387,7 @@ namespace MeshCut
 		void UpdateFromSplit(FPtOnMesh& Pt, int SplitVertex, const FIndex2i& SplitEdges)
 		{
 			// check if within tolerance of the new vtx
-			if (Pt.Pos.DistanceSquared(Mesh->GetVertex(SplitVertex)) < SnapToleranceSq)
+			if (DistanceSquared(Pt.Pos, Mesh->GetVertex(SplitVertex)) < SnapToleranceSq)
 			{
 				Pt.Type = EVertexType::Vertex;
 				Pt.ElemID = SplitVertex;
@@ -406,7 +406,7 @@ namespace MeshCut
 		void UpdateFromPoke(FPtOnMesh& Pt, int PokeVertex, const FIndex3i& PokeEdges, const FIndex3i& PokeTris)
 		{
 			// check if within tolerance of the new vtx
-			if (Pt.Pos.DistanceSquared(Mesh->GetVertex(PokeVertex)) < SnapToleranceSq)
+			if (DistanceSquared(Pt.Pos, Mesh->GetVertex(PokeVertex)) < SnapToleranceSq)
 			{
 				Pt.Type = EVertexType::Vertex;
 				Pt.ElemID = PokeVertex;
diff --git a/Engine/Plugins/Experimental/GeometryProcessing/Source/DynamicMesh/Private/Operations/MeshPlaneCut.cpp b/Engine/Plugins/Experimental/GeometryProcessing/Source/DynamicMesh/Private/Operations/MeshPlaneCut.cpp
--- a/Engine/Plugins/Experimental/GeometryProcessing/Source/DynamicMesh/Private/Operations/MeshPlaneCut.cpp
+++ b/Engine/Plugins/Experimental/GeometryProcessing/Source/DynamicMesh/Private/Operations/MeshPlaneCut.cpp
@@ -505,7 +505,7 @@ void FMeshPlaneCut::CollapseDegenerateEdges(const TSet<int>& OnCutEdges, const T
 					continue;
 				}
 				Mesh->GetEdgeV(EID, A, B);
-				if (A.DistanceSquared(B) > Tol2)
+				if (DistanceSquared(A, B) > Tol2)
 				{
 					continue;
 				}
@@ -616,7 +616,8 @@ bool FMeshPlaneCut::HoleFill(TFunction<TArray<FIndex3i>(const FGeneralPolygon2d&
 				LoopVertices.Add(Span.Vertices);
 			}
 		}
-		FPlanarHoleFiller Filler(Mesh, &LoopVertices, PlanarTriangulationFunc, PlaneOrigin, PlaneNormal*Boundary.NormalSign);
+		FVector3d SignedPlaneNormal = PlaneNormal*(double)Boundary.NormalSign;
+		FPlanarHoleFiller Filler(Mesh, &LoopVertices, PlanarTriangulationFunc, PlaneOrigin, SignedPlaneNormal);
 
 		int GID = ConstantGroupID >= 0 ? ConstantGroupID : Mesh->AllocateTriangleGroup();
 		bool bFullyFilledHole = Filler.Fill(GID);
@@ -625,9 +626,9 @@ bool FMeshPlaneCut::HoleFill(TFunction<TArray<FIndex3i>(const FGeneralPolygon2d&
 		if (Mesh->HasAttributes())
 		{
 			FDynamicMeshEditor Editor(Mesh);
-			Editor.SetTriangleNormals(Filler.NewTriangles, (FVector3f)(PlaneNormal*Boundary.NormalSign));
+			Editor.SetTriangleNormals(Filler.NewTriangles, (FVector3f)(SignedPlaneNormal));
 
-			FFrame3d ProjectionFrame(PlaneOrigin, PlaneNormal*Boundary.NormalSign);
+			FFrame3d ProjectionFrame(PlaneOrigin, SignedPlaneNormal);
 			for (int UVLayerIdx = 0, NumLayers = Mesh->Attributes()->NumUVLayers(); UVLayerIdx < NumLayers; UVLayerIdx++)
 			{
 				Editor.SetTriangleUVsFromProjection(Filler.NewTriangles, ProjectionFrame, UVScaleFactor, FVector2f::Zero(), true, UVLayerIdx);
diff --git a/Engine/Plugins/Experimental/GeometryProcessing/Source/DynamicMesh/Private/Operations/MeshSelfUnion.cpp b/Engine/Plugins/Experimental/GeometryProcessing/Source/DynamicMesh/Private/Operations/MeshSelfUnion.cpp
--- a/Engine/Plugins/Experimental/GeometryProcessing/Source/DynamicMesh/Private/Operations/MeshSelfUnion.cpp
+++ b/Engine/Plugins/Experimental/GeometryProcessing/Source/DynamicMesh/Private/Operations/MeshSelfUnion.cpp
@@ -73,7 +73,7 @@ bool FMeshSelfUnion::Compute()
 			for (int ChainSubIdx = ChainIdx + 1; ChainSubIdx + 1 < ChainEnd; ChainSubIdx++)
 			{
 				int VID[2]{ Cut.VertexChains[ChainSubIdx], Cut.VertexChains[ChainSubIdx + 1] };
-				if (Mesh->GetVertex(VID[0]).DistanceSquared(Mesh->GetVertex(VID[1])) < DegenerateEdgeTolSq)
+				if ( DistanceSquared(Mesh->GetVertex(VID[0]), Mesh->GetVertex(VID[1])) < DegenerateEdgeTolSq)
 				{
 					EIDs.Add(Mesh->FindEdge(VID[0], VID[1]));
 				}
@@ -90,7 +90,7 @@ bool FMeshSelfUnion::Compute()
 			}
 			FVector3d A, B;
 			Mesh->GetEdgeV(EID, A, B);
-			if (A.DistanceSquared(B) > DegenerateEdgeTolSq)
+			if (DistanceSquared(A, B) > DegenerateEdgeTolSq)
 			{
 				continue;
 			}
@@ -353,7 +353,7 @@ bool FMeshSelfUnion::Compute()
 				Pos, SnapTolerance,
 				[this, &Pos](int VID)
 				{
-					return Pos.DistanceSquared(Mesh->GetVertex(VID));
+					return DistanceSquared(Pos, Mesh->GetVertex(VID));
 				},
 				[&ExcludeVertices](int VID)
 				{
@@ -369,7 +369,7 @@ bool FMeshSelfUnion::Compute()
 				int* Match = FoundMatches.Find(NearestVID);
 				if (Match)
 				{
-					double OldDSq = Mesh->GetVertex(*Match).DistanceSquared(Mesh->GetVertex(NearestVID));
+					double OldDSq = DistanceSquared(Mesh->GetVertex(*Match), Mesh->GetVertex(NearestVID));
 					if (DSq < OldDSq) // new vertex is a better match than the old one
 					{
 						int OldVID = *Match; // copy old VID out of match before updating the TMap
@@ -407,7 +407,7 @@ bool FMeshSelfUnion::Compute()
 					FVector3d EdgePts[2];
 					Mesh->GetEdgeV(OtherEID, EdgePts[0], EdgePts[1]);
 					// only accept the match if it's not going to create a degenerate edge -- TODO: filter already-matched edges from the FindNearestEdge query!
-					if (EdgePts[0].DistanceSquared(Pos) > SnapToleranceSq && EdgePts[1].DistanceSquared(Pos) > SnapToleranceSq)
+					if (DistanceSquared(EdgePts[0], Pos) > SnapToleranceSq && DistanceSquared(EdgePts[1], Pos) > SnapToleranceSq)
 					{
 						FSegment3d Seg(EdgePts[0], EdgePts[1]);
 						double Along = Seg.ProjectUnitRange(Pos);
@@ -719,7 +719,7 @@ bool FMeshSelfUnion::MergeEdges(const TArray<int>& CutBoundaryEdges, const TMap<
 			FVector3d OA, OB;
 			Mesh->GetEdgeV(OtherEID, OA, OB);
 
-			if (OA.DistanceSquared(A) < SnapToleranceSq && OB.DistanceSquared(B) < SnapToleranceSq)
+			if (DistanceSquared(OA, A) < SnapToleranceSq && DistanceSquared(OB, B) < SnapToleranceSq)
 			{
 				FDynamicMesh3::FMergeEdgesInfo MergeInfo;
 				EMeshResult EdgeMergeResult = Mesh->MergeEdges(EID, OtherEID, MergeInfo);
diff --git a/Engine/Plugins/Experimental/GeometryProcessing/Source/DynamicMesh/Private/Operations/MinimalHoleFiller.cpp b/Engine/Plugins/Experimental/GeometryProcessing/Source/DynamicMesh/Private/Operations/MinimalHoleFiller.cpp
--- a/Engine/Plugins/Experimental/GeometryProcessing/Source/DynamicMesh/Private/Operations/MinimalHoleFiller.cpp
+++ b/Engine/Plugins/Experimental/GeometryProcessing/Source/DynamicMesh/Private/Operations/MinimalHoleFiller.cpp
@@ -316,8 +316,8 @@ namespace
 		if (bDoFlip == false)
 		{
 			FIndex2i OpposingVertices = FillMesh->GetEdgeOpposingV(EdgeID);
-			double CurrentEdgeLength = FillMesh->GetVertex(EdgeVertices.A).Distance(FillMesh->GetVertex(EdgeVertices.B));
-			double FlipEdgeLength = FillMesh->GetVertex(OpposingVertices.A).Distance(FillMesh->GetVertex(OpposingVertices.B));
+			double CurrentEdgeLength = Distance(FillMesh->GetVertex(EdgeVertices.A), FillMesh->GetVertex(EdgeVertices.B));
+			double FlipEdgeLength = Distance(FillMesh->GetVertex(OpposingVertices.A), FillMesh->GetVertex(OpposingVertices.B));
 			if (FlipEdgeLength < CurrentEdgeLength)
 			{
 				if (CheckIfEdgeFlipCreatesFlip(*FillMesh, EdgeID) == false)
diff --git a/Engine/Plugins/Experimental/GeometryProcessing/Source/DynamicMesh/Private/Operations/OffsetMeshRegion.cpp b/Engine/Plugins/Experimental/GeometryProcessing/Source/DynamicMesh/Private/Operations/OffsetMeshRegion.cpp
--- a/Engine/Plugins/Experimental/GeometryProcessing/Source/DynamicMesh/Private/Operations/OffsetMeshRegion.cpp
+++ b/Engine/Plugins/Experimental/GeometryProcessing/Source/DynamicMesh/Private/Operations/OffsetMeshRegion.cpp
@@ -302,7 +302,7 @@ bool FOffsetMeshRegion::ApplyOffset(FOffsetInfo& Region, FMeshNormals* UseNormal
 
 				if (k > 0)
 				{
-					AccumUVTranslation += Mesh->GetVertex(BaseLoopV[k]).Distance(Mesh->GetVertex(BaseLoopV[k - 1]));
+					AccumUVTranslation += Distance(Mesh->GetVertex(BaseLoopV[k]), Mesh->GetVertex(BaseLoopV[k - 1]));
 				}
 
 				// translate horizontally such that vertical spans are adjacent in UV space (so textures tile/wrap properly)
diff --git a/Engine/Plugins/Experimental/GeometryProcessing/Source/DynamicMesh/Private/Operations/RemoveOccludedTriangles.cpp b/Engine/Plugins/Experimental/GeometryProcessing/Source/DynamicMesh/Private/Operations/RemoveOccludedTriangles.cpp
--- a/Engine/Plugins/Experimental/GeometryProcessing/Source/DynamicMesh/Private/Operations/RemoveOccludedTriangles.cpp
+++ b/Engine/Plugins/Experimental/GeometryProcessing/Source/DynamicMesh/Private/Operations/RemoveOccludedTriangles.cpp
@@ -67,7 +67,7 @@ namespace UE
 				return Mesh->GetVertex(Index);
 			}
 
-			FORCEINLINE void GetTriVertices(int32 TriIndex, FVector3d& V0, FVector3d& V1, FVector3d& V2) const
+			FORCEINLINE void GetTriVertices(int32 TriIndex, UE::Core::TVector<double>& V0, UE::Core::TVector<double>& V1, UE::Core::TVector<double>& V2) const
 			{
 				Mesh->GetTriVertices(Component->Indices[TriIndex], V0, V1, V2);
 			}
diff --git a/Engine/Plugins/Experimental/GeometryProcessing/Source/DynamicMesh/Private/Operations/SimpleHoleFiller.cpp b/Engine/Plugins/Experimental/GeometryProcessing/Source/DynamicMesh/Private/Operations/SimpleHoleFiller.cpp
--- a/Engine/Plugins/Experimental/GeometryProcessing/Source/DynamicMesh/Private/Operations/SimpleHoleFiller.cpp
+++ b/Engine/Plugins/Experimental/GeometryProcessing/Source/DynamicMesh/Private/Operations/SimpleHoleFiller.cpp
@@ -82,7 +82,7 @@ bool FSimpleHoleFiller::Fill_Fan(int GroupID)
 
 bool FSimpleHoleFiller::Fill_EarClip(int GroupID)
 {
-	TArray<FVector3d> Vertices;
+	TArray<FVector3<double>> Vertices;
 	int32 NumVertices = Loop.GetVertexCount();
 	for (int32 i = 0; i < NumVertices; ++i)
 	{
diff --git a/Engine/Plugins/Experimental/GeometryProcessing/Source/DynamicMesh/Private/Parameterization/DynamicMeshUVEditor.cpp b/Engine/Plugins/Experimental/GeometryProcessing/Source/DynamicMesh/Private/Parameterization/DynamicMeshUVEditor.cpp
--- a/Engine/Plugins/Experimental/GeometryProcessing/Source/DynamicMesh/Private/Parameterization/DynamicMeshUVEditor.cpp
+++ b/Engine/Plugins/Experimental/GeometryProcessing/Source/DynamicMesh/Private/Parameterization/DynamicMeshUVEditor.cpp
@@ -270,7 +270,7 @@ bool FDynamicMeshUVEditor::SetTriangleUVsFromFreeBoundaryConformal(const TArray<
 	{
 		for (int32 j = i + 1; j < LoopNum; ++j)
 		{
-			double DistSqr = Submesh.GetVertex(ConstrainLoop[i]).DistanceSquared(Submesh.GetVertex(ConstrainLoop[j]));
+			double DistSqr = DistanceSquared(Submesh.GetVertex(ConstrainLoop[i]), Submesh.GetVertex(ConstrainLoop[j]));
 			if (DistSqr > MaxDistSqr)
 			{
 				MaxDistSqr = DistSqr;
diff --git a/Engine/Plugins/Experimental/GeometryProcessing/Source/DynamicMesh/Private/QueueRemesher.cpp b/Engine/Plugins/Experimental/GeometryProcessing/Source/DynamicMesh/Private/QueueRemesher.cpp
--- a/Engine/Plugins/Experimental/GeometryProcessing/Source/DynamicMesh/Private/QueueRemesher.cpp
+++ b/Engine/Plugins/Experimental/GeometryProcessing/Source/DynamicMesh/Private/QueueRemesher.cpp
@@ -95,7 +95,7 @@ int FQueueRemesher::FastSplitIteration()
 		{
 			FIndex2i EdgeVertices = Mesh->GetEdgeV(eid);
 			int OtherVertex = (EdgeVertices.A == NewVertex) ? EdgeVertices.B : EdgeVertices.A;
-			if (Mesh->GetVertex(OtherVertex).DistanceSquared(NewVertexPosition) > MaxEdgeLengthSqr)
+			if (DistanceSquared(Mesh->GetVertex(OtherVertex), NewVertexPosition) > MaxEdgeLengthSqr)
 			{
 				QueueEdge(eid);
 			}
@@ -344,7 +344,7 @@ void FQueueRemesher::TrackedMoveVerticesParallel(TFunction<FVector3d(int, bool&)
 		{
 			const FVector3d& SmoothedPosition = TempPosBuffer[VertexA];
 			const FVector3d OtherVertexPosition = Mesh->GetVertex(VertexB);
-			const double NewEdgeLengthSqr = SmoothedPosition.DistanceSquared(OtherVertexPosition);
+			const double NewEdgeLengthSqr = DistanceSquared(SmoothedPosition, OtherVertexPosition);
 			if (NewEdgeLengthSqr < MinEdgeLength*MinEdgeLength || NewEdgeLengthSqr > MaxEdgeLength*MaxEdgeLength)
 			{
 				EdgeShouldBeQueuedBuffer[EdgeID] = true;
@@ -355,7 +355,7 @@ void FQueueRemesher::TrackedMoveVerticesParallel(TFunction<FVector3d(int, bool&)
 		{
 			const FVector3d& SmoothedPosition = TempPosBuffer[VertexB];
 			const FVector3d OtherVertexPosition = Mesh->GetVertex(VertexA);
-			const double NewEdgeLengthSqr = SmoothedPosition.DistanceSquared(OtherVertexPosition);
+			const double NewEdgeLengthSqr = DistanceSquared(SmoothedPosition, OtherVertexPosition);
 			if (NewEdgeLengthSqr < MinEdgeLength * MinEdgeLength || NewEdgeLengthSqr > MaxEdgeLength* MaxEdgeLength)
 			{
 				EdgeShouldBeQueuedBuffer[EdgeID] = true;
@@ -425,7 +425,7 @@ void FQueueRemesher::TrackedFullSmoothPass_Buffer(bool bParallel)
 				int OtherVertexID = (EdgeVertices.A == VertexID) ? EdgeVertices.B : EdgeVertices.A;
 				FVector3d OtherVertexPosition = Mesh->GetVertex(OtherVertexID);
 
-				double NewEdgeLength = SmoothedPosition.Distance(OtherVertexPosition);
+				double NewEdgeLength = Distance(SmoothedPosition, OtherVertexPosition);
 				if (NewEdgeLength < MinEdgeLength || NewEdgeLength > MaxEdgeLength)
 				{
 					QueueEdge(EdgeID);
@@ -509,7 +509,7 @@ void FQueueRemesher::TrackedFullProjectionPass(bool bParallel)
 				int OtherVertexID = (EdgeVertices.A == VertexID) ? EdgeVertices.B : EdgeVertices.A;
 				FVector3d OtherVertexPosition = Mesh->GetVertex(OtherVertexID);
 
-				double NewEdgeLength = ProjectedPosition.Distance(OtherVertexPosition);
+				double NewEdgeLength = Distance(ProjectedPosition, OtherVertexPosition);
 				if (NewEdgeLength < MinEdgeLength || NewEdgeLength > MaxEdgeLength)
 				{
 					QueueEdge(EdgeID);
diff --git a/Engine/Plugins/Experimental/GeometryProcessing/Source/DynamicMesh/Private/Remesher.cpp b/Engine/Plugins/Experimental/GeometryProcessing/Source/DynamicMesh/Private/Remesher.cpp
--- a/Engine/Plugins/Experimental/GeometryProcessing/Source/DynamicMesh/Private/Remesher.cpp
+++ b/Engine/Plugins/Experimental/GeometryProcessing/Source/DynamicMesh/Private/Remesher.cpp
@@ -153,7 +153,7 @@ FRemesher::EProcessResult FRemesher::ProcessEdge(int edgeID)
 	bool bIsBoundaryEdge = (t1 == IndexConstants::InvalidID);
 	FVector3d vA(Mesh->GetVertex(a));
 	FVector3d vB(Mesh->GetVertex(b));
-	double edge_len_sqr = vA.DistanceSquared(vB);
+	double edge_len_sqr = DistanceSquared(vA, vB);
 
 	// check if we should collapse, and also find which vertex we should collapse to,
 	// in cases where we have constraints/etc
@@ -197,8 +197,8 @@ FRemesher::EProcessResult FRemesher::ProcessEdge(int edgeID)
 		else
 		{
 			vNewPos = GetProjectedCollapsePosition(iKeep, vNewPos);
-			double div = vA.Distance(vB);
-			collapse_t = (div < FMathd::ZeroTolerance) ? 0.5 : (vNewPos.Distance(Mesh->GetVertex(iKeep))) / div;
+			double div = Distance(vA, vB);
+			collapse_t = (div < FMathd::ZeroTolerance) ? 0.5 : (Distance(vNewPos, Mesh->GetVertex(iKeep))) / div;
 			collapse_t = VectorUtil::Clamp(collapse_t, 0.0, 1.0);
 		}
 
@@ -276,8 +276,8 @@ abort_collapse:
 		}
 		else
 		{
-			double CurDistSqr = Mesh->GetVertex(a).DistanceSquared(Mesh->GetVertex(b));
-			double FlipDistSqr = Mesh->GetVertex(c).DistanceSquared(Mesh->GetVertex(d));
+			double CurDistSqr = DistanceSquared(Mesh->GetVertex(a), Mesh->GetVertex(b));
+			double FlipDistSqr = DistanceSquared(Mesh->GetVertex(c), Mesh->GetVertex(d));
 			bTryFlip = (FlipDistSqr < MinLengthFlipThresh*MinLengthFlipThresh * CurDistSqr);
 		}
 
diff --git a/Engine/Plugins/Experimental/GeometryProcessing/Source/DynamicMesh/Private/Sampling/MeshImageBakingCache.cpp b/Engine/Plugins/Experimental/GeometryProcessing/Source/DynamicMesh/Private/Sampling/MeshImageBakingCache.cpp
--- a/Engine/Plugins/Experimental/GeometryProcessing/Source/DynamicMesh/Private/Sampling/MeshImageBakingCache.cpp
+++ b/Engine/Plugins/Experimental/GeometryProcessing/Source/DynamicMesh/Private/Sampling/MeshImageBakingCache.cpp
@@ -234,7 +234,7 @@ bool FMeshImageBakingCache::ValidateCache()
 			[Mesh, NormalOverlay, this](const FMeshUVSampleInfo& SampleInfo, FCorrespondenceSample& ValueOut)
 		{
 			//FVector3d BaseTriNormal = Mesh->GetTriNormal(SampleInfo.TriangleIndex);
-			NormalOverlay->GetTriBaryInterpolate<double>(SampleInfo.TriangleIndex, &SampleInfo.BaryCoords[0], &ValueOut.BaseNormal[0]);
+			NormalOverlay->GetTriBaryInterpolate<double>(SampleInfo.TriangleIndex, &SampleInfo.BaryCoords.X, &ValueOut.BaseNormal.X);
 			Normalize(ValueOut.BaseNormal);
 			FVector3d RayDir = ValueOut.BaseNormal;
 
diff --git a/Engine/Plugins/Experimental/GeometryProcessing/Source/DynamicMesh/Private/Sampling/MeshNormalMapBaker.cpp b/Engine/Plugins/Experimental/GeometryProcessing/Source/DynamicMesh/Private/Sampling/MeshNormalMapBaker.cpp
--- a/Engine/Plugins/Experimental/GeometryProcessing/Source/DynamicMesh/Private/Sampling/MeshNormalMapBaker.cpp
+++ b/Engine/Plugins/Experimental/GeometryProcessing/Source/DynamicMesh/Private/Sampling/MeshNormalMapBaker.cpp
@@ -31,7 +31,7 @@ void FMeshNormalMapBaker::Bake()
 
 			// sample normal on detail mesh
 			FVector3d DetailNormal;
-			DetailNormalOverlay->GetTriBaryInterpolate<double>(DetailTriID, &SampleData.DetailBaryCoords[0], &DetailNormal.X);
+			DetailNormalOverlay->GetTriBaryInterpolate<double>(DetailTriID, &SampleData.DetailBaryCoords.X, &DetailNormal.X);
 			Normalize(DetailNormal);
 
 			// compute normal in tangent space
@@ -47,14 +47,14 @@ void FMeshNormalMapBaker::Bake()
 
 	NormalsBuilder = MakeUnique<TImageBuilder<FVector3f>>();
 	NormalsBuilder->SetDimensions(BakeCache->GetDimensions());
-	FVector3f DefaultMapNormal = (DefaultNormal + FVector3f::One()) * 0.5;
+	FVector3f DefaultMapNormal = (DefaultNormal + FVector3f::One()) * 0.5f;
 	NormalsBuilder->Clear(DefaultMapNormal);
 
 
 	BakeCache->EvaluateSamples([&](const FVector2i& Coords, const FMeshImageBakingCache::FCorrespondenceSample& Sample)
 	{
 		FVector3f RelativeDetailNormal = NormalSampleFunction(Sample);
-		FVector3f MapNormal = (RelativeDetailNormal + FVector3f::One()) * 0.5;
+		FVector3f MapNormal = (RelativeDetailNormal + FVector3f::One()) * 0.5f;
 		NormalsBuilder->SetPixel(Coords, MapNormal);
 	});
 
diff --git a/Engine/Plugins/Experimental/GeometryProcessing/Source/DynamicMesh/Private/Sampling/MeshOcclusionMapBaker.cpp b/Engine/Plugins/Experimental/GeometryProcessing/Source/DynamicMesh/Private/Sampling/MeshOcclusionMapBaker.cpp
--- a/Engine/Plugins/Experimental/GeometryProcessing/Source/DynamicMesh/Private/Sampling/MeshOcclusionMapBaker.cpp
+++ b/Engine/Plugins/Experimental/GeometryProcessing/Source/DynamicMesh/Private/Sampling/MeshOcclusionMapBaker.cpp
@@ -51,7 +51,7 @@ void FMeshOcclusionMapBaker::Bake()
 			FIndex3i DetailTri = DetailMesh->GetTriangle(DetailTriID);
 			//FVector3d DetailTriNormal = DetailMesh.GetTriNormal(DetailTriID);
 			FVector3d DetailTriNormal;
-			DetailNormalOverlay->GetTriBaryInterpolate<double>(DetailTriID, &SampleData.DetailBaryCoords[0], &DetailTriNormal.X);
+			DetailNormalOverlay->GetTriBaryInterpolate<double>(DetailTriID, &SampleData.DetailBaryCoords.X, &DetailTriNormal.X);
 			Normalize(DetailTriNormal);
 
 			FVector3d DetailBaryCoords = SampleData.DetailBaryCoords;
diff --git a/Engine/Plugins/Experimental/GeometryProcessing/Source/DynamicMesh/Private/Sampling/MeshPropertyMapBaker.cpp b/Engine/Plugins/Experimental/GeometryProcessing/Source/DynamicMesh/Private/Sampling/MeshPropertyMapBaker.cpp
--- a/Engine/Plugins/Experimental/GeometryProcessing/Source/DynamicMesh/Private/Sampling/MeshPropertyMapBaker.cpp
+++ b/Engine/Plugins/Experimental/GeometryProcessing/Source/DynamicMesh/Private/Sampling/MeshPropertyMapBaker.cpp
@@ -89,7 +89,7 @@ void FMeshPropertyMapBaker::Bake()
 						if (DetailNormalOverlay->IsSetTriangle(DetailTriID))
 						{
 							FVector3d DetailNormal;
-							DetailNormalOverlay->GetTriBaryInterpolate<double>(DetailTriID, &SampleData.DetailBaryCoords[0], &DetailNormal.X);
+							DetailNormalOverlay->GetTriBaryInterpolate<double>(DetailTriID, &SampleData.DetailBaryCoords.X, &DetailNormal.X);
 							Normalize(DetailNormal);
 							Color = NormalToColor(DetailNormal);
 						}
@@ -100,7 +100,7 @@ void FMeshPropertyMapBaker::Bake()
 						if (DetailUVOverlay && DetailUVOverlay->IsSetTriangle(DetailTriID))
 						{
 							FVector2d DetailUV;
-							DetailUVOverlay->GetTriBaryInterpolate<double>(DetailTriID, &SampleData.DetailBaryCoords[0], &DetailUV.X);
+							DetailUVOverlay->GetTriBaryInterpolate<double>(DetailTriID, &SampleData.DetailBaryCoords.X, &DetailUV.X);
 							Color = UVToColor(DetailUV);
 						}
 					}
diff --git a/Engine/Plugins/Experimental/GeometryProcessing/Source/DynamicMesh/Private/Sampling/MeshResampleImageBaker.cpp b/Engine/Plugins/Experimental/GeometryProcessing/Source/DynamicMesh/Private/Sampling/MeshResampleImageBaker.cpp
--- a/Engine/Plugins/Experimental/GeometryProcessing/Source/DynamicMesh/Private/Sampling/MeshResampleImageBaker.cpp
+++ b/Engine/Plugins/Experimental/GeometryProcessing/Source/DynamicMesh/Private/Sampling/MeshResampleImageBaker.cpp
@@ -22,7 +22,7 @@ void FMeshResampleImageBaker::Bake()
 		if (DetailMesh->IsTriangle(SampleData.DetailTriID) && DetailUVOverlay)
 		{
 			FVector2d DetailUV;
-			DetailUVOverlay->GetTriBaryInterpolate<double>(DetailTriID, &SampleData.DetailBaryCoords[0], &DetailUV.X);
+			DetailUVOverlay->GetTriBaryInterpolate<double>(DetailTriID, &SampleData.DetailBaryCoords.X, &DetailUV.X);
 
 			Color = SampleFunction(DetailUV);
 		}
diff --git a/Engine/Plugins/Experimental/GeometryProcessing/Source/DynamicMesh/Private/ShapeApproximation/ShapeDetection3.cpp b/Engine/Plugins/Experimental/GeometryProcessing/Source/DynamicMesh/Private/ShapeApproximation/ShapeDetection3.cpp
--- a/Engine/Plugins/Experimental/GeometryProcessing/Source/DynamicMesh/Private/ShapeApproximation/ShapeDetection3.cpp
+++ b/Engine/Plugins/Experimental/GeometryProcessing/Source/DynamicMesh/Private/ShapeApproximation/ShapeDetection3.cpp
@@ -80,7 +80,7 @@ bool UE::Geometry::IsSphereMesh(const FDynamicMesh3& Mesh, FSphere3d& SphereOut,
 		FVector3d A, B;
 		Mesh.GetEdgeV(EdgeID, A, B);
 
-		double HalfChordLen = A.Distance(B) * 0.5;
+		double HalfChordLen = Distance(A, B) * 0.5;
 		double MaxChordHeight = UseRadius - FMathd::Sqrt(UseRadius*UseRadius - HalfChordLen*HalfChordLen);   // "sagitta" height
 
 		double MidpointSignedDist = SphereOut.SignedDistance( (A+B)*0.5 );
@@ -248,7 +248,7 @@ bool UE::Geometry::IsCapsuleMesh(const FDynamicMesh3& Mesh, FCapsule3d& CapsuleO
 		FVector3d A, B;
 		Mesh.GetEdgeV(EdgeID, A, B);
 
-		double HalfChordLen = A.Distance(B) * 0.5;
+		double HalfChordLen = Distance(A, B) * 0.5;
 		double MaxChordHeight = UseRadius - FMathd::Sqrt(UseRadius * UseRadius - HalfChordLen * HalfChordLen);   // "sagitta" height
 
 		double MidpointSignedDist = CapsuleOut.SignedDistance( (A + B)*0.5 );
diff --git a/Engine/Plugins/Experimental/GeometryProcessing/Source/DynamicMesh/Public/DynamicMesh3.h b/Engine/Plugins/Experimental/GeometryProcessing/Source/DynamicMesh/Public/DynamicMesh3.h
--- a/Engine/Plugins/Experimental/GeometryProcessing/Source/DynamicMesh/Public/DynamicMesh3.h
+++ b/Engine/Plugins/Experimental/GeometryProcessing/Source/DynamicMesh/Public/DynamicMesh3.h
@@ -114,7 +114,7 @@ public:
 	/** DuplicateTriangleID is returned by AppendTriangle() to indicate that the added triangle already exists in the mesh, and was ignored because we do not support duplicate triangles */
 	constexpr static int DuplicateTriangleID = -3;
 
-	constexpr static FVector3d InvalidVertex{TNumericLimits<double>::Max(), 0.0, 0.0};
+	const static FVector3d InvalidVertex;
 	constexpr static FIndex3i InvalidTriangle{InvalidID, InvalidID, InvalidID};
 	constexpr static FIndex2i InvalidEdge{InvalidID, InvalidID};
 
@@ -602,7 +602,8 @@ public:
 	FIndex3i GetTriNeighbourTris(int TriangleID) const;
 
 	/** Get the three vertex positions of a triangle */
-	inline void GetTriVertices(int TriangleID, FVector3d& v0, FVector3d& v1, FVector3d& v2) const
+	template<typename VecType>
+	inline void GetTriVertices(int TriangleID, VecType& v0, VecType& v1, VecType& v2) const
 	{
 		const FIndex3i& Triangle = Triangles[TriangleID];
 		v0                       = Vertices[Triangle[0]];
diff --git a/Engine/Plugins/Experimental/GeometryProcessing/Source/DynamicMesh/Public/DynamicMeshOverlay.h b/Engine/Plugins/Experimental/GeometryProcessing/Source/DynamicMesh/Public/DynamicMeshOverlay.h
--- a/Engine/Plugins/Experimental/GeometryProcessing/Source/DynamicMesh/Public/DynamicMeshOverlay.h
+++ b/Engine/Plugins/Experimental/GeometryProcessing/Source/DynamicMesh/Public/DynamicMeshOverlay.h
@@ -474,9 +474,9 @@ public:
 
 protected:
 	/** Set the value at an Element to be a linear interpolation of two other Elements */
-	void SetElementFromLerp(int SetElement, int ElementA, int ElementB, RealType Alpha);
+	void SetElementFromLerp(int SetElement, int ElementA, int ElementB, double Alpha);
 	/** Set the value at an Element to be a barycentric interpolation of three other Elements */
-	void SetElementFromBary(int SetElement, int ElementA, int ElementB, int ElementC, const FVector3<RealType>& BaryCoords);
+	void SetElementFromBary(int SetElement, int ElementA, int ElementB, int ElementC, const FVector3<double>& BaryCoords);
 
 	/** updates the triangles array and optionally the element reference counts */
 	void InternalSetTriangle(int TriangleID, const FIndex3i& TriElements, bool bIncrementRefCounts);
@@ -514,7 +514,18 @@ public:
 	 */
 	inline int AppendElement(const VectorType& Value)
 	{
-		return BaseType::AppendElement((const RealType*)Value);
+		// Cannot use cast operator here because Core Vector types do not define it.
+		// However assuming that vector has .X member is also not good...
+		//return BaseType::AppendElement((const RealType*)Value);
+		return BaseType::AppendElement(&Value.X);
+	}
+
+	/**
+	 * Append a new Element to the overlay
+	 */
+	inline int AppendElement(const RealType* Value)
+	{
+		return BaseType::AppendElement(Value);
 	}
 
 	/**
diff --git a/Engine/Plugins/Experimental/GeometryProcessing/Source/DynamicMesh/Public/EdgeLoop.h b/Engine/Plugins/Experimental/GeometryProcessing/Source/DynamicMesh/Public/EdgeLoop.h
--- a/Engine/Plugins/Experimental/GeometryProcessing/Source/DynamicMesh/Public/EdgeLoop.h
+++ b/Engine/Plugins/Experimental/GeometryProcessing/Source/DynamicMesh/Public/EdgeLoop.h
@@ -164,7 +164,16 @@ public:
 	/**
 	 * Add the vertices of the loop to the Vertices array
 	 */
-	void GetVertices(TArray<FVector3d>& Vertices) const;
+	template<typename VecType>
+	void GetVertices(TArray<VecType>& VerticesOut) const
+	{
+		int NumV = Vertices.Num();
+		for (int i = 0; i < NumV; ++i)
+		{
+			VecType Pos = Mesh->GetVertex(Vertices[i]);
+			VerticesOut.Add(Pos);
+		}
+	}
 
 
 	/**
diff --git a/Engine/Plugins/Experimental/GeometryProcessing/Source/DynamicMesh/Public/MeshTangents.h b/Engine/Plugins/Experimental/GeometryProcessing/Source/DynamicMesh/Public/MeshTangents.h
--- a/Engine/Plugins/Experimental/GeometryProcessing/Source/DynamicMesh/Public/MeshTangents.h
+++ b/Engine/Plugins/Experimental/GeometryProcessing/Source/DynamicMesh/Public/MeshTangents.h
@@ -96,7 +96,7 @@ public:
 	 * @param TriangleID triangle index in mesh
 	 * @param TriVertIdx vertex index in range 0,1,2
 	 */
-	void GetPerTriangleTangent(int32 TriangleID, int32 TriVertIdx, FVector3<RealType>& TangentOut, FVector3<RealType>& BitangentOut) const
+	void GetPerTriangleTangent(int32 TriangleID, int32 TriVertIdx, UE::Core::TVector<RealType>& TangentOut, UE::Core::TVector<RealType>& BitangentOut) const
 	{
 		int32 k = TriangleID * 3 + TriVertIdx;
 		TangentOut = Tangents[k];
@@ -129,7 +129,10 @@ public:
 	 * @param TangentOut interpolated tangent
 	 * @param BitangentOut interpolated bitangent
 	 */
-	void GetInterpolatedTriangleTangent(int32 TriangleID, const FVector3<RealType>& BaryCoords, FVector3<RealType>& TangentOut, FVector3<RealType>& BitangentOut) const
+	void GetInterpolatedTriangleTangent(int32 TriangleID, 
+		const UE::Core::TVector<RealType>& BaryCoords,
+		UE::Core::TVector<RealType>& TangentOut,
+		UE::Core::TVector<RealType>& BitangentOut) const
 	{
 		int32 k = TriangleID * 3;
 		if (k >= 0 && (k+2) < Tangents.Num())
diff --git a/Engine/Plugins/Experimental/GeometryProcessing/Source/DynamicMesh/Public/Operations/MergeCoincidentMeshEdges.h b/Engine/Plugins/Experimental/GeometryProcessing/Source/DynamicMesh/Public/Operations/MergeCoincidentMeshEdges.h
--- a/Engine/Plugins/Experimental/GeometryProcessing/Source/DynamicMesh/Public/Operations/MergeCoincidentMeshEdges.h
+++ b/Engine/Plugins/Experimental/GeometryProcessing/Source/DynamicMesh/Public/Operations/MergeCoincidentMeshEdges.h
@@ -64,8 +64,8 @@ protected:
 	// the vertices here so we try both combinations.
 	inline bool IsSameEdge(const FVector3d& a, const FVector3d& b, const FVector3d& c, const FVector3d& d) const
 	{
-		return (a.DistanceSquared(c) < MergeVtxDistSqr && b.DistanceSquared(d) < MergeVtxDistSqr) ||
-			(a.DistanceSquared(d) < MergeVtxDistSqr && b.DistanceSquared(c) < MergeVtxDistSqr);
+		return (DistanceSquared(a,c) < MergeVtxDistSqr && DistanceSquared(b,d) < MergeVtxDistSqr) ||
+			(DistanceSquared(a,d) < MergeVtxDistSqr && DistanceSquared(b,c) < MergeVtxDistSqr);
 	}
 };
 
diff --git a/Engine/Plugins/Experimental/GeometryProcessing/Source/DynamicMesh/Public/Sampling/MeshSurfaceSampler.h b/Engine/Plugins/Experimental/GeometryProcessing/Source/DynamicMesh/Public/Sampling/MeshSurfaceSampler.h
--- a/Engine/Plugins/Experimental/GeometryProcessing/Source/DynamicMesh/Public/Sampling/MeshSurfaceSampler.h
+++ b/Engine/Plugins/Experimental/GeometryProcessing/Source/DynamicMesh/Public/Sampling/MeshSurfaceSampler.h
@@ -66,7 +66,8 @@ struct FDynamicMeshUVMesh
 		return FVector3d(Elem.X, Elem.Y, 0);
 	}
 
-	inline void GetTriVertices(int32 TriangleIndex, FVector3d& V0, FVector3d& V1, FVector3d& V2) const
+	template<typename VecType>
+	inline void GetTriVertices(int32 TriangleIndex, VecType& V0, VecType& V1, VecType& V2) const
 	{
 		FIndex3i TriIndices = UV->GetTriangle(TriangleIndex);
 		V0 = GetVertex(TriIndices.A);
diff --git a/Engine/Plugins/Experimental/GeometryProcessing/Source/GeometricObjects/Private/Generators/SweepGenerator.cpp b/Engine/Plugins/Experimental/GeometryProcessing/Source/GeometricObjects/Private/Generators/SweepGenerator.cpp
--- a/Engine/Plugins/Experimental/GeometryProcessing/Source/GeometricObjects/Private/Generators/SweepGenerator.cpp
+++ b/Engine/Plugins/Experimental/GeometryProcessing/Source/GeometricObjects/Private/Generators/SweepGenerator.cpp
@@ -319,8 +319,8 @@ FMeshShapeGenerator& FProfileSweepGenerator::Generate()
 					// The currently supported modes are to either always connect diagonally down, or connect the shorter diagonal.
 					// For comparing diagonals, we allow some percent difference to triangulate symmetric quads uniformly.
 					if (QuadSplitMethod == EProfileSweepQuadSplit::Uniform
-						|| Vertices[CurrentVert].DistanceSquared(Vertices[BottomRightVert])
-							/ Vertices[BottomVert].DistanceSquared(Vertices[RightVert]) <= (1 + DiagonalTolerance)*(1 + DiagonalTolerance))
+						|| DistanceSquared(Vertices[CurrentVert], Vertices[BottomRightVert])
+							/ DistanceSquared(Vertices[BottomVert], Vertices[RightVert]) <= (1 + DiagonalTolerance)*(1 + DiagonalTolerance))
 					{
 						FVector3d DiagonalDown = Normalized(Vertices[BottomRightVert] - Vertices[CurrentVert]);
 
@@ -444,7 +444,7 @@ void FProfileSweepGenerator::InitializeUvBuffer(const TArray<int32>& VertPositio
 		if (!bUVsSkipFullyWeldedEdges
 			|| !(WeldedVertices.Contains(ProfileIndex) && WeldedVertices.Contains(NextProfileIndex))) // check for welded-to-welded
 		{
-			CumulativeDistance += ProfileCurve[ProfileIndex].Distance(ProfileCurve[NextProfileIndex]);
+			CumulativeDistance += Distance(ProfileCurve[ProfileIndex], ProfileCurve[NextProfileIndex]);
 		}
 		Vs.Add(CumulativeDistance);
 	}
@@ -474,17 +474,18 @@ void FProfileSweepGenerator::InitializeUvBuffer(const TArray<int32>& VertPositio
 	// the profile vertices across sweep frames. We can divide by NumProfileVertices later, if we need to.
 	TArray<double> Distances;
 	Distances.SetNumZeroed(NumSweepSegments);
-	ParallelFor(NumSweepSegments, 
+	ParallelFor(NumSweepSegments,
 		[this, NumSweepPoints, NumProfilePoints, NumWelded, NumNonWelded, &VertPositionOffsets, &Distances]
-		(int SweepIndex)
+	(int SweepIndex)
 	{
 		int32 NextSweepIndex = (SweepIndex + 1) % NumSweepPoints;
 		for (int32 ProfileIndex = 0; ProfileIndex < NumProfilePoints; ++ProfileIndex)
 		{
 			if (!WeldedVertices.Contains(ProfileIndex))
 			{
-				Distances[SweepIndex] += Vertices[GetVertIndex(false, SweepIndex, ProfileIndex, NumWelded, NumNonWelded, VertPositionOffsets)].Distance(
-					Vertices[GetVertIndex(false, NextSweepIndex, ProfileIndex, NumWelded, NumNonWelded, VertPositionOffsets)]);
+				FVector3d A = Vertices[GetVertIndex(false, SweepIndex, ProfileIndex, NumWelded, NumNonWelded, VertPositionOffsets)];
+				FVector3d B = Vertices[GetVertIndex(false, NextSweepIndex, ProfileIndex, NumWelded, NumNonWelded, VertPositionOffsets)];
+				Distances[SweepIndex] += Distance(A, B);
 			}
 		}
 	});
diff --git a/Engine/Plugins/Experimental/GeometryProcessing/Source/GeometricObjects/Private/Spatial/PointSetHashTable.cpp b/Engine/Plugins/Experimental/GeometryProcessing/Source/GeometricObjects/Private/Spatial/PointSetHashTable.cpp
--- a/Engine/Plugins/Experimental/GeometryProcessing/Source/GeometricObjects/Private/Spatial/PointSetHashTable.cpp
+++ b/Engine/Plugins/Experimental/GeometryProcessing/Source/GeometricObjects/Private/Spatial/PointSetHashTable.cpp
@@ -48,7 +48,7 @@ bool FPointSetHashtable::FindPointsInBall(const FVector3d& QueryPt, double Query
 		{
 			if (Points->IsPoint(vid))		// handle case where points are removed from pointset
 			{
-				if (QueryPt.DistanceSquared(Points->GetPoint(vid)) < RadiusSqr)
+				if (DistanceSquared(QueryPt, Points->GetPoint(vid)) < RadiusSqr)
 				{
 					ResultOut.Add(vid);
 				}
@@ -82,7 +82,7 @@ bool FPointSetHashtable::FindPointsInBall(const FVector3d& QueryPt, double Query
 				{
 					if (Points->IsPoint(vid))		// handle case where points are removed from pointset
 					{
-						if (QueryPt.DistanceSquared(Points->GetPoint(vid)) < RadiusSqr)
+						if (DistanceSquared(QueryPt, Points->GetPoint(vid)) < RadiusSqr)
 						{
 							ResultOut.Add(vid);
 						}
diff --git a/Engine/Plugins/Experimental/GeometryProcessing/Source/GeometricObjects/Private/Util/RevolveUtil.cpp b/Engine/Plugins/Experimental/GeometryProcessing/Source/GeometricObjects/Private/Util/RevolveUtil.cpp
--- a/Engine/Plugins/Experimental/GeometryProcessing/Source/GeometricObjects/Private/Util/RevolveUtil.cpp
+++ b/Engine/Plugins/Experimental/GeometryProcessing/Source/GeometricObjects/Private/Util/RevolveUtil.cpp
@@ -31,7 +31,7 @@ void RevolveUtil::WeldPointsOnAxis(TArray<FVector3d>& ProfileCurve, const FVecto
 		FVector3d VectorToPoint = ProfileCurve[ProfileIndex] - RevolutionAxisOrigin;
 		FVector3d VectorProjectedToAxis = RevolutionAxisDirection.Dot(VectorToPoint) * RevolutionAxisDirection;
 
-		double dist = VectorToPoint.DistanceSquared(VectorProjectedToAxis);
+		double dist = DistanceSquared(VectorToPoint, VectorProjectedToAxis);
 		if (dist < ToleranceSquared)
 		{
 			// Move the point directly onto the axis and mark it as welded
@@ -88,19 +88,19 @@ bool RevolveUtil::ProfileIsCCWRelativeRevolve(TArray<FVector3d>& ProfileCurve, c
 
 	// Check that direction is clockwise by looking at the minimal point and its neighbors in the 2D 
 	// coordinate space of "(distance from axis, height along axis)"
-	double PointDistance = (ProfileCurve[MinAlongAxisIndex] - RevolutionAxisOrigin).Distance(MinAlongAxis * RevolutionAxisDirection);
+	double PointDistance = Distance( (ProfileCurve[MinAlongAxisIndex] - RevolutionAxisOrigin), (MinAlongAxis * RevolutionAxisDirection) );
 	FVector2d Point(PointDistance, MinAlongAxis);
 
 	int32 PreviousIndex = (MinAlongAxisIndex + ProfileCurve.Num() - 1) % ProfileCurve.Num();
 	FVector3d PreviousWrtAxis = ProfileCurve[PreviousIndex] - RevolutionAxisOrigin;
 	double PreviousAlongAxis = PreviousWrtAxis.Dot(RevolutionAxisDirection);
-	double PreviousDistance = PreviousWrtAxis.Distance(PreviousAlongAxis * RevolutionAxisDirection);
+	double PreviousDistance = Distance(PreviousWrtAxis, PreviousAlongAxis * RevolutionAxisDirection);
 	FVector2d Previous(PreviousDistance, PreviousAlongAxis);
 
 	int32 NextIndex = (MinAlongAxisIndex + 1) % ProfileCurve.Num();
 	FVector3d NextWrtAxis = ProfileCurve[NextIndex] - RevolutionAxisOrigin;
 	double NextAlongAxis = NextWrtAxis.Dot(RevolutionAxisDirection);
-	double NextDistance = NextWrtAxis.Distance(NextAlongAxis * RevolutionAxisDirection);
+	double NextDistance = Distance(NextWrtAxis, NextAlongAxis * RevolutionAxisDirection);
 	FVector2d Next(NextDistance, NextAlongAxis);
 	
 	return DotPerp((Previous - Point), (Next - Point)) > 0;
@@ -114,7 +114,7 @@ void RevolveUtil::MakeProfileCurveMidpointOfFirstStep(TArray<FVector3d>& Profile
 	{
 		FVector3d PointWrtAxisOrigin = Point - RevolutionAxisOrigin;
 		FVector3d AdjustmentDirection = -Normalized(RevolutionAxisDirection.Cross(PointWrtAxisOrigin));
-		double DistToAxis = PointWrtAxisOrigin.Distance(PointWrtAxisOrigin.Dot(RevolutionAxisDirection)* RevolutionAxisDirection);
+		double DistToAxis = Distance(PointWrtAxisOrigin, PointWrtAxisOrigin.Dot(RevolutionAxisDirection)*RevolutionAxisDirection);
 		
 		Point += AdjustmentScale * DistToAxis * AdjustmentDirection;
 	}
diff --git a/Engine/Plugins/Experimental/GeometryProcessing/Source/GeometricObjects/Public/BoxTypes.h b/Engine/Plugins/Experimental/GeometryProcessing/Source/GeometricObjects/Public/BoxTypes.h
--- a/Engine/Plugins/Experimental/GeometryProcessing/Source/GeometricObjects/Public/BoxTypes.h
+++ b/Engine/Plugins/Experimental/GeometryProcessing/Source/GeometricObjects/Public/BoxTypes.h
@@ -302,8 +302,8 @@ struct TAxisAlignedBox3
 	}
 	TAxisAlignedBox3(const FBox& Box)
 	{
-		Min = Box.Min;
-		Max = Box.Max;
+		Min = FVector3<RealType>(Box.Min);
+		Max = FVector3<RealType>(Box.Max);
 	}
 
 	/**
diff --git a/Engine/Plugins/Experimental/GeometryProcessing/Source/GeometricObjects/Public/Curve/CurveUtil.h b/Engine/Plugins/Experimental/GeometryProcessing/Source/GeometricObjects/Public/Curve/CurveUtil.h
--- a/Engine/Plugins/Experimental/GeometryProcessing/Source/GeometricObjects/Public/Curve/CurveUtil.h
+++ b/Engine/Plugins/Experimental/GeometryProcessing/Source/GeometricObjects/Public/Curve/CurveUtil.h
@@ -13,15 +13,14 @@ namespace UE
 {
 namespace Geometry
 {
+namespace CurveUtil
+{
 
 /**
  * Curve utility functions
  */
-template<typename RealType>
-struct TCurveUtil
-{
-
-	static FVector3<RealType> Tangent(const TArrayView<const FVector3<RealType>>& Vertices, int32 Idx, bool bLoop = false)
+	template<typename RealType, typename VectorType>
+	FVector3<RealType> Tangent(const TArrayView<const VectorType>& Vertices, int32 Idx, bool bLoop = false)
 	{
 		int32 EndIdx = Idx + 1;
 		int32 StartIdx = Idx - 1;
@@ -40,28 +39,30 @@ struct TCurveUtil
 	}
 
 
-	static RealType ArcLength(const TArrayView<const FVector3<RealType>>& Vertices, bool bLoop = false) {
+	template<typename RealType, typename VectorType>
+	static RealType ArcLength(const TArrayView<const VectorType>& Vertices, bool bLoop = false) {
 		RealType Sum = 0;
 		int32 NV = Vertices.Num();
 		for (int i = 1; i < NV; ++i)
 		{
-			Sum += Vertices[i].Distance(Vertices[i - 1]);
+			Sum += Distance(Vertices[i], Vertices[i-1]);
 		}
 		if (bLoop)
 		{
-			Sum += Vertices[NV - 1].Distance(Vertices[0]);
+			Sum += Distance(Vertices[NV-1], Vertices[0]);
 		}
 		return Sum;
 	}
 
-	static int FindNearestIndex(const TArrayView<const FVector3<RealType>>& Vertices, FVector3<RealType> V)
+	template<typename RealType, typename VectorType>
+	int FindNearestIndex(const TArrayView<const VectorType>& Vertices, const VectorType& V)
 	{
 		int iNearest = -1;
 		double dNear = TMathUtil<double>::MaxReal;
 		int N = Vertices.Num();
 		for ( int i = 0; i < N; ++i )
 		{
-			double dSqr = Vertices[i].DistanceSquared(V);
+			double dSqr = DistanceSquared(Vertices[i], V);
 			if (dSqr < dNear)
 			{
 				dNear = dSqr;
@@ -75,7 +76,8 @@ struct TCurveUtil
 	/**
 	 * smooth vertices in-place (will not produce a symmetric result, but does not require extra buffer)
 	 */
-	static void InPlaceSmooth(TArrayView<FVector3<RealType>> Vertices, int StartIdx, int EndIdx, double Alpha, int NumIterations, bool bClosed)
+	template<typename RealType, typename VectorType>
+	void InPlaceSmooth(TArrayView<VectorType> Vertices, int StartIdx, int EndIdx, double Alpha, int NumIterations, bool bClosed)
 	{
 		int N = Vertices.Num();
 		if ( bClosed )
@@ -116,7 +118,8 @@ struct TCurveUtil
 	/**
 	 * smooth set of vertices using extra buffer
 	 */
-	static void IterativeSmooth(TArrayView<FVector3<RealType>> Vertices, int StartIdx, int EndIdx, double Alpha, int NumIterations, bool bClosed)
+	template<typename RealType, typename VectorType>
+	void IterativeSmooth(TArrayView<VectorType> Vertices, int StartIdx, int EndIdx, double Alpha, int NumIterations, bool bClosed)
 	{
 		int N = Vertices.Num();
 		TArray<FVector3<RealType>> Buffer;
@@ -167,8 +170,8 @@ struct TCurveUtil
 			}
 		}
 	}
-};
 
 
+} // end namespace UE::Geometry::CurveUtil
 } // end namespace UE::Geometry
 } // end namespace UE
\ No newline at end of file
diff --git a/Engine/Plugins/Experimental/GeometryProcessing/Source/GeometricObjects/Public/FrameTypes.h b/Engine/Plugins/Experimental/GeometryProcessing/Source/GeometricObjects/Public/FrameTypes.h
--- a/Engine/Plugins/Experimental/GeometryProcessing/Source/GeometricObjects/Public/FrameTypes.h
+++ b/Engine/Plugins/Experimental/GeometryProcessing/Source/GeometricObjects/Public/FrameTypes.h
@@ -97,7 +97,7 @@ struct TFrame3
 	{
 		FVector Normal(Plane.X, Plane.Y, Plane.Z);
 		Origin = (RealType)Plane.W * FVector3<RealType>(Normal);
-		Rotation.SetFromTo(FVector3<RealType>::UnitZ(), Normal);
+		Rotation.SetFromTo(FVector3<RealType>::UnitZ(), (FVector3<RealType>)Normal);
 	}
 
 
@@ -427,7 +427,7 @@ struct TFrame3
 	 * @param HitPointOut intersection point, or FVector3::Max() if ray does not hit plane or is parallel to plane
 	 * @return true if ray intersects plane and HitPointOut is valid
 	 */
-	bool RayPlaneIntersection(const FVector3<RealType>& RayOrigin, const FVector3<RealType>& RayDirection, int PlaneNormalAxis, FVector3<RealType>& HitPointOut) const
+	bool RayPlaneIntersection(const FVector3<RealType>& RayOrigin, const FVector3<RealType>& RayDirection, int PlaneNormalAxis, UE::Core::TVector<RealType>& HitPointOut) const
 	{
 		FVector3<RealType> Normal = GetAxis(PlaneNormalAxis);
 		RealType PlaneD = -Origin.Dot(Normal);
diff --git a/Engine/Plugins/Experimental/GeometryProcessing/Source/GeometricObjects/Public/Generators/MinimalBoxMeshGenerator.h b/Engine/Plugins/Experimental/GeometryProcessing/Source/GeometricObjects/Public/Generators/MinimalBoxMeshGenerator.h
--- a/Engine/Plugins/Experimental/GeometryProcessing/Source/GeometricObjects/Public/Generators/MinimalBoxMeshGenerator.h
+++ b/Engine/Plugins/Experimental/GeometryProcessing/Source/GeometricObjects/Public/Generators/MinimalBoxMeshGenerator.h
@@ -45,10 +45,10 @@ public:
 		for (int fi = 0; fi < 6; ++fi) 
 		{		
 			int NormalAxisIdx = IndexUtil::BoxFaceNormals[fi];
-			FVector3f FaceNormal = (FVector3f)(FMath::Sign(NormalAxisIdx) * Box.GetAxis(FMath::Abs(NormalAxisIdx) - 1));
+			FVector3f FaceNormal = (FVector3f)((double)FMath::Sign(NormalAxisIdx) * Box.GetAxis(FMath::Abs(NormalAxisIdx) - 1));
 
-			double FaceWidth = Vertices[IndexUtil::BoxFaces[fi][0]].Distance(Vertices[IndexUtil::BoxFaces[fi][1]]);
-			double FaceHeight = Vertices[IndexUtil::BoxFaces[fi][1]].Distance(Vertices[IndexUtil::BoxFaces[fi][2]]);
+			double FaceWidth = Distance(Vertices[IndexUtil::BoxFaces[fi][0]], Vertices[IndexUtil::BoxFaces[fi][1]]);
+			double FaceHeight = Distance(Vertices[IndexUtil::BoxFaces[fi][1]], Vertices[IndexUtil::BoxFaces[fi][2]]);
 			double WidthUVScale = FaceWidth * UVScale;
 			double HeightUVScale = FaceHeight * UVScale;
 
diff --git a/Engine/Plugins/Experimental/GeometryProcessing/Source/GeometricObjects/Public/Generators/SweepGenerator.h b/Engine/Plugins/Experimental/GeometryProcessing/Source/GeometricObjects/Public/Generators/SweepGenerator.h
--- a/Engine/Plugins/Experimental/GeometryProcessing/Source/GeometricObjects/Public/Generators/SweepGenerator.h
+++ b/Engine/Plugins/Experimental/GeometryProcessing/Source/GeometricObjects/Public/Generators/SweepGenerator.h
@@ -87,7 +87,7 @@ protected:
 			int NumPathSegs = bLoop ? Path.Num() : Path.Num() - 1;
 			for (int PIdx = 0; PIdx < NumPathSegs; PIdx++)
 			{
-				double SegLen = Path[PIdx].Distance(Path[(PIdx + 1) % Path.Num()]);
+				double SegLen = Distance(Path[PIdx], Path[(PIdx + 1) % Path.Num()]);
 				TotalPathLength += SegLen;
 				PathPercentages.Add(TotalPathLength);
 			}
@@ -474,7 +474,7 @@ public:
 			{
 				if (Caps[CapIdx] == ECapType::FlatMidpointFan)
 				{
-					Vertices[CapVertStart[CapIdx]] = FVector3d::UnitZ() * Heights[CapIdx];
+					Vertices[CapVertStart[CapIdx]] = FVector3d::UnitZ() * (double)Heights[CapIdx];
 				}
 			}
 
@@ -674,7 +674,7 @@ public:
 		if (bUVScaleRelativeWorld)
 		{
 			double Perimeter = CrossSection.Perimeter();
-			double PathLen = TCurveUtil<double>::ArcLength(Path, bLoop);
+			double PathLen = UE::Geometry::CurveUtil::ArcLength<double, FVector3d>(Path, bLoop);
 			SectionScale.X = Perimeter / UnitUVInWorldCoordinates;
 			SectionScale.Y = PathLen / UnitUVInWorldCoordinates;
 			CapScale.X = CapScale.Y = 1.0f / UnitUVInWorldCoordinates;
@@ -691,7 +691,7 @@ public:
 		FFrame3d CrossSectionFrame = InitialFrame;
 		for (int PathIdx = 0; PathIdx < PathNum; ++PathIdx)
 		{
-			FVector3d Tangent = TCurveUtil<double>::Tangent(Path, PathIdx, bLoop);
+			FVector3d Tangent = UE::Geometry::CurveUtil::Tangent<double, FVector3d>(Path, PathIdx, bLoop);
 			CrossSectionFrame.AlignAxis(2, Tangent);
 			FVector3d C = Path[PathIdx];
 			FVector3d X = CrossSectionFrame.X();
@@ -701,17 +701,17 @@ public:
 				FVector2d XP = CrossSection[SubIdx];
 				FVector2d XN = XNormals[SubIdx];
 				Vertices[SubIdx + PathIdx * XNum] = C + X * XP.X + Y * XP.Y;
-				Normals[SubIdx + PathIdx * XNum] = FVector3f(X * XN.X + Y * XN.Y);
+				Normals[SubIdx + PathIdx * XNum] = (FVector3f)(X * XN.X + Y * XN.Y);
 			}
 		}
 		if (bCapped && !bLoop)
 		{
 			for (int CapIdx = 0; CapIdx < 2; CapIdx++)
 			{
-				FVector3f Normal(TCurveUtil<double>::Tangent(Path, CapIdx * (PathNum-1), bLoop) * (CapIdx * 2 - 1));
+				FVector3d Normal = CurveUtil::Tangent<double, FVector3d>(Path, CapIdx * (PathNum - 1), bLoop) * (double)(CapIdx * 2 - 1);
 				for (int SubIdx = 0; SubIdx < XNum; SubIdx++)
 				{
-					Normals[CapNormalStart[CapIdx] + SubIdx] = Normal;
+					Normals[CapNormalStart[CapIdx] + SubIdx] = (FVector3f)Normal;
 				}
 			}
 		}
diff --git a/Engine/Plugins/Experimental/GeometryProcessing/Source/GeometricObjects/Public/Implicit/Blend.h b/Engine/Plugins/Experimental/GeometryProcessing/Source/GeometricObjects/Public/Implicit/Blend.h
--- a/Engine/Plugins/Experimental/GeometryProcessing/Source/GeometricObjects/Public/Implicit/Blend.h
+++ b/Engine/Plugins/Experimental/GeometryProcessing/Source/GeometricObjects/Public/Implicit/Blend.h
@@ -215,7 +215,7 @@ protected:
 			return;
 		}
 
-		TArray<FVector3d> Seeds;
+		TArray<FVector3<double>> Seeds;
 		for (const TriangleMeshType* Source : Sources)
 		{
 			for (int VID = 0; VID < Source->MaxVertexID(); VID++)
diff --git a/Engine/Plugins/Experimental/GeometryProcessing/Source/GeometricObjects/Public/Implicit/CachingMeshSDF.h b/Engine/Plugins/Experimental/GeometryProcessing/Source/GeometricObjects/Public/Implicit/CachingMeshSDF.h
--- a/Engine/Plugins/Experimental/GeometryProcessing/Source/GeometricObjects/Public/Implicit/CachingMeshSDF.h
+++ b/Engine/Plugins/Experimental/GeometryProcessing/Source/GeometricObjects/Public/Implicit/CachingMeshSDF.h
@@ -139,7 +139,7 @@ public:
 		// figure out Origin & dimensions
 		FAxisAlignedBox3d Bounds = Spatial->GetBoundingBox();
 		
-		float MaxDim = MaxElement(Bounds.Max - Bounds.Min + ExpandBounds * 2) + 4 * MaxOffsetDistance;
+		double MaxDim = MaxElement(Bounds.Max - Bounds.Min + ExpandBounds * 2.0) + 4.0 * MaxOffsetDistance;
 		if (!ensureMsgf(MaxDim / CellSize <= ApproxMaxCellsPerDimension - 4, TEXT("SDF resolution clamped to avoid excessive memory use")))
 		{
 			CellSize = MaxDim / (ApproxMaxCellsPerDimension - 4);
diff --git a/Engine/Plugins/Experimental/GeometryProcessing/Source/GeometricObjects/Public/Implicit/Solidify.h b/Engine/Plugins/Experimental/GeometryProcessing/Source/GeometricObjects/Public/Implicit/Solidify.h
--- a/Engine/Plugins/Experimental/GeometryProcessing/Source/GeometricObjects/Public/Implicit/Solidify.h
+++ b/Engine/Plugins/Experimental/GeometryProcessing/Source/GeometricObjects/Public/Implicit/Solidify.h
@@ -130,7 +130,7 @@ public:
 			};
 		}
 
-		TArray<FVector3d> MCSeeds;
+		TArray<FVector3<double>> MCSeeds;
 		for ( int32 VertIdx : Source->VertexIndicesItr() )
 		{
 			MCSeeds.Add(Source->GetVertex(VertIdx));
diff --git a/Engine/Plugins/Experimental/GeometryProcessing/Source/GeometricObjects/Public/Implicit/SweepingMeshSDF.h b/Engine/Plugins/Experimental/GeometryProcessing/Source/GeometricObjects/Public/Implicit/SweepingMeshSDF.h
--- a/Engine/Plugins/Experimental/GeometryProcessing/Source/GeometricObjects/Public/Implicit/SweepingMeshSDF.h
+++ b/Engine/Plugins/Experimental/GeometryProcessing/Source/GeometricObjects/Public/Implicit/SweepingMeshSDF.h
@@ -200,7 +200,7 @@ public:
 			ExactBandWidth = FMath::Max(ApproxMaxCellsPerDimension/2-1, 1);
 		}
 
-		float MaxDim = MaxElement(Bounds.Max - Bounds.Min + ExpandBounds * 2);
+		double MaxDim = MaxElement(Bounds.Max - Bounds.Min + ExpandBounds * 2.0);
 		if (!ensureMsgf(MaxDim / CellSize <= ApproxMaxCellsPerDimension - 2 * ExactBandWidth, TEXT("SDF resolution clamped to avoid excessive memory use")))
 		{
 			CellSize = MaxDim / (ApproxMaxCellsPerDimension - 2 * ExactBandWidth);
@@ -1182,7 +1182,7 @@ public:
 			s12 = 1;
 		}
 		// and find the distance
-		return x0.Distance(s12 * x1 + (1 - s12) * x2);
+		return Distance(x0, s12*x1 + (1.0 - s12)*x2);
 	}
 
 
@@ -1202,7 +1202,7 @@ public:
 			s12 = 1;
 		}
 		// and find the distance
-		return x0.Distance(s12 * x1 + (1 - s12) * x2);
+		return Distance(x0, s12*x1 + (1.0 - s12)*x2);
 	}
 
 
@@ -1223,7 +1223,7 @@ public:
 		float w12 = 1 - w23 - w31;
 		if (w23 >= 0 && w31 >= 0 && w12 >= 0) // if we're inside the triangle
 		{
-			return x0.Distance(w23 * x1 + w31 * x2 + w12 * x3);
+			return Distance(x0, w23*x1 + w31*x2 + w12*x3);
 		}
 		else // we have to clamp to one of the edges
 		{
@@ -1259,7 +1259,7 @@ public:
 		double w12 = 1 - w23 - w31;
 		if (w23 >= 0 && w31 >= 0 && w12 >= 0) // if we're inside the triangle
 		{
-			return x0.Distance(w23 * x1 + w31 * x2 + w12 * x3);
+			return Distance(x0, w23*x1 + w31*x2 + w12*x3);
 		}
 		else // we have to clamp to one of the edges
 		{
diff --git a/Engine/Plugins/Experimental/GeometryProcessing/Source/GeometricObjects/Public/Intersection/IntersectionUtil.h b/Engine/Plugins/Experimental/GeometryProcessing/Source/GeometricObjects/Public/Intersection/IntersectionUtil.h
--- a/Engine/Plugins/Experimental/GeometryProcessing/Source/GeometricObjects/Public/Intersection/IntersectionUtil.h
+++ b/Engine/Plugins/Experimental/GeometryProcessing/Source/GeometricObjects/Public/Intersection/IntersectionUtil.h
@@ -34,7 +34,7 @@ namespace IntersectionUtil
 	using namespace UE::Geometry;
 
 	template<typename RealType>
-	bool RayTriangleTest(const FVector3<RealType>& RayOrigin, const FVector3<RealType>& RayDirection, const FVector3<RealType>& V0, const FVector3<RealType>& V1, const FVector3<RealType>& V2)
+	bool RayTriangleTest(const UE::Core::TVector<RealType>& RayOrigin, const UE::Core::TVector<RealType>& RayDirection, const UE::Core::TVector<RealType>& V0, const UE::Core::TVector<RealType>& V1, const UE::Core::TVector<RealType>& V2)
 	{
 		// same code as IntrRay3Triangle3, but can be called w/o constructing additional data structures
 			
@@ -101,9 +101,9 @@ namespace IntersectionUtil
 	 */
 	template<typename RealType>
 	bool LineSphereTest(
-		const FVector3<RealType>& LineOrigin, 
-		const FVector3<RealType>& LineDirection, 
-		const FVector3<RealType>& SphereCenter, 
+		const UE::Core::TVector<RealType>& LineOrigin,
+		const UE::Core::TVector<RealType>& LineDirection,
+		const UE::Core::TVector<RealType>& SphereCenter,
 		RealType SphereRadius)
 	{
 		// adapted from GeometricTools GTEngine
@@ -130,9 +130,9 @@ namespace IntersectionUtil
 	 */
 	template<typename RealType>
 	bool LineSphereIntersection(
-		const FVector3<RealType>& LineOrigin, 
-		const FVector3<RealType>& LineDirection,
-		const FVector3<RealType>& SphereCenter,
+		const UE::Core::TVector<RealType>& LineOrigin,
+		const UE::Core::TVector<RealType>& LineDirection,
+		const UE::Core::TVector<RealType>& SphereCenter,
 		RealType SphereRadius,
 		FLinearIntersection& ResultOut)
 	{
@@ -174,9 +174,9 @@ namespace IntersectionUtil
 
 	template<typename RealType>
 	FLinearIntersection LineSphereIntersection(
-		const FVector3<RealType>& LineOrigin, 
-		const FVector3<RealType>& LineDirection, 
-		const FVector3<RealType>& SphereCenter, 
+		const UE::Core::TVector<RealType>& LineOrigin,
+		const UE::Core::TVector<RealType>& LineDirection,
+		const UE::Core::TVector<RealType>& SphereCenter,
 		RealType SphereRadius)
 	{
 		FLinearIntersection result;
@@ -191,9 +191,9 @@ namespace IntersectionUtil
 	 */
 	template<typename RealType>
 	bool RaySphereTest(
-		const FVector3<RealType>& RayOrigin, 
-		const FVector3<RealType>& RayDirection, 
-		const FVector3<RealType>& SphereCenter, 
+		const UE::Core::TVector<RealType>& RayOrigin,
+		const UE::Core::TVector<RealType>& RayDirection,
+		const UE::Core::TVector<RealType>& SphereCenter,
 		RealType SphereRadius)
 	{
 		// adapted from GeometricTools GTEngine
@@ -229,9 +229,9 @@ namespace IntersectionUtil
 	 */
 	template<typename RealType>
 	bool RaySphereIntersection(
-		const FVector3<RealType>& RayOrigin, 
-		const FVector3<RealType>& RayDirection, 
-		const FVector3<RealType>& SphereCenter, 
+		const UE::Core::TVector<RealType>& RayOrigin,
+		const UE::Core::TVector<RealType>& RayDirection,
+		const UE::Core::TVector<RealType>& SphereCenter,
 		RealType SphereRadius, 
 		FLinearIntersection& Result)
 	{
@@ -260,9 +260,9 @@ namespace IntersectionUtil
 
 	template<typename RealType>
 	FLinearIntersection RaySphereIntersection(
-		const FVector3<RealType>& RayOrigin, 
-		const FVector3<RealType>& RayDirection, 
-		const FVector3<RealType>& SphereCenter, 
+		const UE::Core::TVector<RealType>& RayOrigin,
+		const UE::Core::TVector<RealType>& RayDirection,
+		const UE::Core::TVector<RealType>& SphereCenter,
 		RealType SphereRadius)
 	{
 		FLinearIntersection result;
diff --git a/Engine/Plugins/Experimental/GeometryProcessing/Source/GeometricObjects/Public/Intersection/IntrTriangle3Triangle3.h b/Engine/Plugins/Experimental/GeometryProcessing/Source/GeometricObjects/Public/Intersection/IntrTriangle3Triangle3.h
--- a/Engine/Plugins/Experimental/GeometryProcessing/Source/GeometricObjects/Public/Intersection/IntrTriangle3Triangle3.h
+++ b/Engine/Plugins/Experimental/GeometryProcessing/Source/GeometricObjects/Public/Intersection/IntrTriangle3Triangle3.h
@@ -32,7 +32,7 @@ class TIntrTriangle3Triangle3
 protected:
 	// Input
 	TTriangle3<Real> Triangle0, Triangle1;
-	double Tolerance = TMathUtil<Real>::ZeroTolerance;
+	Real Tolerance = TMathUtil<Real>::ZeroTolerance;
 
 	// If true, will return intersection polygons for co-planar triangles.
 	// This is somewhat expensive, default is false.
@@ -62,7 +62,7 @@ public:
 	/**
 	 * Store an externally-computed segment intersection result
 	 */
-	inline void SetResult(const FVector3d& A, const FVector3d& B)
+	inline void SetResult(const FVector3<Real>& A, const FVector3<Real>& B)
 	{
 		Result = EIntersectionResult::Intersects;
 		Type = EIntersectionType::Segment;
@@ -103,7 +103,7 @@ public:
 	{
 		return bReportCoplanarIntersection;
 	}
-	double GetTolerance() const
+	Real GetTolerance() const
 	{
 		return Tolerance;
 	}
@@ -122,7 +122,7 @@ public:
 		Result = EIntersectionResult::NotComputed;
 		bReportCoplanarIntersection = bReportCoplanarIntersectionIn;
 	}
-	void SetTolerance(double ToleranceIn)
+	void SetTolerance(Real ToleranceIn)
 	{
 		Result = EIntersectionResult::NotComputed;
 		Tolerance = ToleranceIn;
@@ -208,7 +208,7 @@ public:
 		// At this point, Triangle1 transversely intersects plane 0.  Compute the
 		// line segment of intersection.  Then test for intersection between this
 		// segment and triangle 0.
-		double t;
+		Real t;
 		FVector3<Real> intr0, intr1;
 		if (zero1 == 0)
 		{
@@ -259,8 +259,8 @@ public:
 		FVector3<Real> N0 = UnitCross(E0[0], E0[1]);
 
 		// Project Triangle1 onto normal line of Triangle0, test for separation.
-		double N0dT0V0 = N0.Dot(Triangle0.V[0]);
-		double min1, max1;
+		Real N0dT0V0 = N0.Dot(Triangle0.V[0]);
+		Real min1, max1;
 		ProjectOntoAxis(Triangle1, N0, min1, max1);
 		if (N0dT0V0 < min1 - Tolerance || N0dT0V0 > max1 + Tolerance)
 		{
@@ -277,7 +277,7 @@ public:
 		FVector3<Real> N1 = UnitCross(E1[0], E1[1]);
 
 		FVector3<Real> dir;
-		double min0, max0;
+		Real min0, max0;
 		int i0, i1;
 
 		FVector3<Real> N0xN1 = UnitCross(N0, N1);
@@ -287,7 +287,7 @@ public:
 
 			// Project Triangle0 onto normal line of Triangle1, test for
 			// separation.
-			double N1dT1V0 = N1.Dot(Triangle1.V[0]);
+			Real N1dT1V0 = N1.Dot(Triangle1.V[0]);
 			ProjectOntoAxis(Triangle0, N1, min0, max0);
 			if (N1dT1V0 < min0 - Tolerance || N1dT1V0 > max0 + Tolerance)
 			{
@@ -364,8 +364,8 @@ public:
 		FVector3<Real> N0 = UnitCross(E0[0], E0[1]);
 
 		// Project Triangle1 onto normal line of Triangle0, test for separation.
-		double N0dT0V0 = N0.Dot(Triangle0.V[0]);
-		double min1, max1;
+		Real N0dT0V0 = N0.Dot(Triangle0.V[0]);
+		Real min1, max1;
 		ProjectOntoAxis(Triangle1, N0, min1, max1);
 		if (N0dT0V0 < min1-Tolerance || N0dT0V0 > max1+Tolerance) {
 			return false;
@@ -381,7 +381,7 @@ public:
 		FVector3<Real> N1 = UnitCross(E1[0], E1[1]);
 
 		FVector3<Real> dir;
-		double min0, max0;
+		Real min0, max0;
 		int i0, i1;
 
 		FVector3<Real> N0xN1 = UnitCross(N0, N1);
@@ -390,7 +390,7 @@ public:
 
 			// Project Triangle0 onto normal line of Triangle1, test for
 			// separation.
-			double N1dT1V0 = N1.Dot(Triangle1.V[0]);
+			Real N1dT1V0 = N1.Dot(Triangle1.V[0]);
 			ProjectOntoAxis(Triangle0, N1, min0, max0);
 			if (N1dT1V0 < min0 - Tolerance || N1dT1V0 > max0 + Tolerance) {
 				return false;
@@ -437,11 +437,11 @@ public:
 
 
 
-	static void ProjectOntoAxis(const TTriangle3<Real>& triangle, const FVector3<Real>& axis, double& fmin, double& fmax)
+	static void ProjectOntoAxis(const TTriangle3<Real>& triangle, const FVector3<Real>& axis, Real& fmin, Real& fmax)
 	{
-		double dot0 = axis.Dot(triangle.V[0]);
-		double dot1 = axis.Dot(triangle.V[1]);
-		double dot2 = axis.Dot(triangle.V[2]);
+		Real dot0 = axis.Dot(triangle.V[0]);
+		Real dot1 = axis.Dot(triangle.V[1]);
+		Real dot2 = axis.Dot(triangle.V[2]);
 
 		fmin = dot0;
 		fmax = fmin;
@@ -522,8 +522,8 @@ public:
 		// Project the triangle and segment onto the coordinate plane most
 		// aligned with the plane normal.
 		int maxNormal = 0;
-		double fmax = FMath::Abs(plane.Normal.X);
-		double absMax = FMath::Abs(plane.Normal.Y);
+		Real fmax = FMath::Abs(plane.Normal.X);
+		Real absMax = FMath::Abs(plane.Normal.Y);
 		if (absMax > fmax)
 		{
 			maxNormal = 1;
@@ -599,39 +599,39 @@ public:
 			intr[0] = calc.Point0;
 		}
 
-		FVector3d* OutPts[2]{ &OutA, &OutB };
+		FVector3<Real>* OutPts[2]{ &OutA, &OutB };
 
 		// Unproject the segment of intersection.
 		if (maxNormal == 0)
 		{
-			double invNX = ((double)1) / plane.Normal.X;
+			Real invNX = ((Real)1) / plane.Normal.X;
 			for (i = 0; i < Quantity; ++i)
 			{
-				double y = intr[i].X;
-				double z = intr[i].Y;
-				double x = invNX * (plane.Constant - plane.Normal.Y * y - plane.Normal.Z * z);
+				Real y = intr[i].X;
+				Real z = intr[i].Y;
+				Real x = invNX * (plane.Constant - plane.Normal.Y * y - plane.Normal.Z * z);
 				*OutPts[i] = FVector3<Real>(x, y, z);
 			}
 		}
 		else if (maxNormal == 1)
 		{
-			double invNY = ((double)1) / plane.Normal.Y;
+			Real invNY = ((Real)1) / plane.Normal.Y;
 			for (i = 0; i < Quantity; ++i)
 			{
-				double x = intr[i].X;
-				double z = intr[i].Y;
-				double y = invNY * (plane.Constant - plane.Normal.X * x - plane.Normal.Z * z);
+				Real x = intr[i].X;
+				Real z = intr[i].Y;
+				Real y = invNY * (plane.Constant - plane.Normal.X * x - plane.Normal.Z * z);
 				*OutPts[i] = FVector3<Real>(x, y, z);
 			}
 		}
 		else
 		{
-			double invNZ = ((double)1) / plane.Normal.Z;
+			Real invNZ = ((Real)1) / plane.Normal.Z;
 			for (i = 0; i < Quantity; ++i)
 			{
-				double x = intr[i].X;
-				double y = intr[i].Y;
-				double z = invNZ * (plane.Constant - plane.Normal.X * x - plane.Normal.Y * y);
+				Real x = intr[i].X;
+				Real y = intr[i].Y;
+				Real z = invNZ * (plane.Constant - plane.Normal.X * x - plane.Normal.Y * y);
 				*OutPts[i] = FVector3<Real>(x, y, z);
 			}
 		}
@@ -707,8 +707,8 @@ protected:
 		// Project triangles onto coordinate plane most aligned with plane
 		// normal.
 		int maxNormal = 0;
-		double fmax = FMath::Abs(plane.Normal.X);
-		double absMax = FMath::Abs(plane.Normal.Y);
+		Real fmax = FMath::Abs(plane.Normal.X);
+		Real absMax = FMath::Abs(plane.Normal.Y);
 		if (absMax > fmax)
 		{
 			maxNormal = 1;
@@ -755,7 +755,7 @@ protected:
 		FVector2<Real> save;
 		FVector2<Real> edge0 = projTri0.V[1] - projTri0.V[0];
 		FVector2<Real> edge1 = projTri0.V[2] - projTri0.V[0];
-		if (DotPerp(edge0, edge1) < (double)0)
+		if (DotPerp(edge0, edge1) < (Real)0)
 		{
 			// Triangle is clockwise, reorder it.
 			save = projTri0.V[1];
@@ -765,7 +765,7 @@ protected:
 
 		edge0 = projTri1.V[1] - projTri1.V[0];
 		edge1 = projTri1.V[2] - projTri1.V[0];
-		if (DotPerp(edge0, edge1) < (double)0)
+		if (DotPerp(edge0, edge1) < (Real)0)
 		{
 			// Triangle is clockwise, reorder it.
 			save = projTri1.V[1];
@@ -783,34 +783,34 @@ protected:
 		Quantity = intr.Quantity;
 		if (maxNormal == 0)
 		{
-			double invNX = ((double)1) / plane.Normal.X;
+			Real invNX = ((Real)1) / plane.Normal.X;
 			for (i = 0; i < Quantity; i++)
 			{
-				double y = intr.Points[i].X;
-				double z = intr.Points[i].Y;
-				double x = invNX * (plane.Constant - plane.Normal.Y * y - plane.Normal.Z * z);
+				Real y = intr.Points[i].X;
+				Real z = intr.Points[i].Y;
+				Real x = invNX * (plane.Constant - plane.Normal.Y * y - plane.Normal.Z * z);
 				Points[i] = FVector3<Real>(x, y, z);
 			}
 		}
 		else if (maxNormal == 1)
 		{
-			double invNY = ((double)1) / plane.Normal.Y;
+			Real invNY = ((Real)1) / plane.Normal.Y;
 			for (i = 0; i < Quantity; i++)
 			{
-				double x = intr.Points[i].X;
-				double z = intr.Points[i].Y;
-				double y = invNY * (plane.Constant - plane.Normal.X * x - plane.Normal.Z * z);
+				Real x = intr.Points[i].X;
+				Real z = intr.Points[i].Y;
+				Real y = invNY * (plane.Constant - plane.Normal.X * x - plane.Normal.Z * z);
 				Points[i] = FVector3<Real>(x, y, z);
 			}
 		}
 		else
 		{
-			double invNZ = ((double)1) / plane.Normal.Z;
+			Real invNZ = ((Real)1) / plane.Normal.Z;
 			for (i = 0; i < Quantity; i++)
 			{
-				double x = intr.Points[i].X;
-				double y = intr.Points[i].Y;
-				double z = invNZ * (plane.Constant - plane.Normal.X * x - plane.Normal.Y * y);
+				Real x = intr.Points[i].X;
+				Real y = intr.Points[i].Y;
+				Real z = invNZ * (plane.Constant - plane.Normal.X * x - plane.Normal.Y * y);
 				Points[i] = FVector3<Real>(x, y, z);
 			}
 		}
diff --git a/Engine/Plugins/Experimental/GeometryProcessing/Source/GeometricObjects/Public/MatrixTypes.h b/Engine/Plugins/Experimental/GeometryProcessing/Source/GeometricObjects/Public/MatrixTypes.h
--- a/Engine/Plugins/Experimental/GeometryProcessing/Source/GeometricObjects/Public/MatrixTypes.h
+++ b/Engine/Plugins/Experimental/GeometryProcessing/Source/GeometricObjects/Public/MatrixTypes.h
@@ -39,7 +39,7 @@ struct TMatrix3
 	 * Construct outer-product of U*transpose(V) of U and V
 	 * result is that Mij = u_i * v_j
 	 */
-	TMatrix3<RealType>(const FVector3<RealType>& U, const FVector3<RealType>& V)
+	TMatrix3<RealType>(const UE::Core::TVector<RealType>& U, const UE::Core::TVector<RealType>& V)
 		: Row0(U.X * V.X, U.X * V.Y, U.X * V.Z),
 		  Row1(U.Y * V.X, U.Y * V.Y, U.Y * V.Z),
 		  Row2(U.Z * V.X, U.Z * V.Y, U.Z * V.Z)
@@ -53,7 +53,7 @@ struct TMatrix3
 	{
 	}
 
-	TMatrix3(const FVector3<RealType>& V1, const FVector3<RealType>& V2, const FVector3<RealType>& V3, bool bRows)
+	TMatrix3(const UE::Core::TVector<RealType>& V1, const UE::Core::TVector<RealType>& V2, const UE::Core::TVector<RealType>& V3, bool bRows)
 	{
 		if (bRows)
 		{
@@ -111,7 +111,7 @@ struct TMatrix3
 			Row2.X * Scale, Row2.Y * Scale, Row2.Z * Scale);
 	}
 
-	FVector3<RealType> operator*(const FVector3<RealType>& V) const
+	FVector3<RealType> operator*(const UE::Core::TVector<RealType>& V) const
 	{
 		return FVector3<RealType>(
 			Row0.X * V.X + Row0.Y * V.Y + Row0.Z * V.Z,
@@ -241,7 +241,7 @@ struct TMatrix3
 			   VectorUtil::EpsilonEqual(Row2, Mat2.Row2, Epsilon);
 	}
 
-	static TMatrix3<RealType> AxisAngleR(const FVector3<RealType>& Axis, RealType AngleRad)
+	static TMatrix3<RealType> AxisAngleR(const UE::Core::TVector<RealType>& Axis, RealType AngleRad)
 	{
 		RealType cs = TMathUtil<RealType>::Cos(AngleRad);
 		RealType sn = TMathUtil<RealType>::Sin(AngleRad);
@@ -261,7 +261,7 @@ struct TMatrix3
 			xzm - ySin, yzm + xSin, z2 * oneMinusCos + cs);
 	}
 
-	static TMatrix3<RealType> AxisAngleD(const FVector3<RealType>& Axis, RealType AngleDeg)
+	static TMatrix3<RealType> AxisAngleD(const UE::Core::TVector<RealType>& Axis, RealType AngleDeg)
 	{
 		return AxisAngleR(Axis, TMathUtil<RealType>::DegreesToRadians(AngleDeg));
 	}
@@ -467,7 +467,7 @@ inline TMatrix2<RealType> operator*(RealType Scale, const TMatrix2<RealType>& Ma
 
 // Skew-Symmetric matrix such that A X B = CrossProductMatrix(A) * B;
 template <typename RealType>
-inline static TMatrix3<RealType> CrossProductMatrix(const FVector3<RealType>& A)
+inline static TMatrix3<RealType> CrossProductMatrix(const UE::Core::TVector<RealType>& A)
 {
 	RealType Zero(0);
 
diff --git a/Engine/Plugins/Experimental/GeometryProcessing/Source/GeometricObjects/Public/MeshAdapter.h b/Engine/Plugins/Experimental/GeometryProcessing/Source/GeometricObjects/Public/MeshAdapter.h
--- a/Engine/Plugins/Experimental/GeometryProcessing/Source/GeometricObjects/Public/MeshAdapter.h
+++ b/Engine/Plugins/Experimental/GeometryProcessing/Source/GeometricObjects/Public/MeshAdapter.h
@@ -30,7 +30,7 @@ struct TTriangleMeshAdapter
 	TFunction<FIndex3i(int32)> GetTriangle;
 	TFunction<FVector3<RealType>(int32)> GetVertex;
 
-	inline void GetTriVertices(int TID, FVector3<RealType>& V0, FVector3<RealType>& V1, FVector3<RealType>& V2) const
+	inline void GetTriVertices(int TID, UE::Core::TVector<RealType>& V0, UE::Core::TVector<RealType>& V1, UE::Core::TVector<RealType>& V2) const
 	{
 		FIndex3i TriIndices = GetTriangle(TID);
 		V0 = GetVertex(TriIndices.A);
@@ -132,7 +132,7 @@ struct TIndexMeshArrayAdapter
 		return FVector3<OutRealType>((*SourceVertices)[Index]);
 	}
 
-	FORCEINLINE void GetTriVertices(int32 TriIndex, FVector3<OutRealType>& V0, FVector3<OutRealType>& V1, FVector3<OutRealType>& V2) const
+	FORCEINLINE void GetTriVertices(int32 TriIndex, UE::Core::TVector<OutRealType>& V0, UE::Core::TVector<OutRealType>& V1, UE::Core::TVector<OutRealType>& V2) const
 	{
 		int32 Start = TriIndex * 3;
 		V0 = FVector3<OutRealType>((*SourceVertices)[(*SourceTriangles)[Start]]);
@@ -214,7 +214,7 @@ struct TIndexVectorMeshArrayAdapter
 		return FVector3<OutRealType>((*SourceVertices)[Index]);
 	}
 
-	FORCEINLINE void GetTriVertices(int32 TriIndex, FVector3<OutRealType>& V0, FVector3<OutRealType>& V1, FVector3<OutRealType>& V2) const
+	FORCEINLINE void GetTriVertices(int32 TriIndex, UE::Core::TVector<OutRealType>& V0, UE::Core::TVector<OutRealType>& V1, UE::Core::TVector<OutRealType>& V2) const
 	{
 		const IndexVectorType& Tri = (*SourceTriangles)[TriIndex];
 		V0 = FVector3<OutRealType>((*SourceVertices)[Tri[0]]);
diff --git a/Engine/Plugins/Experimental/GeometryProcessing/Source/GeometricObjects/Public/MeshQueries.h b/Engine/Plugins/Experimental/GeometryProcessing/Source/GeometricObjects/Public/MeshQueries.h
--- a/Engine/Plugins/Experimental/GeometryProcessing/Source/GeometricObjects/Public/MeshQueries.h
+++ b/Engine/Plugins/Experimental/GeometryProcessing/Source/GeometricObjects/Public/MeshQueries.h
@@ -93,7 +93,7 @@ public:
 		int NumVertices = Mesh.VertexCount();
 		if (NumVertices > 0)
 		{
-			Centroid /= NumVertices;
+			Centroid /= (double)NumVertices;
 		}
 		return Centroid;
 	}
@@ -440,7 +440,7 @@ public:
 			{
 				FVector3d A, B;
 				Mesh.GetEdgeV(EdgeID, A, B);
-				double Length = A.Distance(B);
+				double Length = Distance(A, B);
 				if (Length < MinEdgeLength) { MinEdgeLength = Length; }
 				if (Length > MaxEdgeLength) { MaxEdgeLength = Length; }
 				AverageEdgeLength += Length;
@@ -481,7 +481,7 @@ public:
 			if (Mesh.IsEdge(EdgeID))
 			{
 				Mesh.GetEdgeV(EdgeID, A, B);
-				double Length = A.Distance(B);
+				double Length = Distance(A, B);
 				if (Length < MinEdgeLength) MinEdgeLength = Length;
 				if (Length > MaxEdgeLength) MaxEdgeLength = Length;
 				AverageEdgeLength += Length;
diff --git a/Engine/Plugins/Experimental/GeometryProcessing/Source/GeometricObjects/Public/Parameterization/MeshDijkstra.h b/Engine/Plugins/Experimental/GeometryProcessing/Source/GeometricObjects/Public/Parameterization/MeshDijkstra.h
--- a/Engine/Plugins/Experimental/GeometryProcessing/Source/GeometricObjects/Public/Parameterization/MeshDijkstra.h
+++ b/Engine/Plugins/Experimental/GeometryProcessing/Source/GeometricObjects/Public/Parameterization/MeshDijkstra.h
@@ -330,7 +330,7 @@ private:
 				continue;
 			}
 
-			double NbrDist = ParentDist + ParentPos.Distance(GetPositionFunc(NbrPointID));
+			double NbrDist = ParentDist + Distance(ParentPos, GetPositionFunc(NbrPointID));
 			if (Queue.Contains(NbrPointID))
 			{
 				if (NbrDist < NbrNode->GraphDistance)
diff --git a/Engine/Plugins/Experimental/GeometryProcessing/Source/GeometricObjects/Public/Parameterization/MeshLocalParam.h b/Engine/Plugins/Experimental/GeometryProcessing/Source/GeometricObjects/Public/Parameterization/MeshLocalParam.h
--- a/Engine/Plugins/Experimental/GeometryProcessing/Source/GeometricObjects/Public/Parameterization/MeshLocalParam.h
+++ b/Engine/Plugins/Experimental/GeometryProcessing/Source/GeometricObjects/Public/Parameterization/MeshLocalParam.h
@@ -374,7 +374,7 @@ protected:
 			{
 				FFrame3d NbrFrame(GetFrame(NbrNode->PointID));
 				FVector2d NbrUV = PropagateUV(NodePos, NbrNode->UV, NbrFrame, SeedFrame);
-				double Weight = 1.0 / (NodePos.DistanceSquared(NbrFrame.Origin) + TMathUtil<double>::ZeroTolerance);
+				double Weight = 1.0 / (DistanceSquared(NodePos,NbrFrame.Origin) + TMathUtil<double>::ZeroTolerance);
 				AverageUV += Weight * NbrUV;
 				WeightSum += Weight;
 				NbrCount++;
@@ -443,7 +443,7 @@ protected:
 				continue;
 			}
 
-			double NbrDist = ParentDist + ParentPos.Distance(GetPosition(NbrPointID));
+			double NbrDist = ParentDist + Distance(ParentPos, GetPosition(NbrPointID));
 			if (Queue.Contains(NbrPointID))
 			{
 				if (NbrDist < NbrNode->GraphDistance)
diff --git a/Engine/Plugins/Experimental/GeometryProcessing/Source/GeometricObjects/Public/PlaneTypes.h b/Engine/Plugins/Experimental/GeometryProcessing/Source/GeometricObjects/Public/PlaneTypes.h
--- a/Engine/Plugins/Experimental/GeometryProcessing/Source/GeometricObjects/Public/PlaneTypes.h
+++ b/Engine/Plugins/Experimental/GeometryProcessing/Source/GeometricObjects/Public/PlaneTypes.h
@@ -20,14 +20,14 @@ struct TPlane3
 
 	TPlane3() {}
 
-	TPlane3(FVector3<RealType> Normal, double Constant) : Normal(Normal), Constant(Constant)
+	TPlane3(const UE::Core::TVector<RealType>& Normal, double Constant) : Normal(Normal), Constant(Constant)
 	{
 	}
 
 	/**
 	 * N is specified, c = Dot(N,P) where P is a point on the plane.
 	 */
-	TPlane3(const FVector3<RealType>& Normal, const FVector3<RealType>& Point) : Normal(Normal), Constant(Normal.Dot(Point))
+	TPlane3(const UE::Core::TVector<RealType>& Normal, const UE::Core::TVector<RealType>& Point) : Normal(Normal), Constant(Normal.Dot(Point))
 	{
 	}
 
@@ -35,7 +35,7 @@ struct TPlane3
 	 * N = Cross(P1-P0,P2-P0)/Length(Cross(P1-P0,P2-P0)), c = Dot(N,P0) where
 	 * P0, P1, P2 are points on the plane.
 	 */
-	TPlane3(const FVector3<RealType>& P0, const FVector3<RealType>& P1, const FVector3<RealType>& P2)
+	TPlane3(const UE::Core::TVector<RealType>& P0, const UE::Core::TVector<RealType>& P1, const UE::Core::TVector<RealType>& P2)
 	{
 		Normal = VectorUtil::Normal(P0, P1, P2);
 		Constant = Normal.Dot(P0);
@@ -49,7 +49,7 @@ struct TPlane3
 	 * the point is on the negative side, and zero if the point is on the
 	 * plane.
 	 */
-	double DistanceTo(const FVector3<RealType>& P) const
+	double DistanceTo(const UE::Core::TVector<RealType>& P) const
 	{
 		return Normal.Dot(P) - Constant;
 	}
@@ -60,7 +60,7 @@ struct TPlane3
 	 * function returns +1 when P is on the positive side, -1 when P is on the
 	 * the negative side, or 0 when P is on the plane.
 	 */
-	int WhichSide(const FVector3<RealType>& P) const
+	int WhichSide(const UE::Core::TVector<RealType>& P) const
 	{
 		double Distance = DistanceTo(P);
 		if (Distance < 0)
@@ -87,7 +87,10 @@ struct TPlane3
 	 * @param HitPointOut intersection point, or FVector3::MaxVector() if line is parallel to plane
 	 * @return true if Line intersects Plane and IntersectionPointOut is valid
 	 */
-	bool FindLineIntersection(const FVector3<RealType>& LineOrigin, const FVector3<RealType>& LineDirection, FVector3<RealType>& IntersectionPointOut) const
+	bool FindLineIntersection(
+		const UE::Core::TVector<RealType>& LineOrigin, 
+		const UE::Core::TVector<RealType>& LineDirection, 
+		UE::Core::TVector<RealType>& IntersectionPointOut) const
 	{
 		RealType NormalDot = LineDirection.Dot(Normal);
 		if ( TMathUtil<RealType>::Abs(NormalDot) < TMathUtil<RealType>::ZeroTolerance )
@@ -109,7 +112,7 @@ struct TPlane3
 	 * @param Point1 second point of segment
 	 * @return 0 if line is fully clipped, 1 if line is partially clipped, 2 if line is not clipped
 	 */
-	int ClipSegment(FVector3<RealType>& Point0, FVector3<RealType>& Point1)
+	int ClipSegment(UE::Core::TVector<RealType>& Point0, UE::Core::TVector<RealType>& Point1)
 	{
 		RealType Dist0 = DistanceTo(Point0);
 		RealType Dist1 = DistanceTo(Point1);
diff --git a/Engine/Plugins/Experimental/GeometryProcessing/Source/GeometricObjects/Public/Polyline3.h b/Engine/Plugins/Experimental/GeometryProcessing/Source/GeometricObjects/Public/Polyline3.h
--- a/Engine/Plugins/Experimental/GeometryProcessing/Source/GeometricObjects/Public/Polyline3.h
+++ b/Engine/Plugins/Experimental/GeometryProcessing/Source/GeometricObjects/Public/Polyline3.h
@@ -49,6 +49,15 @@ public:
 	{
 	}
 
+	/**
+	 * Construct a single-segment polyline
+	 */
+	TPolyline3(const FVector3<T>& Point0, const FVector3<T>& Point1) : Timestamp(0)
+	{
+		Vertices.Add(Point0);
+		Vertices.Add(Point1);
+	}
+
 	/** @return the Timestamp for the polyline, which is updated every time the polyline is modified */
 	int GetTimestamp() const
 	{
@@ -146,6 +155,21 @@ public:
 		Timestamp++;
 	}
 
+	/**
+	 * Add a list of Vertices to the polyline
+	 */
+	template<typename VectorType>
+	void AppendVertices(const TArray<VectorType>& NewVertices)
+	{
+		int32 NumV = NewVertices.Num();
+		Vertices.Reserve(Vertices.Num() + NumV);
+		for (int32 k = 0; k < NumV; ++k)
+		{
+			Vertices.Append( (FVector3<T>)NewVertices[k] );
+		}
+		Timestamp++;
+	}
+
 	/**
 	 * Set vertex at given index to a new Position
 	 */
diff --git a/Engine/Plugins/Experimental/GeometryProcessing/Source/GeometricObjects/Public/QuadricError.h b/Engine/Plugins/Experimental/GeometryProcessing/Source/GeometricObjects/Public/QuadricError.h
--- a/Engine/Plugins/Experimental/GeometryProcessing/Source/GeometricObjects/Public/QuadricError.h
+++ b/Engine/Plugins/Experimental/GeometryProcessing/Source/GeometricObjects/Public/QuadricError.h
@@ -161,7 +161,7 @@ struct TQuadricError
 	 * Evaluates p*A*p + 2*dot(p,b) + c
 	 * @return 
 	 */
-	RealType Evaluate(const FVector3<RealType>& pt) const
+	RealType Evaluate(const UE::Core::TVector<RealType>& pt) const
 	{
 		RealType x = Axx * pt.X + Axy * pt.Y + Axz * pt.Z;
 		RealType y = Axy * pt.X + Ayy * pt.Y + Ayz * pt.Z;
@@ -173,7 +173,7 @@ struct TQuadricError
 	/**
 	 * 
 	 */
-	FVector3<RealType> MultiplyA(const FVector3<RealType>& pt) const
+	FVector3<RealType> MultiplyA(const UE::Core::TVector<RealType>& pt) const
 	{
 		RealType x = Axx * pt.X + Axy * pt.Y + Axz * pt.Z;
 		RealType y = Axy * pt.X + Ayy * pt.Y + Ayz * pt.Z;
@@ -182,7 +182,7 @@ struct TQuadricError
 	}
 
 
-	bool SolveAxEqualsb(FVector3<RealType>& OutResult, const RealType bvecx, const RealType bvecy, const RealType bvecz, const RealType minThresh = 1000.0*TMathUtil<RealType>::Epsilon) const 
+	bool SolveAxEqualsb(UE::Core::TVector<RealType>& OutResult, const RealType bvecx, const RealType bvecy, const RealType bvecz, const RealType minThresh = 1000.0*TMathUtil<RealType>::Epsilon) const
 	{
 		RealType a11 = Azz * Ayy - Ayz * Ayz;
 		RealType a12 = Axz * Ayz - Azz * Axy;
@@ -245,14 +245,14 @@ struct TQuadricError
 
 		return FVector3<RealType>(a, b, c);
 	}
-	static FVector3<RealType> MultiplySymmetricMatrix(const RealType SM[6], const FVector3<RealType>& vec)
+	static FVector3<RealType> MultiplySymmetricMatrix(const RealType SM[6], const UE::Core::TVector<RealType>& vec)
 	{
 		RealType vectmp[3] = { vec.X, vec.Y, vec.Z };
 		return MultiplySymmetricMatrix(SM, vectmp);
 	}
 
 
-	bool OptimalPoint(FVector3<RealType>& OutResult, RealType minThresh = 1000.0*TMathUtil<RealType>::Epsilon ) const
+	bool OptimalPoint(UE::Core::TVector<RealType>& OutResult, RealType minThresh = 1000.0*TMathUtil<RealType>::Epsilon ) const
 	{
 		return SolveAxEqualsb(OutResult, -bx, -by, -bz, minThresh);
 	}
@@ -383,7 +383,7 @@ public:
 	/**
 	* The optimal point minimizing the quadric error with respect to a volume conserving constraint
 	*/
-	bool OptimalPoint(FVector3<RealType>& OutResult, RealType minThresh = 1000.0*TMathUtil<RealType>::Epsilon) const
+	bool OptimalPoint(UE::Core::TVector<RealType>& OutResult, RealType minThresh = 1000.0*TMathUtil<RealType>::Epsilon) const
 	{
 		// Compute the unconstrained optimal point
 		bool bValid = BaseStruct::OptimalPoint(OutResult, minThresh);
@@ -597,7 +597,7 @@ public:
 	/**
 	* The optimal point minimizing the quadric error with respect to a volume conserving constraint
 	*/
-	bool OptimalPoint(FVector3<RealType>& OptPoint, RealType minThresh = 1000.0*TMathUtil<RealType>::Epsilon) const
+	bool OptimalPoint(UE::Core::TVector<RealType>& OptPoint, RealType minThresh = 1000.0*TMathUtil<RealType>::Epsilon) const
 	{
 		// Generate symmetric matrix
 		RealType SM[6] = { 0., 0., 0., 0., 0., 0. };
@@ -755,7 +755,7 @@ typedef TAttrBasedQuadricError<double> FAttrBasedQuadricErrord;
 * On return the quadric has been weighted by the length of the (p1-p0).
 */
 template<typename RealType>
-TQuadricError<RealType>  CreateSeamQuadric(const FVector3<RealType>& p0, const FVector3<RealType>& p1, FVector3<RealType> AdjFaceNormal)
+TQuadricError<RealType>  CreateSeamQuadric(const FVector3<RealType>& p0, const FVector3<RealType>& p1, const FVector3<RealType>& AdjFaceNormal)
 {
 
 	RealType LenghtSqrd = AdjFaceNormal.SquaredLength();
diff --git a/Engine/Plugins/Experimental/GeometryProcessing/Source/GeometricObjects/Public/Quaternion.h b/Engine/Plugins/Experimental/GeometryProcessing/Source/GeometricObjects/Public/Quaternion.h
--- a/Engine/Plugins/Experimental/GeometryProcessing/Source/GeometricObjects/Public/Quaternion.h
+++ b/Engine/Plugins/Experimental/GeometryProcessing/Source/GeometricObjects/Public/Quaternion.h
@@ -232,7 +232,7 @@ TQuaternion<RealType> operator -(const TQuaternion<RealType>& A, const TQuaterni
 }
 
 template<typename RealType>
-FVector3<RealType> operator*(const TQuaternion<RealType>& Q, const FVector3<RealType>& V) 
+FVector3<RealType> operator*(const TQuaternion<RealType>& Q, const UE::Core::TVector<RealType>& V)
 {
 	//return q.ToRotationMatrix() * v;
 	// inline-expansion of above:
diff --git a/Engine/Plugins/Experimental/GeometryProcessing/Source/GeometricObjects/Public/Spatial/FastWinding.h b/Engine/Plugins/Experimental/GeometryProcessing/Source/GeometricObjects/Public/Spatial/FastWinding.h
--- a/Engine/Plugins/Experimental/GeometryProcessing/Source/GeometricObjects/Public/Spatial/FastWinding.h
+++ b/Engine/Plugins/Experimental/GeometryProcessing/Source/GeometricObjects/Public/Spatial/FastWinding.h
@@ -82,7 +82,7 @@ namespace FastTriWinding
 			Order2 += a * FMatrix3d(dcp, n);
 
 			// this is just for return value...
-			double maxdist = FMath::Max3(P0.DistanceSquared(P), P1.DistanceSquared(P), P2.DistanceSquared(P));
+			double maxdist = FMath::Max3(DistanceSquared(P0,P), DistanceSquared(P1,P), DistanceSquared(P2,P));
 			R = FMath::Max(R, sqrt(maxdist));
 		}
 	}
@@ -442,7 +442,7 @@ private:
 			return false;
 		}
 
-		double dist_qp = cacheInfo->Center.Distance(Q);
+		double dist_qp = Distance(cacheInfo->Center, Q);
 		if (dist_qp > FWNBeta * cacheInfo->R)
 		{
 			return true;
diff --git a/Engine/Plugins/Experimental/GeometryProcessing/Source/GeometricObjects/Public/Spatial/MeshAABBTree3.h b/Engine/Plugins/Experimental/GeometryProcessing/Source/GeometricObjects/Public/Spatial/MeshAABBTree3.h
--- a/Engine/Plugins/Experimental/GeometryProcessing/Source/GeometricObjects/Public/Spatial/MeshAABBTree3.h
+++ b/Engine/Plugins/Experimental/GeometryProcessing/Source/GeometricObjects/Public/Spatial/MeshAABBTree3.h
@@ -289,7 +289,7 @@ protected:
 				FIndex3i Triangle = Mesh->GetTriangle(ti);
 				for (int j = 0; j < 3; ++j) 
 				{
-					double VertexDistSqr = Mesh->GetVertex(Triangle[j]).DistanceSquared(P);
+					double VertexDistSqr = DistanceSquared(Mesh->GetVertex(Triangle[j]), P);
 					if (VertexDistSqr < NearestDistSqr) 
 					{
 						NearestDistSqr = VertexDistSqr;
@@ -719,7 +719,8 @@ public:
 		{
 			if (TransformF != nullptr)
 			{
-				TestMeshBounds = FAxisAlignedBox3d(TestMeshBounds, TransformF);
+				TestMeshBounds = FAxisAlignedBox3d(TestMeshBounds,
+					[&](const FVector3<double>& P) { return TransformF(P); });
 			}
 			if (box_box_intersect(RootIndex, TestMeshBounds) == false)
 			{
@@ -903,7 +904,8 @@ protected:
 		if (TransformF != nullptr)
 		{
 			FAxisAlignedBox3d box = GetBox(iBox);
-			return FAxisAlignedBox3d(box, TransformF);
+			return FAxisAlignedBox3d(box, 
+				[&](const FVector3<double>& P) { return TransformF(P); });
 		}
 		else
 		{
diff --git a/Engine/Plugins/Experimental/GeometryProcessing/Source/GeometricObjects/Public/TransformTypes.h b/Engine/Plugins/Experimental/GeometryProcessing/Source/GeometricObjects/Public/TransformTypes.h
--- a/Engine/Plugins/Experimental/GeometryProcessing/Source/GeometricObjects/Public/TransformTypes.h
+++ b/Engine/Plugins/Experimental/GeometryProcessing/Source/GeometricObjects/Public/TransformTypes.h
@@ -31,27 +31,34 @@ public:
 		Scale3D = FVector3<RealType>::One();
 	}
 
-	TTransform3(const TQuaternion<RealType>& RotationIn, const FVector3<RealType>& TranslationIn, const FVector3<RealType>& ScaleIn)
+	TTransform3(const TQuaternion<RealType>& RotationIn, const UE::Core::TVector<RealType>& TranslationIn, const UE::Core::TVector<RealType>& ScaleIn)
 	{
 		Rotation = RotationIn;
 		Translation = TranslationIn;
 		Scale3D = ScaleIn;
 	}
 
-	explicit TTransform3(const TQuaternion<RealType>& RotationIn, const FVector3<RealType>& TranslationIn)
+	explicit TTransform3(const TQuaternion<RealType>& RotationIn, const UE::Core::TVector<RealType>& TranslationIn)
 	{
 		Rotation = RotationIn;
 		Translation = TranslationIn;
 		Scale3D = FVector3<RealType>::One();
 	}
 
-	explicit TTransform3(const FVector3<RealType>& TranslationIn)
+	explicit TTransform3(const UE::Core::TVector<RealType>& TranslationIn)
 	{
 		Rotation = TQuaternion<RealType>::Identity();
 		Translation = TranslationIn;
 		Scale3D = FVector3<RealType>::One();
 	}
 
+	explicit TTransform3(const FVector& TranslationIn)
+	{
+		Rotation = TQuaternion<RealType>::Identity();
+		Translation = FVector3<RealType>(TranslationIn);
+		Scale3D = FVector3<RealType>::One();
+	}
+
 	explicit TTransform3(const FTransform& Transform)
 	{
 		Rotation = TQuaternion<RealType>(Transform.GetRotation());
@@ -103,7 +110,7 @@ public:
 	/**
 	 * set Translation portion of transform
 	 */
-	void SetTranslation(const FVector3<RealType>& TranslationIn)
+	void SetTranslation(const UE::Core::TVector<RealType>& TranslationIn)
 	{
 		Translation = TranslationIn;
 	}
@@ -119,7 +126,7 @@ public:
 	/**
 	 * set Scale portion of transform
 	 */
-	void SetScale(const FVector3<RealType>& ScaleIn)
+	void SetScale(const UE::Core::TVector<RealType>& ScaleIn)
 	{
 		Scale3D = ScaleIn;
 	}
@@ -154,7 +161,7 @@ public:
 	/**
 	 * @return input vector with QS transformation applied, ie QS(V) = Rotate(Scale*V)
 	 */
-	FVector3<RealType> TransformVector(const FVector3<RealType>& V) const
+	FVector3<RealType> TransformVector(const UE::Core::TVector<RealType>& V) const
 	{
 		return Rotation * (Scale3D*V);
 	}
@@ -162,7 +169,7 @@ public:
 	/**
 	 * @return input vector with Q transformation applied, ie Q(V) = Rotate(V)
 	 */
-	FVector3<RealType> TransformVectorNoScale(const FVector3<RealType>& V) const
+	FVector3<RealType> TransformVectorNoScale(const UE::Core::TVector<RealType>& V) const
 	{
 		return Rotation * V;
 	}
@@ -172,7 +179,7 @@ public:
 	 * However 1/Scale requires special handling in case any component is near-zero.
 	 * @return input surface normal with transform applied.
 	 */
-	FVector3<RealType> TransformNormal(const FVector3<RealType>& Normal) const
+	FVector3<RealType> TransformNormal(const UE::Core::TVector<RealType>& Normal) const
 	{
 		// transform normal by a safe inverse scale + normalize, and a standard rotation
 		const FVector3<RealType>& S = Scale3D;
@@ -185,7 +192,7 @@ public:
 	/**
 	 * @return input vector with inverse-QST transformation applied, ie QSTinv(P) = InverseScale(InverseRotate(P - Translate))
 	 */
-	FVector3<RealType> InverseTransformPosition(const FVector3<RealType> &P) const
+	FVector3<RealType> InverseTransformPosition(const UE::Core::TVector<RealType> &P) const
 	{
 		return GetSafeScaleReciprocal(Scale3D) * Rotation.InverseMultiply(P - Translation);
 	}
@@ -193,7 +200,7 @@ public:
 	/**
 	 * @return input vector with inverse-QS transformation applied, ie QSinv(V) = InverseScale(InverseRotate(V))
 	 */
-	FVector3<RealType> InverseTransformVector(const FVector3<RealType> &V) const
+	FVector3<RealType> InverseTransformVector(const UE::Core::TVector<RealType> &V) const
 	{
 		return GetSafeScaleReciprocal(Scale3D) * Rotation.InverseMultiply(V);
 	}
@@ -202,7 +209,7 @@ public:
 	/**
 	 * @return input vector with inverse-Q transformation applied, ie Qinv(V) = InverseRotate(V)
 	 */
-	FVector3<RealType> InverseTransformVectorNoScale(const FVector3<RealType> &V) const
+	FVector3<RealType> InverseTransformVectorNoScale(const UE::Core::TVector<RealType> &V) const
 	{
 		return Rotation.InverseMultiply(V);
 	}
@@ -212,7 +219,7 @@ public:
 	 * Surface Normals are special, their inverse transform is InverseRotate( Normalize(Scale * Normal) ) )
 	 * @return input surface normal with inverse transform applied.
 	 */
-	FVector3<RealType> InverseTransformNormal(const FVector3<RealType>& Normal) const
+	FVector3<RealType> InverseTransformNormal(const UE::Core::TVector<RealType>& Normal) const
 	{
 		return InverseTransformVectorNoScale( Normalized(Scale3D*Normal) );
 	}
@@ -279,43 +286,43 @@ public:
 	// entire float transform to a double transform in order to apply to a double vector, which is the only
 	// case where this conversion is an issue
 
-	template<typename RealType2>
-	FVector3<RealType2> TransformPosition(const FVector3<RealType2>& P) const
+	template<typename RealType2, TEMPLATE_REQUIRES(std::is_same<RealType, RealType2>::value == false)>
+	FVector3<RealType2> TransformPosition(const UE::Core::TVector<RealType2>& P) const
 	{
 		return FVector3<RealType2>(TransformPosition(FVector3<RealType>(P)));
 	}
-	template<typename RealType2>
-	FVector3<RealType2> TransformVector(const FVector3<RealType2>& V) const
+	template<typename RealType2, TEMPLATE_REQUIRES(std::is_same<RealType, RealType2>::value == false)>
+	FVector3<RealType2> TransformVector(const UE::Core::TVector<RealType2>& V) const
 	{
 		return FVector3<RealType2>(TransformVector(FVector3<RealType>(V)));
 	}
-	template<typename RealType2>
-	FVector3<RealType2> TransformVectorNoScale(const FVector3<RealType2>& V) const
+	template<typename RealType2, TEMPLATE_REQUIRES(std::is_same<RealType, RealType2>::value == false)>
+	FVector3<RealType2> TransformVectorNoScale(const UE::Core::TVector<RealType2>& V) const
 	{
 		return FVector3<RealType2>(TransformVectorNoScale(FVector3<RealType>(V)));
 	}
-	template<typename RealType2>
-	FVector3<RealType2> TransformNormal(const FVector3<RealType2>& V) const
+	template<typename RealType2, TEMPLATE_REQUIRES(std::is_same<RealType, RealType2>::value == false)>
+	FVector3<RealType2> TransformNormal(const UE::Core::TVector<RealType2>& V) const
 	{
 		return FVector3<RealType2>(TransformNormal(FVector3<RealType>(V)));
 	}
-	template<typename RealType2>
-	FVector3<RealType2> InverseTransformPosition(const FVector3<RealType2>& P) const
+	template<typename RealType2, TEMPLATE_REQUIRES(std::is_same<RealType, RealType2>::value == false)>
+	FVector3<RealType2> InverseTransformPosition(const UE::Core::TVector<RealType2>& P) const
 	{
 		return FVector3<RealType2>(InverseTransformPosition(FVector3<RealType>(P)));
 	}
-	template<typename RealType2>
-	FVector3<RealType2> InverseTransformVector(const FVector3<RealType2>& V) const
+	template<typename RealType2, TEMPLATE_REQUIRES(std::is_same<RealType, RealType2>::value == false)>
+	FVector3<RealType2> InverseTransformVector(const UE::Core::TVector<RealType2>& V) const
 	{
 		return FVector3<RealType2>(InverseTransformVector(FVector3<RealType>(V)));
 	}
-	template<typename RealType2>
-	FVector3<RealType2> InverseTransformVectorNoScale(const FVector3<RealType2>& V) const
+	template<typename RealType2, TEMPLATE_REQUIRES(std::is_same<RealType, RealType2>::value == false)>
+	FVector3<RealType2> InverseTransformVectorNoScale(const UE::Core::TVector<RealType2>& V) const
 	{
 		return FVector3<RealType2>(InverseTransformVectorNoScale(FVector3<RealType>(V)));
 	}
-	template<typename RealType2>
-	FVector3<RealType2> InverseTransformNormal(const FVector3<RealType2>& V) const
+	template<typename RealType2, TEMPLATE_REQUIRES(std::is_same<RealType, RealType2>::value == false)>
+	FVector3<RealType2> InverseTransformNormal(const UE::Core::TVector<RealType2>& V) const
 	{
 		return FVector3<RealType2>(InverseTransformNormal(FVector3<RealType>(V)));
 	}
diff --git a/Engine/Plugins/Experimental/GeometryProcessing/Source/GeometricObjects/Public/VectorTypes.h b/Engine/Plugins/Experimental/GeometryProcessing/Source/GeometricObjects/Public/VectorTypes.h
--- a/Engine/Plugins/Experimental/GeometryProcessing/Source/GeometricObjects/Public/VectorTypes.h
+++ b/Engine/Plugins/Experimental/GeometryProcessing/Source/GeometricObjects/Public/VectorTypes.h
@@ -3,6 +3,7 @@
 #pragma once
 
 #include "Math/Vector.h"
+#include "Math/VectorLWC.h"
 #include "MathUtil.h"
 #include <sstream>
 
@@ -332,45 +333,52 @@ FVector2<T> Lerp(const FVector2<T>& A, const FVector2<T>& B, T Alpha)
  * @todo Possibly can be replaced/merged with Chaos TVector<T,N>
  */
 template <typename T>
-struct FVector3
+struct FVector3 : public UE::Core::TVector<T>
 {
-	T X{}, Y{}, Z{};
+	using UE::Core::TVector<T>::TVector;
+	using UE::Core::TVector<T>::X;
+	using UE::Core::TVector<T>::Y;
+	using UE::Core::TVector<T>::Z;
 
-	constexpr FVector3()
-		: X(0), Y(0), Z(0)
+	FVector3() : UE::Core::TVector<T>((T)0)
 	{
 	}
 
-	constexpr FVector3(T ValX, T ValY, T ValZ)
-		: X(ValX), Y(ValY), Z(ValZ)
+
+	explicit FVector3(const FVector& Vec)
 	{
+		X = (T)Vec.X;
+		Y = (T)Vec.Y;
+		Z = (T)Vec.Z;
 	}
 
-	constexpr FVector3(const T* Data)
-		: X(Data[0]), Y(Data[1]), Z(Data[2])
+	explicit FVector3(const FVector3f& Vec)
 	{
+		X = (T)Vec.X;
+		Y = (T)Vec.Y;
+		Z = (T)Vec.Z;
 	}
 
-	constexpr FVector3(const FVector2<T>& Data)
-		: X(Data.X), Y(Data.Y), Z((T)0)
+	explicit FVector3(const FVector3d& Vec)
 	{
+		X = (T)Vec.X;
+		Y = (T)Vec.Y;
+		Z = (T)Vec.Z;
 	}
 
-	constexpr FVector3(const FVector3& Vec) = default;
-
-	template<typename RealType2>
-	explicit constexpr FVector3(const FVector3<RealType2>& Vec)
-		: X((T)Vec.X), Y((T)Vec.Y), Z((T)Vec.Z)
+	FVector3(const UE::Core::TVector<T>& Vec)
 	{
+		X = Vec.X;
+		Y = Vec.Y;
+		Z = Vec.Z;
 	}
 
-	explicit constexpr operator const T*() const
-	{
-		return &X;
-	};
-	explicit constexpr operator T*()
+	FVector3& operator=(const UE::Core::TVector<T>& Vec)
 	{
-		return &X;
+		X = Vec.X;
+		Y = Vec.Y;
+		Z = Vec.Z;
+		return *this;
 	}
 
 	explicit operator FVector() const
@@ -378,21 +386,47 @@ struct FVector3
 		return FVector((float)X, (float)Y, (float)Z);
 	}
 
-	FVector3(const FVector& Vec)
-		: X((T)Vec.X), Y((T)Vec.Y), Z((T)Vec.Z)
+	explicit operator FVector3f() const
 	{
+		return FVector3f((float)X, (float)Y, (float)Z);
 	}
 
+	explicit operator FVector3d() const
+	{
+		return FVector3d((double)X, (double)Y, (double)Z);
+	}
 
-	explicit constexpr operator FLinearColor() const
+
+	explicit operator FLinearColor() const
 	{
 		return FLinearColor((float)X, (float)Y, (float)Z);
 	}
-	constexpr FVector3(const FLinearColor& Color)
-		: X((T)Color.R), Y((T)Color.G), Z((T)Color.B)
+	FVector3(const FLinearColor& Color)
+	{
+		X = (T)Color.R;
+		Y = (T)Color.G;
+		Z = (T)Color.B;
+	}
+
+
+	explicit FVector3(const T* Data)
 	{
+		X = Data[0];
+		Y = Data[1];
+		Z = Data[2];
 	}
 
+
+	explicit operator const T*() const
+	{
+		return &X;
+	};
+	explicit operator T*()
+	{
+		return &X;
+	}
+
+
 	static FVector3<T> Zero()
 	{
 		return FVector3<T>((T)0, (T)0, (T)0);
@@ -417,23 +451,6 @@ struct FVector3
 	{
 		return FVector3<T>(TNumericLimits<T>::Max(), TNumericLimits<T>::Max(), TNumericLimits<T>::Max());
 	}
-	
-	FVector3<T>& operator=(const FVector3<T>& V2)
-	{
-		X = V2.X;
-		Y = V2.Y;
-		Z = V2.Z;
-		return *this;
-	}
-
-	T& operator[](int Idx)
-	{
-		return (&X)[Idx];
-	}
-	const T& operator[](int Idx) const
-	{
-		return (&X)[Idx];
-	}
 
 	T Length() const
 	{
@@ -444,14 +461,14 @@ struct FVector3
 		return X * X + Y * Y + Z * Z;
 	}
 
-	constexpr T Distance(const FVector3<T>& V2) const
+	T Distance(const UE::Core::TVector<T>& V2) const
 	{
 		T dx = V2.X - X;
 		T dy = V2.Y - Y;
 		T dz = V2.Z - Z;
 		return TMathUtil<T>::Sqrt(dx * dx + dy * dy + dz * dz);
 	}
-	constexpr T DistanceSquared(const FVector3<T>& V2) const
+	T DistanceSquared(const UE::Core::TVector<T>& V2) const
 	{
 		T dx = V2.X - X;
 		T dy = V2.Y - Y;
@@ -459,58 +476,58 @@ struct FVector3
 		return dx * dx + dy * dy + dz * dz;
 	}
 
-	constexpr FVector3<T> operator-() const
+	FVector3<T> operator-() const
 	{
 		return FVector3<T>(-X, -Y, -Z);
 	}
 
-	constexpr FVector3<T> operator+(const FVector3<T>& V2) const
+	FVector3<T> operator+(const UE::Core::TVector<T>& V2) const
 	{
 		return FVector3<T>(X + V2.X, Y + V2.Y, Z + V2.Z);
 	}
 
-	constexpr FVector3<T> operator-(const FVector3<T>& V2) const
+	FVector3<T> operator-(const UE::Core::TVector<T>& V2) const
 	{
 		return FVector3<T>(X - V2.X, Y - V2.Y, Z - V2.Z);
 	}
 
-	constexpr FVector3<T> operator+(const T& Scalar) const
+	FVector3<T> operator+(const T& Scalar) const
 	{
 		return FVector3<T>(X + Scalar, Y + Scalar, Z + Scalar);
 	}
 
-	constexpr FVector3<T> operator-(const T& Scalar) const
+	FVector3<T> operator-(const T& Scalar) const
 	{
 		return FVector3<T>(X - Scalar, Y - Scalar, Z - Scalar);
 	}
 
-	constexpr FVector3<T> operator*(const T& Scalar) const
+	FVector3<T> operator*(const T& Scalar) const
 	{
 		return FVector3<T>(X * Scalar, Y * Scalar, Z * Scalar);
 	}
 
-	template<typename RealType2>
-	constexpr FVector3<T> operator*(const RealType2& Scalar) const
-	{
-		return FVector3<T>(X * (T)Scalar, Y * (T)Scalar, Z * (T)Scalar);
-	}
-
-	constexpr FVector3<T> operator*(const FVector3<T>& V2) const // component-wise
+	FVector3<T> operator*(const UE::Core::TVector<T>& V2) const // component-wise
 	{
 		return FVector3<T>(X * V2.X, Y * V2.Y, Z * V2.Z);
 	}
 
-	constexpr FVector3<T> operator/(const T& Scalar) const
+	FVector3<T> operator/(const T& Scalar) const
 	{
 		return FVector3<T>(X / Scalar, Y / Scalar, Z / Scalar);
 	}
 
-	constexpr FVector3<T> operator/(const FVector3<T>& V2) const // component-wise
+	FVector3<T> operator/(const UE::Core::TVector<T>& V2) const // component-wise
 	{
 		return FVector3<T>(X / V2.X, Y / V2.Y, Z / V2.Z);
 	}
 
-	constexpr FVector3<T>& operator+=(const FVector3<T>& V2)
+	template<typename RealType2, TEMPLATE_REQUIRES(std::is_floating_point<RealType2>::value)>
+	FVector3<T> operator*(const RealType2& Scalar) const
+	{
+		return FVector3<T>(X * (T)Scalar, Y * (T)Scalar, Z * (T)Scalar);
+	}
+
+	FVector3<T>& operator+=(const UE::Core::TVector<T>& V2)
 	{
 		X += V2.X;
 		Y += V2.Y;
@@ -518,7 +535,7 @@ struct FVector3
 		return *this;
 	}
 
-	constexpr FVector3<T>& operator-=(const FVector3<T>& V2)
+	FVector3<T>& operator-=(const UE::Core::TVector<T>& V2)
 	{
 		X -= V2.X;
 		Y -= V2.Y;
@@ -526,7 +543,7 @@ struct FVector3
 		return *this;
 	}
 
-	constexpr FVector3<T>& operator*=(const T& Scalar)
+	FVector3<T>& operator*=(const T& Scalar)
 	{
 		X *= Scalar;
 		Y *= Scalar;
@@ -534,7 +551,7 @@ struct FVector3
 		return *this;
 	}
 
-	constexpr FVector3<T>& operator/=(const T& Scalar)
+	FVector3<T>& operator/=(const T& Scalar)
 	{
 		X /= Scalar;
 		Y /= Scalar;
@@ -542,51 +559,56 @@ struct FVector3
 		return *this;
 	}
 
-	T Dot(const FVector3<T>& V2) const
+	T Dot(const UE::Core::TVector<T>& V2) const
 	{
 		return X * V2.X + Y * V2.Y + Z * V2.Z;
 	}
 
-	FVector3<T> Cross(const FVector3<T>& V2) const
+	FVector3<T> Cross(const UE::Core::TVector<T>& V2) const
 	{
 		return FVector3(
 			Y * V2.Z - Z * V2.Y,
 			Z * V2.X - X * V2.Z,
 			X * V2.Y - Y * V2.X);
 	}
+};
 
-	constexpr bool operator==(const FVector3<T>& Other) const
-	{
-		return X == Other.X && Y == Other.Y && Z == Other.Z;
-	}
 
-	constexpr bool operator!=(const FVector3<T>& Other) const
-	{
-		return X != Other.X || Y != Other.Y || Z != Other.Z;
-	}
-};
 
 
 
 /** @return unit vector along axis X=0, Y=1, Z=2 */
 template <typename T>
-constexpr FVector3<T> MakeUnitVector3(int32 Axis)
+constexpr UE::Core::TVector<T> MakeUnitVector3(int32 Axis)
 {
-	FVector3<T> UnitVec((T)0, (T)0, (T)0);
+	UE::Core::TVector<T> UnitVec((T)0, (T)0, (T)0);
 	UnitVec[FMath::Clamp(Axis, 0, 2)] = (T)1;
 	return UnitVec;
 }
 
 
+template<typename T>
+T Length(const UE::Core::TVector<T>& V)
+{
+	return TMathUtil<T>::Sqrt(V.X*V.X + V.Y*V.Y + V.Z*V.Z);
+}
+
+template<typename T>
+T SquaredLength(const UE::Core::TVector<T>& V)
+{
+	return V.X*V.X + V.Y*V.Y + V.Z*V.Z;
+}
+
+
 template <typename T>
-constexpr bool IsNormalized(const FVector3<T>& Vector, const T Tolerance = TMathUtil<T>::ZeroTolerance)
+constexpr bool IsNormalized(const UE::Core::TVector<T>& Vector, const T Tolerance = TMathUtil<T>::ZeroTolerance)
 {
 	return TMathUtil<T>::Abs((Vector.X*Vector.X + Vector.Y*Vector.Y + Vector.Z*Vector.Z) - 1) < Tolerance;
 }
 
 
 template<typename T>
-T Normalize(FVector3<T>& Vector, const T Epsilon = 0)
+T Normalize(UE::Core::TVector<T>& Vector, const T Epsilon = 0)
 {
 	T length = Vector.Length();
 	if (length > Epsilon)
@@ -602,22 +624,56 @@ T Normalize(FVector3<T>& Vector, const T Epsilon = 0)
 }
 
 template<typename T>
-constexpr FVector3<T> Normalized(const FVector3<T>& Vector, const T Epsilon = 0)
+constexpr UE::Core::TVector<T> Normalized(const UE::Core::TVector<T>& Vector, const T Epsilon = 0)
 {
 	T length = Vector.Length();
 	if (length > Epsilon)
 	{
 		T invLength = ((T)1) / length;
-		return FVector3<T>(Vector.X*invLength, Vector.Y*invLength, Vector.Z*invLength);
+		return UE::Core::TVector<T>(Vector.X*invLength, Vector.Y*invLength, Vector.Z*invLength);
 	}
-	return FVector3<T>((T)0, (T)0, (T)0);
+	return UE::Core::TVector<T>((T)0, (T)0, (T)0);
 }
 
+template<typename T>
+T Distance(const UE::Core::TVector<T>& V1, const UE::Core::TVector<T>& V2)
+{
+	T dx = V2.X - V1.X;
+	T dy = V2.Y - V1.Y;
+	T dz = V2.Z - V1.Z;
+	return TMathUtil<T>::Sqrt(dx * dx + dy * dy + dz * dz);
+}
+
+template<typename T>
+T DistanceSquared(const UE::Core::TVector<T>& V1, const UE::Core::TVector<T>& V2)
+{
+	T dx = V2.X - V1.X;
+	T dy = V2.Y - V1.Y;
+	T dz = V2.Z - V1.Z;
+	return dx * dx + dy * dy + dz * dz;
+}
+
+
+
+template<typename T>
+T Dot(const UE::Core::TVector<T>& V1, const UE::Core::TVector<T>& V2)
+{
+	return V1.X * V2.X + V1.Y * V2.Y + V1.Z * V2.Z;
+}
 
 template<typename T>
-FVector3<T> UnitCross(const FVector3<T>& V1, const FVector3<T>& V2)
+UE::Core::TVector<T> Cross(const UE::Core::TVector<T>& V1, const UE::Core::TVector<T>& V2)
 {
-	FVector3<T> N = V1.Cross(V2);
+	return UE::Core::TVector<T>(
+		V1.Y * V2.Z - V1.Z * V2.Y,
+		V1.Z * V2.X - V1.X * V2.Z,
+		V1.X * V2.Y - V1.Y * V2.X);
+}
+
+template<typename T>
+UE::Core::TVector<T> UnitCross(const UE::Core::TVector<T>& V1, const UE::Core::TVector<T>& V2)
+{
+	UE::Core::TVector<T> N = V1.Cross(V2);
 	return Normalized(N);
 }
 
@@ -626,7 +682,7 @@ FVector3<T> UnitCross(const FVector3<T>& V1, const FVector3<T>& V2)
  * @return the (positive) angle between V1 and V2 in degrees
  */
 template <typename T>
-T AngleD(const FVector3<T>& V1, const FVector3<T>& V2)
+T AngleD(const UE::Core::TVector<T>& V1, const UE::Core::TVector<T>& V2)
 {
 	T DotVal = V1.Dot(V2);
 	T ClampedDot = (DotVal < (T)-1) ? (T)-1 : ((DotVal > (T)1) ? (T)1 : DotVal);
@@ -638,7 +694,7 @@ T AngleD(const FVector3<T>& V1, const FVector3<T>& V2)
  * @return the (positive) angle between V1 and V2 in radians
  */
 template <typename T>
-T AngleR(const FVector3<T>& V1, const FVector3<T>& V2)
+T AngleR(const UE::Core::TVector<T>& V1, const UE::Core::TVector<T>& V2)
 {
 	T DotVal = V1.Dot(V2);
 	T ClampedDot = (DotVal < (T)-1) ? (T)-1 : ((DotVal > (T)1) ? (T)1 : DotVal);
@@ -646,26 +702,26 @@ T AngleR(const FVector3<T>& V1, const FVector3<T>& V2)
 }
 
 template <typename T>
-constexpr FVector2<T> GetXY(const FVector3<T>& V)
+constexpr FVector2<T> GetXY(const UE::Core::TVector<T>& V)
 {
 	return FVector2<T>(V.X, V.Y);
 }
 
 template <typename T>
-constexpr FVector2<T> GetXZ(const FVector3<T>& V)
+constexpr FVector2<T> GetXZ(const UE::Core::TVector<T>& V)
 {
 	return FVector2<T>(V.X, V.Z);
 }
 
 template <typename T>
-constexpr FVector2<T> GetYZ(const FVector3<T>& V)
+constexpr FVector2<T> GetYZ(const UE::Core::TVector<T>& V)
 {
 	return FVector2<T>(V.Y, V.Z);
 }
 
 
 template<typename T>
-constexpr FVector3<T> Min(const FVector3<T>& V0, const FVector3<T>& V1)
+constexpr FVector3<T> Min(const UE::Core::TVector<T>& V0, const UE::Core::TVector<T>& V1)
 {
 	return FVector3<T>(TMathUtil<T>::Min(V0.X, V1.X),
 		TMathUtil<T>::Min(V0.Y, V1.Y),
@@ -673,7 +729,7 @@ constexpr FVector3<T> Min(const FVector3<T>& V0, const FVector3<T>& V1)
 }
 
 template<typename T>
-constexpr FVector3<T> Max(const FVector3<T>& V0, const FVector3<T>& V1)
+constexpr FVector3<T> Max(const UE::Core::TVector<T>& V0, const UE::Core::TVector<T>& V1)
 {
 	return FVector3<T>(TMathUtil<T>::Max(V0.X, V1.X),
 		TMathUtil<T>::Max(V0.Y, V1.Y),
@@ -682,59 +738,59 @@ constexpr FVector3<T> Max(const FVector3<T>& V0, const FVector3<T>& V1)
 
 
 template<typename T>
-constexpr T MaxElement(const FVector3<T>& Vector)
+constexpr T MaxElement(const UE::Core::TVector<T>& Vector)
 {
 	return TMathUtil<T>::Max3(Vector.X, Vector.Y, Vector.Z);
 }
 
 /** @return 0/1/2 index of maximum element */
 template<typename T>
-constexpr int32 MaxElementIndex(const FVector3<T>& Vector)
+constexpr int32 MaxElementIndex(const UE::Core::TVector<T>& Vector)
 {
 	return TMathUtil<T>::Max3Index(Vector.X, Vector.Y, Vector.Z);
 }
 
 template<typename T>
-constexpr T MinElement(const FVector3<T>& Vector)
+constexpr T MinElement(const UE::Core::TVector<T>& Vector)
 {
 	return TMathUtil<T>::Min3(Vector.X, Vector.Y, Vector.Z);
 }
 
 /** @return 0/1/2 index of minimum element */
 template<typename T>
-constexpr int32 MinElementIndex(const FVector3<T>& Vector)
+constexpr int32 MinElementIndex(const UE::Core::TVector<T>& Vector)
 {
 	return TMathUtil<T>::Min3Index(Vector.X, Vector.Y, Vector.Z);
 }
 
 template<typename T>
-constexpr T MaxAbsElement(const FVector3<T>& Vector)
+constexpr T MaxAbsElement(const UE::Core::TVector<T>& Vector)
 {
 	return TMathUtil<T>::Max3(TMathUtil<T>::Abs(Vector.X), TMathUtil<T>::Abs(Vector.Y), TMathUtil<T>::Abs(Vector.Z));
 }
 
 /** @return 0/1/2 index of maximum absolute-value element */
 template<typename T>
-constexpr T MaxAbsElementIndex(const FVector3<T>& Vector)
+constexpr T MaxAbsElementIndex(const UE::Core::TVector<T>& Vector)
 {
 	return TMathUtil<T>::Max3Index(TMathUtil<T>::Abs(Vector.X), TMathUtil<T>::Abs(Vector.Y), TMathUtil<T>::Abs(Vector.Z));
 }
 
 template<typename T>
-constexpr T MinAbsElement(const FVector3<T>& Vector)
+constexpr T MinAbsElement(const UE::Core::TVector<T>& Vector)
 {
 	return TMathUtil<T>::Min3(TMathUtil<T>::Abs(Vector.X), TMathUtil<T>::Abs(Vector.Y), TMathUtil<T>::Abs(Vector.Z));
 }
 
 /** @return 0/1/2 index of minimum absolute-value element */
 template<typename T>
-constexpr T MinAbsElementIndex(const FVector3<T>& Vector)
+constexpr T MinAbsElementIndex(const UE::Core::TVector<T>& Vector)
 {
 	return TMathUtil<T>::Min3Index(TMathUtil<T>::Abs(Vector.X), TMathUtil<T>::Abs(Vector.Y), TMathUtil<T>::Abs(Vector.Z));
 }
 
 template<typename T>
-constexpr FColor ToFColor(const FVector3<T>& Vector)
+constexpr FColor ToFColor(const UE::Core::TVector<T>& Vector)
 {
 	return FColor(
 		FMathf::Clamp((int)((float)Vector.X * 255.0f), 0, 255),
@@ -743,18 +799,24 @@ constexpr FColor ToFColor(const FVector3<T>& Vector)
 }
 
 template<typename T>
-FVector3<T> Lerp(const FVector3<T>& A, const FVector3<T>& B, T Alpha)
+constexpr FLinearColor ToLinearColor(const UE::Core::TVector<T>& Vector)
+{
+	return FLinearColor((float)Vector.X, (float)Vector.Y, (float)Vector.Z);
+}
+
+template<typename T>
+UE::Core::TVector<T> Lerp(const UE::Core::TVector<T>& A, const UE::Core::TVector<T>& B, T Alpha)
 {
 	T OneMinusAlpha = (T)1 - Alpha;
-	return FVector3<T>(OneMinusAlpha * A.X + Alpha * B.X,
+	return UE::Core::TVector<T>(OneMinusAlpha * A.X + Alpha * B.X,
 		OneMinusAlpha * A.Y + Alpha * B.Y,
 		OneMinusAlpha * A.Z + Alpha * B.Z);
 }
 
 template<typename T>
-FVector3<T> Blend3(const FVector3<T>& A, const FVector3<T>& B, const FVector3<T>& C, const T& WeightA, const T& WeightB, const T& WeightC)
+UE::Core::TVector<T> Blend3(const UE::Core::TVector<T>& A, const UE::Core::TVector<T>& B, const UE::Core::TVector<T>& C, const T& WeightA, const T& WeightB, const T& WeightC)
 {
-	return FVector3<T>(
+	return UE::Core::TVector<T>(
 		WeightA * A.X + WeightB * B.X + WeightC * C.X,
 		WeightA * A.Y + WeightB * B.Y + WeightC * C.Y,
 		WeightA * A.Z + WeightB * B.Z + WeightC * C.Z);
@@ -767,7 +829,7 @@ inline FVector3<RealType> operator*(RealType Scalar, const FVector3<RealType>& V
 }
 
 // allow float*Vector3<double> and double*Vector3<float>
-template <typename RealType, typename RealType2>
+template <typename RealType, typename RealType2, TEMPLATE_REQUIRES(std::is_floating_point<RealType2>::value)>
 inline FVector3<RealType> operator*(RealType2 Scalar, const FVector3<RealType>& V)
 {
 	return FVector3<RealType>((RealType)Scalar * V.X, (RealType)Scalar * V.Y, (RealType)Scalar * V.Z);
@@ -1115,6 +1177,8 @@ FORCEINLINE uint32 GetTypeHash(const TVector4<T>& Vector)
 }
 
 
+
+
 } // end namespace UE::Geometry
 } // end namespace UE
 
@@ -1122,8 +1186,8 @@ FORCEINLINE uint32 GetTypeHash(const TVector4<T>& Vector)
 typedef UE::Geometry::FVector2<float> FVector2f;
 typedef UE::Geometry::FVector2<double> FVector2d;
 
-typedef UE::Geometry::FVector3<float> FVector3f;
-typedef UE::Geometry::FVector3<double> FVector3d;
+//typedef UE::Geometry::FVector3<float> FVector3f;
+//typedef UE::Geometry::FVector3<double> FVector3d;
 
 typedef UE::Geometry::TVector4<float> FVector4f;
 typedef UE::Geometry::TVector4<double> FVector4d;
\ No newline at end of file
diff --git a/Engine/Plugins/Experimental/GeometryProcessing/Source/GeometricObjects/Public/VectorUtil.h b/Engine/Plugins/Experimental/GeometryProcessing/Source/GeometricObjects/Public/VectorUtil.h
--- a/Engine/Plugins/Experimental/GeometryProcessing/Source/GeometricObjects/Public/VectorUtil.h
+++ b/Engine/Plugins/Experimental/GeometryProcessing/Source/GeometricObjects/Public/VectorUtil.h
@@ -43,7 +43,7 @@ namespace VectorUtil
 	 * @return true if all components of V are finite
 	 */
 	template <typename RealType>
-	inline bool IsFinite(const FVector3<RealType>& V)
+	inline bool IsFinite(const UE::Core::TVector<RealType>& V)
 	{
 		return TMathUtil<RealType>::IsFinite(V.X) && TMathUtil<RealType>::IsFinite(V.Y) && TMathUtil<RealType>::IsFinite(V.Z);
 	}
@@ -62,7 +62,7 @@ namespace VectorUtil
 	 * @return normalized vector that is perpendicular to triangle V0,V1,V2  (triangle normal)
 	 */
 	template <typename RealType>
-	inline FVector3<RealType> Normal(const FVector3<RealType>& V0, const FVector3<RealType>& V1, const FVector3<RealType>& V2)
+	inline FVector3<RealType> Normal(const UE::Core::TVector<RealType>& V0, const UE::Core::TVector<RealType>& V1, const UE::Core::TVector<RealType>& V2)
 	{
 		FVector3<RealType> edge1(V1 - V0);
 		FVector3<RealType> edge2(V2 - V0);
@@ -78,7 +78,7 @@ namespace VectorUtil
 	 * @return un-normalized direction that is parallel to normal of triangle V0,V1,V2
 	 */
 	template <typename RealType>
-	inline FVector3<RealType> NormalDirection(const FVector3<RealType>& V0, const FVector3<RealType>& V1, const FVector3<RealType>& V2)
+	inline FVector3<RealType> NormalDirection(const UE::Core::TVector<RealType>& V0, const UE::Core::TVector<RealType>& V1, const UE::Core::TVector<RealType>& V2)
 	{
 		// Unreal has Left-Hand Coordinate System so we need to reverse this cross-product to get proper triangle normal
 		return (V2 - V0).Cross(V1 - V0);
@@ -89,7 +89,7 @@ namespace VectorUtil
 	 * @return area of 3D triangle V0,V1,V2
 	 */
 	template <typename RealType>
-	inline RealType Area(const FVector3<RealType>& V0, const FVector3<RealType>& V1, const FVector3<RealType>& V2)
+	inline RealType Area(const UE::Core::TVector<RealType>& V0, const UE::Core::TVector<RealType>& V1, const UE::Core::TVector<RealType>& V2)
 	{
 		FVector3<RealType> Edge1(V1 - V0);
 		FVector3<RealType> Edge2(V2 - V0);
@@ -113,11 +113,11 @@ namespace VectorUtil
 	 * @return true if triangle V1,V2,V3 is obtuse
 	 */
 	template <typename RealType>
-	inline bool IsObtuse(const FVector3<RealType>& V1, const FVector3<RealType>& V2, const FVector3<RealType>& V3)
+	inline bool IsObtuse(const UE::Core::TVector<RealType>& V1, const UE::Core::TVector<RealType>& V2, const UE::Core::TVector<RealType>& V3)
 	{
-		RealType a2 = V1.DistanceSquared(V2);
-		RealType b2 = V1.DistanceSquared(V3);
-		RealType c2 = V2.DistanceSquared(V3);
+		RealType a2 = DistanceSquared(V1, V2);
+		RealType b2 = DistanceSquared(V1, V3);
+		RealType c2 = DistanceSquared(V2, V3);
 		return (a2 + b2 < c2) || (b2 + c2 < a2) || (c2 + a2 < b2);
 	}
 
@@ -126,7 +126,7 @@ namespace VectorUtil
 	 * @return triangle normal
 	 */
 	template <typename RealType>
-	inline FVector3<RealType> NormalArea(const FVector3<RealType>& V0, const FVector3<RealType>& V1, const FVector3<RealType>& V2, RealType& AreaOut)
+	inline FVector3<RealType> NormalArea(const UE::Core::TVector<RealType>& V0, const UE::Core::TVector<RealType>& V1, const UE::Core::TVector<RealType>& V2, RealType& AreaOut)
 	{
 		FVector3<RealType> edge1(V1 - V0);
 		FVector3<RealType> edge2(V2 - V0);
@@ -153,7 +153,7 @@ namespace VectorUtil
 
 	/** @return true if all coordinates of V0 and V1 are within Epsilon of eachother */
 	template <typename RealType>
-	inline bool EpsilonEqual(const FVector3<RealType>& V0, const FVector3<RealType>& V1, RealType Epsilon)
+	inline bool EpsilonEqual(const UE::Core::TVector<RealType>& V0, const UE::Core::TVector<RealType>& V1, RealType Epsilon)
 	{
 		return EpsilonEqual(V0.X, V1.X, Epsilon) && EpsilonEqual(V0.Y, V1.Y, Epsilon) && EpsilonEqual(V0.Z, V1.Z, Epsilon);
 	}
@@ -178,7 +178,7 @@ namespace VectorUtil
 	 * Calculates two vectors perpendicular to input Normal, as efficiently as possible.
 	 */
 	template <typename RealType>
-	inline void MakePerpVectors(const FVector3<RealType>& Normal, FVector3<RealType>& OutPerp1, FVector3<RealType>& OutPerp2)
+	inline void MakePerpVectors(const UE::Core::TVector<RealType>& Normal, UE::Core::TVector<RealType>& OutPerp1, UE::Core::TVector<RealType>& OutPerp2)
 	{
 		// Duff et al method, from https://graphics.pixar.com/library/OrthonormalB/paper.pdf
 		if (Normal.Z < (RealType)0)
@@ -236,7 +236,7 @@ namespace VectorUtil
 	 * @return angle in degrees
 	 */
 	template <typename RealType>
-	inline RealType PlaneAngleSignedD(const FVector3<RealType>& VFrom, const FVector3<RealType>& VTo, const FVector3<RealType>& PlaneN)
+	inline RealType PlaneAngleSignedD(const UE::Core::TVector<RealType>& VFrom, const UE::Core::TVector<RealType>& VTo, const UE::Core::TVector<RealType>& PlaneN)
 	{
 		FVector3<RealType> vFrom = VFrom - VFrom.Dot(PlaneN) * PlaneN;
 		FVector3<RealType> vTo = VTo - VTo.Dot(PlaneN) * PlaneN;
@@ -256,7 +256,7 @@ namespace VectorUtil
 	 * @return positive value of tan(theta/2) where theta is angle between normalized vectors A and B
 	 */
 	template <typename RealType>
-	RealType VectorTanHalfAngle(const FVector3<RealType>& A, const FVector3<RealType>& B)
+	RealType VectorTanHalfAngle(const UE::Core::TVector<RealType>& A, const UE::Core::TVector<RealType>& B)
 	{
 		RealType cosAngle = A.Dot(B);
 		RealType sqr = ((RealType)1 - cosAngle) / ((RealType)1 + cosAngle);
@@ -270,7 +270,7 @@ namespace VectorUtil
 	 * @return cotangent of angle between V1 and V2, or zero if result would be unstable (eg infinity)
 	 */ 
 	template <typename RealType>
-	RealType VectorCot(const FVector3<RealType>& V1, const FVector3<RealType>& V2)
+	RealType VectorCot(const UE::Core::TVector<RealType>& V1, const UE::Core::TVector<RealType>& V2)
 	{
 		// formula from http://www.geometry.caltech.edu/pubs/DMSB_III.pdf
 		RealType fDot = V1.Dot(V2);
@@ -294,7 +294,7 @@ namespace VectorUtil
 	 * TODO: make robust to degenerate triangles?
 	 */
 	template <typename RealType>
-	FVector3<RealType> BarycentricCoords(const FVector3<RealType>& Point, const FVector3<RealType>& V0, const FVector3<RealType>& V1, const FVector3<RealType>& V2)
+	FVector3<RealType> BarycentricCoords(const UE::Core::TVector<RealType>& Point, const UE::Core::TVector<RealType>& V0, const UE::Core::TVector<RealType>& V1, const UE::Core::TVector<RealType>& V2)
 	{
 		FVector3<RealType> kV02 = V0 - V2;
 		FVector3<RealType> kV12 = V1 - V2;
@@ -341,7 +341,7 @@ namespace VectorUtil
 	 * @return solid angle at point P for triangle A,B,C
 	 */
 	template <typename RealType>
-	inline RealType TriSolidAngle(FVector3<RealType> A, FVector3<RealType> B, FVector3<RealType> C, const FVector3<RealType>& P)
+	inline RealType TriSolidAngle(UE::Core::TVector<RealType> A, UE::Core::TVector<RealType> B, UE::Core::TVector<RealType> C, const UE::Core::TVector<RealType>& P)
 	{
 		// Formula from https://igl.ethz.ch/projects/winding-number/
 		A -= P;
@@ -361,7 +361,7 @@ namespace VectorUtil
 	 * @return gradient (3D vector) lying in plane of triangle.
 	 */
 	template <typename RealType>
-	inline FVector3<RealType> TriGradient(FVector3<RealType> Vi, FVector3<RealType> Vj, FVector3<RealType> Vk, RealType fi, RealType fj, RealType fk)
+	inline FVector3<RealType> TriGradient(UE::Core::TVector<RealType> Vi, UE::Core::TVector<RealType> Vj, UE::Core::TVector<RealType> Vk, RealType fi, RealType fj, RealType fk)
 	{
 		// recenter (better for precision)
 		FVector3<RealType> Centroid = (Vi + Vj + Vk) / (RealType)3;
@@ -387,7 +387,7 @@ namespace VectorUtil
 	 * @return angle between vectors (A-CornerPt) and (B-CornerPt)
 	 */
 	template<typename RealType>
-	inline RealType OpeningAngleD(FVector3<RealType> A, FVector3<RealType> B, const FVector3<RealType>& P)
+	inline RealType OpeningAngleD(UE::Core::TVector<RealType> A, UE::Core::TVector<RealType> B, const UE::Core::TVector<RealType>& P)
 	{
 		A -= P; 
 		Normalize(A);
@@ -402,7 +402,7 @@ namespace VectorUtil
 	 * @return sign of Bitangent relative to Normal and Tangent
 	 */
 	template<typename RealType>
-	inline RealType BitangentSign(const FVector3<RealType>& NormalIn, const FVector3<RealType>& TangentIn, const FVector3<RealType>& BitangentIn)
+	inline RealType BitangentSign(const UE::Core::TVector<RealType>& NormalIn, const UE::Core::TVector<RealType>& TangentIn, const UE::Core::TVector<RealType>& BitangentIn)
 	{
 		// following math from RenderUtils.h::GetBasisDeterminantSign()
 		RealType Cross00 = BitangentIn.Y*NormalIn.Z - BitangentIn.Z*NormalIn.Y;
@@ -416,7 +416,7 @@ namespace VectorUtil
 	 * @return Bitangent vector based on given Normal, Tangent, and Sign value (+1/-1)
 	 */
 	template<typename RealType>
-	inline FVector3<RealType> Bitangent(const FVector3<RealType>& NormalIn, const FVector3<RealType>& TangentIn, RealType BitangentSign)
+	inline FVector3<RealType> Bitangent(const UE::Core::TVector<RealType>& NormalIn, const UE::Core::TVector<RealType>& TangentIn, RealType BitangentSign)
 	{
 		return BitangentSign * FVector3<RealType>(
 			NormalIn.Y*TangentIn.Z - NormalIn.Z*TangentIn.Y,
@@ -428,7 +428,7 @@ namespace VectorUtil
 	 * @return Tangent-Space vector based on given Normal and Bitangent
 	 */
 	template<typename RealType>
-	inline FVector3<RealType> TangentFromBitangent(const FVector3<RealType>& NormalIn, const FVector3<RealType>& BitangentIn)
+	inline FVector3<RealType> TangentFromBitangent(const UE::Core::TVector<RealType>& NormalIn, const UE::Core::TVector<RealType>& BitangentIn)
 	{
 		return BitangentIn.Cross(NormalIn);
 	}
@@ -437,7 +437,7 @@ namespace VectorUtil
 	 * @return Bitangent vector based on given Normal and Tangent
 	 */
 	template<typename RealType>
-	inline FVector3<RealType> BitangentFromTangent(const FVector3<RealType>& NormalIn, const FVector3<RealType>& TangentIn)
+	inline FVector3<RealType> BitangentFromTangent(const UE::Core::TVector<RealType>& NormalIn, const UE::Core::TVector<RealType>& TangentIn)
 	{
 		return NormalIn.Cross(TangentIn);
 	}
@@ -445,7 +445,7 @@ namespace VectorUtil
 	/// @return Aspect ratio of triangle 
 	inline double AspectRatio(const FVector3d& v1, const FVector3d& v2, const FVector3d& v3)
 	{
-		double a = v1.Distance(v2), b = v2.Distance(v3), c = v3.Distance(v1);
+		double a = Distance(v1, v2), b = Distance(v2, v3), c = Distance(v3, v1);
 		double s = (a + b + c) / 2.0;
 		return (a * b * c) / (8.0 * (s - a) * (s - b) * (s - c));
 	}
diff --git a/Engine/Plugins/Experimental/GeometryProcessing/Source/GeometryAlgorithms/Public/ConvexHull3.h b/Engine/Plugins/Experimental/GeometryProcessing/Source/GeometryAlgorithms/Public/ConvexHull3.h
--- a/Engine/Plugins/Experimental/GeometryProcessing/Source/GeometryAlgorithms/Public/ConvexHull3.h
+++ b/Engine/Plugins/Experimental/GeometryProcessing/Source/GeometryAlgorithms/Public/ConvexHull3.h
@@ -46,7 +46,8 @@ public:
 	 * @param FilterFunc Optional filter to include only a subset of the points in the output hull
 	 * @return true if hull was generated, false if points span < 2 dimensions
 	 */
-	bool Solve(TArrayView<const FVector3<RealType>> Points, TFunctionRef<bool(int32)> FilterFunc)
+	template<typename VectorType>
+	bool Solve(TArrayView<const VectorType> Points, TFunctionRef<bool(int32)> FilterFunc)
 	{
 		return Solve(Points.Num(), [&Points](int32 Idx)
 			{
@@ -63,7 +64,8 @@ public:
 	 * @param Points Array of points to consider
 	 * @return true if hull was generated, false if points span < 2 dimensions
 	 */
-	bool Solve(TArrayView<const FVector3<RealType>> Points)
+	template<typename VectorType>
+	bool Solve(TArrayView<const VectorType> Points)
 	{
 		return Solve(Points.Num(), [&Points](int32 Idx)
 			{
@@ -102,7 +104,8 @@ public:
 	 * @param Points Array of points to consider
 	 * @return Array of halfspaces
 	 */
-	TArray<THalfspace3<RealType>> GetAsHalfspaces(TArrayView<const FVector3<RealType>> Points) const
+	template<typename VectorType>
+	TArray<THalfspace3<RealType>> GetAsHalfspaces(TArrayView<const VectorType> Points) const
 	{
 		TArray<THalfspace3<RealType>> Halfspaces;
 		for (FIndex3i Tri : Hull)
@@ -121,7 +124,8 @@ public:
 	 * @param GetPointFunc Function providing array-style access into points
 	 * @return Array of halfspaces
 	 */
-	TArray<THalfspace3<RealType>> GetAsHalfspaces(TFunctionRef<FVector3<RealType>(int32)> GetPointFunc) const
+	template<typename VectorType>
+	TArray<THalfspace3<RealType>> GetAsHalfspaces(TFunctionRef<VectorType(int32)> GetPointFunc) const
 	{
 		TArray<THalfspace3<RealType>> Halfspaces;
 		for (FIndex3i Tri : Hull)
diff --git a/Engine/Plugins/Experimental/GeometryProcessing/Source/MeshConversion/Private/IndexedMeshToDynamicMesh.cpp b/Engine/Plugins/Experimental/GeometryProcessing/Source/MeshConversion/Private/IndexedMeshToDynamicMesh.cpp
--- a/Engine/Plugins/Experimental/GeometryProcessing/Source/MeshConversion/Private/IndexedMeshToDynamicMesh.cpp
+++ b/Engine/Plugins/Experimental/GeometryProcessing/Source/MeshConversion/Private/IndexedMeshToDynamicMesh.cpp
@@ -44,19 +44,19 @@ void UE::Conversion::RenderBuffersToDynamicMesh(
 	for (int32 vi = 0; vi < NumVertices; ++vi)
 	{
 		FVector Position = PositionBuffer.VertexPosition(vi);
-		int32 vid = MeshOut.AppendVertex(Position);
+		int32 vid = MeshOut.AppendVertex(FVector3d(Position));
 		ensure(vid == vi);
 
 		FVector Normal = VertexBuffer.VertexTangentZ(vi);
-		int32 nid = Normals->AppendElement((FVector3f)Normal);
+		int32 nid = Normals->AppendElement(&Normal.X);
 		ensure(nid == vi);
 
 		FVector Tangent = VertexBuffer.VertexTangentX(vi);
-		int32 tid = Tangents->AppendElement((FVector3f)Tangent);
+		int32 tid = Tangents->AppendElement(&Tangent.X);
 		ensure(tid == vi);
 
 		FVector BiTangent = VertexBuffer.VertexTangentY(vi);
-		int32 bid = BiTangents->AppendElement((FVector3f)BiTangent);
+		int32 bid = BiTangents->AppendElement(&BiTangent.X);
 		ensure(bid == vi);
 
 		for (int32 ui = 0; ui < NumUVChannels; ++ui)
diff --git a/Engine/Plugins/Experimental/GeometryProcessing/Source/MeshConversion/Private/MeshDescriptionBuilder.cpp b/Engine/Plugins/Experimental/GeometryProcessing/Source/MeshConversion/Private/MeshDescriptionBuilder.cpp
--- a/Engine/Plugins/Experimental/GeometryProcessing/Source/MeshConversion/Private/MeshDescriptionBuilder.cpp
+++ b/Engine/Plugins/Experimental/GeometryProcessing/Source/MeshConversion/Private/MeshDescriptionBuilder.cpp
@@ -315,7 +315,7 @@ FBox FMeshDescriptionBuilder::ComputeBoundingBox() const
 	FAxisAlignedBox3f bounds = FAxisAlignedBox3f::Empty();
 	for ( FVertexID VertexID : MeshDescription->Vertices().GetElementIDs() )
 	{
-		bounds.Contain(VertexPositions.Get(VertexID));
+		bounds.Contain((FVector3f)VertexPositions.Get(VertexID));
 	}
 	return (FBox)bounds;
 }
diff --git a/Engine/Plugins/Experimental/GeometryProcessing/Source/MeshConversion/Private/MeshDescriptionToDynamicMesh.cpp b/Engine/Plugins/Experimental/GeometryProcessing/Source/MeshConversion/Private/MeshDescriptionToDynamicMesh.cpp
--- a/Engine/Plugins/Experimental/GeometryProcessing/Source/MeshConversion/Private/MeshDescriptionToDynamicMesh.cpp
+++ b/Engine/Plugins/Experimental/GeometryProcessing/Source/MeshConversion/Private/MeshDescriptionToDynamicMesh.cpp
@@ -107,7 +107,7 @@ public:
 			return *FoundIndex;
 		}
 
-		int32 NewIndex = NormalOverlay->AppendElement(Normal);
+		int32 NewIndex = NormalOverlay->AppendElement( &Normal.X );
 		UniqueVertexNormals.Add(VertNormal, NewIndex);
 		return NewIndex;
 	}
@@ -140,7 +140,7 @@ void FMeshDescriptionToDynamicMesh::Convert(const FMeshDescription* MeshIn, FDyn
 	for (const FVertexID VertexID : MeshIn->Vertices().GetElementIDs())
 	{
 		const FVector Position = VertexPositions.Get(VertexID);
-		int NewVertIdx = MeshOut.AppendVertex(Position);
+		int NewVertIdx = MeshOut.AppendVertex( (FVector3d)Position);
 		VertIDMap[NewVertIdx] = VertexID;
 	}
 
@@ -293,7 +293,7 @@ void FMeshDescriptionToDynamicMesh::Convert(const FMeshDescription* MeshIn, FDyn
 				{
 					const FVertexID& TriangleVertID = TriangleVertexIDs[i];
 					const FVector Position = VertexPositions[TriangleVertID];
-					const int32 NewVertIdx = MeshOut.AppendVertex(Position);
+					const int32 NewVertIdx = MeshOut.AppendVertex( (FVector3d)Position );
 					VertexIDs[i] = NewVertIdx;
 					VertIDMap.SetNumUninitialized(NewVertIdx + 1);
 					VertIDMap[NewVertIdx] = TriangleVertID;
@@ -765,7 +765,7 @@ static void CopyTangents_Internal(const FMeshDescription* SourceMesh, const FDyn
 			FVector3<RealType> Bitangent = VectorUtil::Bitangent((FVector3<RealType>)Normal, (FVector3<RealType>)Tangent, (RealType)BitangentSign);
 			Tangent.Normalize();
 			UE::Geometry::Normalize(Bitangent);
-			TangentsOut->SetPerTriangleTangent(TriID, j, Tangent, Bitangent);
+			TangentsOut->SetPerTriangleTangent(TriID, j, (FVector3<RealType>)Tangent, (FVector3<RealType>)Bitangent);
 		}
 	}
 }
diff --git a/Engine/Plugins/Experimental/MeshLODToolset/Source/MeshLODToolset/Private/Tools/GenerateStaticMeshLODAssetTool.cpp b/Engine/Plugins/Experimental/MeshLODToolset/Source/MeshLODToolset/Private/Tools/GenerateStaticMeshLODAssetTool.cpp
--- a/Engine/Plugins/Experimental/MeshLODToolset/Source/MeshLODToolset/Private/Tools/GenerateStaticMeshLODAssetTool.cpp
+++ b/Engine/Plugins/Experimental/MeshLODToolset/Source/MeshLODToolset/Private/Tools/GenerateStaticMeshLODAssetTool.cpp
@@ -273,7 +273,7 @@ void UGenerateStaticMeshLODAssetTool::Setup()
 	// For the first computation, display a bounding box with the working material. Otherwise it looks like nothing
 	// is happening. And we don't want to copy over the potentially huge input mesh to be the preview mesh.
 	FGridBoxMeshGenerator MeshGen;
-	MeshGen.Box = FOrientedBox3d(Bounds.Origin, Bounds.BoxExtent);
+	MeshGen.Box = FOrientedBox3d((FVector3d)Bounds.Origin, (FVector3d)Bounds.BoxExtent);
 	MeshGen.Generate();
 	FDynamicMesh3 BoxMesh(&MeshGen);
 	PreviewWithBackgroundCompute->PreviewMesh->UpdatePreview(MoveTemp(BoxMesh));
diff --git a/Engine/Plugins/Experimental/MeshModelingToolset/Source/MeshModelingTools/Private/AddPatchTool.cpp b/Engine/Plugins/Experimental/MeshModelingToolset/Source/MeshModelingTools/Private/AddPatchTool.cpp
--- a/Engine/Plugins/Experimental/MeshModelingToolset/Source/MeshModelingTools/Private/AddPatchTool.cpp
+++ b/Engine/Plugins/Experimental/MeshModelingToolset/Source/MeshModelingTools/Private/AddPatchTool.cpp
@@ -181,7 +181,7 @@ void UAddPatchTool::UpdatePreviewPosition(const FInputDeviceRay& DeviceClickPos)
 	bool bHit = ToolSceneQueriesUtil::FindNearestVisibleObjectHit(TargetWorld, Result, ClickPosWorldRay);
 	if (bHit)
 	{
-		ShapeFrame = FFrame3f(Result.ImpactPoint, Result.ImpactNormal);
+		ShapeFrame = FFrame3f((FVector3f)Result.ImpactPoint, (FVector3f)Result.ImpactNormal);
 		//ShapeFrame.ConstrainedAlignPerpAxes();
 	}
 
@@ -207,7 +207,7 @@ void UAddPatchTool::UpdatePreviewPosition(const FInputDeviceRay& DeviceClickPos)
 
 void UAddPatchTool::UpdatePreviewMesh()
 {
-	FVector Direction = FVector(0, 0, 1);
+	FVector3d Direction = FVector3d(0, 0, 1.0);
 	float WorldMaxHeight = WorldBounds.Max.Z;
 	float WorldMinHeight = WorldBounds.Min.Z;
 	float WorldHeight = WorldMaxHeight - WorldMinHeight;
@@ -235,10 +235,10 @@ void UAddPatchTool::UpdatePreviewMesh()
 		bool bHit = ToolSceneQueriesUtil::FindNearestVisibleObjectHit(TargetWorld, Result, RayStart, RayEnd);
 		if (bHit)
 		{
-			FVector3d HitPoint = Result.ImpactPoint + ShapeSettings->Shift * Direction;
+			FVector3d HitPoint = (FVector3d)Result.ImpactPoint + (double)ShapeSettings->Shift * Direction;
 			FVector HitPosWorld = (FVector)HitPoint;
 			FVector HitPosLocal = MoveTransform.InverseTransformPosition(HitPosWorld);
-			Projected.SetVertex(vid, HitPosLocal);
+			Projected.SetVertex(vid, (FVector3d)HitPosLocal);
 		}
 		else
 		{
diff --git a/Engine/Plugins/Experimental/MeshModelingToolset/Source/MeshModelingTools/Private/AddPrimitiveTool.cpp b/Engine/Plugins/Experimental/MeshModelingToolset/Source/MeshModelingTools/Private/AddPrimitiveTool.cpp
--- a/Engine/Plugins/Experimental/MeshModelingToolset/Source/MeshModelingTools/Private/AddPrimitiveTool.cpp
+++ b/Engine/Plugins/Experimental/MeshModelingToolset/Source/MeshModelingTools/Private/AddPrimitiveTool.cpp
@@ -217,7 +217,7 @@ void UAddPrimitiveTool::UpdatePreviewPosition(const FInputDeviceRay& DeviceClick
 	FPlane DrawPlane(FVector::ZeroVector, FVector(0, 0, 1));
 	if (ShapeSettings->PlaceMode == EMakeMeshPlacementType::GroundPlane)
 	{
-		FVector DrawPlanePos = FMath::RayPlaneIntersection(ClickPosWorldRay.Origin, ClickPosWorldRay.Direction, DrawPlane);
+		FVector3f DrawPlanePos = (FVector3f)FMath::RayPlaneIntersection(ClickPosWorldRay.Origin, ClickPosWorldRay.Direction, DrawPlane);
 		bHit = true;
 		ShapeFrame = FFrame3f(DrawPlanePos);
 	}
@@ -228,18 +228,18 @@ void UAddPrimitiveTool::UpdatePreviewPosition(const FInputDeviceRay& DeviceClick
 		bHit = ToolSceneQueriesUtil::FindNearestVisibleObjectHit(TargetWorld, Result, ClickPosWorldRay);
 		if (bHit)
 		{
-			FVector3f Normal = Result.ImpactNormal;
+			FVector3f Normal = (FVector3f)Result.ImpactNormal;
 			if (!ShapeSettings->bAlignShapeToPlacementSurface)
 			{
 				Normal = FVector3f::UnitZ();
 			}
-			ShapeFrame = FFrame3f(Result.ImpactPoint, Normal);
+			ShapeFrame = FFrame3f((FVector3f)Result.ImpactPoint, Normal);
 			ShapeFrame.ConstrainedAlignPerpAxes();
 		}
 		else
 		{
 			// fall back to ground plane if we don't have a scene hit
-			FVector DrawPlanePos = FMath::RayPlaneIntersection(ClickPosWorldRay.Origin, ClickPosWorldRay.Direction, DrawPlane);
+			FVector3f DrawPlanePos = (FVector3f)FMath::RayPlaneIntersection(ClickPosWorldRay.Origin, ClickPosWorldRay.Direction, DrawPlane);
 			bHit = true;
 			ShapeFrame = FFrame3f(DrawPlanePos);
 		}
@@ -256,7 +256,7 @@ void UAddPrimitiveTool::UpdatePreviewPosition(const FInputDeviceRay& DeviceClick
 		TArray<FSceneSnapQueryResult> Results;
 		if (GetToolManager()->GetContextQueriesAPI()->ExecuteSceneSnapQuery(Request, Results))
 		{
-			ShapeFrame.Origin = Results[0].Position;
+			ShapeFrame.Origin = (FVector3f)Results[0].Position;
 		}
 	}
 
@@ -489,7 +489,7 @@ void UAddTorusPrimitiveTool::GenerateMesh(FDynamicMesh3* OutMesh) const
 	FPolygon2d PathCircle = FPolygon2d::MakeCircle(TorusSettings->MajorRadius, TorusSettings->TubeSlices);
 	for (int Idx = 0; Idx < PathCircle.VertexCount(); Idx++)
 	{
-		Gen.Path.Add(PathCircle[Idx]);
+		Gen.Path.Add( FVector3d(PathCircle[Idx].X, PathCircle[Idx].Y, 0) );
 	}
 	Gen.bLoop = true;
 	Gen.bCapped = false;
diff --git a/Engine/Plugins/Experimental/MeshModelingToolset/Source/MeshModelingTools/Private/DeformMeshPolygonsTool.cpp b/Engine/Plugins/Experimental/MeshModelingToolset/Source/MeshModelingTools/Private/DeformMeshPolygonsTool.cpp
--- a/Engine/Plugins/Experimental/MeshModelingToolset/Source/MeshModelingTools/Private/DeformMeshPolygonsTool.cpp
+++ b/Engine/Plugins/Experimental/MeshModelingToolset/Source/MeshModelingTools/Private/DeformMeshPolygonsTool.cpp
@@ -529,7 +529,7 @@ void FConstrainedMeshDeformerTask::ApplyAttenuation()
 		}
 	}
 
-	double ExtentLength = Min.Distance(Max);
+	double ExtentLength = Distance(Min, Max);
 
 	// Is this why the system has memory?
 	for (int32 SubVertexID = 0; SubVertexID < SubsetConstraintBuffer.Num(); ++SubVertexID)
@@ -537,9 +537,8 @@ void FConstrainedMeshDeformerTask::ApplyAttenuation()
 		if (!Handles.Contains(SubVertexID))
 		{
 			FDeformerVertexConstraintData& CData = SubsetConstraintBuffer[SubVertexID];
-			double T = CData.Position.Distance(FMath::ClosestPointOnSegment((FVector)CData.Position,
-			                                                                (FVector)MinHandles, (FVector)MaxHandles)) /
-			    ExtentLength;
+			FVector3d OtherPoint = (FVector3d)FMath::ClosestPointOnSegment((FVector)CData.Position, (FVector)MinHandles, (FVector)MaxHandles);
+			double T = Distance(CData.Position, OtherPoint) / ExtentLength;
 			CData.Weight = WeightAttenuationCurve.Eval(T) * LeastWeight;
 		}
 	}
@@ -1041,8 +1040,8 @@ FDynamicMeshAABBTree3& UDeformMeshPolygonsTool::GetSpatial()
 bool UDeformMeshPolygonsTool::HitTest(const FRay& WorldRay, FHitResult& OutHit)
 {
 	UE::Geometry::FTransform3d Transform(Cast<IPrimitiveComponentBackedTarget>(Target)->GetWorldTransform());
-	FRay3d LocalRay(Transform.InverseTransformPosition(WorldRay.Origin),
-	                Transform.InverseTransformVector(WorldRay.Direction));
+	FRay3d LocalRay(Transform.InverseTransformPosition((FVector3d)WorldRay.Origin),
+	                Transform.InverseTransformVector((FVector3d)WorldRay.Direction));
 	UE::Geometry::Normalize(LocalRay.Direction);
 
 	FGroupTopologySelection Selection;
@@ -1087,8 +1086,8 @@ bool UDeformMeshPolygonsTool::HitTest(const FRay& WorldRay, FHitResult& OutHit)
 void UDeformMeshPolygonsTool::OnBeginDrag(const FRay& WorldRay)
 {
 	UE::Geometry::FTransform3d Transform(Cast<IPrimitiveComponentBackedTarget>(Target)->GetWorldTransform());
-	FRay3d LocalRay(Transform.InverseTransformPosition(WorldRay.Origin),
-	                Transform.InverseTransformVector(WorldRay.Direction));
+	FRay3d LocalRay(Transform.InverseTransformPosition((FVector3d)WorldRay.Origin),
+	                Transform.InverseTransformVector((FVector3d)WorldRay.Direction));
 	UE::Geometry::Normalize(LocalRay.Direction);
 
 	HilightSelection.Clear();
@@ -1118,7 +1117,7 @@ void UDeformMeshPolygonsTool::OnBeginDrag(const FRay& WorldRay)
 	UpdateActiveSurfaceFrame(HilightSelection);
 	UpdateQuickTransformer();
 
-	LastBrushPosLocal  = (FVector)Transform.InverseTransformPosition(LastHitPosWorld);
+	LastBrushPosLocal  = (FVector)Transform.InverseTransformPosition((FVector3d)LastHitPosWorld);
 	StartBrushPosLocal = LastBrushPosLocal;
 
 	// Record the requested deformation strategy - NB: will be forced to linear if there aren't any free points to solve.
@@ -1198,14 +1197,14 @@ void UDeformMeshPolygonsTool::UpdateActiveSurfaceFrame(FGroupTopologySelection&
 	UE::Geometry::FTransform3d Transform(Cast<IPrimitiveComponentBackedTarget>(Target)->GetWorldTransform());
 
 	// update surface frame
-	ActiveSurfaceFrame.Origin = StartHitPosWorld;
+	ActiveSurfaceFrame.Origin = (FVector3d)StartHitPosWorld;
 	if (HilightSelection.SelectedCornerIDs.Num() == 1)
 	{
 		// just keeping existing axes...we don't have enough info to do something smarter
 	}
 	else
 	{
-		ActiveSurfaceFrame.AlignAxis(2, StartHitNormalWorld);
+		ActiveSurfaceFrame.AlignAxis(2, (FVector3d)StartHitNormalWorld);
 		if (HilightSelection.SelectedEdgeIDs.Num() == 1)
 		{
 			FVector3d Tangent;
@@ -1242,7 +1241,7 @@ void UDeformMeshPolygonsTool::UpdateQuickTransformer()
 	}
 	else
 	{
-		GetActiveQuickTransformer()->SetActiveFrameFromWorldAxes(StartHitPosWorld);
+		GetActiveQuickTransformer()->SetActiveFrameFromWorldAxes((FVector3d)StartHitPosWorld);
 	}
 }
 
@@ -1302,8 +1301,8 @@ bool UDeformMeshPolygonsTool::OnUpdateHover(const FInputDeviceRay& DevicePos)
 	if (ActiveVertexChange == nullptr)
 	{
 		UE::Geometry::FTransform3d Transform(Cast<IPrimitiveComponentBackedTarget>(Target)->GetWorldTransform());
-		FRay3d LocalRay(Transform.InverseTransformPosition(DevicePos.WorldRay.Origin),
-		                Transform.InverseTransformVector(DevicePos.WorldRay.Direction));
+		FRay3d LocalRay(Transform.InverseTransformPosition((FVector3d)DevicePos.WorldRay.Origin),
+		                Transform.InverseTransformVector((FVector3d)DevicePos.WorldRay.Direction));
 		UE::Geometry::Normalize(LocalRay.Direction);
 
 		HilightSelection.Clear();
@@ -1452,7 +1451,7 @@ void UDeformMeshPolygonsTool::ComputeUpdate_Rotate()
 
 	FVector2d RotateStartVec = RotationStartFrame.ToPlaneUV(RotationStartPointWorld, 2);
 	UE::Geometry::Normalize(RotateStartVec);
-	FVector2d RotateToVec = RotationStartFrame.ToPlaneUV(NewHitPosWorld, 2);
+	FVector2d RotateToVec = RotationStartFrame.ToPlaneUV((FVector3d)NewHitPosWorld, 2);
 	UE::Geometry::Normalize(RotateToVec);
 	double AngleRad = UE::Geometry::SignedAngleR(RotateStartVec, RotateToVec);
 	FQuaterniond Rotation(Transform.InverseTransformVectorNoScale(RotationStartFrame.Z()), AngleRad, false);
@@ -1506,7 +1505,7 @@ void UDeformMeshPolygonsTool::ComputeUpdate_Translate()
 	}
 
 	FVector NewBrushPosLocal = Transform.InverseTransformPosition(NewHitPosWorld);
-	FVector3d NewMoveDelta   = NewBrushPosLocal - StartBrushPosLocal;
+	FVector3d NewMoveDelta   = (FVector3d)NewBrushPosLocal - (FVector3d)StartBrushPosLocal;
 
 	FDynamicMesh3* Mesh = DynamicMeshComponent->GetMesh();
 	if (LastMoveDelta.SquaredLength() > 0.)
diff --git a/Engine/Plugins/Experimental/MeshModelingToolset/Source/MeshModelingTools/Private/DisplaceMeshTool.cpp b/Engine/Plugins/Experimental/MeshModelingToolset/Source/MeshModelingTools/Private/DisplaceMeshTool.cpp
--- a/Engine/Plugins/Experimental/MeshModelingToolset/Source/MeshModelingTools/Private/DisplaceMeshTool.cpp
+++ b/Engine/Plugins/Experimental/MeshModelingToolset/Source/MeshModelingTools/Private/DisplaceMeshTool.cpp
@@ -508,7 +508,7 @@ namespace DisplaceMeshToolLocals{
 				IntensityFunc,
 				Parameters.SineWaveFrequency,
 				Parameters.SineWavePhaseShift,
-				Parameters.SineWaveDirection,
+				(FVector3d)Parameters.SineWaveDirection,
 				DisplacedPositions);
 			break;
 		}
diff --git a/Engine/Plugins/Experimental/MeshModelingToolset/Source/MeshModelingTools/Private/DrawAndRevolveTool.cpp b/Engine/Plugins/Experimental/MeshModelingToolset/Source/MeshModelingTools/Private/DrawAndRevolveTool.cpp
--- a/Engine/Plugins/Experimental/MeshModelingToolset/Source/MeshModelingTools/Private/DrawAndRevolveTool.cpp
+++ b/Engine/Plugins/Experimental/MeshModelingToolset/Source/MeshModelingTools/Private/DrawAndRevolveTool.cpp
@@ -194,8 +194,8 @@ void UDrawAndRevolveTool::Setup()
 /** Uses the settings currently stored in the properties object to update the revolution axis. */
 void UDrawAndRevolveTool::UpdateRevolutionAxis()
 {
-	RevolutionAxisOrigin = Settings->DrawPlaneOrigin;
-	RevolutionAxisDirection = Settings->DrawPlaneOrientation.RotateVector(FVector(1,0,0));
+	RevolutionAxisOrigin = (FVector3d)Settings->DrawPlaneOrigin;
+	RevolutionAxisDirection = (FVector3d)Settings->DrawPlaneOrientation.RotateVector(FVector(1,0,0));
 
 	const int32 AXIS_SNAP_TARGET_ID = 1;
 	ControlPointsMechanic->RemoveSnapLine(AXIS_SNAP_TARGET_ID);
diff --git a/Engine/Plugins/Experimental/MeshModelingToolset/Source/MeshModelingTools/Private/DrawPolyPathTool.cpp b/Engine/Plugins/Experimental/MeshModelingToolset/Source/MeshModelingTools/Private/DrawPolyPathTool.cpp
--- a/Engine/Plugins/Experimental/MeshModelingToolset/Source/MeshModelingTools/Private/DrawPolyPathTool.cpp
+++ b/Engine/Plugins/Experimental/MeshModelingToolset/Source/MeshModelingTools/Private/DrawPolyPathTool.cpp
@@ -468,7 +468,7 @@ void UDrawPolyPathTool::OnCompleteSurfacePath()
 		int NextJ = (j + 1) % NumPoints;
 		int PrevJ = (j - 1 + NumPoints) % NumPoints;
 		FVector3d Prev(CurPathPoints[PrevJ].Origin), Next(CurPathPoints[NextJ].Origin), Cur(CurPathPoints[j].Origin);
-		ArcLengths[j] = ArcLengths[PrevJ] + Cur.Distance(Prev);
+		ArcLengths[j] = ArcLengths[PrevJ] + Distance(Cur, Prev);
 		FLine3d Line1(FLine3d::FromPoints(Prev, Cur)), Line2(FLine3d::FromPoints(Cur, Next));
 		Line1.Origin += DistOffsetDelta * PlaneNormal.Cross(Line1.Direction);
 		Line2.Origin += DistOffsetDelta * PlaneNormal.Cross(Line2.Direction);
@@ -480,10 +480,10 @@ void UDrawPolyPathTool::OnCompleteSurfacePath()
 		}
 		else
 		{
-			FDistLine3Line3d Distance(Line1, Line2);
-			Distance.GetSquared();
-			FVector3d OffsetPoint = 0.5 * (Distance.Line1ClosestPoint + Distance.Line2ClosestPoint);
-			OffsetScaleFactors[j] = OffsetPoint.Distance(Cur) / DistOffsetDelta;
+			FDistLine3Line3d LineDist(Line1, Line2);
+			LineDist.GetSquared();
+			FVector3d OffsetPoint = 0.5 * (LineDist.Line1ClosestPoint + LineDist.Line2ClosestPoint);
+			OffsetScaleFactors[j] = Distance(OffsetPoint, Cur) / DistOffsetDelta;
 			FVector3d TangentDir = UE::Geometry::Normalized(OffsetPoint - Cur).Cross(PlaneNormal);
 			CurPathPoints[j].ConstrainedAlignAxis(0, TangentDir, PlaneNormal);
 		}
diff --git a/Engine/Plugins/Experimental/MeshModelingToolset/Source/MeshModelingTools/Private/DrawPolygonTool.cpp b/Engine/Plugins/Experimental/MeshModelingToolset/Source/MeshModelingTools/Private/DrawPolygonTool.cpp
--- a/Engine/Plugins/Experimental/MeshModelingToolset/Source/MeshModelingTools/Private/DrawPolygonTool.cpp
+++ b/Engine/Plugins/Experimental/MeshModelingToolset/Source/MeshModelingTools/Private/DrawPolygonTool.cpp
@@ -101,7 +101,7 @@ void UDrawPolygonTool::Setup()
 	FSelectClickedAction* SetPlaneAction = new FSelectClickedAction();
 	SetPlaneAction->World = this->TargetWorld;
 	SetPlaneAction->OnClickedPositionFunc = [this](const FHitResult& Hit) {
-		SetDrawPlaneFromWorldPos(Hit.ImpactPoint, Hit.ImpactNormal);
+		SetDrawPlaneFromWorldPos((FVector3d)Hit.ImpactPoint, (FVector3d)Hit.ImpactNormal);
 	};
 	SetPointInWorldConnector = SetPlaneAction;
 
@@ -349,10 +349,10 @@ void UDrawPolygonTool::Render(IToolsContextRenderAPI* RenderAPI)
 			// clip this line to the view plane because if it goes through the view plane the pixel-line-thickness
 			// calculation appears to fail
 			FLine3d DrawSnapLine = SnapEngine.GetActiveSnapLine();
-			FVector3d P0(DrawSnapLine.PointAt(-9999)), P1(DrawSnapLine.PointAt(9999));		// should be smarter here...
+			FVector3d P0(DrawSnapLine.PointAt(-99999.0)), P1(DrawSnapLine.PointAt(99999.0));		// should be smarter here...
 			if (RenderCameraState.bIsOrthographic == false)
 			{
-				UE::Geometry::FPlane3d CameraPlane(RenderCameraState.Forward(), RenderCameraState.Position + 1.0*RenderCameraState.Forward());
+				UE::Geometry::FPlane3d CameraPlane((FVector3d)RenderCameraState.Forward(), (FVector3d)RenderCameraState.Position + 1.0*(FVector3d)RenderCameraState.Forward());
 				CameraPlane.ClipSegment(P0, P1);
 			}
 			PDI->DrawLine((FVector)P0, (FVector)P1, SnapLineColor, SDPG_Foreground, 0.5 * PDIScale, 0.0f, true);
@@ -479,7 +479,7 @@ void UDrawPolygonTool::UpdatePreviewVertex(const FVector3d& PreviewVertexIn)
 	if (PolygonVertices.Num() > 0)
 	{
 		FVector3d LastVertex = PolygonVertices[PolygonVertices.Num() - 1];
-		SnapProperties->SegmentLength = LastVertex.Distance(PreviewVertex);
+		SnapProperties->SegmentLength = Distance(LastVertex, PreviewVertex);
 	}
 }
 
@@ -494,7 +494,7 @@ bool UDrawPolygonTool::FindDrawPlaneHitPoint(const FInputDeviceRay& ClickPos, FV
 
 	FFrame3d Frame(DrawPlaneOrigin, DrawPlaneOrientation);
 	FVector3d HitPos;
-	bool bHit = Frame.RayPlaneIntersection(ClickPos.WorldRay.Origin, ClickPos.WorldRay.Direction, 2, HitPos);
+	bool bHit = Frame.RayPlaneIntersection((FVector3d)ClickPos.WorldRay.Origin, (FVector3d)ClickPos.WorldRay.Direction, 2, HitPos);
 	if (bHit == false)
 	{
 		return false;
@@ -558,8 +558,8 @@ bool UDrawPolygonTool::FindDrawPlaneHitPoint(const FInputDeviceRay& ClickPos, FV
 		if (bWorldHit)
 		{
 			bHaveSurfaceHit = true;
-			SurfaceHitPoint = Result.ImpactPoint;
-			FVector3d UseHitPos = Result.ImpactPoint + SnapProperties->HitNormalOffset*Result.Normal;
+			SurfaceHitPoint = (FVector3d)Result.ImpactPoint;
+			FVector3d UseHitPos = (FVector3d)Result.ImpactPoint + (double)SnapProperties->HitNormalOffset*(FVector3d)Result.Normal;
 			HitPos = Frame.ToPlane(UseHitPos, 2);
 			SurfaceOffsetPoint = UseHitPos;
 		}
@@ -702,7 +702,7 @@ bool UDrawPolygonTool::OnNextSequenceClick(const FInputDeviceRay& ClickPos)
 				PolygonVertices[j-SelfIntersectSegmentIdx] = PolygonVertices[j];
 			}
 			PolygonVertices.SetNum(PolygonVertices.Num() - SelfIntersectSegmentIdx);
-			PolygonVertices[0] = PreviewVertex = (FVector)SelfIntersectionPoint;
+			PolygonVertices[0] = PreviewVertex = (FVector3d)SelfIntersectionPoint;
 			bDonePolygon = true;
 		}
 	}
@@ -1080,7 +1080,7 @@ bool UDrawPolygonTool::GeneratePolygonMesh(const TArray<FVector3d>& Polygon, con
 	// add preview vertex
 	if (bIncludePreviewVtx)
 	{
-		if (PreviewVertex.Distance(Polygon[NumVerts - 1]) > 0.1)
+		if (Distance(PreviewVertex, Polygon[NumVerts-1]) > 0.1)
 		{
 			OuterPolygon.AppendVertex(WorldFrameOut.ToPlaneUV(PreviewVertex, 2));
 		}
diff --git a/Engine/Plugins/Experimental/MeshModelingToolset/Source/MeshModelingTools/Private/DynamicMeshSculptTool.cpp b/Engine/Plugins/Experimental/MeshModelingToolset/Source/MeshModelingTools/Private/DynamicMeshSculptTool.cpp
--- a/Engine/Plugins/Experimental/MeshModelingToolset/Source/MeshModelingTools/Private/DynamicMeshSculptTool.cpp
+++ b/Engine/Plugins/Experimental/MeshModelingToolset/Source/MeshModelingTools/Private/DynamicMeshSculptTool.cpp
@@ -338,8 +338,8 @@ void UDynamicMeshSculptTool::OnPropertyModified(UObject* PropertySet, FProperty*
 
 bool UDynamicMeshSculptTool::HitTest(const FRay& Ray, FHitResult& OutHit)
 {
-	FRay3d LocalRay(CurTargetTransform.InverseTransformPosition(Ray.Origin),
-		CurTargetTransform.InverseTransformVector(Ray.Direction));
+	FRay3d LocalRay(CurTargetTransform.InverseTransformPosition((FVector3d)Ray.Origin),
+		CurTargetTransform.InverseTransformVector((FVector3d)Ray.Direction));
 	UE::Geometry::Normalize(LocalRay.Direction);
 	FDynamicMesh3* Mesh = DynamicMeshComponent->GetMesh();
 
@@ -368,12 +368,12 @@ void UDynamicMeshSculptTool::OnBeginDrag(const FRay& Ray)
 	FHitResult OutHit;
 	if (HitTest(Ray, OutHit))
 	{
-		BrushStartCenterWorld = Ray.PointAt(OutHit.Distance) + BrushProperties->Depth*CurrentBrushRadius*Ray.Direction;
+		BrushStartCenterWorld = (FVector3d)Ray.PointAt(OutHit.Distance) + (double)BrushProperties->Depth*CurrentBrushRadius*(FVector3d)Ray.Direction;
 
 		bInDrag = true;
 
-		ActiveDragPlane = FFrame3d(BrushStartCenterWorld, -Ray.Direction);
-		ActiveDragPlane.RayPlaneIntersection(Ray.Origin, Ray.Direction, 2, LastHitPosWorld);
+		ActiveDragPlane = FFrame3d(BrushStartCenterWorld, -(FVector3d)Ray.Direction);
+		ActiveDragPlane.RayPlaneIntersection((FVector3d)Ray.Origin, (FVector3d)Ray.Direction, 2, LastHitPosWorld);
 
 		LastBrushPosWorld = LastHitPosWorld;
 		LastBrushPosNormalWorld = ActiveDragPlane.Z();
@@ -434,7 +434,7 @@ void UDynamicMeshSculptTool::UpdateROI(const FVector3d& BrushPos)
 				{
 					//const FVector3d& Position = Mesh->GetVertexRef(TriV[j]);
 					FVector3d Position = Mesh->GetVertex(TriV[j]);
-					if (BrushPos.DistanceSquared(Position) < RadiusSqr)
+					if (DistanceSquared(BrushPos, Position) < RadiusSqr)
 					{
 						VertexROIBuilder.Add(TriV[j]);
 					}
@@ -670,7 +670,7 @@ bool UDynamicMeshSculptTool::ApplySmoothBrush(const FRay& WorldRay)
 
 		FVector3d OrigPos = Mesh->GetVertex(VertIdx);
 
-		double Falloff = CalculateBrushFalloff(OrigPos.Distance(NewBrushPosLocal));
+		double Falloff = CalculateBrushFalloff(Distance(OrigPos, NewBrushPosLocal));
 
 		FVector3d SmoothedPos = (SculptProperties->bPreserveUVFlow) ?
 			FMeshWeights::CotanCentroidSafe(*Mesh, VertIdx, 10.0) : FMeshWeights::UniformCentroid(*Mesh, VertIdx);
@@ -754,7 +754,7 @@ bool UDynamicMeshSculptTool::ApplyOffsetBrush(const FRay& WorldRay, bool bUseVie
 		else
 		{
 			FVector3d MoveVec = (bUseViewDirection) ?  (UseSpeed*LocalNormal) : (UseSpeed*BaseNormal);
-			double Falloff = CalculateBrushFalloff(OrigPos.Distance(NewBrushPosLocal));
+			double Falloff = CalculateBrushFalloff(Distance(OrigPos, NewBrushPosLocal));
 			FVector3d NewPos = OrigPos + Falloff * MoveVec;
 			ROIPositionBuffer[k] = NewPos;
 		}
@@ -798,7 +798,7 @@ bool UDynamicMeshSculptTool::ApplySculptMaxBrush(const FRay& WorldRay)
 		else
 		{
 			FVector3d MoveVec = UseSpeed * BaseNormal;
-			double Falloff = CalculateBrushFalloff(OrigPos.Distance(NewBrushPosLocal));
+			double Falloff = CalculateBrushFalloff(Distance(OrigPos, NewBrushPosLocal));
 			FVector3d NewPos = OrigPos + Falloff * MoveVec;
 
 			FVector3d DeltaPos = NewPos - BasePos;
@@ -847,8 +847,8 @@ bool UDynamicMeshSculptTool::ApplyPinchBrush(const FRay& WorldRay)
 		FVector3d MoveVec = UseSpeed * Delta;
 
 		// pinch uses 1/x falloff, shifted so that
-		double Distance = OrigPos.Distance(NewBrushPosLocal);
-		double NormalizedDistance = Distance / CurrentBrushRadius + FMathf::ZeroTolerance;
+		double Dist = Distance(OrigPos, NewBrushPosLocal);
+		double NormalizedDistance = Dist / CurrentBrushRadius + FMathf::ZeroTolerance;
 		double Falloff = (1.0/NormalizedDistance) - 1.0;
 		Falloff = FMathd::Clamp(Falloff, 0.0, 1.0);
 
@@ -877,7 +877,7 @@ FFrame3d UDynamicMeshSculptTool::ComputeROIBrushPlane(const FVector3d& BrushCent
 	for (int TriID : TriangleROI)
 	{
 		FVector3d Centroid = Mesh->GetTriCentroid(TriID);
-		double Weight = CalculateBrushFalloff(BrushCenter.Distance(Centroid));
+		double Weight = CalculateBrushFalloff(Distance(BrushCenter, Centroid));
 
 		AverageNormal += Weight * Mesh->GetTriNormal(TriID);
 		AveragePos += Weight * Centroid;
@@ -888,7 +888,7 @@ FFrame3d UDynamicMeshSculptTool::ComputeROIBrushPlane(const FVector3d& BrushCent
 
 	if (bViewAligned)
 	{
-		AverageNormal = -CameraState.Forward();
+		AverageNormal = -(FVector3d)CameraState.Forward();
 	}
 
 	FFrame3d Result = FFrame3d(AveragePos, AverageNormal);
@@ -930,7 +930,7 @@ bool UDynamicMeshSculptTool::ApplyPlaneBrush(const FRay& WorldRay)
 		if (Dot * PlaneSign >= 0)
 		{
 			FVector3d MoveVec = UseSpeed * Delta;
-			double Falloff = CalculateBrushFalloff(OrigPos.Distance(NewBrushPosLocal));
+			double Falloff = CalculateBrushFalloff(Distance(OrigPos, NewBrushPosLocal));
 			NewPos = OrigPos + Falloff * MoveVec;
 		}
 		ROIPositionBuffer[k] = NewPos;
@@ -957,7 +957,7 @@ bool UDynamicMeshSculptTool::ApplyFixedPlaneBrush(const FRay& WorldRay)
 	double UseSpeed = CurrentBrushRadius * FMathd::Sqrt(SculptProperties->PrimaryBrushSpeed) * 0.1 * ActivePressure;
 
 	FFrame3d FixedPlaneLocal(
-		CurTargetTransform.InverseTransformPosition(GizmoProperties->Position),
+		CurTargetTransform.InverseTransformPosition((FVector3d)GizmoProperties->Position),
 		CurTargetTransform.GetRotation().Inverse() * (FQuaterniond)GizmoProperties->Rotation);
 
 	FDynamicMesh3* Mesh = DynamicMeshComponent->GetMesh();
@@ -975,7 +975,7 @@ bool UDynamicMeshSculptTool::ApplyFixedPlaneBrush(const FRay& WorldRay)
 		if (Dot * PlaneSign >= 0)
 		{
 			double MaxDist = UE::Geometry::Normalize(Delta);
-			double Falloff = CalculateBrushFalloff(OrigPos.Distance(NewBrushPosLocal));
+			double Falloff = CalculateBrushFalloff(Distance(OrigPos, NewBrushPosLocal));
 			FVector3d MoveVec = Falloff * UseSpeed * Delta;
 			NewPos = (MoveVec.SquaredLength() > MaxDist* MaxDist) ?
 				PlanePos : OrigPos + Falloff * MoveVec;
@@ -1022,7 +1022,7 @@ bool UDynamicMeshSculptTool::ApplyFlattenBrush(const FRay& WorldRay)
 		if (Dot * PlaneSign >= 0)
 		{
 			double MaxDist = UE::Geometry::Normalize(Delta);
-			double Falloff = CalculateBrushFalloff(OrigPos.Distance(NewBrushPosLocal));
+			double Falloff = CalculateBrushFalloff(Distance(OrigPos, NewBrushPosLocal));
 			FVector3d MoveVec = Falloff * UseSpeed * Delta;
 			NewPos = (MoveVec.SquaredLength() > MaxDist*MaxDist) ?
 				PlanePos : OrigPos + Falloff * MoveVec;
@@ -1069,7 +1069,7 @@ bool UDynamicMeshSculptTool::ApplyInflateBrush(const FRay& WorldRay)
 
 		FVector3d MoveVec = UseSpeed * Normal;
 
-		double Falloff = CalculateBrushFalloff(OrigPos.Distance(NewBrushPosLocal));
+		double Falloff = CalculateBrushFalloff(Distance(OrigPos, NewBrushPosLocal));
 
 		FVector3d NewPos = OrigPos + Falloff*MoveVec;
 		ROIPositionBuffer[k] = NewPos;
@@ -1231,7 +1231,7 @@ int UDynamicMeshSculptTool::FindHitSculptMeshTriangle(const FRay3d& LocalRay)
 
 		FViewCameraState StateOut;
 		GetToolManager()->GetContextQueriesAPI()->GetCurrentViewState(StateOut);
-		FVector3d LocalEyePosition(CurTargetTransform.InverseTransformPosition(StateOut.Position));
+		FVector3d LocalEyePosition(CurTargetTransform.InverseTransformPosition((FVector3d)StateOut.Position));
 		int HitTID = DynamicMeshComponent->GetOctree()->FindNearestHitObject(LocalRay,
 			[this, Mesh, &LocalEyePosition](int TriangleID) {
 			FVector3d Normal, Centroid;
@@ -1255,7 +1255,7 @@ int UDynamicMeshSculptTool::FindHitTargetMeshTriangle(const FRay3d& LocalRay)
 
 		FViewCameraState StateOut;
 		GetToolManager()->GetContextQueriesAPI()->GetCurrentViewState(StateOut);
-		FVector3d LocalEyePosition(CurTargetTransform.InverseTransformPosition(StateOut.Position));
+		FVector3d LocalEyePosition(CurTargetTransform.InverseTransformPosition((FVector3d)StateOut.Position));
 
 		RaycastOptions.TriangleFilterF = [this, Mesh, LocalEyePosition](int TriangleID) {
 			FVector3d Normal, Centroid;
@@ -1273,7 +1273,7 @@ int UDynamicMeshSculptTool::FindHitTargetMeshTriangle(const FRay3d& LocalRay)
 bool UDynamicMeshSculptTool::UpdateBrushPositionOnActivePlane(const FRay& WorldRay)
 {
 	FVector3d NewHitPosWorld;
-	ActiveDragPlane.RayPlaneIntersection(WorldRay.Origin, WorldRay.Direction, 2, NewHitPosWorld);
+	ActiveDragPlane.RayPlaneIntersection((FVector3d)WorldRay.Origin, (FVector3d)WorldRay.Direction, 2, NewHitPosWorld);
 	LastBrushPosWorld = NewHitPosWorld;
 	LastBrushPosNormalWorld = ActiveDragPlane.Z();
 	return true;
@@ -1283,8 +1283,8 @@ bool UDynamicMeshSculptTool::UpdateBrushPositionOnTargetMesh(const FRay& WorldRa
 {
 	PendingTargetUpdate.Wait();
 
-	FRay3d LocalRay(CurTargetTransform.InverseTransformPosition(WorldRay.Origin),
-		CurTargetTransform.InverseTransformVector(WorldRay.Direction));
+	FRay3d LocalRay(CurTargetTransform.InverseTransformPosition((FVector3d)WorldRay.Origin),
+		CurTargetTransform.InverseTransformVector((FVector3d)WorldRay.Direction));
 	UE::Geometry::Normalize(LocalRay.Direction);
 
 	const FDynamicMesh3* TargetMesh = BrushTargetMeshSpatial.GetMesh();
@@ -1304,9 +1304,9 @@ bool UDynamicMeshSculptTool::UpdateBrushPositionOnTargetMesh(const FRay& WorldRa
 
 	if (bFallbackToViewPlane)
 	{
-		FFrame3d BrushPlane(LastBrushPosWorld, CameraState.Forward());
+		FFrame3d BrushPlane(LastBrushPosWorld, (FVector3d)CameraState.Forward());
 		FVector3d NewHitPosWorld;
-		BrushPlane.RayPlaneIntersection(WorldRay.Origin, WorldRay.Direction, 2, NewHitPosWorld);
+		BrushPlane.RayPlaneIntersection((FVector3d)WorldRay.Origin, (FVector3d)WorldRay.Direction, 2, NewHitPosWorld);
 		LastBrushPosWorld = NewHitPosWorld;
 		LastBrushPosNormalWorld = ActiveDragPlane.Z();
 		return true;
@@ -1317,8 +1317,8 @@ bool UDynamicMeshSculptTool::UpdateBrushPositionOnTargetMesh(const FRay& WorldRa
 
 bool UDynamicMeshSculptTool::UpdateBrushPositionOnSculptMesh(const FRay& WorldRay, bool bFallbackToViewPlane)
 {
-	FRay3d LocalRay(CurTargetTransform.InverseTransformPosition(WorldRay.Origin),
-		CurTargetTransform.InverseTransformVector(WorldRay.Direction));
+	FRay3d LocalRay(CurTargetTransform.InverseTransformPosition((FVector3d)WorldRay.Origin),
+		CurTargetTransform.InverseTransformVector((FVector3d)WorldRay.Direction));
 	UE::Geometry::Normalize(LocalRay.Direction);
 
 	int HitTID = FindHitSculptMeshTriangle(LocalRay);
@@ -1339,9 +1339,9 @@ bool UDynamicMeshSculptTool::UpdateBrushPositionOnSculptMesh(const FRay& WorldRa
 
 	if (bFallbackToViewPlane)
 	{
-		FFrame3d BrushPlane(LastBrushPosWorld, CameraState.Forward());
+		FFrame3d BrushPlane(LastBrushPosWorld, (FVector3d)CameraState.Forward());
 		FVector3d NewHitPosWorld;
-		BrushPlane.RayPlaneIntersection(WorldRay.Origin, WorldRay.Direction, 2, NewHitPosWorld);
+		BrushPlane.RayPlaneIntersection((FVector3d)WorldRay.Origin, (FVector3d)WorldRay.Direction, 2, NewHitPosWorld);
 		LastBrushPosWorld = NewHitPosWorld;
 		LastBrushPosNormalWorld = ActiveDragPlane.Z();
 		LastBrushTriangleID = -1;
@@ -1353,7 +1353,7 @@ bool UDynamicMeshSculptTool::UpdateBrushPositionOnSculptMesh(const FRay& WorldRa
 
 void UDynamicMeshSculptTool::AlignBrushToView()
 {
-	LastBrushPosNormalWorld = -CameraState.Forward();
+	LastBrushPosNormalWorld = -(FVector3d)CameraState.Forward();
 }
 
 
@@ -1443,7 +1443,7 @@ bool UDynamicMeshSculptTool::OnUpdateHover(const FInputDeviceRay& DevicePos)
 	if (bInDrag)
 	{
 		FVector3d NewHitPosWorld;
-		ActiveDragPlane.RayPlaneIntersection(DevicePos.WorldRay.Origin, DevicePos.WorldRay.Direction, 2, NewHitPosWorld);
+		ActiveDragPlane.RayPlaneIntersection((FVector3d)DevicePos.WorldRay.Origin, (FVector3d)DevicePos.WorldRay.Direction, 2, NewHitPosWorld);
 		LastBrushPosWorld = NewHitPosWorld;
 		LastBrushPosNormalWorld = ActiveDragPlane.Z();
 	}
@@ -1523,7 +1523,8 @@ void UDynamicMeshSculptTool::OnTick(float DeltaTime)
 		FRay CursorWorldRay = UMeshSurfacePointTool::LastWorldRay;
 		FHitResult Result;
 		bool bWorldHit = ToolSceneQueriesUtil::FindNearestVisibleObjectHit(TargetWorld, Result, CursorWorldRay);
-		FRay3d LocalRay(CurTargetTransform.InverseTransformPosition(CursorWorldRay.Origin), CurTargetTransform.InverseTransformVector(CursorWorldRay.Direction));
+		FRay3d LocalRay(CurTargetTransform.InverseTransformPosition((FVector3d)CursorWorldRay.Origin),
+			CurTargetTransform.InverseTransformVector((FVector3d)CursorWorldRay.Direction));
 		UE::Geometry::Normalize(LocalRay.Direction);
 		bool bObjectHit = (FindHitSculptMeshTriangle(LocalRay) != IndexConstants::InvalidID);
 		if (bWorldHit &&
@@ -1786,7 +1787,7 @@ void UDynamicMeshSculptTool::ConfigureRemesher(FSubRegionRemesher& Remesher)
 		Remesher.CustomSmoothSpeedF = [this, UseSmoothing](const FDynamicMesh3& Mesh, int vID)
 		{
 			FVector3d Pos = Mesh.GetVertex(vID);
-			double Falloff = CalculateBrushFalloff(Pos.Distance((FVector3d)LastBrushPosLocal));
+			double Falloff = CalculateBrushFalloff(Distance(Pos, (FVector3d)LastBrushPosLocal));
 			return (1.0f - Falloff) * UseSmoothing;
 		};
 	}
@@ -1804,7 +1805,7 @@ void UDynamicMeshSculptTool::ConfigureRemesher(FSubRegionRemesher& Remesher)
 			Remesher.CustomSmoothSpeedF = [this, UseSmoothing, DetailT](const FDynamicMesh3& Mesh, int vID)
 			{
 				FVector3d Pos = Mesh.GetVertex(vID);
-				double FalloffT = 1.0 - CalculateBrushFalloff(Pos.Distance((FVector3d)LastBrushPosLocal));
+				double FalloffT = 1.0 - CalculateBrushFalloff(Distance(Pos, (FVector3d)LastBrushPosLocal));
 				FalloffT = FMathd::Lerp(1.0, FalloffT, DetailT);
 				return FalloffT * UseSmoothing;
 			};
diff --git a/Engine/Plugins/Experimental/MeshModelingToolset/Source/MeshModelingTools/Private/EdgeLoopInsertionTool.cpp b/Engine/Plugins/Experimental/MeshModelingToolset/Source/MeshModelingTools/Private/EdgeLoopInsertionTool.cpp
--- a/Engine/Plugins/Experimental/MeshModelingToolset/Source/MeshModelingTools/Private/EdgeLoopInsertionTool.cpp
+++ b/Engine/Plugins/Experimental/MeshModelingToolset/Source/MeshModelingTools/Private/EdgeLoopInsertionTool.cpp
@@ -346,9 +346,9 @@ FInputRayHit UEdgeLoopInsertionTool::HitTest(const FRay& WorldRay)
 	FInputRayHit Hit;
 
 	// See if we hit an edge
-	FTransform LocalToWorld = Cast<IPrimitiveComponentBackedTarget>(Target)->GetWorldTransform();
-	FRay3d LocalRay(LocalToWorld.InverseTransformPosition(WorldRay.Origin),
-		LocalToWorld.InverseTransformVector(WorldRay.Direction), false);
+	FTransform3d LocalToWorld(Cast<IPrimitiveComponentBackedTarget>(Target)->GetWorldTransform());
+	FRay3d LocalRay(LocalToWorld.InverseTransformPosition((FVector3d)WorldRay.Origin),
+		LocalToWorld.InverseTransformVector((FVector3d)WorldRay.Direction), false);
 	FGroupTopologySelection Selection;
 	FVector3d Position, Normal;
 	if (TopologySelector.FindSelectedElement(
@@ -371,9 +371,9 @@ FInputRayHit UEdgeLoopInsertionTool::HitTest(const FRay& WorldRay)
 bool UEdgeLoopInsertionTool::UpdateHoveredItem(const FRay& WorldRay)
 {
 	// Check that we hit an edge
-	FTransform LocalToWorld = Cast<IPrimitiveComponentBackedTarget>(Target)->GetWorldTransform();
-	FRay3d LocalRay(LocalToWorld.InverseTransformPosition(WorldRay.Origin),
-		LocalToWorld.InverseTransformVector(WorldRay.Direction), false);
+	FTransform3d LocalToWorld(Cast<IPrimitiveComponentBackedTarget>(Target)->GetWorldTransform());
+	FRay3d LocalRay(LocalToWorld.InverseTransformPosition((FVector3d)WorldRay.Origin),
+		LocalToWorld.InverseTransformVector((FVector3d)WorldRay.Direction), false);
 
 	FGroupTopologySelection Selection;
 	FVector3d Position, Normal;
diff --git a/Engine/Plugins/Experimental/MeshModelingToolset/Source/MeshModelingTools/Private/EditMeshPolygonsTool.cpp b/Engine/Plugins/Experimental/MeshModelingToolset/Source/MeshModelingTools/Private/EditMeshPolygonsTool.cpp
--- a/Engine/Plugins/Experimental/MeshModelingToolset/Source/MeshModelingTools/Private/EditMeshPolygonsTool.cpp
+++ b/Engine/Plugins/Experimental/MeshModelingToolset/Source/MeshModelingTools/Private/EditMeshPolygonsTool.cpp
@@ -807,7 +807,7 @@ void UEditMeshPolygonsTool::ComputeUpdate_Gizmo()
 	FVector3d TranslationDelta = CurFrame.Origin - InitialGizmoFrame.Origin;
 	FQuaterniond RotateDelta = CurFrame.Rotation - InitialGizmoFrame.Rotation;
 	FVector3d CurScaleDelta = CurScale - InitialGizmoScale;
-	FVector3d LocalTranslation = WorldTransform.InverseTransformVector((FVector)TranslationDelta);
+	FVector3d LocalTranslation = WorldTransform.InverseTransformVector(TranslationDelta);
 
 	FDynamicMesh3* Mesh = DynamicMeshComponent->GetMesh();
 	if (TranslationDelta.SquaredLength() > 0.0001 || RotateDelta.SquaredLength() > 0.0001 || CurScaleDelta.SquaredLength() > 0.0001)
@@ -1420,8 +1420,8 @@ void UEditMeshPolygonsTool::ApplyCutFaces()
 	FVector3d PlaneNormal;
 	if (CutProperties->Orientation == EPolyEditCutPlaneOrientation::ViewDirection)
 	{
-		FVector3d Direction0 = UE::Geometry::Normalized(Point0.Origin - CameraState.Position);
-		FVector3d Direction1 = UE::Geometry::Normalized(Point1.Origin - CameraState.Position);
+		FVector3d Direction0 = UE::Geometry::Normalized(Point0.Origin - (FVector3d)CameraState.Position);
+		FVector3d Direction1 = UE::Geometry::Normalized(Point1.Origin - (FVector3d)CameraState.Position);
 		PlaneNormal = Direction1.Cross(Direction0);
 	}
 	else
@@ -1551,7 +1551,7 @@ void UEditMeshPolygonsTool::ApplySetUVs()
 	FTransform3d ToLocalXForm(WorldTransform.Inverse());
 	PlanarFrame.Transform(ToLocalXForm);
 	ScalePt = ToLocalXForm.TransformPosition(ScalePt);
-	UVScale = ScalePt.Distance(PlanarFrame.Origin);
+	UVScale = Distance(ScalePt, PlanarFrame.Origin);
 
 	// track changes
 	FDynamicMeshChangeTracker ChangeTracker(Mesh);
@@ -1986,10 +1986,10 @@ void UEditMeshPolygonsTool::ApplyFillHole()
 				// Compute normals and UVs
 				if (Mesh->HasAttributes())
 				{
-					TArray<FVector3d> VertexPositions;
+					TArray<FVector3<double>> VertexPositions;
 					Loop.GetVertices(VertexPositions);
-					FVector3d PlaneOrigin;
-					FVector3d PlaneNormal;
+					FVector3<double> PlaneOrigin;
+					FVector3<double> PlaneNormal;
 					PolygonTriangulation::ComputePolygonPlane<double>(VertexPositions, PlaneNormal, PlaneOrigin);
 
 					FDynamicMeshEditor Editor(Mesh);
diff --git a/Engine/Plugins/Experimental/MeshModelingToolset/Source/MeshModelingTools/Private/EditPivotTool.cpp b/Engine/Plugins/Experimental/MeshModelingToolset/Source/MeshModelingTools/Private/EditPivotTool.cpp
--- a/Engine/Plugins/Experimental/MeshModelingToolset/Source/MeshModelingTools/Private/EditPivotTool.cpp
+++ b/Engine/Plugins/Experimental/MeshModelingToolset/Source/MeshModelingTools/Private/EditPivotTool.cpp
@@ -177,8 +177,8 @@ void UEditPivotTool::Precompute()
 
 		FMeshDescription* Mesh = TargetMeshProviderInterface(0)->GetMeshDescription();
 		VertexIteration(Mesh, [&](int32 VertexID, const FVector& Position) {
-			ObjectBounds.Contain(Position);
-			WorldBounds.Contain(Transform.TransformPosition(Position));
+			ObjectBounds.Contain((FVector3d)Position);
+			WorldBounds.Contain(Transform.TransformPosition((FVector3d)Position));
 		});
 	}
 	else
@@ -191,8 +191,8 @@ void UEditPivotTool::Precompute()
 			UE::Geometry::FTransform3d CurTransform(TargetComponent->GetWorldTransform());
 			FMeshDescription* Mesh = TargetMeshProvider->GetMeshDescription();
 			VertexIteration(Mesh, [&](int32 VertexID, const FVector& Position) {
-				ObjectBounds.Contain(CurTransform.TransformPosition(Position));
-				WorldBounds.Contain(CurTransform.TransformPosition(Position));
+				ObjectBounds.Contain(CurTransform.TransformPosition((FVector3d)Position));
+				WorldBounds.Contain(CurTransform.TransformPosition((FVector3d)Position));
 			});
 		}
 	}
@@ -375,7 +375,7 @@ void UEditPivotTool::OnClickDrag(const FInputDeviceRay& DragPos)
 	FVector TargetNormal = (-NormalSign) * Result.Normal;
 
 	FQuaterniond AlignRotation = (bRotate) ?
-		FQuaterniond(FVector3d::UnitZ(), TargetNormal) : FQuaterniond::Identity();
+		FQuaterniond(FVector3d::UnitZ(), (FVector3d)TargetNormal) : FQuaterniond::Identity();
 
 	FTransform NewTransform = StartDragTransform;
 	NewTransform.SetRotation((FQuat)AlignRotation);
diff --git a/Engine/Plugins/Experimental/MeshModelingToolset/Source/MeshModelingTools/Private/GroupEdgeInsertionTool.cpp b/Engine/Plugins/Experimental/MeshModelingToolset/Source/MeshModelingTools/Private/GroupEdgeInsertionTool.cpp
--- a/Engine/Plugins/Experimental/MeshModelingToolset/Source/MeshModelingTools/Private/GroupEdgeInsertionTool.cpp
+++ b/Engine/Plugins/Experimental/MeshModelingToolset/Source/MeshModelingTools/Private/GroupEdgeInsertionTool.cpp
@@ -543,9 +543,9 @@ void UGroupEdgeInsertionTool::OnClicked(const FInputDeviceRay& ClickPos)
 bool UGroupEdgeInsertionTool::TopologyHitTest(const FRay& WorldRay, 
 	FVector3d& RayPositionOut, FRay3d* LocalRayOut)
 {
-	FTransform LocalToWorld = Cast<IPrimitiveComponentBackedTarget>(Target)->GetWorldTransform();
-	FRay3d LocalRay(LocalToWorld.InverseTransformPosition(WorldRay.Origin),
-		LocalToWorld.InverseTransformVector(WorldRay.Direction), false);
+	FTransform3d LocalToWorld(Cast<IPrimitiveComponentBackedTarget>(Target)->GetWorldTransform());
+	FRay3d LocalRay(LocalToWorld.InverseTransformPosition((FVector3d)WorldRay.Origin),
+		LocalToWorld.InverseTransformVector((FVector3d)WorldRay.Direction), false);
 
 	if (LocalRayOut)
 	{
@@ -557,7 +557,7 @@ bool UGroupEdgeInsertionTool::TopologyHitTest(const FRay& WorldRay,
 	if (TopologySelector.FindSelectedElement(TopologySelectorSettings,
 		LocalRay, Selection, Position, Normal))
 	{
-		RayPositionOut =  LocalToWorld.TransformPosition((FVector)Position);
+		RayPositionOut =  LocalToWorld.TransformPosition(Position);
 		return true;
 	}
 	return false;
@@ -572,9 +572,9 @@ bool UGroupEdgeInsertionTool::GetHoveredItem(const FRay& WorldRay,
 	PointOut.ElementID = FDynamicMesh3::InvalidID;
 
 	// Cast the ray to see what we hit.
-	FTransform LocalToWorld = Cast<IPrimitiveComponentBackedTarget>(Target)->GetWorldTransform();
-	FRay3d LocalRay(LocalToWorld.InverseTransformPosition(WorldRay.Origin),
-		LocalToWorld.InverseTransformVector(WorldRay.Direction), false);
+	FTransform3d LocalToWorld(Cast<IPrimitiveComponentBackedTarget>(Target)->GetWorldTransform());
+	FRay3d LocalRay(LocalToWorld.InverseTransformPosition((FVector3d)WorldRay.Origin),
+		LocalToWorld.InverseTransformVector((FVector3d)WorldRay.Direction), false);
 	if (LocalRayOut)
 	{
 		*LocalRayOut = LocalRay;
@@ -621,9 +621,9 @@ bool UGroupEdgeInsertionTool::GetHoveredItem(const FRay& WorldRay,
 		PointOut.Tangent = EdgeVector / EdgeLength;
 
 		FRay EdgeRay((FVector)StartVert, (FVector)PointOut.Tangent, true);
-		double DistDownEdge = EdgeRay.GetParameter((FVector)Position);
+		float DistDownEdge = EdgeRay.GetParameter((FVector)Position);
 
-		PositionOut = EdgeRay.PointAt(DistDownEdge);
+		PositionOut = (FVector3d)EdgeRay.PointAt(DistDownEdge);
 
 		// See if the point is at a vertex in the group edge span.
 		if (DistDownEdge <= Settings->VertexTolerance)
diff --git a/Engine/Plugins/Experimental/MeshModelingToolset/Source/MeshModelingTools/Private/MeshAttributePaintTool.cpp b/Engine/Plugins/Experimental/MeshModelingToolset/Source/MeshModelingTools/Private/MeshAttributePaintTool.cpp
--- a/Engine/Plugins/Experimental/MeshModelingToolset/Source/MeshModelingTools/Private/MeshAttributePaintTool.cpp
+++ b/Engine/Plugins/Experimental/MeshModelingToolset/Source/MeshModelingTools/Private/MeshAttributePaintTool.cpp
@@ -346,14 +346,14 @@ bool UMeshAttributePaintTool::OnUpdateHover(const FInputDeviceRay& DevicePos)
 void UMeshAttributePaintTool::CalculateVertexROI(const FBrushStampData& Stamp, TArray<int>& VertexROI)
 {
 	UE::Geometry::FTransform3d Transform(Cast<IPrimitiveComponentBackedTarget>(Target)->GetWorldTransform());
-	FVector3d StampPosLocal = Transform.InverseTransformPosition(Stamp.WorldPosition);
+	FVector3d StampPosLocal = Transform.InverseTransformPosition((FVector3d)Stamp.WorldPosition);
 
 	float Radius = GetCurrentBrushRadiusLocal();
 	float RadiusSqr = Radius * Radius;
 	const FDynamicMesh3* Mesh = PreviewMesh->GetPreviewDynamicMesh();
 	FAxisAlignedBox3d QueryBox(StampPosLocal, Radius);
 	VerticesOctree.RangeQuery(QueryBox,
-		[&](int32 VertexID) { return Mesh->GetVertex(VertexID).DistanceSquared(StampPosLocal) < RadiusSqr; },
+		[&](int32 VertexID) { return DistanceSquared(Mesh->GetVertex(VertexID), StampPosLocal) < RadiusSqr; },
 		VertexROI);
 }
 
@@ -422,7 +422,7 @@ void UMeshAttributePaintTool::UpdateVisibleAttribute()
 			for (int32 vid : Mesh.VertexIndicesItr())
 			{
 				float Value = AttribData.CurrentValues[vid];
-				FVector3f Color = ColorMapper->ToColor(Value);
+				FVector3f Color = ColorMapper->ToColor<FVector3f>(Value);
 				Mesh.SetVertexColor(vid, Color);
 			}
 		});
@@ -493,7 +493,7 @@ void UMeshAttributePaintTool::ApplyStamp(const FBrushStampData& Stamp)
 		{
 			int32 vid = ActionData.ROIVertices[k];
 			AttribData.CurrentValues[vid] = ActionData.ROIAfter[k];
-			FVector3f NewColor = ColorMapper->ToColor(ActionData.ROIAfter[k]);
+			FVector3f NewColor = ColorMapper->ToColor<FVector3f>(ActionData.ROIAfter[k]);
 			Mesh.SetVertexColor(vid, NewColor);
 		}
 	}, false);
@@ -504,7 +504,7 @@ void UMeshAttributePaintTool::ApplyStamp(const FBrushStampData& Stamp)
 void UMeshAttributePaintTool::ApplyStamp_Paint(const FBrushStampData& Stamp, FStampActionData& ActionData)
 {
 	UE::Geometry::FTransform3d Transform(Cast<IPrimitiveComponentBackedTarget>(Target)->GetWorldTransform());
-	FVector3d StampPosLocal = Transform.InverseTransformPosition(Stamp.WorldPosition);
+	FVector3d StampPosLocal = Transform.InverseTransformPosition((FVector3d)Stamp.WorldPosition);
 
 	int32 NumVertices = ActionData.ROIVertices.Num();
 	ActionData.ROIBefore.SetNum(NumVertices);
@@ -525,13 +525,13 @@ void UMeshAttributePaintTool::ApplyStamp_Paint(const FBrushStampData& Stamp, FSt
 			for (int32 NbrVID : CurrentMesh->VtxVerticesItr(vid))
 			{
 				FVector3d NbrPos = CurrentMesh->GetVertex(NbrVID);
-				float Weight = FMathf::Clamp(1.0f / NbrPos.DistanceSquared(Position), 0.0001f, 1000.0f);
+				float Weight = FMathf::Clamp(1.0f / DistanceSquared(NbrPos, Position), 0.0001f, 1000.0f);
 				ValueSum += Weight * AttribData.CurrentValues[NbrVID];
 				WeightSum += Weight;
 			}
 			ValueSum /= WeightSum;
 
-			float Falloff = (float)CalculateBrushFalloff(Position.Distance(StampPosLocal));
+			float Falloff = (float)CalculateBrushFalloff(Distance(Position, StampPosLocal));
 			float NewValue = FMathf::Lerp(AttribData.CurrentValues[vid], ValueSum, SmoothSpeed*Falloff);
 
 			ActionData.ROIBefore[k] = AttribData.CurrentValues[vid];
@@ -547,7 +547,7 @@ void UMeshAttributePaintTool::ApplyStamp_Paint(const FBrushStampData& Stamp, FSt
 		{
 			int32 vid = ActionData.ROIVertices[k];
 			FVector3d Position = CurrentMesh->GetVertex(vid);
-			float Falloff = (float)CalculateBrushFalloff(Position.Distance(StampPosLocal));
+			float Falloff = (float)CalculateBrushFalloff(Distance(Position, StampPosLocal));
 			ActionData.ROIBefore[k] = AttribData.CurrentValues[vid];
 			ActionData.ROIAfter[k] = CurrentValueRange.Clamp(ActionData.ROIBefore[k] + UseStrength*Falloff);
 		}
@@ -704,7 +704,7 @@ void UMeshAttributePaintTool::ExternalUpdateValues(int32 AttribIndex, const TArr
 		{
 			for (int32 vid : VertexIndices)
 			{
-				FVector3f NewColor = ColorMapper->ToColor(AttribData.CurrentValues[vid]);
+				FVector3f NewColor = ColorMapper->ToColor<FVector3f>(AttribData.CurrentValues[vid]);
 				Mesh.SetVertexColor(vid, NewColor);
 			}
 		});
diff --git a/Engine/Plugins/Experimental/MeshModelingToolset/Source/MeshModelingTools/Private/MeshGroupPaintTool.cpp b/Engine/Plugins/Experimental/MeshModelingToolset/Source/MeshModelingTools/Private/MeshGroupPaintTool.cpp
--- a/Engine/Plugins/Experimental/MeshModelingToolset/Source/MeshModelingTools/Private/MeshGroupPaintTool.cpp
+++ b/Engine/Plugins/Experimental/MeshModelingToolset/Source/MeshModelingTools/Private/MeshGroupPaintTool.cpp
@@ -938,7 +938,7 @@ int32 UMeshGroupPaintTool::FindHitSculptMeshTriangle(const FRay3d& LocalRay)
 
 		FViewCameraState StateOut;
 		GetToolManager()->GetContextQueriesAPI()->GetCurrentViewState(StateOut);
-		FVector3d LocalEyePosition(CurTargetTransform.InverseTransformPosition(StateOut.Position));
+		FVector3d LocalEyePosition(CurTargetTransform.InverseTransformPosition((FVector3d)StateOut.Position));
 		int HitTID = Octree.FindNearestHitObject(LocalRay,
 			[this, Mesh, &LocalEyePosition](int TriangleID) {
 			FVector3d Normal, Centroid;
diff --git a/Engine/Plugins/Experimental/MeshModelingToolset/Source/MeshModelingTools/Private/MeshSelectionTool.cpp b/Engine/Plugins/Experimental/MeshModelingToolset/Source/MeshModelingTools/Private/MeshSelectionTool.cpp
--- a/Engine/Plugins/Experimental/MeshModelingToolset/Source/MeshModelingTools/Private/MeshSelectionTool.cpp
+++ b/Engine/Plugins/Experimental/MeshModelingToolset/Source/MeshModelingTools/Private/MeshSelectionTool.cpp
@@ -386,8 +386,8 @@ TUniquePtr<FDynamicMeshOctree3>& UMeshSelectionTool::GetOctree()
 
 void UMeshSelectionTool::CalculateVertexROI(const FBrushStampData& Stamp, TArray<int>& VertexROI)
 {
-	FTransform Transform = Cast<IPrimitiveComponentBackedTarget>(Target)->GetWorldTransform();
-	FVector StampPosLocal = Transform.InverseTransformPosition(Stamp.WorldPosition);
+	FTransform3d Transform(Cast<IPrimitiveComponentBackedTarget>(Target)->GetWorldTransform());
+	FVector3d StampPosLocal = Transform.InverseTransformPosition((FVector3d)Stamp.WorldPosition);
 
 	// TODO: need dynamic vertex hash table!
 	float Radius = GetCurrentBrushRadiusLocal();
@@ -408,8 +408,8 @@ void UMeshSelectionTool::CalculateVertexROI(const FBrushStampData& Stamp, TArray
 
 void UMeshSelectionTool::CalculateTriangleROI(const FBrushStampData& Stamp, TArray<int>& TriangleROI)
 {
-	UE::Geometry::FTransform3d Transform(Cast<IPrimitiveComponentBackedTarget>(Target)->GetWorldTransform());
-	FVector3d StampPosLocal = Transform.InverseTransformPosition(Stamp.WorldPosition);
+	FTransform3d Transform(Cast<IPrimitiveComponentBackedTarget>(Target)->GetWorldTransform());
+	FVector3d StampPosLocal = Transform.InverseTransformPosition((FVector3d)Stamp.WorldPosition);
 
 	// always select first triangle
 	const FDynamicMesh3* Mesh = PreviewMesh->GetPreviewDynamicMesh();
diff --git a/Engine/Plugins/Experimental/MeshModelingToolset/Source/MeshModelingTools/Private/MeshSpaceDeformerTool.cpp b/Engine/Plugins/Experimental/MeshModelingToolset/Source/MeshModelingTools/Private/MeshSpaceDeformerTool.cpp
--- a/Engine/Plugins/Experimental/MeshModelingToolset/Source/MeshModelingTools/Private/MeshSpaceDeformerTool.cpp
+++ b/Engine/Plugins/Experimental/MeshModelingToolset/Source/MeshModelingTools/Private/MeshSpaceDeformerTool.cpp
@@ -153,7 +153,7 @@ void UMeshSpaceDeformerTool::Setup()
 	// bounding box dimensions, after scaling them with the transform.
 	FAxisAlignedBox3d BBox = OriginalDynamicMesh->GetBounds();
 	FVector3d Dimensions = BBox.Max - BBox.Min;
-	MeshCenter = MeshTransform.TransformPosition((FVector)BBox.Center());
+	MeshCenter = (FVector3d)MeshTransform.TransformPosition((FVector)BBox.Center());
 
 	Dimensions = FVector3d(MeshTransform.GetScale3D().GetAbs()) * Dimensions;
 	double WorldMajorLength = 0;
@@ -463,7 +463,7 @@ void UMeshSpaceDeformerTool::UpdatePreview()
 
 void UMeshSpaceDeformerTool::SetGizmoFrameFromWorldPos(const FVector& Position, const FVector& Normal, bool bAlignNormal)
 {
-	GizmoFrame.Origin = Position;
+	GizmoFrame.Origin = (FVector3d)Position;
 	if (bAlignNormal)
 	{
 		// It's not clear whether aligning the Z axis to the normal is the right idea here. The Z axis
@@ -473,8 +473,8 @@ void UMeshSpaceDeformerTool::SetGizmoFrameFromWorldPos(const FVector& Position,
 		// Still, it's hard to come up with a clean alternative.
 		FVector3d FrameZ(Normal);
 		FVector3d FrameY = UE::Geometry::Normalized(FrameZ.Cross(FVector3d::UnitZ())); // orthogonal to world Z and frame Z 
-		FVector3d FrameX = FrameY.Cross(Normal); // safe to not normalize because already orthogonal
-		GizmoFrame = FFrame3d(Position, FrameX, FrameY, Normal);
+		FVector3d FrameX = FrameY.Cross(FrameZ); // safe to not normalize because already orthogonal
+		GizmoFrame = FFrame3d((FVector3d)Position, FrameX, FrameY, FrameZ);
 	}
 
 	TransformGizmo->ReinitializeGizmoTransform(GizmoFrame.ToFTransform());
diff --git a/Engine/Plugins/Experimental/MeshModelingToolset/Source/MeshModelingTools/Private/MeshVertexSculptTool.cpp b/Engine/Plugins/Experimental/MeshModelingToolset/Source/MeshModelingTools/Private/MeshVertexSculptTool.cpp
--- a/Engine/Plugins/Experimental/MeshModelingToolset/Source/MeshModelingTools/Private/MeshVertexSculptTool.cpp
+++ b/Engine/Plugins/Experimental/MeshModelingToolset/Source/MeshModelingTools/Private/MeshVertexSculptTool.cpp
@@ -412,9 +412,9 @@ void UMeshVertexSculptTool::UpdateROI(const FVector3d& BrushPos)
 			}
 
 			const FIndex3i& TriV = Mesh->GetTriangleRef(tid);
-			TriangleROIInBuf[k].A = (BrushPos.DistanceSquared(Mesh->GetVertexRef(TriV.A)) < RadiusSqr) ? 1 : 0;
-			TriangleROIInBuf[k].B = (BrushPos.DistanceSquared(Mesh->GetVertexRef(TriV.B)) < RadiusSqr) ? 1 : 0;
-			TriangleROIInBuf[k].C = (BrushPos.DistanceSquared(Mesh->GetVertexRef(TriV.C)) < RadiusSqr) ? 1 : 0;
+			TriangleROIInBuf[k].A = (DistanceSquared(BrushPos, Mesh->GetVertexRef(TriV.A)) < RadiusSqr) ? 1 : 0;
+			TriangleROIInBuf[k].B = (DistanceSquared(BrushPos, Mesh->GetVertexRef(TriV.B)) < RadiusSqr) ? 1 : 0;
+			TriangleROIInBuf[k].C = (DistanceSquared(BrushPos, Mesh->GetVertexRef(TriV.C)) < RadiusSqr) ? 1 : 0;
 			if (TriangleROIInBuf[k].A + TriangleROIInBuf[k].B + TriangleROIInBuf[k].C == 0)
 			{
 				RangeQueryTriBuffer[k] = -1;
@@ -649,7 +649,7 @@ int32 UMeshVertexSculptTool::FindHitSculptMeshTriangle(const FRay3d& LocalRay)
 
 		FViewCameraState StateOut;
 		GetToolManager()->GetContextQueriesAPI()->GetCurrentViewState(StateOut);
-		FVector3d LocalEyePosition(CurTargetTransform.InverseTransformPosition(StateOut.Position));
+		FVector3d LocalEyePosition(CurTargetTransform.InverseTransformPosition((FVector3d)StateOut.Position));
 		int HitTID = Octree.FindNearestHitObject(LocalRay,
 			[this, Mesh, &LocalEyePosition](int TriangleID) {
 			FVector3d Normal, Centroid;
@@ -671,7 +671,7 @@ int32 UMeshVertexSculptTool::FindHitTargetMeshTriangle(const FRay3d& LocalRay)
 
 		FViewCameraState StateOut;
 		GetToolManager()->GetContextQueriesAPI()->GetCurrentViewState(StateOut);
-		FVector3d LocalEyePosition(CurTargetTransform.InverseTransformPosition(StateOut.Position));
+		FVector3d LocalEyePosition(CurTargetTransform.InverseTransformPosition((FVector3d)StateOut.Position));
 
 		RaycastOptions.TriangleFilterF = [this, Mesh, LocalEyePosition](int TriangleID) {
 			FVector3d Normal, Centroid;
diff --git a/Engine/Plugins/Experimental/MeshModelingToolset/Source/MeshModelingTools/Private/MirrorTool.cpp b/Engine/Plugins/Experimental/MeshModelingToolset/Source/MeshModelingTools/Private/MirrorTool.cpp
--- a/Engine/Plugins/Experimental/MeshModelingToolset/Source/MeshModelingTools/Private/MirrorTool.cpp
+++ b/Engine/Plugins/Experimental/MeshModelingToolset/Source/MeshModelingTools/Private/MirrorTool.cpp
@@ -247,7 +247,7 @@ void UMirrorTool::Setup()
 
 	// Set the initial mirror plane. We want the plane to start in the middle if we're doing a simple
 	// mirror (i.e., not appending, and not cropping). Otherwise, we want the plane to start to one side.
-	MirrorPlaneOrigin = CombinedBounds.GetCenter();
+	MirrorPlaneOrigin = (FVector3d)CombinedBounds.GetCenter();
 	MirrorPlaneNormal = FVector3d(0, -1, 0);
 	if (Settings->OperationMode == EMirrorOperationMode::MirrorAndAppend || Settings->bCropAlongMirrorPlaneFirst)
 	{
@@ -273,7 +273,7 @@ void UMirrorTool::Setup()
 	PlaneMechanic->SetPlaneCtrlClickBehaviorTarget->OnClickedPositionFunc = [this](const FHitResult& Hit)
 	{
 		bool bIgnoreNormal = (Settings->CtrlClickBehavior == EMirrorCtrlClickBehavior::Reposition);
-		PlaneMechanic->SetDrawPlaneFromWorldPos(Hit.ImpactPoint, Hit.ImpactNormal, bIgnoreNormal);
+		PlaneMechanic->SetDrawPlaneFromWorldPos((FVector3d)Hit.ImpactPoint, (FVector3d)Hit.ImpactNormal, bIgnoreNormal);
 	};
 	// Also include the original components in the ctrl+click hit testing even though we made them 
 	// invisible, since we want to be able to reposition the plane onto the original mesh.
@@ -515,7 +515,7 @@ void UMirrorTool::RequestAction(EMirrorToolAction ActionType)
 
 void UMirrorTool::ApplyAction(EMirrorToolAction ActionType)
 {
-	FVector3d ShiftedPlaneOrigin = CombinedBounds.GetCenter();
+	FVector3d ShiftedPlaneOrigin = (FVector3d)CombinedBounds.GetCenter();
 
 	if (ActionType == EMirrorToolAction::ShiftToCenter)
 	{
@@ -530,27 +530,27 @@ void UMirrorTool::ApplyAction(EMirrorToolAction ActionType)
 		{
 		case EMirrorToolAction::Left:
 			ShiftedPlaneOrigin.Y = CombinedBounds.Min.Y;
-			DirectionVector = FVector(0, -1, 0);
+			DirectionVector = FVector3d(0, -1.0, 0);
 			break;
 		case EMirrorToolAction::Right:
 			ShiftedPlaneOrigin.Y = CombinedBounds.Max.Y;
-			DirectionVector = FVector(0, 1, 0);
+			DirectionVector = FVector3d(0, 1.0, 0);
 			break;
 		case EMirrorToolAction::Up:
 			ShiftedPlaneOrigin.Z = CombinedBounds.Max.Z;
-			DirectionVector = FVector(0, 0, 1);
+			DirectionVector = FVector3d(0, 0, 1.0);
 			break;
 		case EMirrorToolAction::Down:
 			ShiftedPlaneOrigin.Z = CombinedBounds.Min.Z;
-			DirectionVector = FVector(0, 0, -1);
+			DirectionVector = FVector3d(0, 0, -1.0);
 			break;
 		case EMirrorToolAction::Forward:
 			ShiftedPlaneOrigin.X = CombinedBounds.Max.X;
-			DirectionVector = FVector(1, 0, 0);
+			DirectionVector = FVector3d(1.0, 0, 0);
 			break;
 		case EMirrorToolAction::Backward:
 			ShiftedPlaneOrigin.X = CombinedBounds.Min.X;
-			DirectionVector = FVector(-1, 0, 0);
+			DirectionVector = FVector3d(-1.0, 0, 0);
 			break;
 		}
 
diff --git a/Engine/Plugins/Experimental/MeshModelingToolset/Source/MeshModelingTools/Private/PlaneCutTool.cpp b/Engine/Plugins/Experimental/MeshModelingToolset/Source/MeshModelingTools/Private/PlaneCutTool.cpp
--- a/Engine/Plugins/Experimental/MeshModelingToolset/Source/MeshModelingTools/Private/PlaneCutTool.cpp
+++ b/Engine/Plugins/Experimental/MeshModelingToolset/Source/MeshModelingTools/Private/PlaneCutTool.cpp
@@ -351,15 +351,15 @@ TUniquePtr<FDynamicMeshOperator> UPlaneCutOperatorFactory::MakeNewOperator()
 	FVector LocalOrigin = WorldToLocal.TransformPosition(CutTool->CutPlaneOrigin);
 	FVector WorldNormal = CutTool->CutPlaneOrientation.GetAxisZ();
 	UE::Geometry::FTransform3d W2LForNormal(WorldToLocal);
-	FVector LocalNormal = (FVector)W2LForNormal.TransformNormal(WorldNormal);
+	FVector LocalNormal = (FVector)W2LForNormal.TransformNormal((FVector3d)WorldNormal);
 	FVector BackTransformed = LocalToWorld.TransformVector(LocalNormal);
 	float NormalScaleFactor = FVector::DotProduct(BackTransformed, WorldNormal);
 	if (NormalScaleFactor >= FLT_MIN)
 	{
 		NormalScaleFactor = 1.0 / NormalScaleFactor;
 	}
-	CutOp->LocalPlaneOrigin = LocalOrigin;
-	CutOp->LocalPlaneNormal = LocalNormal;
+	CutOp->LocalPlaneOrigin = (FVector3d)LocalOrigin;
+	CutOp->LocalPlaneNormal = (FVector3d)LocalNormal;
 	CutOp->OriginalMesh = CutTool->MeshesToCut[ComponentIndex]->GetMesh();
 	CutOp->bKeepBothHalves = CutTool->BasicProperties->bKeepBothHalves;
 	CutOp->CutPlaneLocalThickness = CutTool->BasicProperties->SpacingBetweenHalves * NormalScaleFactor;
@@ -453,7 +453,7 @@ void UPlaneCutTool::SetCutPlaneFromWorldPos(const FVector& Position, const FVect
 {
 	CutPlaneOrigin = Position;
 
-	FFrame3f CutPlane(Position, Normal);
+	FFrame3f CutPlane((FVector3f)Position, (FVector3f)Normal);
 	CutPlaneOrientation = (FQuat)CutPlane.Rotation;
 
 	PlaneTransformGizmo->SetActiveTarget(PlaneTransformProxy, GetToolManager());
diff --git a/Engine/Plugins/Experimental/MeshModelingToolset/Source/MeshModelingTools/Private/PositionPlaneGizmo.cpp b/Engine/Plugins/Experimental/MeshModelingToolset/Source/MeshModelingTools/Private/PositionPlaneGizmo.cpp
--- a/Engine/Plugins/Experimental/MeshModelingToolset/Source/MeshModelingTools/Private/PositionPlaneGizmo.cpp
+++ b/Engine/Plugins/Experimental/MeshModelingToolset/Source/MeshModelingTools/Private/PositionPlaneGizmo.cpp
@@ -125,7 +125,8 @@ void UPositionPlaneGizmo::OnBeginDrag(const FRay& Ray)
 {
 	bInTransformDrag = true;
 	FTransform CurTransform = CenterBallShape->GetTransform();
-	QuickTransformer.SetActiveFrameFromWorldNormal(CurTransform.GetTranslation(), CurTransform.GetRotation().GetAxisZ(), true);
+	QuickTransformer.SetActiveFrameFromWorldNormal(
+		(FVector3d)CurTransform.GetTranslation(), (FVector3d)CurTransform.GetRotation().GetAxisZ(), true);
 }
 
 
@@ -149,8 +150,8 @@ void UPositionPlaneGizmo::OnUpdateDrag(const FRay& Ray)
 		if (bHitWorld)
 		{
 			FFrame3f UpdatedFrame(CenterBallShape->GetTransform());
-			UpdatedFrame.AlignAxis(2, Result.ImpactNormal);
-			UpdatedFrame.Origin = Result.ImpactPoint;
+			UpdatedFrame.AlignAxis(2, (FVector3f)Result.ImpactNormal);
+			UpdatedFrame.Origin = (FVector3f)Result.ImpactPoint;
 			CenterBallShape->SetTransform(UpdatedFrame.ToFTransform());
 			PostUpdatedPosition();
 		}
diff --git a/Engine/Plugins/Experimental/MeshModelingToolset/Source/MeshModelingTools/Private/RevolveBoundaryTool.cpp b/Engine/Plugins/Experimental/MeshModelingToolset/Source/MeshModelingTools/Private/RevolveBoundaryTool.cpp
--- a/Engine/Plugins/Experimental/MeshModelingToolset/Source/MeshModelingTools/Private/RevolveBoundaryTool.cpp
+++ b/Engine/Plugins/Experimental/MeshModelingToolset/Source/MeshModelingTools/Private/RevolveBoundaryTool.cpp
@@ -52,7 +52,8 @@ TUniquePtr<FDynamicMeshOperator> URevolveBoundaryOperatorFactory::MakeNewOperato
 			{
 				int32 VertIndex = VertexIndices[i];
 
-				CurveSweepOp->ProfileCurve.Add(ToWorld.TransformPosition((FVector)RevolveBoundaryTool->OriginalMesh->GetVertex(VertIndex)));
+				FVector3d NewPos = (FVector3d)ToWorld.TransformPosition((FVector)RevolveBoundaryTool->OriginalMesh->GetVertex(VertIndex));
+				CurveSweepOp->ProfileCurve.Add(NewPos);
 			}
 			CurveSweepOp->bProfileCurveIsClosed = true;
 		}
@@ -177,11 +178,11 @@ void URevolveBoundaryTool::OnClicked(const FInputDeviceRay& ClickPos)
 			FVector3d VertexA, VertexB;
 			OriginalMesh->GetEdgeV(ClickedEid, VertexA, VertexB);
 			FTransform ToWorldTranform = Cast<IPrimitiveComponentBackedTarget>(Target)->GetWorldTransform();
-			FLine3d EdgeLine = FLine3d::FromPoints(ToWorldTranform.TransformPosition((FVector)VertexA), 
-				ToWorldTranform.TransformPosition((FVector)VertexB));
+			FLine3d EdgeLine = FLine3d::FromPoints((FVector3d)ToWorldTranform.TransformPosition((FVector)VertexA), 
+				(FVector3d)ToWorldTranform.TransformPosition((FVector)VertexB));
 			
 			FFrame3d RevolutionAxisFrame;
-			RevolutionAxisFrame.Origin = EdgeLine.NearestPoint(HitResult.ImpactPoint);
+			RevolutionAxisFrame.Origin = EdgeLine.NearestPoint((FVector3d)HitResult.ImpactPoint);
 			RevolutionAxisFrame.AlignAxis(0, EdgeLine.Direction);
 
 			PlaneMechanic->SetPlaneWithoutBroadcast(RevolutionAxisFrame);
@@ -215,8 +216,8 @@ bool URevolveBoundaryTool::CanAccept() const
  */
 void URevolveBoundaryTool::UpdateRevolutionAxis()
 {
-	RevolutionAxisOrigin = Settings->AxisOrigin;
-	RevolutionAxisDirection = FRotator(Settings->AxisPitch, Settings->AxisYaw, 0).RotateVector(FVector(1, 0, 0));
+	RevolutionAxisOrigin = (FVector3d)Settings->AxisOrigin;
+	RevolutionAxisDirection = (FVector3d)FRotator(Settings->AxisPitch, Settings->AxisYaw, 0).RotateVector(FVector(1, 0, 0));
 	if (Preview)
 	{
 		Preview->InvalidateResult();
diff --git a/Engine/Plugins/Experimental/MeshModelingToolset/Source/MeshModelingTools/Private/Sculpting/KelvinletBrushOp.h b/Engine/Plugins/Experimental/MeshModelingToolset/Source/MeshModelingTools/Private/Sculpting/KelvinletBrushOp.h
--- a/Engine/Plugins/Experimental/MeshModelingToolset/Source/MeshModelingTools/Private/Sculpting/KelvinletBrushOp.h
+++ b/Engine/Plugins/Experimental/MeshModelingToolset/Source/MeshModelingTools/Private/Sculpting/KelvinletBrushOp.h
@@ -449,7 +449,7 @@ public:
 		// compute the displacement vector in the local frame of the stamp
 		FVector3d Dir = Stamp.LocalFrame.ToFrameVector(Stamp.LocalFrame.Origin - Stamp.PrevLocalFrame.Origin);
 
-		FMatrix3d ForceMatrix = CrossProductMatrix(Dir);
+		FMatrix3d ForceMatrix = UE::Geometry::CrossProductMatrix(Dir);
 		// make symmetric
 		ForceMatrix.Row0[1] = -ForceMatrix.Row0[1];
 		ForceMatrix.Row0[2] = -ForceMatrix.Row0[2];
@@ -586,7 +586,7 @@ public:
 		}
 		case EKelvinletBrushMode::PullKelvinlet:
 		{
-			const FVector& Dir = Properties.Direction;
+			FVector Dir = Properties.Direction;
 			FVector3d Force(Dir.X, Dir.Y, Dir.Z);
 
 			FLaplacianPullKelvinlet LaplacianPullKelvinlet(Force, Properties.Size, Properties.Mu, Properties.Nu);
@@ -600,7 +600,7 @@ public:
 		}
 		case EKelvinletBrushMode::SharpPullKelvinlet:
 		{
-			const FVector& Dir = Properties.Direction;
+			FVector Dir = Properties.Direction;
 			FVector3d Force(Dir.X, Dir.Y, Dir.Z);
 
 			FSharpLaplacianPullKelvinlet SharpLaplacianPullKelvinlet(Force, Properties.Size, Properties.Mu, Properties.Nu);
@@ -614,7 +614,7 @@ public:
 		}
 		case EKelvinletBrushMode::LaplacianPullKelvinlet:
 		{
-			const FVector& Dir = Properties.Direction;
+			FVector Dir = Properties.Direction;
 			FVector3d Force(Dir.X, Dir.Y, Dir.Z);
 			FLaplacianPullKelvinlet PullKelvinlet(Force, Properties.Size, Properties.Mu, Properties.Nu);
 			ApplyKelvinlet(PullKelvinlet, VertRIO, ROIPositionBuffer, TimeStep, NumSteps);
@@ -622,7 +622,7 @@ public:
 		}
 		case EKelvinletBrushMode::BiLaplacianPullKelvinlet:
 		{
-			const FVector& Dir = Properties.Direction;
+			FVector Dir = Properties.Direction;
 			FVector3d Force(Dir.X, Dir.Y, Dir.Z);
 			FBiLaplacianPullKelvinlet PullKelvinlet(Force, Properties.Size, Properties.Mu, Properties.Nu);
 			ApplyKelvinlet(PullKelvinlet, VertRIO, ROIPositionBuffer, TimeStep, NumSteps);
@@ -630,16 +630,16 @@ public:
 		}
 		case EKelvinletBrushMode::TwistKelvinlet:
 		{
-			const FVector3d& TwistAxis = Properties.Direction;
-			FTwistKelvinlet TwistKelvinlet(TwistAxis, Properties.Size, Properties.Mu, Properties.Nu);
+			FVector TwistAxis = Properties.Direction;
+			FTwistKelvinlet TwistKelvinlet((FVector3d)TwistAxis, Properties.Size, Properties.Mu, Properties.Nu);
 			ApplyKelvinlet(TwistKelvinlet, VertRIO, ROIPositionBuffer, TimeStep, NumSteps);
 			break;
 		}
 		case EKelvinletBrushMode::PinchKelvinlet:
 		{
-			const FVector& Dir = Properties.Direction;
+			FVector Dir = Properties.Direction;
 
-			FMatrix3d ForceMatrix = CrossProductMatrix(FVector3d(Dir.X, Dir.Y, Dir.Z));
+			FMatrix3d ForceMatrix = UE::Geometry::CrossProductMatrix(FVector3d(Dir.X, Dir.Y, Dir.Z));
 			// make symmetric
 			ForceMatrix.Row0[1] = -ForceMatrix.Row0[1];
 			ForceMatrix.Row0[2] = -ForceMatrix.Row0[2];
@@ -650,12 +650,12 @@ public:
 		}
 		case EKelvinletBrushMode::LaplacianTwistPullKelvinlet:
 		{
-			FVector3d TwistAxis = Properties.Direction;
+			FVector3d TwistAxis = (FVector3d)Properties.Direction;
 			UE::Geometry::Normalize(TwistAxis);
 			TwistAxis *= Properties.Speed;
 			FTwistKelvinlet TwistKelvinlet(TwistAxis, Properties.Size, Properties.Mu, Properties.Nu);
 
-			const FVector& Dir = Properties.Direction;
+			FVector Dir = Properties.Direction;
 			FVector3d Force(Dir.X, Dir.Y, Dir.Z);
 			FLaplacianPullKelvinlet PullKelvinlet(Force, Properties.Size, Properties.Mu, Properties.Nu);
 
diff --git a/Engine/Plugins/Experimental/MeshModelingToolset/Source/MeshModelingTools/Private/Sculpting/MeshSculptToolBase.cpp b/Engine/Plugins/Experimental/MeshModelingToolset/Source/MeshModelingTools/Private/Sculpting/MeshSculptToolBase.cpp
--- a/Engine/Plugins/Experimental/MeshModelingToolset/Source/MeshModelingTools/Private/Sculpting/MeshSculptToolBase.cpp
+++ b/Engine/Plugins/Experimental/MeshModelingToolset/Source/MeshModelingTools/Private/Sculpting/MeshSculptToolBase.cpp
@@ -495,8 +495,8 @@ void UMeshSculptToolBase::OnEndDrag(const FRay& Ray)
 
 FRay3d UMeshSculptToolBase::GetLocalRay(const FRay& WorldRay) const
 {
-	FRay3d LocalRay(CurTargetTransform.InverseTransformPosition(WorldRay.Origin),
-		CurTargetTransform.InverseTransformVector(WorldRay.Direction));
+	FRay3d LocalRay(CurTargetTransform.InverseTransformPosition((FVector3d)WorldRay.Origin),
+		CurTargetTransform.InverseTransformVector((FVector3d)WorldRay.Direction));
 	UE::Geometry::Normalize(LocalRay.Direction);
 	return LocalRay;
 }
@@ -512,7 +512,7 @@ retry_frame_update:
 	FFrame3d NewFrame = LastBrushFrameWorld;
 	NewFrame.Origin = NewPosition;
 	NewFrame.AlignAxis(2, NewNormal);
-	FVector3d CameraUp = CameraState.Up();
+	FVector3d CameraUp = (FVector3d)CameraState.Up();
 
 	if (FMathd::Abs(CameraUp.Dot(NewNormal)) < 0.98)
 	{
@@ -548,32 +548,35 @@ retry_frame_update:
 
 void UMeshSculptToolBase::AlignBrushToView()
 {
-	UpdateBrushFrameWorld(GetBrushFrameWorld().Origin, -CameraState.Forward());
+	UpdateBrushFrameWorld(GetBrushFrameWorld().Origin, -(FVector3d)CameraState.Forward());
 }
 
 
 
-void UMeshSculptToolBase::UpdateBrushTargetPlaneFromHit(const FRay& WorldRay, const FHitResult& Hit)
+void UMeshSculptToolBase::UpdateBrushTargetPlaneFromHit(const FRay& WorldRayIn, const FHitResult& Hit)
 {
+	FRay3d WorldRay(WorldRayIn);
 	FVector3d WorldPosWithBrushDepth = WorldRay.PointAt(Hit.Distance) + GetCurrentBrushDepth() * GetCurrentBrushRadius() * WorldRay.Direction;
 	ActiveBrushTargetPlaneWorld = FFrame3d(WorldPosWithBrushDepth, -WorldRay.Direction);
 }
 
-bool UMeshSculptToolBase::UpdateBrushPositionOnActivePlane(const FRay& WorldRay)
+bool UMeshSculptToolBase::UpdateBrushPositionOnActivePlane(const FRay& WorldRayIn)
 {
 	LastBrushTriangleID = IndexConstants::InvalidID;
 
+	FRay3d WorldRay(WorldRayIn);
 	FVector3d NewHitPosWorld;
 	ActiveBrushTargetPlaneWorld.RayPlaneIntersection(WorldRay.Origin, WorldRay.Direction, 2, NewHitPosWorld);
 	UpdateBrushFrameWorld(NewHitPosWorld, ActiveBrushTargetPlaneWorld.Z());
 	return true;
 }
 
-bool UMeshSculptToolBase::UpdateBrushPositionOnTargetMesh(const FRay& WorldRay, bool bFallbackToViewPlane)
+bool UMeshSculptToolBase::UpdateBrushPositionOnTargetMesh(const FRay& WorldRayIn, bool bFallbackToViewPlane)
 {
 	LastBrushTriangleID = IndexConstants::InvalidID;
 
-	FRay3d LocalRay = GetLocalRay(WorldRay);
+	FRay3d WorldRay(WorldRayIn);
+	FRay3d LocalRay = GetLocalRay(WorldRayIn);
 	int32 HitTID = FindHitTargetMeshTriangle(LocalRay);
 	if (HitTID != IndexConstants::InvalidID)
 	{
@@ -588,7 +591,7 @@ bool UMeshSculptToolBase::UpdateBrushPositionOnTargetMesh(const FRay& WorldRay,
 
 	if (bFallbackToViewPlane)
 	{
-		FFrame3d BrushPlane(GetBrushFrameWorld().Origin, CameraState.Forward());
+		FFrame3d BrushPlane(GetBrushFrameWorld().Origin, (FVector3d)CameraState.Forward());
 		FVector3d NewHitPosWorld;
 		BrushPlane.RayPlaneIntersection(WorldRay.Origin, WorldRay.Direction, 2, NewHitPosWorld);
 		UpdateBrushFrameWorld(NewHitPosWorld, ActiveBrushTargetPlaneWorld.Z());
@@ -598,11 +601,12 @@ bool UMeshSculptToolBase::UpdateBrushPositionOnTargetMesh(const FRay& WorldRay,
 	return false;
 }
 
-bool UMeshSculptToolBase::UpdateBrushPositionOnSculptMesh(const FRay& WorldRay, bool bFallbackToViewPlane)
+bool UMeshSculptToolBase::UpdateBrushPositionOnSculptMesh(const FRay& WorldRayIn, bool bFallbackToViewPlane)
 {
 	LastBrushTriangleID = IndexConstants::InvalidID;
 
-	FRay3d LocalRay = GetLocalRay(WorldRay);
+	FRay3d WorldRay(WorldRayIn);
+	FRay3d LocalRay = GetLocalRay(WorldRayIn);
 	int32 HitTID = FindHitSculptMeshTriangle(LocalRay);
 	if (HitTID != IndexConstants::InvalidID)
 	{
@@ -617,7 +621,7 @@ bool UMeshSculptToolBase::UpdateBrushPositionOnSculptMesh(const FRay& WorldRay,
 
 	if (bFallbackToViewPlane)
 	{
-		FFrame3d BrushPlane(GetBrushFrameWorld().Origin, CameraState.Forward());
+		FFrame3d BrushPlane(GetBrushFrameWorld().Origin, (FVector3d)CameraState.Forward());
 		FVector3d NewHitPosWorld;
 		BrushPlane.RayPlaneIntersection(WorldRay.Origin, WorldRay.Direction, 2, NewHitPosWorld);
 		UpdateBrushFrameWorld(NewHitPosWorld, ActiveBrushTargetPlaneWorld.Z());
@@ -764,7 +768,7 @@ FFrame3d UMeshSculptToolBase::ComputeStampRegionPlane(const FFrame3d& StampFrame
 
 	if (bViewAligned)
 	{
-		AverageNormal = -CameraState.Forward();
+		AverageNormal = -(FVector3d)CameraState.Forward();
 	}
 
 	FFrame3d Result = FFrame3d(AveragePos, AverageNormal);
@@ -786,7 +790,7 @@ void UMeshSculptToolBase::UpdateStrokeReferencePlaneForROI(const FFrame3d& Stamp
 void UMeshSculptToolBase::UpdateStrokeReferencePlaneFromWorkPlane()
 {
 	StrokePlane = FFrame3d(
-		CurTargetTransform.InverseTransformPosition(GizmoProperties->Position),
+		CurTargetTransform.InverseTransformPosition((FVector3d)GizmoProperties->Position),
 		CurTargetTransform.GetRotation().Inverse() * (FQuaterniond)GizmoProperties->Rotation);
 }
 
diff --git a/Engine/Plugins/Experimental/MeshModelingToolset/Source/MeshModelingTools/Private/SeamSculptTool.cpp b/Engine/Plugins/Experimental/MeshModelingToolset/Source/MeshModelingTools/Private/SeamSculptTool.cpp
--- a/Engine/Plugins/Experimental/MeshModelingToolset/Source/MeshModelingTools/Private/SeamSculptTool.cpp
+++ b/Engine/Plugins/Experimental/MeshModelingToolset/Source/MeshModelingTools/Private/SeamSculptTool.cpp
@@ -192,7 +192,7 @@ bool USeamSculptTool::OnUpdateHover(const FInputDeviceRay& DevicePos)
 	bool bHit = UDynamicMeshBrushTool::OnUpdateHover(DevicePos);
 	if (bHit)
 	{
-		FVector3d LocalRayDir = MeshTransform.InverseTransformVector(DevicePos.WorldRay.Direction);
+		FVector3d LocalRayDir = MeshTransform.InverseTransformVector((FVector3d)DevicePos.WorldRay.Direction);
 		FVector3d LocalPos = MeshTransform.InverseTransformPosition((FVector3d)LastBrushStamp.WorldPosition);
 		double NearDistSqr;
 		int32 NewSnapVertex = PreviewMesh->GetSpatial()->FindNearestVertex(LocalPos, NearDistSqr);
diff --git a/Engine/Plugins/Experimental/MeshModelingToolset/Source/MeshModelingTools/Private/Tests/DisplaceMeshTool_Tests.inl b/Engine/Plugins/Experimental/MeshModelingToolset/Source/MeshModelingTools/Private/Tests/DisplaceMeshTool_Tests.inl
--- a/Engine/Plugins/Experimental/MeshModelingToolset/Source/MeshModelingTools/Private/Tests/DisplaceMeshTool_Tests.inl
+++ b/Engine/Plugins/Experimental/MeshModelingToolset/Source/MeshModelingTools/Private/Tests/DisplaceMeshTool_Tests.inl
@@ -295,7 +295,7 @@ IMPLEMENT_SIMPLE_AUTOMATION_TEST(UDisplaceMeshTestDirectionalFilter,
 	{
 		double DistanceMoved = (InputMesh->GetVertex(VertexID) - FinalMesh->GetVertex(VertexID)).Length();
 		FVector3d VertexNormal = InputMeshNormals[VertexID];
-		double Dot = VertexNormal.Dot(Params.FilterDirection);
+		double Dot = VertexNormal.Dot((FVector3d)Params.FilterDirection);
 
 		if (Dot < -0.5)
 		{
diff --git a/Engine/Plugins/Experimental/MeshModelingToolset/Source/MeshModelingTools/Private/TransformMeshesTool.cpp b/Engine/Plugins/Experimental/MeshModelingToolset/Source/MeshModelingTools/Private/TransformMeshesTool.cpp
--- a/Engine/Plugins/Experimental/MeshModelingToolset/Source/MeshModelingTools/Private/TransformMeshesTool.cpp
+++ b/Engine/Plugins/Experimental/MeshModelingToolset/Source/MeshModelingTools/Private/TransformMeshesTool.cpp
@@ -324,7 +324,7 @@ void UTransformMeshesTool::OnClickPress(const FInputDeviceRay& PressPos)
 
 	if (TransformProps->SnapDragSource == ETransformMeshesSnapDragSource::ClickPoint)
 	{
-		StartDragFrameWorld = FFrame3d(PressPos.WorldRay.PointAt(HitPos.HitDepth), HitPos.HitNormal);
+		StartDragFrameWorld = FFrame3d((FVector3d)PressPos.WorldRay.PointAt(HitPos.HitDepth), (FVector3d)HitPos.HitNormal);
 	}
 	else
 	{
@@ -369,7 +369,7 @@ void UTransformMeshesTool::OnClickDrag(const FInputDeviceRay& DragPos)
 		FVector TargetNormal = (-NormalSign) * Result.Normal;
 
 		FQuaterniond AlignRotation = (bRotate) ?
-			FQuaterniond(FVector3d::UnitZ(), TargetNormal) : FQuaterniond::Identity();
+			FQuaterniond(FVector3d::UnitZ(), (FVector3d)TargetNormal) : FQuaterniond::Identity();
 
 		FTransform NewTransform = StartDragTransform;
 		NewTransform.SetRotation((FQuat)AlignRotation);
@@ -386,7 +386,7 @@ void UTransformMeshesTool::OnClickDrag(const FInputDeviceRay& DragPos)
 		FVector TargetNormal = NormalSign * Result.Normal;
 
 		FFrame3d FromFrameWorld = StartDragFrameWorld;
-		FFrame3d ToFrameWorld(HitPos, TargetNormal);
+		FFrame3d ToFrameWorld((FVector3d)HitPos, (FVector3d)TargetNormal);
 		FFrame3d ObjectFrameWorld(StartDragTransform);
 
 		FVector3d CenterShift = FromFrameWorld.Origin - ObjectFrameWorld.Origin;
diff --git a/Engine/Plugins/Experimental/MeshModelingToolset/Source/MeshModelingTools/Public/MeshAttributePaintTool.h b/Engine/Plugins/Experimental/MeshModelingToolset/Source/MeshModelingTools/Public/MeshAttributePaintTool.h
--- a/Engine/Plugins/Experimental/MeshModelingToolset/Source/MeshModelingTools/Public/MeshAttributePaintTool.h
+++ b/Engine/Plugins/Experimental/MeshModelingToolset/Source/MeshModelingTools/Public/MeshAttributePaintTool.h
@@ -33,6 +33,13 @@ public:
 			FMathf::Lerp(LowColor.B, HighColor.B, t),
 			1.0f);
 	}
+
+	template<typename VectorType>
+	VectorType ToColor(float Value)
+	{
+		FLinearColor Color = ToColor(Value);
+		return VectorType(Color.R, Color.G, Color.B);
+	}
 };
 
 
diff --git a/Engine/Plugins/Experimental/MeshModelingToolset/Source/MeshModelingTools/Public/Sculpting/MeshPinchBrushOps.h b/Engine/Plugins/Experimental/MeshModelingToolset/Source/MeshModelingTools/Public/Sculpting/MeshPinchBrushOps.h
--- a/Engine/Plugins/Experimental/MeshModelingToolset/Source/MeshModelingTools/Public/Sculpting/MeshPinchBrushOps.h
+++ b/Engine/Plugins/Experimental/MeshModelingToolset/Source/MeshModelingTools/Public/Sculpting/MeshPinchBrushOps.h
@@ -81,8 +81,8 @@ public:
 			FVector3d MoveVec = Delta;
 
 			// pinch uses 1/x falloff
-			double Distance = OrigPos.Distance(NewSmoothBrushPosLocal);
-			double NormalizedDistance = (Distance / Stamp.Radius) + 0.0001;
+			double Dist = UE::Geometry::Distance(OrigPos, NewSmoothBrushPosLocal);
+			double NormalizedDistance = (Dist / Stamp.Radius) + 0.0001;
 			double LinearFalloff = FMathd::Clamp(1.0 - NormalizedDistance, 0.0, 1.0);
 			double InvFalloff = LinearFalloff * LinearFalloff * LinearFalloff;
 			double UseFalloff = FMathd::Lerp(LinearFalloff, InvFalloff, Stamp.Falloff);
diff --git a/Engine/Plugins/Experimental/MeshModelingToolset/Source/MeshModelingTools/Public/Sculpting/StampFalloffs.h b/Engine/Plugins/Experimental/MeshModelingToolset/Source/MeshModelingTools/Public/Sculpting/StampFalloffs.h
--- a/Engine/Plugins/Experimental/MeshModelingToolset/Source/MeshModelingTools/Public/Sculpting/StampFalloffs.h
+++ b/Engine/Plugins/Experimental/MeshModelingToolset/Source/MeshModelingTools/Public/Sculpting/StampFalloffs.h
@@ -15,9 +15,9 @@ namespace UE
 		{
 			return [](const FSculptBrushStamp& Stamp, const FVector3d& Position)
 			{
-				double Distance = Position.Distance(Stamp.LocalFrame.Origin);
+				double Dist = Distance(Position, Stamp.LocalFrame.Origin);
 				double FalloffT = FMathd::Clamp(1.0 - Stamp.Falloff, 0.0, 1.0);
-				double UnitDistance = Distance / Stamp.Radius;
+				double UnitDistance = Dist / Stamp.Radius;
 				if (UnitDistance > FalloffT)
 				{
 					UnitDistance = FMathd::Clamp((UnitDistance - FalloffT) / (1.0 - FalloffT), 0.0, 1.0);
@@ -34,9 +34,9 @@ namespace UE
 		{
 			return [](const FSculptBrushStamp& Stamp, const FVector3d& Position)
 			{
-				double Distance = Position.Distance(Stamp.LocalFrame.Origin);
+				double Dist = Distance(Position, Stamp.LocalFrame.Origin);
 				double FalloffT = FMathd::Clamp(1.0 - Stamp.Falloff, 0.0, 1.0);
-				double UnitDistance = Distance / Stamp.Radius;
+				double UnitDistance = Dist / Stamp.Radius;
 				if (UnitDistance > FalloffT)
 				{
 					UnitDistance = FMathd::Clamp((UnitDistance - FalloffT) / (1.0 - FalloffT), 0.0, 1.0);
@@ -54,9 +54,9 @@ namespace UE
 		{
 			return [](const FSculptBrushStamp& Stamp, const FVector3d& Position)
 			{
-				double Distance = Position.Distance(Stamp.LocalFrame.Origin);
+				double Dist = Distance(Position, Stamp.LocalFrame.Origin);
 				double FalloffT = FMathd::Clamp(1.0 - Stamp.Falloff, 0.0, 1.0);
-				double UnitDistance = Distance / Stamp.Radius;
+				double UnitDistance = Dist / Stamp.Radius;
 				if (UnitDistance > FalloffT)
 				{
 					UnitDistance = FMathd::Clamp((UnitDistance - FalloffT) / (1.0 - FalloffT), 0.0, 1.0);
@@ -73,9 +73,9 @@ namespace UE
 		{
 			return [](const FSculptBrushStamp& Stamp, const FVector3d& Position)
 			{
-				double Distance = Position.Distance(Stamp.LocalFrame.Origin);
+				double Dist = Distance(Position, Stamp.LocalFrame.Origin);
 				double FalloffT = FMathd::Clamp(1.0 - Stamp.Falloff, 0.0, 1.0);
-				double UnitDistance = Distance / Stamp.Radius;
+				double UnitDistance = Dist / Stamp.Radius;
 				if (UnitDistance > FalloffT)
 				{
 					UnitDistance = FMathd::Clamp((UnitDistance - FalloffT) / (1.0 - FalloffT), 0.0, 1.0);
diff --git a/Engine/Plugins/Experimental/MeshModelingToolset/Source/MeshModelingToolsEditorOnly/Private/BspConversionTool.cpp b/Engine/Plugins/Experimental/MeshModelingToolset/Source/MeshModelingToolsEditorOnly/Private/BspConversionTool.cpp
--- a/Engine/Plugins/Experimental/MeshModelingToolset/Source/MeshModelingToolsEditorOnly/Private/BspConversionTool.cpp
+++ b/Engine/Plugins/Experimental/MeshModelingToolset/Source/MeshModelingToolsEditorOnly/Private/BspConversionTool.cpp
@@ -354,7 +354,7 @@ bool UBspConversionTool::CombineThenConvert(FText*)
 	FVector3d NewPivot;
 	if (BrushForPivot)
 	{
-		NewPivot = BrushForPivot->GetTransform().GetLocation();
+		NewPivot = (FVector3d)BrushForPivot->GetTransform().GetLocation();
 	}
 	else
 	{
@@ -530,7 +530,7 @@ bool UBspConversionTool::ConvertThenCombine(FText *ErrorMessage)
 	// If there was only a single additive brush, we should keep its pivot in the original location.
 	if (BrushForPivot)
 	{
-		FVector3d BrushTranslation = BrushForPivot->GetTransform().GetLocation();
+		FVector3d BrushTranslation = (FVector3d)BrushForPivot->GetTransform().GetLocation();
 		MeshTransforms::Translate(*OutputMesh, OutputTransform->GetTranslation() - BrushTranslation);
 		OutputTransform->SetTranslation(BrushTranslation);
 	}
diff --git a/Engine/Plugins/Experimental/MeshModelingToolset/Source/MeshModelingToolsEditorOnly/Private/Hair/GroomQueryUtil.cpp b/Engine/Plugins/Experimental/MeshModelingToolset/Source/MeshModelingToolsEditorOnly/Private/Hair/GroomQueryUtil.cpp
--- a/Engine/Plugins/Experimental/MeshModelingToolset/Source/MeshModelingToolsEditorOnly/Private/Hair/GroomQueryUtil.cpp
+++ b/Engine/Plugins/Experimental/MeshModelingToolset/Source/MeshModelingToolsEditorOnly/Private/Hair/GroomQueryUtil.cpp
@@ -62,7 +62,7 @@ void UE::GroomQueries::ExtractAllHairCards(AGroomActor* GroomActor,
 			int32 NewVID = MeshOut.AppendVertex((FVector3d)CardGeo.Positions[k]);
 			MapV[k] = NewVID;
 
-			int32 NewNormID = Normals->AppendElement((FVector3f)CardGeo.Normals[k]);
+			int32 NewNormID = Normals->AppendElement(&CardGeo.Normals[k].X);
 			MapN[k] = NewNormID;
 
 			int32 NewUVID = UVs->AppendElement( FVector2f(CardGeo.UVs[k].X, CardGeo.UVs[k].Y) );
diff --git a/Engine/Plugins/Experimental/MeshModelingToolset/Source/MeshModelingToolsEditorOnly/Private/Hair/GroomToMeshTool.cpp b/Engine/Plugins/Experimental/MeshModelingToolset/Source/MeshModelingToolsEditorOnly/Private/Hair/GroomToMeshTool.cpp
--- a/Engine/Plugins/Experimental/MeshModelingToolset/Source/MeshModelingToolsEditorOnly/Private/Hair/GroomToMeshTool.cpp
+++ b/Engine/Plugins/Experimental/MeshModelingToolset/Source/MeshModelingToolsEditorOnly/Private/Hair/GroomToMeshTool.cpp
@@ -481,7 +481,7 @@ TSharedPtr<FDynamicMesh3, ESPMode::ThreadSafe> UGroomToMeshTool::UpdateVoxelizat
 	//MarchingCubes.RootMode = ERootfindingModes::Bisection;
 	MarchingCubes.RootModeSteps = 4;
 
-	TArray<FVector3d> Seeds;
+	TArray<FVector3<double>> Seeds;
 	for (int32 k = 0; k < Lines.Num(); ++k)
 	{
 		Seeds.Add( Lines[k].Segment.EndPoint() + Lines[k].GetRadius()*Lines[k].Segment.Direction );
diff --git a/Engine/Plugins/Experimental/MeshModelingToolset/Source/MeshModelingToolsEditorOnly/Private/MeshTangentsTool.cpp b/Engine/Plugins/Experimental/MeshModelingToolset/Source/MeshModelingToolsEditorOnly/Private/MeshTangentsTool.cpp
--- a/Engine/Plugins/Experimental/MeshModelingToolset/Source/MeshModelingToolsEditorOnly/Private/MeshTangentsTool.cpp
+++ b/Engine/Plugins/Experimental/MeshModelingToolset/Source/MeshModelingToolsEditorOnly/Private/MeshTangentsTool.cpp
@@ -179,11 +179,11 @@ void UMeshTangentsTool::Render(IToolsContextRenderAPI* RenderAPI)
 			if (ErrorPt.MaxAngleDeg > Settings->AngleThreshDeg)
 			{
 				Visualizer.DrawPoint(ErrorPt.VertexPos, FLinearColor(0.95, 0.05, 0.05), 6 * Settings->LineThickness, false);
-				Visualizer.DrawLine(ErrorPt.VertexPos, ErrorPt.VertexPos + Settings->LineLength * ErrorPt.MikktTangent, FLinearColor(0.95, 0.05, 0.05), 2 * Settings->LineThickness, false);
-				Visualizer.DrawLine(ErrorPt.VertexPos, ErrorPt.VertexPos + Settings->LineLength * ErrorPt.MikktBitangent, FLinearColor(0.05, 0.95, 0.05), 2 * Settings->LineThickness, false);
+				Visualizer.DrawLine<FVector3f>(ErrorPt.VertexPos, ErrorPt.VertexPos + Settings->LineLength * ErrorPt.MikktTangent, FLinearColor(0.95, 0.05, 0.05), 2 * Settings->LineThickness, false);
+				Visualizer.DrawLine<FVector3f>(ErrorPt.VertexPos, ErrorPt.VertexPos + Settings->LineLength * ErrorPt.MikktBitangent, FLinearColor(0.05, 0.95, 0.05), 2 * Settings->LineThickness, false);
 
-				Visualizer.DrawLine(ErrorPt.VertexPos, ErrorPt.VertexPos + (1.1f*Settings->LineLength) * ErrorPt.OtherTangent, FLinearColor(0.95, 0.50, 0.05), Settings->LineThickness, false);
-				Visualizer.DrawLine(ErrorPt.VertexPos, ErrorPt.VertexPos + (1.1f*Settings->LineLength) * ErrorPt.OtherBitangent, FLinearColor(0.05, 0.95, 0.95), Settings->LineThickness, false);
+				Visualizer.DrawLine<FVector3f>(ErrorPt.VertexPos, ErrorPt.VertexPos + (1.1f*Settings->LineLength) * ErrorPt.OtherTangent, FLinearColor(0.95, 0.50, 0.05), Settings->LineThickness, false);
+				Visualizer.DrawLine<FVector3f>(ErrorPt.VertexPos, ErrorPt.VertexPos + (1.1f*Settings->LineLength) * ErrorPt.OtherBitangent, FLinearColor(0.05, 0.95, 0.95), Settings->LineThickness, false);
 			}
 		}
 		Visualizer.EndFrame();
@@ -311,19 +311,19 @@ void UMeshTangentsTool::OnTangentsUpdated(const TUniquePtr<FMeshTangentsd>& NewR
 				InputMesh->GetTriVertices(Index, Verts[0], Verts[1], Verts[2]);
 				for (int j = 0; j < 3; ++j)
 				{
-					FVector3d TangentMikkt, BitangentMikkt;
-					MikktTangents->GetPerTriangleTangent(Index, j, TangentMikkt, BitangentMikkt);
+					FVector3f TangentMikkt, BitangentMikkt;
+					MikktTangents->GetPerTriangleTangent<FVector3f, float>(Index, j, TangentMikkt, BitangentMikkt);
 					UE::Geometry::Normalize(TangentMikkt); UE::Geometry::Normalize(BitangentMikkt);
-					FVector3d TangentNew, BitangentNew;
-					NewTangents->GetPerTriangleTangent(Index, j, TangentNew, BitangentNew);
+					FVector3f TangentNew, BitangentNew;
+					NewTangents->GetPerTriangleTangent<FVector3f, float>(Index, j, TangentNew, BitangentNew);
 					UE::Geometry::Normalize(TangentNew); 
 					UE::Geometry::Normalize(BitangentNew);
 					ensure(UE::Geometry::IsNormalized(TangentMikkt) && UE::Geometry::IsNormalized(BitangentMikkt));
 					ensure(UE::Geometry::IsNormalized(TangentNew) && UE::Geometry::IsNormalized(BitangentNew));
-					double MaxAngleDeg = FMathd::Max(UE::Geometry::AngleD(TangentMikkt, TangentNew), UE::Geometry::AngleD(BitangentMikkt, BitangentNew));
-					if (MaxAngleDeg > 0.5)
+					float MaxAngleDeg = FMathf::Max(UE::Geometry::AngleD(TangentMikkt, TangentNew), UE::Geometry::AngleD(BitangentMikkt, BitangentNew));
+					if (MaxAngleDeg > 0.5f)
 					{
-						FMikktDeviation Deviation{ static_cast<float>(MaxAngleDeg), Index, j, Verts[j], TangentMikkt, BitangentMikkt, TangentNew, BitangentNew };
+						FMikktDeviation Deviation{ MaxAngleDeg, Index, j, (FVector3f)Verts[j], TangentMikkt, BitangentMikkt, TangentNew, BitangentNew };
 						Deviations.Add(Deviation);
 					}
 				}
diff --git a/Engine/Plugins/Experimental/MeshModelingToolset/Source/MeshModelingToolsEditorOnly/Private/ShapeSprayTool.cpp b/Engine/Plugins/Experimental/MeshModelingToolset/Source/MeshModelingToolsEditorOnly/Private/ShapeSprayTool.cpp
--- a/Engine/Plugins/Experimental/MeshModelingToolset/Source/MeshModelingToolsEditorOnly/Private/ShapeSprayTool.cpp
+++ b/Engine/Plugins/Experimental/MeshModelingToolset/Source/MeshModelingToolsEditorOnly/Private/ShapeSprayTool.cpp
@@ -134,8 +134,8 @@ void UShapeSprayTool::OnUpdateDrag(const FRay& Ray)
 {
 	UDynamicMeshBrushTool::OnUpdateDrag(Ray);
 
-	FFrame3f WorldFrame(LastBrushStamp.WorldPosition, LastBrushStamp.WorldNormal);
-	FTransform Transform = Cast<IPrimitiveComponentBackedTarget>(Target)->GetWorldTransform();
+	FFrame3f WorldFrame((FVector3f)LastBrushStamp.WorldPosition, (FVector3f)LastBrushStamp.WorldNormal);
+	FTransform3d Transform(Cast<IPrimitiveComponentBackedTarget>(Target)->GetWorldTransform());
 
 
 	FDynamicMesh3* Mesh = AccumMeshComponent->GetMesh();
@@ -163,8 +163,8 @@ void UShapeSprayTool::OnUpdateDrag(const FRay& Ray)
 			float ObjSize = (1.0f + (Random.GetFraction()-0.5f) ) *  (Settings->ObjectSize);
 
 			FFrame3d SurfFrame(
-				Transform.InverseTransformPosition(Hit.ImpactPoint),
-				Transform.InverseTransformVector(Hit.Normal));
+				Transform.InverseTransformPosition((FVector3d)Hit.ImpactPoint),
+				Transform.InverseTransformVector((FVector3d)Hit.Normal));
 			SplatShape(SurfFrame, ObjSize, Mesh);
 		}
 	}
@@ -176,7 +176,7 @@ void UShapeSprayTool::OnUpdateDrag(const FRay& Ray)
 
 void UShapeSprayTool::SplatShape(const FFrame3d& LocalFrame, double Scale, FDynamicMesh3* TargetMesh)
 {
-	FVector3f UseColor = (Settings->bRandomColor) ? FLinearColor::MakeRandomColor() : Settings->Color;
+	FLinearColor UseColor = (Settings->bRandomColor) ? FLinearColor::MakeRandomColor() : Settings->Color;
 	FQuaternionf Rotationf = (FQuaternionf)LocalFrame.Rotation;
 
 	VertexMap.Reset();
@@ -190,7 +190,7 @@ void UShapeSprayTool::SplatShape(const FFrame3d& LocalFrame, double Scale, FDyna
 		FVector3f Normal = Rotationf * ShapeMesh.GetVertexNormal(vid);
 		TargetMesh->SetVertexNormal(VertexMap[vid], Normal);
 
-		TargetMesh->SetVertexColor(VertexMap[vid], UseColor);
+		TargetMesh->SetVertexColor(VertexMap[vid], (FVector3<float>)UseColor);
 	}
 	for (int tid : ShapeMesh.TriangleIndicesItr())
 	{
diff --git a/Engine/Plugins/Experimental/MeshModelingToolset/Source/MeshModelingToolsEditorOnly/Public/MeshTangentsTool.h b/Engine/Plugins/Experimental/MeshModelingToolset/Source/MeshModelingToolsEditorOnly/Public/MeshTangentsTool.h
--- a/Engine/Plugins/Experimental/MeshModelingToolset/Source/MeshModelingToolsEditorOnly/Public/MeshTangentsTool.h
+++ b/Engine/Plugins/Experimental/MeshModelingToolset/Source/MeshModelingToolsEditorOnly/Public/MeshTangentsTool.h
@@ -133,9 +133,9 @@ protected:
 		float MaxAngleDeg;
 		int32 TriangleID;
 		int32 TriVertIndex;
-		FVector3d VertexPos;
-		FVector3d MikktTangent, MikktBitangent;
-		FVector3d OtherTangent, OtherBitangent;
+		FVector3f VertexPos;
+		FVector3f MikktTangent, MikktBitangent;
+		FVector3f OtherTangent, OtherBitangent;
 	};
 	TArray<FMikktDeviation> Deviations;
 
diff --git a/Engine/Plugins/Experimental/MeshModelingToolset/Source/ModelingComponents/Private/Drawing/UVLayoutPreview.cpp b/Engine/Plugins/Experimental/MeshModelingToolset/Source/ModelingComponents/Private/Drawing/UVLayoutPreview.cpp
--- a/Engine/Plugins/Experimental/MeshModelingToolset/Source/ModelingComponents/Private/Drawing/UVLayoutPreview.cpp
+++ b/Engine/Plugins/Experimental/MeshModelingToolset/Source/ModelingComponents/Private/Drawing/UVLayoutPreview.cpp
@@ -149,7 +149,7 @@ void UUVLayoutPreview::UpdateUVMesh(const FDynamicMesh3* SourceMesh, int32 Sourc
 			{
 				NewUVTri[j] = NewUVOverlay->AppendElement(UVs[j]);
 				NewTri[j] = UVMesh.AppendVertex(FVector3d(UVs[j].X, UVs[j].Y, 0));
-				NewNormalTri[j] = NewNormalOverlay->AppendElement(FVector3f::UnitZ());
+				NewNormalTri[j] = NewNormalOverlay->AppendElement( FVector3<float>::UnitZ() );
 			}
 
 			FVector2f EdgeUV1 = UVs[1] - UVs[0];
@@ -192,7 +192,7 @@ void UUVLayoutPreview::UpdateUVMesh(const FDynamicMesh3* SourceMesh, int32 Sourc
 void UUVLayoutPreview::RecalculatePosition()
 {
 	FFrame3d ObjFrame;
-	ObjFrame.AlignAxis(2, -CameraState.Forward());
+	ObjFrame.AlignAxis(2, -(FVector3d)CameraState.Forward());
 	ObjFrame.ConstrainedAlignAxis(0, (FVector3d)CameraState.Right(), ObjFrame.Z());
 	ObjFrame.Origin = SourceObjectWorldBounds.Center(); // SourceObjectFrame.Origin;
 
diff --git a/Engine/Plugins/Experimental/MeshModelingToolset/Source/ModelingComponents/Private/Mechanics/CollectSurfacePathMechanic.cpp b/Engine/Plugins/Experimental/MeshModelingToolset/Source/ModelingComponents/Private/Mechanics/CollectSurfacePathMechanic.cpp
--- a/Engine/Plugins/Experimental/MeshModelingToolset/Source/ModelingComponents/Private/Mechanics/CollectSurfacePathMechanic.cpp
+++ b/Engine/Plugins/Experimental/MeshModelingToolset/Source/ModelingComponents/Private/Mechanics/CollectSurfacePathMechanic.cpp
@@ -14,7 +14,7 @@ using namespace UE::Geometry;
 
 UCollectSurfacePathMechanic::UCollectSurfacePathMechanic()
 {
-	SpatialSnapPointsFunc = [this](FVector3d A, FVector3d B) { return A.DistanceSquared(B) < (ConstantSnapDistance * ConstantSnapDistance); };
+	SpatialSnapPointsFunc = [this](FVector3d A, FVector3d B) { return DistanceSquared(A, B) < (ConstantSnapDistance * ConstantSnapDistance); };
 
 	PathColor = LinearColors::DarkOrange3f();
 	PreviewColor = LinearColors::Green3f();
diff --git a/Engine/Plugins/Experimental/MeshModelingToolset/Source/ModelingComponents/Private/Mechanics/CurveControlPointsMechanic.cpp b/Engine/Plugins/Experimental/MeshModelingToolset/Source/ModelingComponents/Private/Mechanics/CurveControlPointsMechanic.cpp
--- a/Engine/Plugins/Experimental/MeshModelingToolset/Source/ModelingComponents/Private/Mechanics/CurveControlPointsMechanic.cpp
+++ b/Engine/Plugins/Experimental/MeshModelingToolset/Source/ModelingComponents/Private/Mechanics/CurveControlPointsMechanic.cpp
@@ -87,7 +87,7 @@ void UCurveControlPointsMechanic::Setup(UInteractiveTool* ParentToolIn)
 			// As in PointSnapQuery, we convert our angle-based tolerance to one we can use in an ortho viewport (instead of
 			// dividing our field of view into 90 visual angle degrees, we divide the plane into 90 units).
 			float OrthoTolerance = ToolSceneQueriesUtil::GetDefaultVisualAngleSnapThreshD() * CameraState.OrthoWorldCoordinateWidth / 90.0;
-			return Position1.DistanceSquared(Position2) < OrthoTolerance * OrthoTolerance;
+			return DistanceSquared(Position1, Position2) < OrthoTolerance * OrthoTolerance;
 		}
 		else
 		{
@@ -285,7 +285,7 @@ int32 UCurveControlPointsMechanic::InsertPointAt(int32 SequencePosition, const F
 			int32 PreviousSequencePosition = (SequencePosition + ControlPoints.Num() - 1) % ControlPoints.Num();
 			int32 PreviousID = ControlPoints.GetPointIDAt(PreviousSequencePosition);
 
-			FPolyline3d SegmentPolyline(TArray<FVector3d>{ControlPoints.GetPointCoordinates(PreviousID), NewPointCoordinates});
+			FPolyline3d SegmentPolyline(ControlPoints.GetPointCoordinates(PreviousID), NewPointCoordinates);
 
 			if (DrawnControlSegments->IsLineValid(PreviousID))
 			{
@@ -307,7 +307,7 @@ int32 UCurveControlPointsMechanic::InsertPointAt(int32 SequencePosition, const F
 			// Create a segment going to the next point
 			int32 NextSequencePosition = (SequencePosition + 1) % ControlPoints.Num();
 
-			FPolyline3d SegmentPolyline(TArray<FVector3d>{ControlPoints.GetPointCoordinatesAt(NextSequencePosition), NewPointCoordinates});
+			FPolyline3d SegmentPolyline(ControlPoints.GetPointCoordinatesAt(NextSequencePosition), NewPointCoordinates);
 			GeometrySet.AddCurve(NewPointID, SegmentPolyline);
 
 			FRenderableLine RenderableSegment((FVector)NewPointCoordinates, (FVector)ControlPoints.GetPointCoordinatesAt(NextSequencePosition),
@@ -349,7 +349,7 @@ void UCurveControlPointsMechanic::SetIsLoop(bool bIsLoopIn)
 		{
 			if (bIsLoopIn)
 			{
-				FPolyline3d SegmentPolyline(TArray<FVector3d>{ControlPoints.GetPointCoordinates(ControlPoints.Last()), ControlPoints.GetPointCoordinates(ControlPoints.First())});
+				FPolyline3d SegmentPolyline(ControlPoints.GetPointCoordinates(ControlPoints.Last()), ControlPoints.GetPointCoordinates(ControlPoints.First()));
 				GeometrySet.AddCurve(ControlPoints.Last(), SegmentPolyline);
 
 				FRenderableLine RenderableSegment((FVector)ControlPoints.GetPointCoordinates(ControlPoints.Last()),
@@ -450,7 +450,7 @@ void UCurveControlPointsMechanic::GizmoTransformChanged(UTransformProxy* Proxy,
 		// Do snapping only if we have a single point selected.
 		if (SelectedPointIDs.Num() == 1 && (bSnappingEnabled ^ bSnapToggle))
 		{
-			FVector3d NewLocation = SelectedPointStartPositions[0] + Displacement;
+			FVector3d NewLocation = SelectedPointStartPositions[0] + (FVector3d)Displacement;
 			if (bSnappingEnabled ^ bSnapToggle)
 			{
 				SnapEngine.UpdateSnappedPoint(NewLocation);
@@ -471,7 +471,7 @@ void UCurveControlPointsMechanic::GizmoTransformChanged(UTransformProxy* Proxy,
 		{
 			for (int32 i = 0; i < SelectedPointIDs.Num(); ++i)
 			{
-				UpdatePointLocation(SelectedPointIDs[i], SelectedPointStartPositions[i] + Displacement);
+				UpdatePointLocation(SelectedPointIDs[i], SelectedPointStartPositions[i] + (FVector3d)Displacement);
 			}
 			bPointsChanged = true;
 		}
@@ -527,7 +527,7 @@ void UCurveControlPointsMechanic::UpdatePointLocation(int32 PointID, const FVect
 		int32 PreviousSequencePosition = (SequencePosition + ControlPoints.Num() - 1) % ControlPoints.Num();
 		DrawnControlSegments->SetLineEnd(ControlPoints.GetPointIDAt(PreviousSequencePosition), (FVector)NewLocation);
 
-		FPolyline3d SegmentPolyline(TArray<FVector3d>{ControlPoints.GetPointCoordinatesAt(PreviousSequencePosition), NewLocation});
+		FPolyline3d SegmentPolyline(ControlPoints.GetPointCoordinatesAt(PreviousSequencePosition), NewLocation);
 		GeometrySet.UpdateCurve(ControlPoints.GetPointIDAt(PreviousSequencePosition), SegmentPolyline);
 	}
 
@@ -536,7 +536,7 @@ void UCurveControlPointsMechanic::UpdatePointLocation(int32 PointID, const FVect
 	{
 		DrawnControlSegments->SetLineStart(PointID, (FVector)NewLocation);
 
-		FPolyline3d SegmentPolyline(TArray<FVector3d>{NewLocation, ControlPoints.GetPointCoordinatesAt((SequencePosition + 1) % ControlPoints.Num())});
+		FPolyline3d SegmentPolyline(NewLocation, ControlPoints.GetPointCoordinatesAt((SequencePosition + 1) % ControlPoints.Num()));
 		GeometrySet.UpdateCurve(PointID, SegmentPolyline);
 	}
 }
@@ -553,7 +553,7 @@ bool UCurveControlPointsMechanic::HitTest(const FInputDeviceRay& ClickPos, FInpu
 		&& (SelectedPointIDs[0] == ControlPoints.First() || SelectedPointIDs[0] == ControlPoints.Last())))
 	{
 		FVector3d HitPoint;
-		bool bHit = DrawPlane.RayPlaneIntersection(ClickPos.WorldRay.Origin, ClickPos.WorldRay.Direction, 2, HitPoint);
+		bool bHit = DrawPlane.RayPlaneIntersection((FVector3d)ClickPos.WorldRay.Origin, (FVector3d)ClickPos.WorldRay.Direction, 2, HitPoint);
 		if (bHit)
 		{
 			ResultOut = FInputRayHit(ClickPos.WorldRay.GetParameter((FVector)HitPoint));
@@ -592,7 +592,7 @@ void UCurveControlPointsMechanic::OnClicked(const FInputDeviceRay& ClickPos)
 	if (bInteractiveInitializationMode)
 	{
 		FVector3d NewPointCoordinates;
-		if (!DrawPlane.RayPlaneIntersection(ClickPos.WorldRay.Origin, ClickPos.WorldRay.Direction, 2, NewPointCoordinates))
+		if (!DrawPlane.RayPlaneIntersection((FVector3d)ClickPos.WorldRay.Origin, (FVector3d)ClickPos.WorldRay.Direction, 2, NewPointCoordinates))
 		{
 			// Missed the plane entirely (probably in ortho mode). Nothing to do.
 			return;
@@ -674,7 +674,7 @@ void UCurveControlPointsMechanic::OnClicked(const FInputDeviceRay& ClickPos)
 
 			int32 NewPointID = -1;
 			FVector3d NewPointCoordinates;
-			DrawPlane.RayPlaneIntersection(ClickPos.WorldRay.Origin, ClickPos.WorldRay.Direction, 2, NewPointCoordinates);
+			DrawPlane.RayPlaneIntersection((FVector3d)ClickPos.WorldRay.Origin, (FVector3d)ClickPos.WorldRay.Direction, 2, NewPointCoordinates);
 
 			if (bSnappingEnabled ^ bSnapToggle)
 			{
@@ -780,7 +780,7 @@ void UCurveControlPointsMechanic::UpdateGizmoLocation()
 		{
 			NewGizmoLocation += ControlPoints.GetPointCoordinates(PointID);
 		}
-		NewGizmoLocation /= SelectedPointIDs.Num();
+		NewGizmoLocation /= (double)SelectedPointIDs.Num();
 
 		PointTransformGizmo->ReinitializeGizmoTransform(FTransform((FQuat)DrawPlane.Rotation, (FVector)NewGizmoLocation));
 	}
@@ -922,7 +922,7 @@ bool UCurveControlPointsMechanic::OnUpdateHover(const FInputDeviceRay& DevicePos
 		(bInsertPointToggle && SelectedPointIDs.Num() == 1 && (SelectedPointIDs[0] == ControlPoints.First() || SelectedPointIDs[0] == ControlPoints.Last())))
 	{
 		FVector3d HitPoint;
-		if (!DrawPlane.RayPlaneIntersection(DevicePos.WorldRay.Origin, DevicePos.WorldRay.Direction, 2, HitPoint))
+		if (!DrawPlane.RayPlaneIntersection((FVector3d)DevicePos.WorldRay.Origin, (FVector3d)DevicePos.WorldRay.Direction, 2, HitPoint))
 		{
 			// We're probably looking in an ortho viewport and missing the plane. End the hover.
 			return false;
@@ -1081,7 +1081,7 @@ int32 UCurveControlPointsMechanic::DeletePoint(int32 PointID)
 			FVector3d NextPointCoordinates = ControlPoints.GetPointCoordinatesAt((SequencePosition + 1) % ControlPoints.Num());
 
 			DrawnControlSegments->SetLineEnd(PreviousPointID, (FVector)NextPointCoordinates);
-			FPolyline3d SegmentPolyline(TArray<FVector3d>{ControlPoints.GetPointCoordinates(PreviousPointID), NextPointCoordinates});
+			FPolyline3d SegmentPolyline(ControlPoints.GetPointCoordinates(PreviousPointID), NextPointCoordinates);
 			GeometrySet.UpdateCurve(PreviousPointID, SegmentPolyline);
 		}
 		else
diff --git a/Engine/Plugins/Experimental/MeshModelingToolset/Source/ModelingComponents/Private/Mechanics/DragAlignmentMechanic.cpp b/Engine/Plugins/Experimental/MeshModelingToolset/Source/ModelingComponents/Private/Mechanics/DragAlignmentMechanic.cpp
--- a/Engine/Plugins/Experimental/MeshModelingToolset/Source/ModelingComponents/Private/Mechanics/DragAlignmentMechanic.cpp
+++ b/Engine/Plugins/Experimental/MeshModelingToolset/Source/ModelingComponents/Private/Mechanics/DragAlignmentMechanic.cpp
@@ -31,7 +31,7 @@ void UDragAlignmentMechanic::Setup(UInteractiveTool* ParentToolIn)
 		const TArray<const UPrimitiveComponent*>* InvisibleComponentsToInclude)
 	{
 		bool bSuccess = false;
-		FVector3d QueryPoint = WorldRay.PointAt(1);
+		FVector3d QueryPoint = (FVector3d)WorldRay.PointAt(1);
 		FVector3d SnapPoint;
 		ToolSceneQueriesUtil::FSnapGeometry SnapGeometry;
 
@@ -170,7 +170,7 @@ void UDragAlignmentMechanic::OnGizmoTransformChanged(FTransform NewTransform)
 	// Ignore changes that weren't preceded by a successful query (likely programatic)
 	if (bWaitingOnProjectedResult) 
 	{
-		LastProjectedResult = NewTransform.GetLocation();
+		LastProjectedResult = (FVector3d)NewTransform.GetLocation();
 		bPreviewEndpointsValid = true;
 	}
 	bWaitingOnProjectedResult = false;
@@ -198,8 +198,8 @@ void UDragAlignmentMechanic::InitializeDeformedMeshRayCast(TFunction<FDynamicMes
 	MeshRayCast = [this, GetSpatialIn, TargetTransform, LinearDeformer](const FRay& WorldRay, FHitResult& HitResult, bool bUseFilter)
 	{
 		// Transform ray into local space
-		FRay3d LocalRay(TargetTransform.InverseTransformPosition(WorldRay.Origin),
-			TargetTransform.InverseTransformVector(WorldRay.Direction));
+		FRay3d LocalRay(TargetTransform.InverseTransformPosition((FVector3d)WorldRay.Origin),
+			TargetTransform.InverseTransformVector((FVector3d)WorldRay.Direction));
 		UE::Geometry::Normalize(LocalRay.Direction);
 
 		FDynamicMeshAABBTree3* Spatial = GetSpatialIn();
@@ -224,7 +224,7 @@ void UDragAlignmentMechanic::InitializeDeformedMeshRayCast(TFunction<FDynamicMes
 
 		// See if we're close enough to one of the triangle's vertices to snap to it. We have
 		// to do this check in world space since scaling will affect things.
-		FVector3d WorldHitPoint = WorldRay.PointAt(RayTValue);
+		FVector3d WorldHitPoint = (FVector3d)WorldRay.PointAt(RayTValue);
 
 		FVector3d Vert1, Vert2, Vert3;
 		Mesh->GetTriVertices(HitTid, Vert1, Vert2, Vert3);
@@ -233,14 +233,14 @@ void UDragAlignmentMechanic::InitializeDeformedMeshRayCast(TFunction<FDynamicMes
 		Vert3 = TargetTransform.TransformPosition(Vert3);
 
 		FVector3d* ClosestVert = &Vert1;
-		double MinDistSquared = WorldHitPoint.DistanceSquared(Vert1);
-		double DistSquared = WorldHitPoint.DistanceSquared(Vert2);
+		double MinDistSquared = DistanceSquared(WorldHitPoint, Vert1);
+		double DistSquared = DistanceSquared(WorldHitPoint, Vert2);
 		if (DistSquared < MinDistSquared)
 		{
 			ClosestVert = &Vert2;
 			MinDistSquared = DistSquared;
 		}
-		DistSquared = WorldHitPoint.DistanceSquared(Vert3);
+		DistSquared = DistanceSquared(WorldHitPoint, Vert3);
 		if (DistSquared < MinDistSquared)
 		{
 			ClosestVert = &Vert3;
@@ -330,7 +330,7 @@ void UDragAlignmentMechanic::Render(IToolsContextRenderAPI* RenderAPI)
 	if (LastHitResult.Item == TRIANGLE_HIT_TYPE_ID)
 	{
 		// We'll draw a circle with an X in it lying on the face we hit.
-		FFrame3d HitFrame = FFrame3d(LastHitResult.ImpactPoint, LastHitResult.ImpactNormal);
+		FFrame3d HitFrame = FFrame3d((FVector3d)LastHitResult.ImpactPoint, (FVector3d)LastHitResult.ImpactNormal);
 
 		float TickLength = (float)ToolSceneQueriesUtil::CalculateDimensionFromVisualAngleD(CameraState, HitFrame.Origin, 0.8);
 		float TickThickness = 4 * PDIScale;
@@ -352,7 +352,7 @@ void UDragAlignmentMechanic::Render(IToolsContextRenderAPI* RenderAPI)
 
 	// If the two endpoints are far enough apart, draw a line connecting them.
 	if (ToolSceneQueriesUtil::CalculateNormalizedViewVisualAngleD(
-		CameraState, LastHitResult.ImpactPoint, LastProjectedResult) > VisualAngleRenderThreshold)
+		CameraState, (FVector3d)LastHitResult.ImpactPoint, LastProjectedResult) > VisualAngleRenderThreshold)
 	{
 		PDI->DrawLine(LastHitResult.ImpactPoint, (FVector)LastProjectedResult, Color,
 			SDPG_Foreground, 0.5f * PDIScale, 0.0f, true);
diff --git a/Engine/Plugins/Experimental/MeshModelingToolset/Source/ModelingComponents/Private/Mechanics/LatticeControlPointsMechanic.cpp b/Engine/Plugins/Experimental/MeshModelingToolset/Source/ModelingComponents/Private/Mechanics/LatticeControlPointsMechanic.cpp
--- a/Engine/Plugins/Experimental/MeshModelingToolset/Source/ModelingComponents/Private/Mechanics/LatticeControlPointsMechanic.cpp
+++ b/Engine/Plugins/Experimental/MeshModelingToolset/Source/ModelingComponents/Private/Mechanics/LatticeControlPointsMechanic.cpp
@@ -86,7 +86,7 @@ void ULatticeControlPointsMechanic::Setup(UInteractiveTool* ParentToolIn)
 			// As in PointSnapQuery, we convert our angle-based tolerance to one we can use in an ortho viewport (instead of
 			// dividing our field of view into 90 visual angle degrees, we divide the plane into 90 units).
 			float OrthoTolerance = ToolSceneQueriesUtil::GetDefaultVisualAngleSnapThreshD() * CachedCameraState.OrthoWorldCoordinateWidth / 90.0;
-			return Position1.DistanceSquared(Position2) < OrthoTolerance * OrthoTolerance;
+			return DistanceSquared(Position1, Position2) < OrthoTolerance * OrthoTolerance;
 		}
 		else
 		{
@@ -302,9 +302,9 @@ void ULatticeControlPointsMechanic::GizmoTransformChanged(UTransformProxy* Proxy
 	FVector DeltaScale = Transform.GetScale3D() / GizmoStartScale;
 
 	FTransform3d DeltaTransform;
-	DeltaTransform.SetScale(DeltaScale);
-	DeltaTransform.SetRotation(DeltaRotation);
-	DeltaTransform.SetTranslation(Transform.GetTranslation());
+	DeltaTransform.SetScale((FVector3d)DeltaScale);
+	DeltaTransform.SetRotation((FQuaterniond)DeltaRotation);
+	DeltaTransform.SetTranslation((FVector3d)Transform.GetTranslation());
 
 	// If any deltas are non-zero
 	if (Displacement != FVector::ZeroVector || !DeltaRotation.EpsilonEqual(FQuaterniond::Identity(), SMALL_NUMBER) || DeltaScale != FVector::OneVector)
@@ -314,7 +314,7 @@ void ULatticeControlPointsMechanic::GizmoTransformChanged(UTransformProxy* Proxy
 			FVector3d PointPosition = SelectedPointStartPositions[PointID];
 
 			// Translate to origin, scale, rotate, and translate back (DeltaTransform has "translate back" baked in.)
-			PointPosition -= GizmoStartPosition;
+			PointPosition -= (FVector3d)GizmoStartPosition;
 			PointPosition = DeltaTransform.TransformPosition(PointPosition);
 
 			ControlPoints[PointID] = PointPosition;
@@ -473,7 +473,7 @@ void ULatticeControlPointsMechanic::UpdateGizmoLocation()
 		{
 			NewGizmoLocation += ControlPoints[PointID];
 		}
-		NewGizmoLocation /= SelectedPointIDs.Num();
+		NewGizmoLocation /= (double)SelectedPointIDs.Num();
 
 		PointTransformGizmo->SetVisibility(true);
 	}
diff --git a/Engine/Plugins/Experimental/MeshModelingToolset/Source/ModelingComponents/Private/Mechanics/PlaneDistanceFromHitMechanic.cpp b/Engine/Plugins/Experimental/MeshModelingToolset/Source/ModelingComponents/Private/Mechanics/PlaneDistanceFromHitMechanic.cpp
--- a/Engine/Plugins/Experimental/MeshModelingToolset/Source/ModelingComponents/Private/Mechanics/PlaneDistanceFromHitMechanic.cpp
+++ b/Engine/Plugins/Experimental/MeshModelingToolset/Source/ModelingComponents/Private/Mechanics/PlaneDistanceFromHitMechanic.cpp
@@ -82,7 +82,7 @@ void UPlaneDistanceFromHitMechanic::UpdateCurrentDistance(const FRay& WorldRay)
 			float WorldHitDist = WorldRay.GetParameter(WorldHitResult.ImpactPoint);
 			if (WorldHitDist < NearestHitDist)
 			{
-				NearestHitFrameWorld = FFrame3d(WorldHitResult.ImpactPoint, WorldHitResult.ImpactNormal);
+				NearestHitFrameWorld = FFrame3d((FVector3d)WorldHitResult.ImpactPoint, (FVector3d)WorldHitResult.ImpactNormal);
 				FVector3d HitPosLocal = PreviewHeightFrame.ToFramePoint(FVector3d(WorldHitResult.ImpactPoint));
 				NearestHitHeight = HitPosLocal.Z;
 				NearestHitDist = WorldHitDist;
@@ -166,10 +166,10 @@ void UPlaneDistanceFromHitMechanic::Render(IToolsContextRenderAPI* RenderAPI)
 	FVector3d PreviewOrigin = PreviewHeightFrame.Origin;
 	FVector3d DrawPlaneNormal = PreviewHeightFrame.Z();
 
-	Renderer.DrawLine(
+	Renderer.DrawLine<FVector3d>(
 		PreviewOrigin - 1000*DrawPlaneNormal, PreviewOrigin + 1000*DrawPlaneNormal, AxisColor, 1.0, false);
-	Renderer.DrawLine(
-		PreviewOrigin + CurrentHeight*DrawPlaneNormal, CurrentHitPosFrameWorld.Origin, HeightPosColor, 1.0, false);
+	Renderer.DrawLine<FVector3d>(
+		PreviewOrigin + (double)CurrentHeight*DrawPlaneNormal, CurrentHitPosFrameWorld.Origin, HeightPosColor, 1.0, false);
 
 	Renderer.EndFrame();
 }
\ No newline at end of file
diff --git a/Engine/Plugins/Experimental/MeshModelingToolset/Source/ModelingComponents/Private/Mechanics/SpaceCurveDeformationMechanic.cpp b/Engine/Plugins/Experimental/MeshModelingToolset/Source/ModelingComponents/Private/Mechanics/SpaceCurveDeformationMechanic.cpp
--- a/Engine/Plugins/Experimental/MeshModelingToolset/Source/ModelingComponents/Private/Mechanics/SpaceCurveDeformationMechanic.cpp
+++ b/Engine/Plugins/Experimental/MeshModelingToolset/Source/ModelingComponents/Private/Mechanics/SpaceCurveDeformationMechanic.cpp
@@ -536,7 +536,7 @@ void USpaceCurveDeformationMechanic::UpdateGizmoLocation()
 			{
 				SharedOrigin += CurvePoints[PointID].Origin;
 			}
-			NewGizmoFrame = FFrame3d(SharedOrigin / SelectedPointIDs.Num());
+			NewGizmoFrame = FFrame3d(SharedOrigin / (double)SelectedPointIDs.Num());
 		}
 
 		//if (SelectedPointIDs.Num() > 1 && TransformProperties->TransformMode == ESpaceCurveControlPointTransformMode::Shared)
diff --git a/Engine/Plugins/Experimental/MeshModelingToolset/Source/ModelingComponents/Private/Selection/GroupTopologySelector.cpp b/Engine/Plugins/Experimental/MeshModelingToolset/Source/ModelingComponents/Private/Selection/GroupTopologySelector.cpp
--- a/Engine/Plugins/Experimental/MeshModelingToolset/Source/ModelingComponents/Private/Selection/GroupTopologySelector.cpp
+++ b/Engine/Plugins/Experimental/MeshModelingToolset/Source/ModelingComponents/Private/Selection/GroupTopologySelector.cpp
@@ -22,7 +22,7 @@ FGroupTopologySelector::FGroupTopologySelector()
 {
 	// initialize to sane values
 	PointsWithinToleranceTest =
-		[](const FVector3d& A, const FVector3d& B, double TolScale) { return A.Distance(B) < (TolScale*1.0); };
+		[](const FVector3d& A, const FVector3d& B, double TolScale) { return Distance(A, B) < (TolScale*1.0); };
 	GetSpatial =
 		[]() { return nullptr; };
 }
@@ -599,7 +599,7 @@ bool FGroupTopologySelector::FindSelectedElement(const FSelectionSettings& Setti
 	FDynamicMeshAABBTree3* Spatial = GetSpatial();
 
 	// Needed for occlusion test, which happens in local space.
-	FVector3d LocalCameraOrigin = TargetTransform.InverseTransformPosition(CameraRectangle.CameraOrigin);
+	FVector3d LocalCameraOrigin = TargetTransform.InverseTransformPosition((FVector3d)CameraRectangle.CameraOrigin);
 	
 	// Corner selection takes priority over edges.
 	if (Settings.bEnableCornerHits)
@@ -629,7 +629,7 @@ bool FGroupTopologySelector::FindSelectedElement(const FSelectionSettings& Setti
 				}
 
 				// Check whether any of the component segments intersect the rectangle
-				const TArray<FVector3d>& Verts = Curve.GetVertices();
+				const TArray<FVector3<double>>& Verts = Curve.GetVertices();
 				FVector3d CurrentVert = TargetTransform.TransformPosition(Verts[0]);
 				for (int32 i = 1; i < Verts.Num(); ++i)
 				{
@@ -851,7 +851,7 @@ void FGroupTopologySelector::DrawSelection(const FGroupTopologySelection& Select
 		// Depending on whether we're in an orthographic view or not, we set the radius based on visual angle or based on ortho 
 		// viewport width (divided into 90 segments like the FOV is divided into 90 degrees).
 		float Radius = (CameraState->bIsOrthographic) ? CameraState->OrthoWorldCoordinateWidth * 0.5 / 90.0
-			: (float)ToolSceneQueriesUtil::CalculateDimensionFromVisualAngleD(*CameraState, WorldPosition, 0.5);
+			: (float)ToolSceneQueriesUtil::CalculateDimensionFromVisualAngleD(*CameraState, (FVector3d)WorldPosition, 0.5);
 		Renderer->DrawViewFacingCircle(Position, Radius, 16, UseColor, LineWidth, false);
 	}
 
diff --git a/Engine/Plugins/Experimental/MeshModelingToolset/Source/ModelingComponents/Private/Selection/PolygonSelectionMechanic.cpp b/Engine/Plugins/Experimental/MeshModelingToolset/Source/ModelingComponents/Private/Selection/PolygonSelectionMechanic.cpp
--- a/Engine/Plugins/Experimental/MeshModelingToolset/Source/ModelingComponents/Private/Selection/PolygonSelectionMechanic.cpp
+++ b/Engine/Plugins/Experimental/MeshModelingToolset/Source/ModelingComponents/Private/Selection/PolygonSelectionMechanic.cpp
@@ -275,8 +275,8 @@ bool UPolygonSelectionMechanic::TopologyHitTest(const FRay& WorldRay, FHitResult
 
 bool UPolygonSelectionMechanic::TopologyHitTest(const FRay& WorldRay, FHitResult& OutHit, FGroupTopologySelection& OutSelection, bool bUseOrthoSettings)
 {
-	FRay3d LocalRay(TargetTransform.InverseTransformPosition(WorldRay.Origin),
-		TargetTransform.InverseTransformVector(WorldRay.Direction));
+	FRay3d LocalRay(TargetTransform.InverseTransformPosition((FVector3d)WorldRay.Origin),
+		TargetTransform.InverseTransformVector((FVector3d)WorldRay.Direction));
 	UE::Geometry::Normalize(LocalRay.Direction);
 
 	FVector3d LocalPosition, LocalNormal;
@@ -359,8 +359,8 @@ bool UPolygonSelectionMechanic::UpdateHighlight(const FRay& WorldRay)
 {
 	checkf(DrawnTriangleSetComponent != nullptr, TEXT("Initialize() not called on UPolygonSelectionMechanic."));
 
-	FRay3d LocalRay(TargetTransform.InverseTransformPosition(WorldRay.Origin),
-		TargetTransform.InverseTransformVector(WorldRay.Direction));
+	FRay3d LocalRay(TargetTransform.InverseTransformPosition((FVector3d)WorldRay.Origin),
+		TargetTransform.InverseTransformVector((FVector3d)WorldRay.Direction));
 	UE::Geometry::Normalize(LocalRay.Direction);
 
 	HilightSelection.Clear();
@@ -438,8 +438,8 @@ bool UPolygonSelectionMechanic::HasSelection() const
 
 bool UPolygonSelectionMechanic::UpdateSelection(const FRay& WorldRay, FVector3d& LocalHitPositionOut, FVector3d& LocalHitNormalOut)
 {
-	FRay3d LocalRay(TargetTransform.InverseTransformPosition(WorldRay.Origin),
-		TargetTransform.InverseTransformVector(WorldRay.Direction));
+	FRay3d LocalRay(TargetTransform.InverseTransformPosition((FVector3d)WorldRay.Origin),
+		TargetTransform.InverseTransformVector((FVector3d)WorldRay.Direction));
 	UE::Geometry::Normalize(LocalRay.Direction);
 
 	bool bSelectionModified = false;
diff --git a/Engine/Plugins/Experimental/MeshModelingToolset/Source/ModelingComponents/Private/Snapping/BasePositionSnapSolver3.cpp b/Engine/Plugins/Experimental/MeshModelingToolset/Source/ModelingComponents/Private/Snapping/BasePositionSnapSolver3.cpp
--- a/Engine/Plugins/Experimental/MeshModelingToolset/Source/ModelingComponents/Private/Snapping/BasePositionSnapSolver3.cpp
+++ b/Engine/Plugins/Experimental/MeshModelingToolset/Source/ModelingComponents/Private/Snapping/BasePositionSnapSolver3.cpp
@@ -12,7 +12,7 @@ FBasePositionSnapSolver3::FBasePositionSnapSolver3()
 {
 	// initialize to sane values
 	SnapMetricFunc =
-		[](const FVector3d& A, const FVector3d& B) { return A.Distance(B); };
+		[](const FVector3d& A, const FVector3d& B) { return Distance(A, B); };
 	SnapMetricTolerance = 1.0;
 
 	Reset();
diff --git a/Engine/Plugins/Experimental/MeshModelingToolset/Source/ModelingComponents/Private/Snapping/PointPlanarSnapSolver.cpp b/Engine/Plugins/Experimental/MeshModelingToolset/Source/ModelingComponents/Private/Snapping/PointPlanarSnapSolver.cpp
--- a/Engine/Plugins/Experimental/MeshModelingToolset/Source/ModelingComponents/Private/Snapping/PointPlanarSnapSolver.cpp
+++ b/Engine/Plugins/Experimental/MeshModelingToolset/Source/ModelingComponents/Private/Snapping/PointPlanarSnapSolver.cpp
@@ -44,7 +44,7 @@ void FPointPlanarSnapSolver::UpdatePointHistory(const TArray<FVector>& Points)
 	int Num = Points.Num();
 	for (int k = 0; k < Num; ++k)
 	{
-		PointHistory.Add(Points[k]);
+		PointHistory.Add((FVector3d)Points[k]);
 	}
 	ResetGenerated();
 }
@@ -197,7 +197,7 @@ void FPointPlanarSnapSolver::GenerateLineIntersectionTargets()
 			FVector2d IntersectionPoint;
 			if (LinesInPlane[i].IntersectionPoint(LinesInPlane[j], IntersectionPoint))
 			{
-				Intersection.Position = Plane.FromFramePoint(FVector3d(IntersectionPoint));
+				Intersection.Position = Plane.FromFramePoint(FVector3d(IntersectionPoint.X, IntersectionPoint.Y, 0.0));
 
 				// Lines in plane are 1:1 with our OriginalLinePointers array, which has more info per line.
 				// Store the more important priority line in the snap target, and the second line in the 
@@ -252,7 +252,7 @@ void FPointPlanarSnapSolver::GenerateTargets(const FVector3d& PointIn)
 		int NumHistoryPts = PointHistory.Num();
 		for (int j = 0; j < NumHistoryPts - 1; ++j)
 		{
-			double SegLength = PointHistory[j].Distance(PointHistory[j + 1]);
+			double SegLength = Distance(PointHistory[j], PointHistory[j+1]);
 			for (int k = 0; k < NumSnapLines; ++k)
 			{
 				FSnapTargetPoint Target;
diff --git a/Engine/Plugins/Experimental/MeshModelingToolset/Source/ModelingComponents/Private/Snapping/RaySpatialSnapSolver.cpp b/Engine/Plugins/Experimental/MeshModelingToolset/Source/ModelingComponents/Private/Snapping/RaySpatialSnapSolver.cpp
--- a/Engine/Plugins/Experimental/MeshModelingToolset/Source/ModelingComponents/Private/Snapping/RaySpatialSnapSolver.cpp
+++ b/Engine/Plugins/Experimental/MeshModelingToolset/Source/ModelingComponents/Private/Snapping/RaySpatialSnapSolver.cpp
@@ -177,7 +177,7 @@ void FRaySpatialSnapSolver::Draw(FToolDataVisualizer* Renderer, float LineLength
 			const FLine3d& Line = GetActiveSnapLine();
 			const FLinearColor* LineColor = (ColorMap == nullptr) ? nullptr : ColorMap->Find(ActiveSnapID);
 			FLinearColor UseColor = (LineColor == nullptr) ? Renderer->LineColor : *LineColor;
-			Renderer->DrawLine(Line.PointAt(0), GetActiveSnapToPoint(), UseColor, Renderer->LineThickness*0.5);
+			Renderer->DrawLine( (FVector3d)Line.PointAt(0), GetActiveSnapToPoint(), UseColor, Renderer->LineThickness*0.5);
 		}
 		Renderer->DrawPoint(GetActiveSnapToPoint());
 	}
diff --git a/Engine/Plugins/Experimental/MeshModelingToolset/Source/ModelingComponents/Private/ToolSceneQueriesUtil.cpp b/Engine/Plugins/Experimental/MeshModelingToolset/Source/ModelingComponents/Private/ToolSceneQueriesUtil.cpp
--- a/Engine/Plugins/Experimental/MeshModelingToolset/Source/ModelingComponents/Private/ToolSceneQueriesUtil.cpp
+++ b/Engine/Plugins/Experimental/MeshModelingToolset/Source/ModelingComponents/Private/ToolSceneQueriesUtil.cpp
@@ -40,7 +40,7 @@ bool ToolSceneQueriesUtil::PointSnapQuery(const FViewCameraState& CameraState, c
 		// where we divide the FOV into 90 visual angle degrees, we divide the plane into 90 segments and use the same tolerance.
 		double AngleThreshold = (VisualAngleThreshold <= 0) ? GetDefaultVisualAngleSnapThreshD() : VisualAngleThreshold;
 		double OrthoThreshold = AngleThreshold * CameraState.OrthoWorldCoordinateWidth / 90.0;
-		FVector3d ViewPlaneNormal = CameraState.Orientation.GetForwardVector();
+		FVector3d ViewPlaneNormal = (FVector3d)CameraState.Orientation.GetForwardVector();
 		FVector3d DistanceVector = Point1 - Point2;
 
 		// Project the vector into the plane and check its length
@@ -61,7 +61,7 @@ double ToolSceneQueriesUtil::PointSnapMetric(const FViewCameraState& CameraState
 	}
 	else
 	{
-		FVector3d ViewPlaneNormal = CameraState.Orientation.GetForwardVector();
+		FVector3d ViewPlaneNormal = (FVector3d)CameraState.Orientation.GetForwardVector();
 
 		// Get projected distance in the plane
 		FVector3d DistanceVector = Point1 - Point2;
@@ -109,9 +109,9 @@ double ToolSceneQueriesUtil::CalculateDimensionFromVisualAngleD(const FViewCamer
 	FVector3d EyePos = (FVector3d)CameraState.Position;
 	FVector3d PointVec = Point - EyePos;
 	TargetVisualAngleDeg *= CameraState.GetFOVAngleNormalizationFactor();
-	FVector3d RotPointPos = EyePos + FQuaterniond(CameraState.Up(), TargetVisualAngleDeg, true)*PointVec;
+	FVector3d RotPointPos = EyePos + FQuaterniond((FVector3d)CameraState.Up(), TargetVisualAngleDeg, true)*PointVec;
 	double ActualAngleDeg = CalculateViewVisualAngleD(CameraState, Point, RotPointPos);
-	return Point.Distance(RotPointPos) * (TargetVisualAngleDeg/ActualAngleDeg);
+	return Distance(Point, RotPointPos) * (TargetVisualAngleDeg/ActualAngleDeg);
 }
 
 
@@ -120,9 +120,9 @@ bool ToolSceneQueriesUtil::IsPointVisible(const FViewCameraState& CameraState, c
 {
 	if (CameraState.bIsOrthographic == false)
 	{
-		FVector3d PointDir = (Point - CameraState.Position);
+		FVector3d PointDir = (Point - (FVector3d)CameraState.Position);
 		//@todo should use view frustum here!
-		if (PointDir.Dot(CameraState.Forward()) < 0.25)		// ballpark estimate
+		if (PointDir.Dot((FVector3d)CameraState.Forward()) < 0.25)		// ballpark estimate
 		{
 			return false;
 		}
@@ -181,16 +181,16 @@ bool ToolSceneQueriesUtil::FindSceneSnapPoint(FFindSceneSnapPointParams& Params)
 	TArray<FSceneSnapQueryResult> Results;
 	if (QueryAPI->ExecuteSceneSnapQuery(Request, Results))
 	{
-		*Params.SnapPointOut = Results[0].Position;
+		*Params.SnapPointOut = (FVector3d)Results[0].Position;
 
 		if (Params.SnapGeometryOut != nullptr)
 		{
 			int iSnap = Results[0].TriSnapIndex;
-			Params.SnapGeometryOut->Points[0] = Results[0].TriVertices[iSnap];
+			Params.SnapGeometryOut->Points[0] = (FVector3d)Results[0].TriVertices[iSnap];
 			Params.SnapGeometryOut->PointCount = 1;
 			if (Results[0].TargetType == ESceneSnapQueryTargetType::MeshEdge)
 			{
-				Params.SnapGeometryOut->Points[1] = Results[0].TriVertices[(iSnap+1)%3];
+				Params.SnapGeometryOut->Points[1] = (FVector3d)Results[0].TriVertices[(iSnap+1)%3];
 				Params.SnapGeometryOut->PointCount = 2;
 			}
 		}
@@ -218,7 +218,7 @@ bool ToolSceneQueriesUtil::FindWorldGridSnapPoint(const UInteractiveTool* Tool,
 	TArray<FSceneSnapQueryResult> Results;
 	if ( QueryAPI->ExecuteSceneSnapQuery(Request, Results) )
 	{
-		GridSnapPointOut = Results[0].Position;
+		GridSnapPointOut = (FVector3d)Results[0].Position;
 		return true;
 	};
 	return false;
diff --git a/Engine/Plugins/Experimental/MeshModelingToolset/Source/ModelingComponents/Public/AssetUtils/Texture2DBuilder.h b/Engine/Plugins/Experimental/MeshModelingToolset/Source/ModelingComponents/Public/AssetUtils/Texture2DBuilder.h
--- a/Engine/Plugins/Experimental/MeshModelingToolset/Source/ModelingComponents/Public/AssetUtils/Texture2DBuilder.h
+++ b/Engine/Plugins/Experimental/MeshModelingToolset/Source/ModelingComponents/Public/AssetUtils/Texture2DBuilder.h
@@ -346,7 +346,7 @@ public:
 			Pixel.X = FMathf::Clamp(Pixel.X, 0.0, 1.0);
 			Pixel.Y = FMathf::Clamp(Pixel.Y, 0.0, 1.0);
 			Pixel.Z = FMathf::Clamp(Pixel.Z, 0.0, 1.0);
-			FColor Texel = ((FLinearColor)Pixel).ToFColor(bSRGB);
+			FColor Texel = ToLinearColor(Pixel).ToFColor(bSRGB);
 			SetTexel(i, Texel);
 		}
 		return true;
diff --git a/Engine/Plugins/Experimental/MeshModelingToolset/Source/ModelingOperators/Private/CleaningOps/HoleFillOp.cpp b/Engine/Plugins/Experimental/MeshModelingToolset/Source/ModelingOperators/Private/CleaningOps/HoleFillOp.cpp
--- a/Engine/Plugins/Experimental/MeshModelingToolset/Source/ModelingOperators/Private/CleaningOps/HoleFillOp.cpp
+++ b/Engine/Plugins/Experimental/MeshModelingToolset/Source/ModelingOperators/Private/CleaningOps/HoleFillOp.cpp
@@ -368,10 +368,10 @@ void FHoleFillOp::CalculateResult(FProgressCancel* Progress)
 		NewGroupIDs.Add(NewGroupID);
 
 		// Compute a best-fit plane of the boundary vertices
-		TArray<FVector3d> VertexPositions;
+		TArray<FVector3<double>> VertexPositions;
 		Loop.GetVertices(VertexPositions);
-		FVector3d PlaneOrigin;
-		FVector3d PlaneNormal;
+		FVector3<double> PlaneOrigin;
+		FVector3<double> PlaneNormal;
 		PolygonTriangulation::ComputePolygonPlane<double>(VertexPositions, PlaneNormal, PlaneOrigin);
 		PlaneNormal *= -1.0;	// Previous function seems to orient the normal opposite to what's expected elsewhere
 
diff --git a/Engine/Plugins/Experimental/MeshModelingToolset/Source/ModelingOperators/Private/CompositionOps/CurveSweepOp.cpp b/Engine/Plugins/Experimental/MeshModelingToolset/Source/ModelingOperators/Private/CompositionOps/CurveSweepOp.cpp
--- a/Engine/Plugins/Experimental/MeshModelingToolset/Source/ModelingOperators/Private/CompositionOps/CurveSweepOp.cpp
+++ b/Engine/Plugins/Experimental/MeshModelingToolset/Source/ModelingOperators/Private/CompositionOps/CurveSweepOp.cpp
@@ -91,10 +91,10 @@ void FCurveSweepOp::CalculateResult(FProgressCancel* Progress)
 			}
 
 			// Compute a best-fit plane of the vertices
-			TArray<FVector3d> VertexPositions;
+			TArray<FVector3<double>> VertexPositions;
 			Loop.GetVertices(VertexPositions);
-			FVector3d PlaneOrigin;
-			FVector3d PlaneNormal;
+			FVector3<double> PlaneOrigin;
+			FVector3<double> PlaneNormal;
 			PolygonTriangulation::ComputePolygonPlane<double>(VertexPositions, PlaneNormal, PlaneOrigin);
 			PlaneNormal *= -1.0;	// Previous function seems to orient the normal opposite to what's expected elsewhere
 
diff --git a/Engine/Plugins/Experimental/MeshModelingToolset/Source/ModelingOperators/Private/CuttingOps/PlaneCutOp.cpp b/Engine/Plugins/Experimental/MeshModelingToolset/Source/ModelingOperators/Private/CuttingOps/PlaneCutOp.cpp
--- a/Engine/Plugins/Experimental/MeshModelingToolset/Source/ModelingOperators/Private/CuttingOps/PlaneCutOp.cpp
+++ b/Engine/Plugins/Experimental/MeshModelingToolset/Source/ModelingOperators/Private/CuttingOps/PlaneCutOp.cpp
@@ -64,7 +64,7 @@ void FPlaneCutOp::CalculateResult(FProgressCancel* Progress)
 			{
 				SecondCutMaxID = FMath::Max(SecondCutMaxID, SubObjectAttrib->GetValue(TID));
 			}
-			Cut.PlaneOrigin += Cut.PlaneNormal * (2 * CutPlaneLocalThickness);
+			Cut.PlaneOrigin += Cut.PlaneNormal * (2.0 * CutPlaneLocalThickness);
 			Cut.CutWithoutDelete(true, 0, SubObjectAttrib, SecondCutMaxID + 1, false, true);
 
 			for (int TID : ResultMesh->TriangleIndicesItr())
diff --git a/Engine/Plugins/Experimental/MeshModelingToolset/Source/ModelingOperators/Public/CuttingOps/PlaneCutOp.h b/Engine/Plugins/Experimental/MeshModelingToolset/Source/ModelingOperators/Public/CuttingOps/PlaneCutOp.h
--- a/Engine/Plugins/Experimental/MeshModelingToolset/Source/ModelingOperators/Public/CuttingOps/PlaneCutOp.h
+++ b/Engine/Plugins/Experimental/MeshModelingToolset/Source/ModelingOperators/Public/CuttingOps/PlaneCutOp.h
@@ -19,11 +19,12 @@ public:
 	virtual ~FPlaneCutOp() {}
 
 	// inputs
-	FVector LocalPlaneOrigin, LocalPlaneNormal;
-	bool bFillCutHole, bFillSpans;
-	bool bKeepBothHalves;
-	float CutPlaneLocalThickness = 0; // plane thickness in the local space of the mesh
-	float UVScaleFactor;
+	FVector3d LocalPlaneOrigin, LocalPlaneNormal;
+	bool bFillCutHole = true;
+	bool bFillSpans = false;
+	bool bKeepBothHalves = false;
+	double CutPlaneLocalThickness = 0; // plane thickness in the local space of the mesh
+	double UVScaleFactor = 0;
 	static const FName ObjectIndexAttribute;
 	TSharedPtr<const FDynamicMesh3, ESPMode::ThreadSafe> OriginalMesh;
 
diff --git a/Engine/Plugins/Experimental/MeshModelingToolset/Source/ModelingOperatorsEditorOnly/Private/CuttingOps/EmbedPolygonsOp.cpp b/Engine/Plugins/Experimental/MeshModelingToolset/Source/ModelingOperatorsEditorOnly/Private/CuttingOps/EmbedPolygonsOp.cpp
--- a/Engine/Plugins/Experimental/MeshModelingToolset/Source/ModelingOperatorsEditorOnly/Private/CuttingOps/EmbedPolygonsOp.cpp
+++ b/Engine/Plugins/Experimental/MeshModelingToolset/Source/ModelingOperatorsEditorOnly/Private/CuttingOps/EmbedPolygonsOp.cpp
@@ -60,7 +60,7 @@ void CollapseDegenerateEdgesOnVertexPath(FDynamicMesh3& Mesh, TArray<int>& Verte
 				continue;
 			}
 			Mesh.GetEdgeV(EID, A, B);
-			double DSq = A.DistanceSquared(B);
+			double DSq = DistanceSquared(A,B);
 			if (DSq > Tol2)
 			{
 				continue;
diff --git a/Engine/Plugins/Experimental/PlanarCutPlugin/Source/PlanarCut/Private/FractureAutoUV.cpp b/Engine/Plugins/Experimental/PlanarCutPlugin/Source/PlanarCut/Private/FractureAutoUV.cpp
--- a/Engine/Plugins/Experimental/PlanarCutPlugin/Source/PlanarCut/Private/FractureAutoUV.cpp
+++ b/Engine/Plugins/Experimental/PlanarCutPlugin/Source/PlanarCut/Private/FractureAutoUV.cpp
@@ -166,7 +166,7 @@ struct FGeomMesh : public UE::Geometry::FUVPacker::IUVMeshView
 		GlobalVertices.SetNum(Collection->Vertex.Num());
 		for (int32 Idx = 0; Idx < GlobalVertices.Num(); Idx++)
 		{
-			GlobalVertices[Idx] = GlobalTransformArray[Collection->BoneMap[Idx]].TransformPosition(Collection->Vertex[Idx]);
+			GlobalVertices[Idx] = (FVector3d)GlobalTransformArray[Collection->BoneMap[Idx]].TransformPosition(Collection->Vertex[Idx]);
 		}
 	}
 
@@ -235,7 +235,7 @@ struct FGeomMesh : public UE::Geometry::FUVPacker::IUVMeshView
 		return 0;
 	}
 
-	inline void GetTriVertices(int TID, FVector3d& V0, FVector3d& V1, FVector3d& V2) const
+	inline void GetTriVertices(int TID, UE::Core::TVector<double>& V0, UE::Core::TVector<double>& V1, UE::Core::TVector<double>& V2) const
 	{
 		FIntVector TriRaw = Collection->Indices[TID];
 
@@ -297,7 +297,7 @@ struct FGeomFlatUVMesh
 		return 0;
 	}
 
-	inline void GetTriVertices(int TID, FVector3d& V0, FVector3d& V1, FVector3d& V2) const
+	inline void GetTriVertices(int TID, UE::Core::TVector<double>& V0, UE::Core::TVector<double>& V1, UE::Core::TVector<double>& V2) const
 	{
 		FIntVector TriRaw = Collection->Indices[TID];
 
diff --git a/Engine/Plugins/Experimental/PlanarCutPlugin/Source/PlanarCut/Private/PlanarCut.cpp b/Engine/Plugins/Experimental/PlanarCutPlugin/Source/PlanarCut/Private/PlanarCut.cpp
--- a/Engine/Plugins/Experimental/PlanarCutPlugin/Source/PlanarCut/Private/PlanarCut.cpp
+++ b/Engine/Plugins/Experimental/PlanarCutPlugin/Source/PlanarCut/Private/PlanarCut.cpp
@@ -187,8 +187,8 @@ namespace UE
 				{
 					FVector3f Tangents[2];
 					UE::PlanarCutInternals::AugmentDynamicMesh::GetTangent(Mesh, VID, Tangents[0], Tangents[1]);
-					int TID = TangentOverlays[0]->AppendElement(Tangents[0]);
-					int TID2 = TangentOverlays[1]->AppendElement(Tangents[1]);
+					int TID = TangentOverlays[0]->AppendElement(&Tangents[0].X);
+					int TID2 = TangentOverlays[1]->AppendElement(&Tangents[1].X);
 					check(TID == TID2);
 					VertToTangentMap[VID] = TID;
 				}
@@ -223,8 +223,8 @@ namespace UE
 				FMeshTangentsf Tangents(&Mesh);
 				Tangents.ComputeTriVertexTangents(Normals, UVs, Options);
 
-				const TArray<FVector3f>& TanU = Tangents.GetTangents();
-				const TArray<FVector3f>& TanV = Tangents.GetBitangents();
+				const TArray<FVector3<float>>& TanU = Tangents.GetTangents();
+				const TArray<FVector3<float>>& TanV = Tangents.GetBitangents();
 				FDynamicMeshMaterialAttribute* MaterialIDs = Mesh.Attributes()->GetMaterialID();
 				for (int TID : Mesh.TriangleIndicesItr())
 				{
@@ -296,7 +296,7 @@ namespace UE
 						double MaxEdgeLenSq = 0;
 						for (int i = 2, j = 0; j < 3; i = j++)
 						{
-							double EdgeLenSq = Triangle.V[i].DistanceSquared(Triangle.V[j]);
+							double EdgeLenSq = DistanceSquared(Triangle.V[i], Triangle.V[j]);
 							if (EdgeLenSq > MaxEdgeLenSq)
 							{
 								MaxEdgeIdx = i;
@@ -346,7 +346,7 @@ namespace UE
 									}
 									TPair<int, double> VIDDist = KnownSamples[ComponentIdx].FindNearestInRadius(SamplePos, Spacing * .5, [&Mesh, SamplePos](int VID)
 										{
-											return Mesh.GetVertex(VID).DistanceSquared(SamplePos);
+											return DistanceSquared(Mesh.GetVertex(VID), SamplePos);
 										});
 									// No point within radius Spacing/2 -> Add a new sample
 									if (VIDDist.Key == -1)
@@ -530,8 +530,8 @@ struct FCellMeshes
 
 	void ApplyNoise(FDynamicMesh3& Mesh, FVector3d Normal, const FNoiseSettings& Settings, bool bProjectBoundariesToNormal = false)
 	{
-		float Amplitude = Settings.Amplitude;
-		float Frequency = Settings.Frequency;
+		double Amplitude = (double)Settings.Amplitude;
+		double Frequency = (double)Settings.Frequency;
 		int32 Octaves = Settings.Octaves;
 		FVector3d Z = Normal * Amplitude;
 
@@ -657,7 +657,7 @@ struct FCellMeshes
 			{
 				VertexCentroid += V;
 			}
-			VertexCentroid /= Mesh.VertexCount();
+			VertexCentroid /= (double)Mesh.VertexCount();
 			FAxisAlignedBox3d Bounds = Mesh.GetCachedBounds();
 			double BoundsSize = Bounds.MaxDim();
 			// currently just scale the meshes down so they leave half-a-grout worth of space on their longest axis
@@ -742,7 +742,7 @@ private:
 				PlaneVertInfo.Normal = Normal;
 				if (MeshIdx == 1 && OtherCell != OutsideCellIndex)
 				{
-					PlaneVertInfo.Normal *= -1;
+					PlaneVertInfo.Normal *= -1.0f;
 				}
 				VertStart[MeshIdx] = Meshes[MeshIdx]->MaxVertexID();
 				FVector2f MinUV(FMathf::MaxReal, FMathf::MaxReal);
@@ -1010,7 +1010,7 @@ private:
 						for (int T1SubIdx = 0; T1SubIdx < 3; T1SubIdx++)
 						{
 							FVector3d V1 = Mesh.GetVertex(Tri[1][T1SubIdx]);
-							if (V0.DistanceSquared(V1) < FMathd::ZeroTolerance)
+							if (DistanceSquared(V0, V1) < FMathd::ZeroTolerance)
 							{
 								MatchedVertices++;
 								break;
@@ -1033,7 +1033,7 @@ private:
 								MeshesVertices.FindPointsInBall(P, FMathd::ZeroTolerance, [this, P](FIndex2i IDs)
 									{
 										FVector3d Pos = CellMeshes[IDs.A].AugMesh.GetVertex(IDs.B);
-										return P.DistanceSquared(Pos);
+										return DistanceSquared(P, Pos);
 									}, Update.IDs);
 							}
 						}
@@ -1272,7 +1272,7 @@ private:
 					Meshes[Side]->ReverseOrientation(true);
 					Offset = ZRange.Min;
 				}
-				PlaneVertInfo.Normal = FVector3f(Plane.GetNormal()) * (-1 + Side * 2);
+				PlaneVertInfo.Normal = FVector3f(Plane.GetNormal()) * (-1.0f + (float)Side * 2.0f);
 				FVector3d OffsetVec = FVector3d(Plane.GetNormal()) * Offset;
 
 				for (int CornerIdx = 0; CornerIdx < 4; CornerIdx++)
@@ -1300,13 +1300,13 @@ private:
 		{
 			int GroutIdx = bOnlyGrout ? 0 : 2;
 			FDynamicMesh3* GroutMesh = &CellMeshes[GroutIdx].AugMesh;
-			FVector3d GroutOffset = Plane.GetNormal() * (Grout * .5);
+			FVector3d GroutOffset = (FVector3d)Plane.GetNormal() * (Grout * .5);
 			if (!bOnlyGrout)
 			{
 				for (int Side = 0; Side < 2; Side++)
 				{
 					// shift both sides out by Grout/2
-					MeshTransforms::Translate(*Meshes[Side], GroutOffset * (-1 + Side * 2));
+					MeshTransforms::Translate(*Meshes[Side], GroutOffset * (-1.0 + (double)Side * 2.0));
 				}
 			}
 
@@ -1412,11 +1412,11 @@ struct FDynamicMeshCollection
 				VertexInfo.Position = CollectionToLocal.TransformPosition(FVector3d(Collection->Vertex[Idx]));
 				VertexInfo.UV = FVector2f(Collection->UV[Idx]);
 				VertexInfo.Color = FVector3f(Collection->Color[Idx]);
-				VertexInfo.Normal = CollectionToLocal.TransformVectorNoScale(FVector3f(Collection->Normal[Idx]));
+				VertexInfo.Normal = (FVector3f)CollectionToLocal.TransformVectorNoScale(FVector3d(Collection->Normal[Idx]));
 				int VID = Mesh.AppendVertex(VertexInfo);
 				UE::PlanarCutInternals::AugmentDynamicMesh::SetTangent(Mesh, VID, VertexInfo.Normal,
-					CollectionToLocal.TransformVectorNoScale(FVector3f(Collection->TangentU[Idx])),
-					CollectionToLocal.TransformVectorNoScale(FVector3f(Collection->TangentV[Idx])));
+					(FVector3f)CollectionToLocal.TransformVectorNoScale(FVector3d(Collection->TangentU[Idx])),
+					(FVector3f)CollectionToLocal.TransformVectorNoScale(FVector3d(Collection->TangentV[Idx])));
 			}
 			FIntVector VertexOffset(VertexStart, VertexStart, VertexStart);
 			for (int32 Idx = Collection->FaceStart[GeometryIdx], N = Collection->FaceStart[GeometryIdx] + FaceCount; Idx < N; Idx++)
@@ -1433,8 +1433,8 @@ struct FDynamicMeshCollection
 						int32 SrcIdx = AddTri[SubIdx] + VertexStart;
 						UE::PlanarCutInternals::AugmentDynamicMesh::SetTangent(Mesh, NewVID,
 							Mesh.GetVertexNormal(NewVID), // TODO: we don't actually use the vertex normal; consider removing this arg from the function entirely
-							CollectionToLocal.TransformVectorNoScale(FVector3f(Collection->TangentU[SrcIdx])),
-							CollectionToLocal.TransformVectorNoScale(FVector3f(Collection->TangentV[SrcIdx])));
+							(FVector3f)CollectionToLocal.TransformVectorNoScale(FVector3d(Collection->TangentU[SrcIdx])),
+							(FVector3f)CollectionToLocal.TransformVectorNoScale(FVector3d(Collection->TangentV[SrcIdx])));
 						NewTri[SubIdx] = NewVID;
 					}
 					TID = Mesh.AppendTriangle(NewTri, 0);
@@ -1581,7 +1581,7 @@ struct FDynamicMeshCollection
 			{
 				TPair<int, double> Nearest = VerticesHashes[B]->FindNearestInRadius(V, FMathd::ZeroTolerance * 10, [&RefMesh, &V](int VID)
 					{
-						return RefMesh.GetVertex(VID).DistanceSquared(V);
+						return DistanceSquared(RefMesh.GetVertex(VID), V);
 					});
 				if (Nearest.Key != -1)
 				{
@@ -2026,7 +2026,7 @@ struct FDynamicMeshCollection
 		{
 			TPair<int, double> Nearest = VertHash[B]->FindNearestInRadius(V, FMathd::ZeroTolerance * 10, [&RefMesh, &V](int VID)
 				{
-					return RefMesh.GetVertex(VID).DistanceSquared(V);
+					return DistanceSquared(RefMesh.GetVertex(VID), V);
 				});
 			if (Nearest.Key != -1)
 			{
@@ -2048,7 +2048,7 @@ struct FDynamicMeshCollection
 		{
 			FVector3d Pt = Source.GetVertex(VID);
 			Neighbors.Reset();
-			VertHash.FindPointsInBall(Pt, SnapDistance, [&Source, Pt](int OtherVID) {return Pt.DistanceSquared(Source.GetVertex(OtherVID));}, Neighbors);
+			VertHash.FindPointsInBall(Pt, SnapDistance, [&Source, Pt](int OtherVID) {return DistanceSquared(Pt, Source.GetVertex(OtherVID));}, Neighbors);
 			for (int NbrVID : Neighbors)
 			{
 				VertComponents.UnionSequential(VID, NbrVID);
diff --git a/Engine/Plugins/Experimental/Water/Source/Runtime/Private/WaterBooleanUtils.h b/Engine/Plugins/Experimental/Water/Source/Runtime/Private/WaterBooleanUtils.h
--- a/Engine/Plugins/Experimental/Water/Source/Runtime/Private/WaterBooleanUtils.h
+++ b/Engine/Plugins/Experimental/Water/Source/Runtime/Private/WaterBooleanUtils.h
@@ -303,9 +303,9 @@ namespace UE {
 			{
 				FIndex3i Tri = Mesh.GetTriangle(tid);
 				FVector3d Normal = Mesh.GetTriNormal(tid);
-				int32 e0 = Normals->AppendElement((FVector3f)Normal);
-				int32 e1 = Normals->AppendElement((FVector3f)Normal);
-				int32 e2 = Normals->AppendElement((FVector3f)Normal);
+				int32 e0 = Normals->AppendElement((FVector3<float>)Normal);
+				int32 e1 = Normals->AppendElement((FVector3<float>)Normal);
+				int32 e2 = Normals->AppendElement((FVector3<float>)Normal);
 				Normals->SetTriangle(tid, FIndex3i(e0, e1, e2));
 			}
 		}
