Fix geometry TVector4's One() to return the correct type

#rnx

diff --git a/Engine/Plugins/Experimental/GeometryProcessing/Source/GeometricObjects/Public/VectorTypes.h b/Engine/Plugins/Experimental/GeometryProcessing/Source/GeometricObjects/Public/VectorTypes.h
--- a/Engine/Plugins/Experimental/GeometryProcessing/Source/GeometricObjects/Public/VectorTypes.h
+++ b/Engine/Plugins/Experimental/GeometryProcessing/Source/GeometricObjects/Public/VectorTypes.h
@@ -951,7 +951,7 @@ struct TVector4
 	}
 	static TVector4<T> One()
 	{
-		return FVector3<T>((T)1, (T)1, (T)1, (T)1);
+		return TVector4<T>((T)1, (T)1, (T)1, (T)1);
 	}
 
 
