add an "IndexOf" to FIndex4i

#rnx

diff --git a/Engine/Plugins/Experimental/GeometryProcessing/Source/GeometricObjects/Public/IndexTypes.h b/Engine/Plugins/Experimental/GeometryProcessing/Source/GeometricObjects/Public/IndexTypes.h
--- a/Engine/Plugins/Experimental/GeometryProcessing/Source/GeometricObjects/Public/IndexTypes.h
+++ b/Engine/Plugins/Experimental/GeometryProcessing/Source/GeometricObjects/Public/IndexTypes.h
@@ -225,6 +225,11 @@ struct FIndex4i
 		return A != Other.A || B != Other.B || C != Other.C || D != Other.D;
 	}
 
+	int IndexOf(int Value) const
+	{
+		return (A == Value) ? 0 : ((B == Value) ? 1 : ((C == Value) ? 2 : ((D == Value) ? 3 : -1)));
+	}
+
 	bool Contains(int Idx) const
 	{
 		return A == Idx || B == Idx || C == Idx || D == Idx;
