GeometryProcessing: Fix inverted parallel flag for MeshImageBakingCache.

#rnx

diff --git a/Engine/Plugins/Experimental/GeometryProcessing/Source/DynamicMesh/Private/Sampling/MeshImageBakingCache.cpp b/Engine/Plugins/Experimental/GeometryProcessing/Source/DynamicMesh/Private/Sampling/MeshImageBakingCache.cpp
--- a/Engine/Plugins/Experimental/GeometryProcessing/Source/DynamicMesh/Private/Sampling/MeshImageBakingCache.cpp
+++ b/Engine/Plugins/Experimental/GeometryProcessing/Source/DynamicMesh/Private/Sampling/MeshImageBakingCache.cpp
@@ -297,5 +297,5 @@ void FMeshImageBakingCache::EvaluateSamples(
 
 		SampleFunction(Coords, Sample);
 
-	}, bParallel ? EParallelForFlags::ForceSingleThread : EParallelForFlags::None);
+	}, !bParallel ? EParallelForFlags::ForceSingleThread : EParallelForFlags::None);
 }
\ No newline at end of file
