less roundabout SetPixel in GeometricObjects' ImageBuilder

#rnx

diff --git a/Engine/Plugins/Experimental/GeometryProcessing/Source/GeometricObjects/Public/Image/ImageBuilder.h b/Engine/Plugins/Experimental/GeometryProcessing/Source/GeometricObjects/Public/Image/ImageBuilder.h
--- a/Engine/Plugins/Experimental/GeometryProcessing/Source/GeometricObjects/Public/Image/ImageBuilder.h
+++ b/Engine/Plugins/Experimental/GeometryProcessing/Source/GeometricObjects/Public/Image/ImageBuilder.h
@@ -78,7 +78,7 @@ public:
 	 */
 	void SetPixel(int64 LinearIndex, const PixelType& NewValue)
 	{
-		SetPixel(Dimensions.GetCoords(LinearIndex), NewValue);
+		Image[LinearIndex] = NewValue;
 	}
 
 
