Change ParallelFor calls in geometry imageoccupancymap to be per-row to account for int32 index type

#rnx

diff --git a/Engine/Plugins/Experimental/GeometryProcessing/Source/GeometricObjects/Public/Image/ImageOccupancyMap.h b/Engine/Plugins/Experimental/GeometryProcessing/Source/GeometricObjects/Public/Image/ImageOccupancyMap.h
--- a/Engine/Plugins/Experimental/GeometryProcessing/Source/GeometricObjects/Public/Image/ImageOccupancyMap.h
+++ b/Engine/Plugins/Experimental/GeometryProcessing/Source/GeometricObjects/Public/Image/ImageOccupancyMap.h
@@ -95,48 +95,54 @@ public:
 		// find interior texels
 		TAtomic<int64> InteriorCounter;
 		FCriticalSection GutterLock;
-		ParallelFor(Dimensions.Num(), [&](int64 LinearIdx)
+		ParallelFor(Dimensions.GetHeight(), 
+			[this, &UVSpaceMesh, &GetTriangleIDFunc, &FlatSpatial, TexelDiag, &QueryOptions, &InteriorCounter, &GutterLock]
+		(int32 ImgY)
 		{
-			FVector2d UVPoint = Dimensions.GetTexelUV(LinearIdx);
-			FVector3d UVPoint3d(UVPoint.X, UVPoint.Y, 0);
-
-			double NearDistSqr;
-			int32 NearestTriID = FlatSpatial.FindNearestTriangle(UVPoint3d, NearDistSqr, QueryOptions);
-			if (NearestTriID >= 0)
+			for (int32 ImgX = 0; ImgX < Dimensions.GetWidth(); ++ImgX)
 			{
-				FVector3d A, B, C;
-				UVSpaceMesh.GetTriVertices(NearestTriID, A, B, C);
-				FTriangle2d UVTriangle(GetXY(A), GetXY(B), GetXY(C));
+				int64 LinearIdx = Dimensions.GetIndex(ImgX, ImgY);
+				FVector2d UVPoint = Dimensions.GetTexelUV(LinearIdx);
+				FVector3d UVPoint3d(UVPoint.X, UVPoint.Y, 0);
 
-				if (UVTriangle.IsInsideOrOn(UVPoint))
-				{
-					TexelType[LinearIdx] = InteriorTexel;
-					TexelQueryUV[LinearIdx] = (FVector2f)UVPoint;
-					TexelQueryTriangle[LinearIdx] = GetTriangleIDFunc(NearestTriID);
-					InteriorCounter.IncrementExchange();
-				}
-				else if (NearDistSqr < TexelDiag * TexelDiag)
-				{
-					FDistPoint3Triangle3d DistQuery = TMeshQueries<MeshType>::TriangleDistance(UVSpaceMesh, NearestTriID, UVPoint3d);
-					FVector2d NearestUV = GetXY(DistQuery.ClosestTrianglePoint);
-					// nudge point into triangle to improve numerical behavior of things like barycentric coord calculation
-					NearestUV += (10.0 * FMathf::ZeroTolerance) * Normalized(NearestUV - UVPoint);
-
-					TexelType[LinearIdx] = InteriorTexel;
-					TexelQueryUV[LinearIdx] = (FVector2f)NearestUV;
-					TexelQueryTriangle[LinearIdx] = GetTriangleIDFunc(NearestTriID);
-				}
-				else
+				double NearDistSqr;
+				int32 NearestTriID = FlatSpatial.FindNearestTriangle(UVPoint3d, NearDistSqr, QueryOptions);
+				if (NearestTriID >= 0)
 				{
-					TexelType[LinearIdx] = GutterTexel;
-					FDistPoint3Triangle3d DistQuery = TMeshQueries<MeshType>::TriangleDistance(UVSpaceMesh, NearestTriID, UVPoint3d);
-					FVector2d NearestUV = GetXY(DistQuery.ClosestTrianglePoint);
-					FVector2i NearestCoords = Dimensions.UVToCoords(NearestUV);
-					int64 NearestLinearIdx = Dimensions.GetIndex(NearestCoords);
-
-					GutterLock.Lock();
-					GutterTexels.Add(TTuple<int64, int64>(LinearIdx, NearestLinearIdx));
-					GutterLock.Unlock();
+					FVector3d A, B, C;
+					UVSpaceMesh.GetTriVertices(NearestTriID, A, B, C);
+					FTriangle2d UVTriangle(GetXY(A), GetXY(B), GetXY(C));
+
+					if (UVTriangle.IsInsideOrOn(UVPoint))
+					{
+						TexelType[LinearIdx] = InteriorTexel;
+						TexelQueryUV[LinearIdx] = (FVector2f)UVPoint;
+						TexelQueryTriangle[LinearIdx] = GetTriangleIDFunc(NearestTriID);
+						InteriorCounter.IncrementExchange();
+					}
+					else if (NearDistSqr < TexelDiag * TexelDiag)
+					{
+						FDistPoint3Triangle3d DistQuery = TMeshQueries<MeshType>::TriangleDistance(UVSpaceMesh, NearestTriID, UVPoint3d);
+						FVector2d NearestUV = GetXY(DistQuery.ClosestTrianglePoint);
+						// nudge point into triangle to improve numerical behavior of things like barycentric coord calculation
+						NearestUV += (10.0 * FMathf::ZeroTolerance) * Normalized(NearestUV - UVPoint);
+
+						TexelType[LinearIdx] = InteriorTexel;
+						TexelQueryUV[LinearIdx] = (FVector2f)NearestUV;
+						TexelQueryTriangle[LinearIdx] = GetTriangleIDFunc(NearestTriID);
+					}
+					else
+					{
+						TexelType[LinearIdx] = GutterTexel;
+						FDistPoint3Triangle3d DistQuery = TMeshQueries<MeshType>::TriangleDistance(UVSpaceMesh, NearestTriID, UVPoint3d);
+						FVector2d NearestUV = GetXY(DistQuery.ClosestTrianglePoint);
+						FVector2i NearestCoords = Dimensions.UVToCoords(NearestUV);
+						int64 NearestLinearIdx = Dimensions.GetIndex(NearestCoords);
+
+						GutterLock.Lock();
+						GutterTexels.Add(TTuple<int64, int64>(LinearIdx, NearestLinearIdx));
+						GutterLock.Unlock();
+					}
 				}
 			}
 		});
@@ -160,41 +166,47 @@ public:
 	{
 		int64 N = Dimensions.Num();
 		PassBuffer.SetNum(N);
-
-		ParallelFor(N, [&](int64 LinearIdx)
+		
+		ParallelFor(Dimensions.GetHeight(), 
+			[this, &BeginTexel, & AccumulateTexel, &CompleteTexel, &WeightFunction, FilterWidth, &PassBuffer]
+		(int32 ImgY)
 		{
-			if (TexelType[LinearIdx] != EmptyTexel)
+			for (int32 ImgX = 0; ImgX < Dimensions.GetWidth(); ++ImgX)
 			{
-				TexelValueType AccumValue = BeginTexel(LinearIdx);
-				float WeightSum = 0;
+				int64 LinearIdx = Dimensions.GetIndex(ImgX, ImgY);
+				if (TexelType[LinearIdx] != EmptyTexel)
+				{
+					TexelValueType AccumValue = BeginTexel(LinearIdx);
+					float WeightSum = 0;
 
-				FVector2i Coords = Dimensions.GetCoords(LinearIdx);
-				FVector2i MaxNbr = Coords + FVector2i(FilterWidth, FilterWidth);
-				Dimensions.Clamp(MaxNbr);
-				FVector2i MinNbr = Coords - FVector2i(FilterWidth, FilterWidth);
-				Dimensions.Clamp(MinNbr);
+					FVector2i Coords(ImgX, ImgY);
+					FVector2i MaxNbr = Coords + FVector2i(FilterWidth, FilterWidth);
+					Dimensions.Clamp(MaxNbr);
+					FVector2i MinNbr = Coords - FVector2i(FilterWidth, FilterWidth);
+					Dimensions.Clamp(MinNbr);
 
-				for (int32 Y = MinNbr.Y; Y <= MaxNbr.Y; ++Y)
-				{
-					for (int32 X = MinNbr.X; X <= MaxNbr.X; ++X)
+					for (int32 Y = MinNbr.Y; Y <= MaxNbr.Y; ++Y)
 					{
-						FVector2i NbrCoords(X, Y);
-						if (Dimensions.IsValidCoords(NbrCoords))
+						for (int32 X = MinNbr.X; X <= MaxNbr.X; ++X)
 						{
-							FVector2i Offset = NbrCoords - Coords;
-							int64 LinearNbrIndex = Dimensions.GetIndex(NbrCoords);
-							if (TexelType[LinearNbrIndex] != EmptyTexel)
+							FVector2i NbrCoords(X, Y);
+							if (Dimensions.IsValidCoords(NbrCoords))
 							{
-								float NbrWeight = WeightFunction(Offset);
-								AccumulateTexel(LinearNbrIndex, NbrWeight, AccumValue);
-								WeightSum += NbrWeight;
+								FVector2i Offset = NbrCoords - Coords;
+								int64 LinearNbrIndex = Dimensions.GetIndex(NbrCoords);
+								if (TexelType[LinearNbrIndex] != EmptyTexel)
+								{
+									float NbrWeight = WeightFunction(Offset);
+									AccumulateTexel(LinearNbrIndex, NbrWeight, AccumValue);
+									WeightSum += NbrWeight;
+								}
 							}
 						}
 					}
-				}
 
-				CompleteTexel(LinearIdx, WeightSum, AccumValue);
-				PassBuffer[LinearIdx] = AccumValue;
+					CompleteTexel(LinearIdx, WeightSum, AccumValue);
+					PassBuffer[LinearIdx] = AccumValue;
+				}
 			}
 		});
 
