Minor fixes for AO map baker: Clear uninitialized data from AO textures before writing them + only get tangents if needed and check that they are there

#rnx

diff --git a/Engine/Plugins/Experimental/GeometryProcessing/Source/DynamicMesh/Private/Sampling/MeshOcclusionMapBaker.cpp b/Engine/Plugins/Experimental/GeometryProcessing/Source/DynamicMesh/Private/Sampling/MeshOcclusionMapBaker.cpp
--- a/Engine/Plugins/Experimental/GeometryProcessing/Source/DynamicMesh/Private/Sampling/MeshOcclusionMapBaker.cpp
+++ b/Engine/Plugins/Experimental/GeometryProcessing/Source/DynamicMesh/Private/Sampling/MeshOcclusionMapBaker.cpp
@@ -87,10 +87,14 @@ void FMeshOcclusionMapBaker::Bake()
 			Normalize(DetailTriNormal);
 
 			FVector3d BaseTangentX, BaseTangentY;
-			BaseMeshTangents->GetInterpolatedTriangleTangent(
-				SampleData.BaseSample.TriangleIndex,
-				SampleData.BaseSample.BaryCoords,
-				BaseTangentX, BaseTangentY);
+			if (NormalSpace == ESpace::Tangent)
+			{
+				check(BaseMeshTangents);
+				BaseMeshTangents->GetInterpolatedTriangleTangent(
+					SampleData.BaseSample.TriangleIndex,
+					SampleData.BaseSample.BaryCoords,
+					BaseTangentX, BaseTangentY);
+			}
 
 			FVector3d DetailBaryCoords = SampleData.DetailBaryCoords;
 			FVector3d DetailPos = DetailMesh->GetTriBaryPoint(DetailTriID, DetailBaryCoords.X, DetailBaryCoords.Y, DetailBaryCoords.Z);
@@ -162,8 +166,10 @@ void FMeshOcclusionMapBaker::Bake()
 
 	OcclusionBuilder = MakeUnique<TImageBuilder<FVector3f>>();
 	OcclusionBuilder->SetDimensions(BakeCache->GetDimensions());
+	OcclusionBuilder->Clear(FVector3f::Zero());
 	NormalBuilder = MakeUnique<TImageBuilder<FVector3f>>();
 	NormalBuilder->SetDimensions(BakeCache->GetDimensions());
+	NormalBuilder->Clear(FVector3f::Zero());
 
 	BakeCache->EvaluateSamples([&](const FVector2i& Coords, const FMeshImageBakingCache::FCorrespondenceSample& Sample)
 	{
diff --git a/Engine/Plugins/Experimental/GeometryProcessing/Source/DynamicMesh/Public/Sampling/MeshOcclusionMapBaker.h b/Engine/Plugins/Experimental/GeometryProcessing/Source/DynamicMesh/Public/Sampling/MeshOcclusionMapBaker.h
--- a/Engine/Plugins/Experimental/GeometryProcessing/Source/DynamicMesh/Public/Sampling/MeshOcclusionMapBaker.h
+++ b/Engine/Plugins/Experimental/GeometryProcessing/Source/DynamicMesh/Public/Sampling/MeshOcclusionMapBaker.h
@@ -21,7 +21,7 @@ public:
 	// Inputs
 	//
 
-	const TMeshTangents<double>* BaseMeshTangents;
+	const TMeshTangents<double>* BaseMeshTangents = nullptr;
 
 	//
 	// Options
