change MeshImageBakingCache ParallelFor iteration to be over rows instead of over pixels, b/c ParallelFor doesn't support int64 indices

#rnx

diff --git a/Engine/Plugins/Experimental/GeometryProcessing/Source/DynamicMesh/Private/Sampling/MeshImageBakingCache.cpp b/Engine/Plugins/Experimental/GeometryProcessing/Source/DynamicMesh/Private/Sampling/MeshImageBakingCache.cpp
--- a/Engine/Plugins/Experimental/GeometryProcessing/Source/DynamicMesh/Private/Sampling/MeshImageBakingCache.cpp
+++ b/Engine/Plugins/Experimental/GeometryProcessing/Source/DynamicMesh/Private/Sampling/MeshImageBakingCache.cpp
@@ -255,25 +255,30 @@ bool FMeshImageBakingCache::ValidateCache()
 		SampleMap.Resize(Dimensions.GetWidth(), Dimensions.GetHeight());
 
 		// calculate interior texels
-		ParallelFor(Dimensions.Num(), [&](int64 LinearIdx)
+		ParallelFor(Dimensions.GetHeight(), [this, &DetailMeshSampler](int32 ImgY)
 		{
 			if (CancelF())
 			{
 				return;
 			}
 
-			if (OccupancyMap->IsInterior(LinearIdx) == false)
+			for (int32 ImgX = 0; ImgX < Dimensions.GetWidth(); ImgX++)
 			{
-				return;
-			}
+				int64 LinearIdx = Dimensions.GetIndex(ImgX, ImgY);
 
-			FVector2d UVPosition = (FVector2d)OccupancyMap->TexelQueryUV[LinearIdx];
-			int32 UVTriangleID = OccupancyMap->TexelQueryTriangle[LinearIdx];
+				if (OccupancyMap->IsInterior(LinearIdx) == false)
+				{
+					continue;
+				}
 
-			FCorrespondenceSample Sample;
-			DetailMeshSampler.SampleUV(UVTriangleID, UVPosition, Sample);
+				FVector2d UVPosition = (FVector2d)OccupancyMap->TexelQueryUV[LinearIdx];
+				int32 UVTriangleID = OccupancyMap->TexelQueryTriangle[LinearIdx];
 
-			SampleMap[LinearIdx] = Sample;
+				FCorrespondenceSample Sample;
+				DetailMeshSampler.SampleUV(UVTriangleID, UVPosition, Sample);
+
+				SampleMap[LinearIdx] = Sample;
+			}
 		});
 
 		bSamplesValid = true;
@@ -290,22 +295,26 @@ void FMeshImageBakingCache::EvaluateSamples(
 {
 	check(IsCacheValid());
 
-	ParallelFor(Dimensions.Num(), [&](int64 LinearIdx)
+	ParallelFor(Dimensions.GetHeight(), [this, &SampleFunction](int32 ImgY)
 	{
 		if (CancelF())
 		{
 			return;
 		}
 
-		if (OccupancyMap->IsInterior(LinearIdx) == false)
+		for (int32 ImgX = 0; ImgX < Dimensions.GetWidth(); ImgX++)
 		{
-			return;
-		}
-		FVector2i Coords = Dimensions.GetCoords(LinearIdx);
+			int64 LinearIdx = Dimensions.GetIndex(ImgX, ImgY);
+			if (OccupancyMap->IsInterior(LinearIdx) == false)
+			{
+				continue;
+			}
+			FVector2i Coords(ImgX, ImgY);
 
-		const FCorrespondenceSample& Sample = SampleMap[LinearIdx];
+			const FCorrespondenceSample& Sample = SampleMap[LinearIdx];
 
-		SampleFunction(Coords, Sample);
+			SampleFunction(Coords, Sample);
+		}
 
 	}, !bParallel ? EParallelForFlags::ForceSingleThread : EParallelForFlags::None);
 }
\ No newline at end of file
