GeometryProcessing: Add selective occlusion map type support in OcclusionMapBaker.

#rnx

diff --git a/Engine/Plugins/Experimental/GeometryProcessing/Source/DynamicMesh/Private/Sampling/MeshOcclusionMapBaker.cpp b/Engine/Plugins/Experimental/GeometryProcessing/Source/DynamicMesh/Private/Sampling/MeshOcclusionMapBaker.cpp
--- a/Engine/Plugins/Experimental/GeometryProcessing/Source/DynamicMesh/Private/Sampling/MeshOcclusionMapBaker.cpp
+++ b/Engine/Plugins/Experimental/GeometryProcessing/Source/DynamicMesh/Private/Sampling/MeshOcclusionMapBaker.cpp
@@ -11,6 +11,11 @@ using namespace UE::Geometry;
 
 void FMeshOcclusionMapBaker::Bake()
 {
+	if (OcclusionType == EOcclusionMapType::None)
+	{
+		return;
+	}
+
 	const FMeshImageBakingCache* BakeCache = GetCache();
 	check(BakeCache);
 	const FDynamicMesh3* DetailMesh = BakeCache->GetDetailMesh();
@@ -87,7 +92,7 @@ void FMeshOcclusionMapBaker::Bake()
 			Normalize(DetailTriNormal);
 
 			FVector3d BaseTangentX, BaseTangentY;
-			if (NormalSpace == ESpace::Tangent)
+			if (WantBentNormal() && NormalSpace == ESpace::Tangent)
 			{
 				check(BaseMeshTangents);
 				BaseMeshTangents->GetInterpolatedTriangleTangent(
@@ -116,46 +121,61 @@ void FMeshOcclusionMapBaker::Bake()
 				FRay3d OcclusionRay(DetailPos, SurfaceFrame.FromFrameVector(SphereDir));
 				ensure(OcclusionRay.Direction.Dot(DetailTriNormal) > 0);
 
-				// Have weight of point fall off as it becomes more coplanar with face. 
-				// This reduces faceting artifacts that we would otherwise see because geometry does not vary smoothly
-				double PointWeight = 1.0;
-				double BiasDot = OcclusionRay.Direction.Dot(DetailTriNormal);
-				if (BiasDot < BiasDotThreshold)
-				{
-					PointWeight = FMathd::Lerp(0.0, 1.0, FMathd::Clamp(BiasDot / BiasDotThreshold, 0.0, 1.0));
-					PointWeight *= PointWeight;
-				}
-				TotalPointWeight += PointWeight;
+				bool bHit = DetailSpatial->TestAnyHitTriangle(OcclusionRay, QueryOptions);
 
-				FVector3d BentNormal = OcclusionRay.Direction;
-				switch (NormalSpace)
-				{
-				case ESpace::Tangent:
+				if (WantAmbientOcclusion())
 				{
-					// compute normal in tangent space
-					double dx = BentNormal.Dot(BaseTangentX);
-					double dy = BentNormal.Dot(BaseTangentY);
-					double dz = BentNormal.Dot(SampleData.BaseNormal);
-					BentNormal = FVector3d(dx, dy, dz);;
-					break;
-				}
-				case ESpace::Object:
-					break;
+					// Have weight of point fall off as it becomes more coplanar with face. 
+					// This reduces faceting artifacts that we would otherwise see because geometry does not vary smoothly
+					double PointWeight = 1.0;
+					double BiasDot = OcclusionRay.Direction.Dot(DetailTriNormal);
+					if (BiasDot < BiasDotThreshold)
+					{
+						PointWeight = FMathd::Lerp(0.0, 1.0, FMathd::Clamp(BiasDot / BiasDotThreshold, 0.0, 1.0));
+						PointWeight *= PointWeight;
+					}
+					TotalPointWeight += PointWeight;
+
+					if (bHit)
+					{
+						AccumOcclusion += PointWeight;
+					}
 				}
-				TotalNormalWeight += 1.0;
 
-				if (DetailSpatial->TestAnyHitTriangle(OcclusionRay, QueryOptions))
+				if (WantBentNormal())
 				{
-					AccumOcclusion += PointWeight;
-				}
-				else
-				{
-					AccumNormal += BentNormal;
+					FVector3d BentNormal = OcclusionRay.Direction;
+					switch (NormalSpace)
+					{
+					case ESpace::Tangent:
+					{
+						// compute normal in tangent space
+						double dx = BentNormal.Dot(BaseTangentX);
+						double dy = BentNormal.Dot(BaseTangentY);
+						double dz = BentNormal.Dot(SampleData.BaseNormal);
+						BentNormal = FVector3d(dx, dy, dz);;
+						break;
+					}
+					case ESpace::Object:
+						break;
+					}
+					TotalNormalWeight += 1.0;
+
+					if (!bHit)
+					{
+						AccumNormal += BentNormal;
+					}
 				}
 			}
 
-			AccumOcclusion = (TotalPointWeight > 0.0001) ? (AccumOcclusion / TotalPointWeight) : 0.0;
-			AccumNormal = (TotalNormalWeight > 0.0 && AccumNormal.Length() > 0.0) ? Normalized(AccumNormal / TotalNormalWeight) : DefaultNormal;
+			if (WantAmbientOcclusion())
+			{
+				AccumOcclusion = (TotalPointWeight > 0.0001) ? (AccumOcclusion / TotalPointWeight) : 0.0;
+			}
+			if (WantBentNormal())
+			{
+				AccumNormal = (TotalNormalWeight > 0.0 && AccumNormal.Length() > 0.0) ? Normalized(AccumNormal / TotalNormalWeight) : DefaultNormal;
+			}
 			OutOcclusion = AccumOcclusion;
 			OutNormal = AccumNormal;
 			return;
@@ -164,33 +184,51 @@ void FMeshOcclusionMapBaker::Bake()
 		OutNormal = DefaultNormal;
 	};
 
-	OcclusionBuilder = MakeUnique<TImageBuilder<FVector3f>>();
-	OcclusionBuilder->SetDimensions(BakeCache->GetDimensions());
-	OcclusionBuilder->Clear(FVector3f::Zero());
-	NormalBuilder = MakeUnique<TImageBuilder<FVector3f>>();
-	NormalBuilder->SetDimensions(BakeCache->GetDimensions());
-	NormalBuilder->Clear(FVector3f::Zero());
+	if (WantAmbientOcclusion())
+	{
+		OcclusionBuilder = MakeUnique<TImageBuilder<FVector3f>>();
+		OcclusionBuilder->SetDimensions(BakeCache->GetDimensions());
+		OcclusionBuilder->Clear(FVector3f::Zero());
+	}
+	if (WantBentNormal())
+	{
+		NormalBuilder = MakeUnique<TImageBuilder<FVector3f>>();
+		NormalBuilder->SetDimensions(BakeCache->GetDimensions());
+		NormalBuilder->Clear(FVector3f::Zero());
+	}
 
 	BakeCache->EvaluateSamples([&](const FVector2i& Coords, const FMeshImageBakingCache::FCorrespondenceSample& Sample)
 	{
 		double Occlusion = 0.0;
 		FVector3d BentNormal;
 		OcclusionSampleFunction(Sample, Occlusion, BentNormal);
-		FVector3f OcclusionColor = FMathd::Clamp(1.0f - (float)Occlusion, 0.0f, 1.0f) * FVector3f::One();
-		OcclusionBuilder->SetPixel(Coords, OcclusionColor);
-		FVector3f NormalColor = (FVector3f(BentNormal.X, BentNormal.Y, BentNormal.Z) + FVector3f::One()) * 0.5;
-		NormalBuilder->SetPixel(Coords, NormalColor);
+		if (WantAmbientOcclusion())
+		{
+			FVector3f OcclusionColor = FMathd::Clamp(1.0f - (float)Occlusion, 0.0f, 1.0f) * FVector3f::One();
+			OcclusionBuilder->SetPixel(Coords, OcclusionColor);
+		}
+		if (WantBentNormal())
+		{
+			FVector3f NormalColor = (FVector3f(BentNormal.X, BentNormal.Y, BentNormal.Z) + FVector3f::One()) * 0.5;
+			NormalBuilder->SetPixel(Coords, NormalColor);
+		}
 	});
 
 	const FImageOccupancyMap& Occupancy = *BakeCache->GetOccupancyMap();
 	for (int64 k = 0; k < Occupancy.GutterTexels.Num(); k++)
 	{
 		TPair<int64, int64> GutterTexel = Occupancy.GutterTexels[k];
-		OcclusionBuilder->CopyPixel(GutterTexel.Value, GutterTexel.Key);
-		NormalBuilder->CopyPixel(GutterTexel.Value, GutterTexel.Key);
+		if (WantAmbientOcclusion())
+		{
+			OcclusionBuilder->CopyPixel(GutterTexel.Value, GutterTexel.Key);
+		}
+		if (WantBentNormal())
+		{
+			NormalBuilder->CopyPixel(GutterTexel.Value, GutterTexel.Key);
+		}
 	}
 
-	if (BlurRadius > 0.01)
+	if (WantAmbientOcclusion() && BlurRadius > 0.01)
 	{
 		TDiscreteKernel2f BlurKernel2d;
 		TGaussian2f::MakeKernelFromRadius(BlurRadius, BlurKernel2d);
diff --git a/Engine/Plugins/Experimental/GeometryProcessing/Source/DynamicMesh/Public/Sampling/MeshOcclusionMapBaker.h b/Engine/Plugins/Experimental/GeometryProcessing/Source/DynamicMesh/Public/Sampling/MeshOcclusionMapBaker.h
--- a/Engine/Plugins/Experimental/GeometryProcessing/Source/DynamicMesh/Public/Sampling/MeshOcclusionMapBaker.h
+++ b/Engine/Plugins/Experimental/GeometryProcessing/Source/DynamicMesh/Public/Sampling/MeshOcclusionMapBaker.h
@@ -2,6 +2,7 @@
 
 #pragma once
 
+#include "Misc/EnumClassFlags.h"
 #include "Sampling/MeshImageBaker.h"
 #include "Sampling/MeshImageBakingCache.h"
 #include "Image/ImageBuilder.h"
@@ -12,6 +13,15 @@ namespace UE
 namespace Geometry
 {
 
+enum class EOcclusionMapType : uint8
+{
+	None = 0,
+	AmbientOcclusion = 1,
+	BentNormal = 2,
+	All = 3
+};
+ENUM_CLASS_FLAGS(EOcclusionMapType);
+
 class DYNAMICMESH_API FMeshOcclusionMapBaker : public FMeshImageBaker
 {
 public:
@@ -39,6 +49,7 @@ public:
 		Object
 	};
 
+	EOcclusionMapType OcclusionType = EOcclusionMapType::All;
 	int32 NumOcclusionRays = 32;
 	double MaxDistance = TNumericLimits<double>::Max();
 	double SpreadAngle = 180.0;
@@ -96,6 +107,20 @@ public:
 		}
 	}
 
+	//
+	// Utility
+	//
+
+	inline bool WantAmbientOcclusion() const
+	{
+		return ((OcclusionType & EOcclusionMapType::AmbientOcclusion) == EOcclusionMapType::AmbientOcclusion);
+	}
+
+	inline bool WantBentNormal() const
+	{
+		return ((OcclusionType & EOcclusionMapType::BentNormal) == EOcclusionMapType::BentNormal);
+	}
+
 protected:
 	TUniquePtr<TImageBuilder<FVector3f>> OcclusionBuilder;
 	TUniquePtr<TImageBuilder<FVector3f>> NormalBuilder;
diff --git a/Engine/Plugins/Experimental/MeshModelingToolset/Source/MeshModelingTools/Private/BakeMeshAttributeMapsTool.cpp b/Engine/Plugins/Experimental/MeshModelingToolset/Source/MeshModelingTools/Private/BakeMeshAttributeMapsTool.cpp
--- a/Engine/Plugins/Experimental/MeshModelingToolset/Source/MeshModelingTools/Private/BakeMeshAttributeMapsTool.cpp
+++ b/Engine/Plugins/Experimental/MeshModelingToolset/Source/MeshModelingTools/Private/BakeMeshAttributeMapsTool.cpp
@@ -174,6 +174,7 @@ public:
 
 		TUniquePtr<FMeshOcclusionMapBaker> Baker = MakeUnique<FMeshOcclusionMapBaker>();
 		Baker->SetCache(BakeCache.Get());
+		Baker->OcclusionType = EOcclusionMapType::All;
         Baker->BaseMeshTangents = BaseMeshTangents;
 		Baker->NumOcclusionRays = Settings.OcclusionRays;
         Baker->MaxDistance = Settings.MaxDistance;
@@ -1194,17 +1195,25 @@ void UBakeMeshAttributeMapsTool::OnMapsUpdated(const TUniquePtr<FMeshImageBaker>
 	case EBakeMapType::Occlusion:
 	{
 		FMeshOcclusionMapBaker* Baker = static_cast<FMeshOcclusionMapBaker*>(NewResult.Get());
-		FTexture2DBuilder TextureBuilder;
-		TextureBuilder.Initialize(FTexture2DBuilder::ETextureType::AmbientOcclusion, CachedOcclusionMapSettings.Dimensions);
-		TextureBuilder.Copy(*Baker->GetResult(FMeshOcclusionMapBaker::EResult::AmbientOcclusion));
-		TextureBuilder.Commit(false);
-		CachedOcclusionMap = TextureBuilder.GetTexture2D();
+		const TUniquePtr<TImageBuilder<FVector3f>>& AOResult = Baker->GetResult(FMeshOcclusionMapBaker::EResult::AmbientOcclusion);
+		if (ensure(AOResult))
+		{
+			FTexture2DBuilder TextureBuilder;
+			TextureBuilder.Initialize(FTexture2DBuilder::ETextureType::AmbientOcclusion, CachedOcclusionMapSettings.Dimensions);
+			TextureBuilder.Copy(*AOResult);
+			TextureBuilder.Commit(false);
+			CachedOcclusionMap = TextureBuilder.GetTexture2D();
+		}
 
-		FTexture2DBuilder TextureNormalBuilder;
-		TextureNormalBuilder.Initialize(FTexture2DBuilder::ETextureType::NormalMap, CachedOcclusionMapSettings.Dimensions);
-		TextureNormalBuilder.Copy(*Baker->GetResult(FMeshOcclusionMapBaker::EResult::BentNormal));
-		TextureNormalBuilder.Commit(false);
-		CachedBentNormalMap = TextureNormalBuilder.GetTexture2D();
+		const TUniquePtr<TImageBuilder<FVector3f>>& BNResult = Baker->GetResult(FMeshOcclusionMapBaker::EResult::BentNormal);
+		if (ensure(BNResult))
+		{
+			FTexture2DBuilder TextureNormalBuilder;
+			TextureNormalBuilder.Initialize(FTexture2DBuilder::ETextureType::NormalMap, CachedOcclusionMapSettings.Dimensions);
+			TextureNormalBuilder.Copy(*BNResult);
+			TextureNormalBuilder.Commit(false);
+			CachedBentNormalMap = TextureNormalBuilder.GetTexture2D();
+		}
 		break;
 	}
 	case EBakeMapType::Curvature:
diff --git a/Engine/Plugins/Experimental/PlanarCutPlugin/Source/PlanarCut/Private/FractureAutoUV.cpp b/Engine/Plugins/Experimental/PlanarCutPlugin/Source/PlanarCut/Private/FractureAutoUV.cpp
--- a/Engine/Plugins/Experimental/PlanarCutPlugin/Source/PlanarCut/Private/FractureAutoUV.cpp
+++ b/Engine/Plugins/Experimental/PlanarCutPlugin/Source/PlanarCut/Private/FractureAutoUV.cpp
@@ -572,9 +572,7 @@ void TextureInternalSurfaces(
 			FMeshOcclusionMapBaker OcclusionBaker;
 			OcclusionBaker.SetCache(&BakeCache);
 			
-			// set NormalSpace to Object to avoid need to compute tangents for bent normals computation by the AO baker, which we don't need here.
-			// TODO: consider adding a toggle to fully disable bent normals computation in FMeshOcclusionMapBaker
-			OcclusionBaker.NormalSpace = FMeshOcclusionMapBaker::ESpace::Object;
+			OcclusionBaker.OcclusionType = EOcclusionMapType::AmbientOcclusion;
 			OcclusionBaker.NumOcclusionRays = AttributeSettings.AO_Rays;
 			OcclusionBaker.MaxDistance = AttributeSettings.AO_MaxDistance == 0 ? TNumericLimits<double>::Max() : AttributeSettings.AO_MaxDistance;
 			OcclusionBaker.BiasAngleDeg = AttributeSettings.AO_BiasAngleDeg;
