fix bug where unset UVs on overlay leads to UV seams everywhere:
 - on dmesh->meshd, don't give every unset UV a new ID
 - on meshd->dmesh, when splitting over-shared UVs, make sure that UVs from the same source and same target parent still map to the same UV element in the result

#rnx

diff --git a/Engine/Plugins/Experimental/GeometryProcessing/Source/MeshConversion/Private/DynamicMeshToMeshDescription.cpp b/Engine/Plugins/Experimental/GeometryProcessing/Source/MeshConversion/Private/DynamicMeshToMeshDescription.cpp
--- a/Engine/Plugins/Experimental/GeometryProcessing/Source/MeshConversion/Private/DynamicMeshToMeshDescription.cpp
+++ b/Engine/Plugins/Experimental/GeometryProcessing/Source/MeshConversion/Private/DynamicMeshToMeshDescription.cpp
@@ -627,6 +627,7 @@ void FDynamicMeshToMeshDescription::Convert_NoSharedInstances(const FDynamicMesh
 	Builder.SetNumUVLayers(NumUVLayers);
 
 	TArray<TArray<FUVID>> MapUVArray; MapUVArray.SetNum(NumUVLayers);
+	TArray<FUVID> UnsetUVIDs; UnsetUVIDs.Init(INDEX_NONE, NumUVLayers); // UVIDs for unset UVs to share, if needed
 	for (int32 k = 0; k < NumUVLayers; ++k)
 	{
 		const FDynamicMeshUVOverlay* UVOverlay = UVLayers[k];
@@ -768,12 +769,16 @@ void FDynamicMeshToMeshDescription::Convert_NoSharedInstances(const FDynamicMesh
 		{ 
 			FUVID UVIDs[3] = {FUVID(-1), FUVID(-1), FUVID(-1)};
 
-			// add zero uvs for unset triangles.  this mimics the old version of this code.
+			// add zero uvs for unset triangles.
 			if (!UVLayers[k]->IsSetTriangle(TriID))
 			{
+				if (UnsetUVIDs[k].GetValue() == INDEX_NONE)
+				{
+					UnsetUVIDs[k] = Builder.AppendUV(FVector2D::ZeroVector, k);
+				}
 				for (int32 j = 0; j < 3; ++j)
-				{ 
-					UVIDs[j] = Builder.AppendUV(FVector2D::ZeroVector, k);
+				{
+					UVIDs[j] = UnsetUVIDs[k];
 				}
 			}
 			else
diff --git a/Engine/Plugins/Experimental/GeometryProcessing/Source/MeshConversion/Private/MeshDescriptionToDynamicMesh.cpp b/Engine/Plugins/Experimental/GeometryProcessing/Source/MeshConversion/Private/MeshDescriptionToDynamicMesh.cpp
--- a/Engine/Plugins/Experimental/GeometryProcessing/Source/MeshConversion/Private/MeshDescriptionToDynamicMesh.cpp
+++ b/Engine/Plugins/Experimental/GeometryProcessing/Source/MeshConversion/Private/MeshDescriptionToDynamicMesh.cpp
@@ -492,6 +492,7 @@ void FMeshDescriptionToDynamicMesh::Convert(const FMeshDescription* MeshIn, FDyn
 					// copy uv "index buffer"
 					FUVID InvalidUVID(INDEX_NONE);
 					TArray<int32> TrisMissingUVs; 
+					TMap<FIndex2i, int> SplitUVMapping; // mapping from (ParentVID, OrigElID) -> SplitElID, for UV ElIDs that need to be split to multiple vertices
 					for (int32 TriID : MeshOut.TriangleIndicesItr())
 					{
 						FTriangleID TriangleID = TriIDMap[TriID];
@@ -514,6 +515,24 @@ void FMeshDescriptionToDynamicMesh::Convert(const FMeshDescription* MeshIn, FDyn
 						///--  We have to do some clean-up on the shared UVs that come from MeshDecription --///
 						// This clean up should go away if MeshDescription can solve these problems during import from the source fbx files
 						{
+							// Helper to split UVs out on a triangle while keeping UVs with the same source ID and parent vertex the same element in the overlay
+							auto SplitVertexUV = [&SplitUVMapping, &UVCoordinates, &UVIndices, &UVOverlay](int ParentVID, FIndex3i& TriUV, int Idx)
+							{
+								FIndex2i SplitUVID(ParentVID, TriUV[Idx]);
+								int* FoundUVEl = SplitUVMapping.Find(SplitUVID);
+								if (FoundUVEl)
+								{
+									TriUV[Idx] = *FoundUVEl;
+								}
+								else
+								{
+									const FVector2D UVvalue = UVCoordinates[UVIndices[Idx]];
+									int NewUVEl = UVOverlay->AppendElement(FVector2f(UVvalue));
+									TriUV[Idx] = NewUVEl;
+									SplitUVMapping.Add(SplitUVID, NewUVEl);
+								}
+							};
+
 							// MeshDescription can attach multiple vertices to the same UV element.  DynamicMesh does not.
 							// if we have already used this element for a different mesh vertex, split it.
 							const FIndex3i ParentTriangle = MeshOut.GetTriangle(TriID);
@@ -522,42 +541,34 @@ void FMeshDescriptionToDynamicMesh::Convert(const FMeshDescription* MeshIn, FDyn
 								int32 ParentVID = UVOverlay->GetParentVertex(TriUV[i]);
 								if (ParentVID != FDynamicMesh3::InvalidID && ParentVID != ParentTriangle[i])
 								{
-									const FVector2D UVvalue = UVCoordinates[UVIndices[i]];
-									TriUV[i] = UVOverlay->AppendElement(FVector2f(UVvalue));
+									SplitVertexUV(ParentTriangle[i], TriUV, i);
 								}
 							}
 
 							// MeshDescription allows for degenerate UV tris.  Dynamic Mesh does not.
-							// if the UV tri is degenerate we split the degenerate UV edge by adding two new UVs
-							// in its place, or if it is totally degenerate we add 3 new UVs
+							// if the UV tri is degenerate we split the degenerate UV edge by adding a new UV
+							// in its place, or if it is totally degenerate we add 2 new UVs
 
 							if (TriUV[0] == TriUV[1] && TriUV[0] == TriUV[2])
 							{
 								const FVector2D UVvalue = UVCoordinates[UVIndices[1]];
-								TriUV[0] = UVOverlay->AppendElement(FVector2f(UVvalue));
-								TriUV[1] = UVOverlay->AppendElement(FVector2f(UVvalue));
-								TriUV[2] = UVOverlay->AppendElement(FVector2f(UVvalue));
-
+								SplitVertexUV(ParentTriangle[1], TriUV, 1);
+								SplitVertexUV(ParentTriangle[2], TriUV, 2);
 							}
 							else
 							{
+								int SplitWhich = -1;
 								if (TriUV[0] == TriUV[1])
 								{
-									const FVector2D UVvalue = UVCoordinates[UVIndices[0]];
-									TriUV[0] = UVOverlay->AppendElement(FVector2f(UVvalue));
-									TriUV[1] = UVOverlay->AppendElement(FVector2f(UVvalue));
+									SplitWhich = 1;
 								}
-								if (TriUV[0] == TriUV[2])
+								else if (TriUV[0] == TriUV[2] || TriUV[1] == TriUV[2])
 								{
-									const FVector2D UVvalue = UVCoordinates[UVIndices[0]];
-									TriUV[0] = UVOverlay->AppendElement(FVector2f(UVvalue));
-									TriUV[2] = UVOverlay->AppendElement(FVector2f(UVvalue));
+									SplitWhich = 2;
 								}
-								if (TriUV[1] == TriUV[2])
+								if (SplitWhich > -1)
 								{
-									const FVector2D UVvalue = UVCoordinates[UVIndices[1]];
-									TriUV[1] = UVOverlay->AppendElement(FVector2f(UVvalue));
-									TriUV[2] = UVOverlay->AppendElement(FVector2f(UVvalue));
+									SplitVertexUV(ParentTriangle[SplitWhich], TriUV, SplitWhich);
 								}
 							}
 						}
