Geometry Processing: Add Color Overlay  as FVector4f to the default Dynamic Mesh Attribute Set.  This should be used in place the mesh vertex color as it  mirrors the wedge-instance granularity of MeshDescription used in static meshes, and the SkeletalMeshLODModel used in skeletal meshes.   Also fix up the mesh conversion code to support this, as well as the undo system, and previous clients of the mesh vertex color.

diff --git a/Engine/Plugins/Experimental/GeometryProcessing/Source/DynamicMesh/Private/DynamicMeshAttributeSet.cpp b/Engine/Plugins/Experimental/GeometryProcessing/Source/DynamicMesh/Private/DynamicMeshAttributeSet.cpp
--- a/Engine/Plugins/Experimental/GeometryProcessing/Source/DynamicMesh/Private/DynamicMeshAttributeSet.cpp
+++ b/Engine/Plugins/Experimental/GeometryProcessing/Source/DynamicMesh/Private/DynamicMeshAttributeSet.cpp
@@ -2,7 +2,7 @@
 
 #include "DynamicMeshAttributeSet.h"
 #include "IndexTypes.h"
-#include "Async/ParallelFor.h"
+#include "Async/Async.h"
 
 #include "ExplicitUseGeometryMathTypes.h"		// using UE::Geometry::(math types)
 using namespace UE::Geometry;
@@ -20,7 +20,11 @@ void FDynamicMeshAttributeSet::Copy(const FDynamicMeshAttributeSet& Copy)
 	{
 		NormalLayers[NormalLayerIndex].Copy(Copy.NormalLayers[NormalLayerIndex]);
 	}
-
+	if (Copy.ColorLayer)
+	{
+		EnablePrimaryColors();
+		ColorLayer->Copy(*(Copy.ColorLayer));
+	}
 	if (Copy.MaterialIDAttrib)
 	{
 		EnableMaterialID();
@@ -64,6 +68,13 @@ bool FDynamicMeshAttributeSet::IsCompact()
 			return false;
 		}
 	}
+	if (HasPrimaryColors())
+	{
+		if (!ColorLayer->IsCompact())
+		{
+			return false;
+		}
+	}
 	
 	// material ID and generic per-element attributes currently cannot be non-compact
 	return true;
@@ -83,6 +94,11 @@ void FDynamicMeshAttributeSet::CompactCopy(const FCompactMaps& CompactMaps, cons
 	{
 		NormalLayers[NormalLayerIndex].CompactCopy(CompactMaps, Copy.NormalLayers[NormalLayerIndex]);
 	}
+	if (Copy.ColorLayer)
+	{
+		EnablePrimaryColors();
+		ColorLayer->CompactCopy(CompactMaps, *(Copy.ColorLayer));
+	}
 
 	if (Copy.MaterialIDAttrib)
 	{
@@ -123,7 +139,10 @@ void FDynamicMeshAttributeSet::CompactInPlace(const FCompactMaps& CompactMaps)
 	{
 		NormalLayers[NormalLayerIndex].CompactInPlace(CompactMaps);
 	}
-
+	if (ColorLayer.IsValid())
+	{
+		ColorLayer->CompactInPlace(CompactMaps);
+	}
 	if (MaterialIDAttrib.IsValid())
 	{
 		MaterialIDAttrib->CompactInPlace(CompactMaps);
@@ -146,18 +165,45 @@ void FDynamicMeshAttributeSet::SplitAllBowties(bool bParallel)
 {
 	int32 UVLayerCount = NumUVLayers();
 	int32 NormalLayerCount = NumNormalLayers();
-
-	ParallelFor(UVLayerCount + NormalLayerCount, [&](int32 idx)
-	{
-		if (idx < NormalLayerCount)
-		{
-			GetNormalLayer(idx)->SplitBowties();
+	
+	TArray<TFuture<void>> Pending;
+	auto ASyncOrRunSplit = [&Pending, bParallel](auto Overlay)->void
+	{
+		if (bParallel)
+		{	
+			auto AsyncTask = Async(EAsyncExecution::ThreadPool, [Overlay]() 
+			{
+				Overlay->SplitBowties();
+			});
+			Pending.Add(MoveTemp(AsyncTask));
 		}
 		else
 		{
-			GetUVLayer(idx - NormalLayerCount)->SplitBowties();
+			Overlay->SplitBowties();
 		}
-	}, (bParallel) ? EParallelForFlags::Unbalanced : EParallelForFlags::ForceSingleThread );
+	};
+
+	for (int32 i = 0; i < UVLayerCount; ++i)
+	{	
+		FDynamicMeshUVOverlay* UVLayer = GetUVLayer(i);
+		ASyncOrRunSplit(UVLayer);
+	}
+	for (int32 i = 0; i < NormalLayerCount; ++i)
+	{
+		FDynamicMeshNormalOverlay* NormalLayer = GetNormalLayer(i);
+		ASyncOrRunSplit(NormalLayer);
+	}
+	if (HasPrimaryColors())
+	{
+		FDynamicMeshColorOverlay* Colors = PrimaryColors();
+		ASyncOrRunSplit(Colors);
+	}
+
+	// this array will be empty if bParallel == false
+	for (TFuture<void>& Future : Pending)
+	{
+		Future.Wait();
+	}
 }
 
 
@@ -174,6 +220,14 @@ void FDynamicMeshAttributeSet::EnableMatchingAttributes(const FDynamicMeshAttrib
 	{
 		NormalLayers[NormalLayerIndex].ClearElements();
 	}
+	if (ToMatch.ColorLayer)
+	{
+		EnablePrimaryColors();
+	}
+	else
+	{
+		DisablePrimaryColors();
+	}
 
 	if (ToMatch.MaterialIDAttrib)
 	{
@@ -212,6 +266,10 @@ void FDynamicMeshAttributeSet::Reparent(FDynamicMesh3* NewParent)
 	{
 		NormalLayers[NormalLayerIndex].Reparent(NewParent);
 	}
+	if (ColorLayer)
+	{
+		ColorLayer->Reparent(NewParent);
+	}
 
 	if (MaterialIDAttrib)
 	{
@@ -284,7 +342,18 @@ void FDynamicMeshAttributeSet::SetNumNormalLayers(int Num)
 	check(NormalLayers.Num() == Num);
 }
 
+void FDynamicMeshAttributeSet::EnablePrimaryColors()
+{
+	if (HasPrimaryColors() == false)
+	{ 
+		ColorLayer = MakeUnique<FDynamicMeshColorOverlay>(ParentMesh);
+	}
+}
 
+void FDynamicMeshAttributeSet::DisablePrimaryColors()
+{
+	ColorLayer.Reset();
+}
 
 int32 FDynamicMeshAttributeSet::NumPolygroupLayers() const
 {
@@ -358,6 +427,11 @@ bool FDynamicMeshAttributeSet::IsSeamEdge(int eid) const
 		}
 	}
 
+	if (ColorLayer && ColorLayer->IsSeamEdge(eid))
+	{
+		return true;
+	}
+	
 	return false;
 }
 
@@ -379,10 +453,14 @@ bool FDynamicMeshAttributeSet::IsSeamEndEdge(int eid) const
 		}
 	}
 
+	if (ColorLayer && ColorLayer->IsSeamEndEdge(eid))
+	{
+		return true;
+	}
 	return false;
 }
 
-bool FDynamicMeshAttributeSet::IsSeamEdge(int EdgeID, bool& bIsUVSeamOut, bool& bIsNormalSeamOut) const
+bool FDynamicMeshAttributeSet::IsSeamEdge(int EdgeID, bool& bIsUVSeamOut, bool& bIsNormalSeamOut, bool& bIsColorSeamOut) const
 {
 	bIsUVSeamOut = false;
 	for (const FDynamicMeshUVOverlay& UVLayer : UVLayers)
@@ -402,7 +480,12 @@ bool FDynamicMeshAttributeSet::IsSeamEdge(int EdgeID, bool& bIsUVSeamOut, bool&
 		}
 	}
 
-	return (bIsUVSeamOut || bIsNormalSeamOut);
+	bIsColorSeamOut = false;
+	if (ColorLayer && ColorLayer->IsSeamEdge(EdgeID))
+	{
+		bIsColorSeamOut = true;
+	}
+	return (bIsUVSeamOut || bIsNormalSeamOut || bIsColorSeamOut);
 }
 
 
@@ -422,6 +505,10 @@ bool FDynamicMeshAttributeSet::IsSeamVertex(int VID, bool bBoundaryIsSeam) const
 			return true;
 		}
 	}
+	if (ColorLayer && ColorLayer->IsSeamVertex(VID, bBoundaryIsSeam))
+	{
+		return true;
+	}
 	return false;
 }
 
@@ -468,7 +555,10 @@ void FDynamicMeshAttributeSet::OnNewTriangle(int TriangleID, bool bInserted)
 	{
 		NormalLayer.InitializeNewTriangle(TriangleID);
 	}
-
+	if (ColorLayer)
+	{
+		ColorLayer->InitializeNewTriangle(TriangleID);
+	}
 	if (MaterialIDAttrib)
 	{
 		int NewValue = 0;
@@ -495,6 +585,10 @@ void FDynamicMeshAttributeSet::OnRemoveTriangle(int TriangleID)
 	{
 		NormalLayer.OnRemoveTriangle(TriangleID);
 	}
+	if (ColorLayer)
+	{
+		ColorLayer->OnRemoveTriangle(TriangleID);
+	}
 
 	// has no effect on MaterialIDAttrib
 }
@@ -511,7 +605,10 @@ void FDynamicMeshAttributeSet::OnReverseTriOrientation(int TriangleID)
 	{
 		NormalLayer.OnReverseTriOrientation(TriangleID);
 	}
-
+	if (ColorLayer)
+	{
+		ColorLayer->OnReverseTriOrientation(TriangleID);
+	}
 	// has no effect on MaterialIDAttrib
 }
 
@@ -527,7 +624,10 @@ void FDynamicMeshAttributeSet::OnSplitEdge(const FDynamicMesh3::FEdgeSplitInfo&
 	{
 		NormalLayer.OnSplitEdge(SplitInfo);
 	}
-
+	if (ColorLayer)
+	{
+		ColorLayer->OnSplitEdge(SplitInfo);
+	}
 	if (MaterialIDAttrib)
 	{
 		MaterialIDAttrib->OnSplitEdge(SplitInfo);
@@ -551,7 +651,10 @@ void FDynamicMeshAttributeSet::OnFlipEdge(const FDynamicMesh3::FEdgeFlipInfo & f
 	{
 		NormalLayer.OnFlipEdge(flipInfo);
 	}
-
+	if (ColorLayer)
+	{
+		ColorLayer->OnFlipEdge(flipInfo);
+	}
 	if (MaterialIDAttrib)
 	{
 		MaterialIDAttrib->OnFlipEdge(flipInfo);
@@ -576,7 +679,10 @@ void FDynamicMeshAttributeSet::OnCollapseEdge(const FDynamicMesh3::FEdgeCollapse
 	{
 		NormalLayer.OnCollapseEdge(collapseInfo);
 	}
-
+	if (ColorLayer)
+	{
+		ColorLayer->OnCollapseEdge(collapseInfo);
+	}
 	if (MaterialIDAttrib)
 	{
 		MaterialIDAttrib->OnCollapseEdge(collapseInfo);
@@ -600,7 +706,10 @@ void FDynamicMeshAttributeSet::OnPokeTriangle(const FDynamicMesh3::FPokeTriangle
 	{
 		NormalLayer.OnPokeTriangle(pokeInfo);
 	}
-
+	if (ColorLayer)
+	{
+		ColorLayer->OnPokeTriangle(pokeInfo);
+	}
 	if (MaterialIDAttrib)
 	{
 		MaterialIDAttrib->OnPokeTriangle(pokeInfo);
@@ -624,7 +733,10 @@ void FDynamicMeshAttributeSet::OnMergeEdges(const FDynamicMesh3::FMergeEdgesInfo
 	{
 		NormalLayer.OnMergeEdges(mergeInfo);
 	}
-
+	if (ColorLayer)
+	{
+		ColorLayer->OnMergeEdges(mergeInfo);
+	}
 	if (MaterialIDAttrib)
 	{
 		MaterialIDAttrib->OnMergeEdges(mergeInfo);
@@ -648,7 +760,10 @@ void FDynamicMeshAttributeSet::OnSplitVertex(const DynamicMeshInfo::FVertexSplit
 	{
 		NormalLayer.OnSplitVertex(SplitInfo, TrianglesToUpdate);
 	}
-
+	if (ColorLayer)
+	{
+		ColorLayer->OnSplitVertex(SplitInfo, TrianglesToUpdate);
+	}
 	if (MaterialIDAttrib)
 	{
 		MaterialIDAttrib->OnSplitVertex(SplitInfo, TrianglesToUpdate);
diff --git a/Engine/Plugins/Experimental/GeometryProcessing/Source/DynamicMesh/Private/DynamicMeshChangeTracker.cpp b/Engine/Plugins/Experimental/GeometryProcessing/Source/DynamicMesh/Private/DynamicMeshChangeTracker.cpp
--- a/Engine/Plugins/Experimental/GeometryProcessing/Source/DynamicMesh/Private/DynamicMeshChangeTracker.cpp
+++ b/Engine/Plugins/Experimental/GeometryProcessing/Source/DynamicMesh/Private/DynamicMeshChangeTracker.cpp
@@ -154,6 +154,8 @@ template class DYNAMICMESH_API TDynamicMeshAttributeChange<int, 2>;
 template class DYNAMICMESH_API TDynamicMeshAttributeChange<float, 3>;
 template class DYNAMICMESH_API TDynamicMeshAttributeChange<double, 3>;
 template class DYNAMICMESH_API TDynamicMeshAttributeChange<int, 3>;
+template class DYNAMICMESH_API TDynamicMeshAttributeChange<float, 4>;
+template class DYNAMICMESH_API TDynamicMeshAttributeChange<double, 4>;
 
 } // end namespace UE::Geometry
 } // end namespace UE
@@ -399,6 +401,20 @@ void FDynamicMeshAttributeSetChangeTracker::BeginChange()
 		}
 	}
 
+	// initialize Color layer state tracking if needed
+	if (Attribs->HasPrimaryColors())
+	{
+		Change->ColorChange.Emplace();
+		const FDynamicMeshColorOverlay* ColorLayer = Attribs->PrimaryColors();
+		ColorState.MaxElementID = ColorLayer->MaxElementID();
+		ColorState.StartElements.Init(false, ColorState.MaxElementID);
+		ColorState.ChangedElements.Init(false,  ColorState.MaxElementID);
+		for ( int ElemID : ColorLayer->ElementIndicesItr())
+		{
+			ColorState.StartElements[ElemID] = true;
+		}
+	}
+
 	if (Attribs->HasMaterialID())
 	{
 		Change->MaterialIDAttribChange.Emplace();
@@ -456,6 +472,19 @@ void FDynamicMeshAttributeSetChangeTracker::SaveInitialTriangle(int TriangleID)
 			NormalChange.SaveInitialTriangle(NormalLayer, TriangleID);
 		}
 	}
+	if (Change->ColorChange.IsSet())
+	{
+		FDynamicMeshColorChange& ColorChange = *Change->ColorChange;
+		const FDynamicMeshColorOverlay* ColorLayer = Attribs->PrimaryColors();
+		if (ColorLayer->IsSetTriangle(TriangleID))
+		{
+			FIndex3i ColorTriangle = ColorLayer->GetTriangle(TriangleID);
+			SaveElement(ColorTriangle.A, ColorState, ColorLayer, ColorChange);
+			SaveElement(ColorTriangle.B, ColorState, ColorLayer, ColorChange);
+			SaveElement(ColorTriangle.C, ColorState, ColorLayer, ColorChange);
+			ColorChange.SaveInitialTriangle(ColorLayer, TriangleID);
+		}
+	}
 
 	if (Change->MaterialIDAttribChange.IsSet())
 	{
@@ -541,6 +570,31 @@ void FDynamicMeshAttributeSetChangeTracker::StoreAllFinalTriangles(const TArray<
 		}
 	}
 
+	// store final Color elements for all modified triangles, and final triangles if needed
+	if (Change->ColorChange.IsSet())
+	{
+		FDynamicMeshColorChange& ColorChange = *Change->ColorChange;
+		const FDynamicMeshColorOverlay* ColorLayer = Attribs->PrimaryColors();
+		StoredVertices.Reset();
+
+		for (int tid : TriangleIDs)
+		{
+			if (ColorLayer->IsSetTriangle(tid))
+			{
+				FIndex3i Tri = ColorLayer->GetTriangle(tid);
+				for (int j = 0; j < 3; ++j)
+				{
+					if (StoredVertices.Contains(Tri[j]) == false)
+					{ 
+						ColorChange.StoreFinalElement(ColorLayer, Tri[j]);
+						StoredVertices.Add(Tri[j]);
+					}
+				}
+				ColorChange.StoreFinalTriangle(ColorLayer, tid);
+			}
+		}
+	}
+
 	if (Change->MaterialIDAttribChange.IsSet())
 	{
 		Change->MaterialIDAttribChange->StoreAllFinalTriangles(Attribs->GetMaterialID(), TriangleIDs);
@@ -590,7 +644,11 @@ bool FDynamicMeshAttributeChangeSet::Apply(FDynamicMeshAttributeSet* Attributes,
 	{
 		RegisteredAttributeChanges[k]->Apply(Attributes->GetRegisteredAttribute(k), bRevert);
 	}
-
+	if (ColorChange.IsSet())
+	{
+		FDynamicMeshColorOverlay* ColorLayer = Attributes->PrimaryColors();
+		ColorChange->Apply(ColorLayer, bRevert);
+	}
 	if (MaterialIDAttribChange.IsSet())
 	{
 		MaterialIDAttribChange->Apply(Attributes->GetMaterialID(), bRevert);
diff --git a/Engine/Plugins/Experimental/GeometryProcessing/Source/DynamicMesh/Private/DynamicMeshEditor.cpp b/Engine/Plugins/Experimental/GeometryProcessing/Source/DynamicMesh/Private/DynamicMeshEditor.cpp
--- a/Engine/Plugins/Experimental/GeometryProcessing/Source/DynamicMesh/Private/DynamicMeshEditor.cpp
+++ b/Engine/Plugins/Experimental/GeometryProcessing/Source/DynamicMesh/Private/DynamicMeshEditor.cpp
@@ -1310,6 +1310,22 @@ void FDynamicMeshEditor::CopyAttributes(int FromTriangleID, int ToTriangleID, FM
 			NormalOverlay->SetTriangle(ToTriangleID, ToElemTri);
 		}
 	}
+	
+	if (Mesh->Attributes()->HasPrimaryColors())
+	{
+		FDynamicMeshColorOverlay* ColorOverlay = Mesh->Attributes()->PrimaryColors();
+		if (ColorOverlay->IsSetTriangle(FromTriangleID))
+		{
+			FIndex3i FromElemTri = ColorOverlay->GetTriangle(FromTriangleID);
+			FIndex3i ToElemTri = ColorOverlay->GetTriangle(ToTriangleID);
+			for (int j = 0; j < 3; ++j)
+			{
+				int NewElemID = FindOrCreateDuplicateColor(FromElemTri[j], IndexMaps, &ResultOut);
+				ToElemTri[j] = NewElemID;
+			}
+			ColorOverlay->SetTriangle(ToTriangleID, ToElemTri);
+		}
+	}
 
 	if (Mesh->Attributes()->HasMaterialID())
 	{
@@ -1359,6 +1375,22 @@ int FDynamicMeshEditor::FindOrCreateDuplicateNormal(int ElementID, int NormalLay
 }
 
 
+int FDynamicMeshEditor::FindOrCreateDuplicateColor(int ElementID, FMeshIndexMappings& IndexMaps, FDynamicMeshEditResult* ResultOut)
+{
+	int NewElementID = IndexMaps.GetNewColor(ElementID);
+	if (NewElementID == IndexMaps.InvalidID())
+	{
+		FDynamicMeshColorOverlay* ColorOverlay = Mesh->Attributes()->PrimaryColors();
+		NewElementID = ColorOverlay->AppendElement(ColorOverlay->GetElement(ElementID));
+		IndexMaps.SetColor(ElementID, NewElementID);
+		if (ResultOut)
+		{
+			ResultOut->NewColorOverlayElements.Add(NewElementID);
+		}
+	}
+	return NewElementID;
+}
+
 
 int FDynamicMeshEditor::FindOrCreateDuplicateVertex(int VertexID, FMeshIndexMappings& IndexMaps, FDynamicMeshEditResult& ResultOut)
 {
@@ -1500,6 +1532,19 @@ void FDynamicMeshEditor::AppendMesh(const FDynamicMesh3* AppendMesh,
 			}
 		}
 
+		if (AppendMesh->Attributes()->HasPrimaryColors() && Mesh->Attributes()->HasPrimaryColors())
+		{
+			const FDynamicMeshColorOverlay* FromColors = AppendMesh->Attributes()->PrimaryColors();
+			FDynamicMeshColorOverlay* ToColors = Mesh->Attributes()->PrimaryColors();
+			if (FromColors != nullptr && ToColors != nullptr)
+			{
+				FIndexMapi& ColorMap = IndexMapsOut.GetColorMap();
+				ColorMap.Reserve(FromColors->ElementCount());
+				AppendColors(AppendMesh, FromColors, ToColors,
+					VertexMap, TriangleMap, ColorMap);
+			}
+		}
+
 		if (AppendMesh->Attributes()->HasMaterialID() && Mesh->Attributes()->HasMaterialID())
 		{
 			const FDynamicMeshMaterialAttribute* FromMaterialIDs = AppendMesh->Attributes()->GetMaterialID();
@@ -1600,6 +1645,34 @@ void FDynamicMeshEditor::AppendUVs(const FDynamicMesh3* AppendMesh,
 	}
 }
 
+void FDynamicMeshEditor::AppendColors(const FDynamicMesh3* AppendMesh,
+	const FDynamicMeshColorOverlay* FromOverlay, FDynamicMeshColorOverlay* ToOverlay,
+	const FIndexMapi& VertexMap, const FIndexMapi& TriangleMap,
+	FIndexMapi& MapOut)
+{
+	// copy over color elements
+	for (int ElemID : FromOverlay->ElementIndicesItr())
+	{
+		int NewElemID = ToOverlay->AppendElement(FromOverlay->GetElement(ElemID));
+		MapOut.Add(ElemID, NewElemID);
+	}
+
+	// now set new triangles
+	for (int TriID : AppendMesh->TriangleIndicesItr())
+	{
+		if (FromOverlay->IsSetTriangle(TriID))
+		{
+			FIndex3i ElemTri = FromOverlay->GetTriangle(TriID);
+			int NewTriID = TriangleMap.GetTo(TriID);
+			for (int j = 0; j < 3; ++j)
+			{
+				ElemTri[j] = FromOverlay->IsElement(ElemTri[j]) ? MapOut.GetTo(ElemTri[j]) : FDynamicMesh3::InvalidID;
+			}
+			ToOverlay->SetTriangle(NewTriID, ElemTri);
+		}
+	}
+}
+
 
 
 
@@ -1649,6 +1722,19 @@ static int AppendTriangleNormalAttribute(const FDynamicMesh3* FromMesh, int From
 	return NewElementID;
 }
 
+// Utility function for ::AppendTriangles()
+static int AppendTriangleColorAttribute(const FDynamicMesh3* FromMesh, int FromElementID, FDynamicMesh3* ToMesh, FMeshIndexMappings& IndexMaps)
+{
+	int NewElementID = IndexMaps.GetNewColor(FromElementID);
+	if (NewElementID == IndexMaps.InvalidID())
+	{
+		const FDynamicMeshColorOverlay* FromOverlay = FromMesh->Attributes()->PrimaryColors();
+		FDynamicMeshColorOverlay* ToOverlay = ToMesh->Attributes()->PrimaryColors();
+		NewElementID = ToOverlay->AppendElement(FromOverlay->GetElement(FromElementID));
+		IndexMaps.SetColor(FromElementID, NewElementID);
+	}
+	return NewElementID;
+}
 
 
 
@@ -1698,6 +1784,24 @@ static void AppendTriangleAttributes(const FDynamicMesh3* FromMesh, int FromTria
 		}
 	}
 
+	if (FromMesh->Attributes()->HasPrimaryColors() && ToMesh->Attributes()->HasPrimaryColors())
+	{
+		const FDynamicMeshColorOverlay* FromOverlay = FromMesh->Attributes()->PrimaryColors();
+		FDynamicMeshColorOverlay* ToOverlay = ToMesh->Attributes()->PrimaryColors();
+		if (FromOverlay->IsSetTriangle(FromTriangleID))
+		{
+			FIndex3i FromElemTri = FromOverlay->GetTriangle(FromTriangleID);
+			FIndex3i ToElemTri = ToOverlay->GetTriangle(ToTriangleID);
+			for (int j = 0; j < 3; ++j)
+			{
+				check(FromElemTri[j] != FDynamicMesh3::InvalidID);
+				int NewElemID = AppendTriangleColorAttribute(FromMesh, FromElemTri[j], ToMesh, IndexMaps);
+				ToElemTri[j] = NewElemID;
+			}
+			ToOverlay->SetTriangle(ToTriangleID, ToElemTri);
+		}
+	}
+
 	if (FromMesh->Attributes()->HasMaterialID() && ToMesh->Attributes()->HasMaterialID())
 	{
 		const FDynamicMeshMaterialAttribute* FromMaterialIDs = FromMesh->Attributes()->GetMaterialID();
diff --git a/Engine/Plugins/Experimental/GeometryProcessing/Source/DynamicMesh/Private/DynamicMeshOverlay.cpp b/Engine/Plugins/Experimental/GeometryProcessing/Source/DynamicMesh/Private/DynamicMeshOverlay.cpp
--- a/Engine/Plugins/Experimental/GeometryProcessing/Source/DynamicMesh/Private/DynamicMeshOverlay.cpp
+++ b/Engine/Plugins/Experimental/GeometryProcessing/Source/DynamicMesh/Private/DynamicMeshOverlay.cpp
@@ -1411,6 +1411,8 @@ template class DYNAMICMESH_API TDynamicMeshOverlay<int, 2>;
 template class DYNAMICMESH_API TDynamicMeshOverlay<float, 3>;
 template class DYNAMICMESH_API TDynamicMeshOverlay<double, 3>;
 template class DYNAMICMESH_API TDynamicMeshOverlay<int, 3>;
+template class DYNAMICMESH_API TDynamicMeshOverlay<float, 4>;
+template class DYNAMICMESH_API TDynamicMeshOverlay<double, 4>;
 
 template class DYNAMICMESH_API TDynamicMeshVectorOverlay<float, 2, FVector2f>;
 template class DYNAMICMESH_API TDynamicMeshVectorOverlay<double, 2, FVector2d>;
@@ -1418,6 +1420,8 @@ template class DYNAMICMESH_API TDynamicMeshVectorOverlay<int, 2, FVector2i>;
 template class DYNAMICMESH_API TDynamicMeshVectorOverlay<float, 3, FVector3f>;
 template class DYNAMICMESH_API TDynamicMeshVectorOverlay<double, 3, FVector3d>;
 template class DYNAMICMESH_API TDynamicMeshVectorOverlay<int, 3, FVector3i>;
+template class DYNAMICMESH_API TDynamicMeshVectorOverlay<float, 4, FVector4f>;
+template class DYNAMICMESH_API TDynamicMeshVectorOverlay<double, 4, FVector4d>;
 
 } // end namespace UE::Geometry
 } // end namespace UE
\ No newline at end of file
diff --git a/Engine/Plugins/Experimental/GeometryProcessing/Source/DynamicMesh/Public/DynamicMeshAttributeSet.h b/Engine/Plugins/Experimental/GeometryProcessing/Source/DynamicMesh/Public/DynamicMeshAttributeSet.h
--- a/Engine/Plugins/Experimental/GeometryProcessing/Source/DynamicMesh/Public/DynamicMeshAttributeSet.h
+++ b/Engine/Plugins/Experimental/GeometryProcessing/Source/DynamicMesh/Public/DynamicMeshAttributeSet.h
@@ -19,7 +19,8 @@ namespace Geometry
 typedef TDynamicMeshVectorOverlay<float, 2, FVector2f> FDynamicMeshUVOverlay;
 /** Standard Normal overlay type - 3-element float */
 typedef TDynamicMeshVectorOverlay<float, 3, FVector3f> FDynamicMeshNormalOverlay;
-
+/** Standard Color overlay type - 4-element float (rbga) */
+typedef TDynamicMeshVectorOverlay<float, 4, FVector4f> FDynamicMeshColorOverlay;
 /** Standard per-triangle integer material ID */
 typedef TDynamicMeshScalarTriangleAttribute<int32> FDynamicMeshMaterialAttribute;
 
@@ -110,7 +111,7 @@ public:
 	virtual bool IsSeamEndEdge(int EdgeID) const;
 
 	/** @return true if the given edge is a seam edge in any overlay */
-	virtual bool IsSeamEdge(int EdgeID, bool& bIsUVSeamOut, bool& bIsNormalSeamOut) const;
+	virtual bool IsSeamEdge(int EdgeID, bool& bIsUVSeamOut, bool& bIsNormalSeamOut, bool& bIsColorSeamOut) const;
 
 	/** @return true if the given vertex is a seam vertex in any overlay */
 	virtual bool IsSeamVertex(int VertexID, bool bBoundaryIsSeam = true) const;
@@ -240,6 +241,24 @@ public:
 		return (PrimaryNormals() != nullptr && PrimaryTangents()  != nullptr && PrimaryBiTangents() != nullptr);
 	}
 
+	bool HasPrimaryColors() const
+	{
+		return !!ColorLayer;
+	}
+
+	FDynamicMeshColorOverlay* PrimaryColors()
+	{
+		return ColorLayer.Get();
+	}
+
+	const FDynamicMeshColorOverlay* PrimaryColors() const
+	{
+		return ColorLayer.Get();
+	}
+
+	void EnablePrimaryColors();
+	
+	void DisablePrimaryColors();
 
 	//
 	// Polygroup layers
@@ -327,10 +346,11 @@ public:
 protected:
 	/** Parent mesh of this attribute set */
 	FDynamicMesh3* ParentMesh;
-
 	
+
 	TIndirectArray<FDynamicMeshUVOverlay> UVLayers;
 	TIndirectArray<FDynamicMeshNormalOverlay> NormalLayers;
+	TUniquePtr<FDynamicMeshColorOverlay> ColorLayer;
 
 	TUniquePtr<FDynamicMeshMaterialAttribute> MaterialIDAttrib;
 
diff --git a/Engine/Plugins/Experimental/GeometryProcessing/Source/DynamicMesh/Public/DynamicMeshChangeTracker.h b/Engine/Plugins/Experimental/GeometryProcessing/Source/DynamicMesh/Public/DynamicMeshChangeTracker.h
--- a/Engine/Plugins/Experimental/GeometryProcessing/Source/DynamicMesh/Public/DynamicMeshChangeTracker.h
+++ b/Engine/Plugins/Experimental/GeometryProcessing/Source/DynamicMesh/Public/DynamicMeshChangeTracker.h
@@ -87,6 +87,8 @@ protected:
 typedef TDynamicMeshAttributeChange<float,2> FDynamicMeshUVChange;
 /** Standard Normal overlay change type - 3-element float */
 typedef TDynamicMeshAttributeChange<float,3> FDynamicMeshNormalChange;
+/** Standard Color overlay change type - 4-element float */
+typedef TDynamicMeshAttributeChange<float, 4> FDynamicMeshColorChange;
 /** Standard per-triangle integer attribute change type */
 typedef FDynamicMeshTriangleAttributeChange<int32,1> FDynamicMeshTriGroupChange;
 
@@ -101,6 +103,7 @@ public:
 
 	TArray<FDynamicMeshUVChange> UVChanges;
 	TArray<FDynamicMeshNormalChange> NormalChanges;
+	TOptional<FDynamicMeshColorChange> ColorChange;
 	TOptional<FDynamicMeshTriGroupChange> MaterialIDAttribChange;
 	TArray<FDynamicMeshTriGroupChange> PolygroupChanges;
 	TArray<TUniquePtr<FDynamicMeshAttributeChangeBase>> RegisteredAttributeChanges;
@@ -232,6 +235,7 @@ protected:
 	};
 	TArray<FElementState> UVStates;
 	TArray<FElementState> NormalStates;
+	FElementState ColorState;
 
 	template<typename AttribOverlayType, typename AttribChangeType>
 	void SaveElement(int ElementID, FElementState& State, const AttribOverlayType* Overlay, AttribChangeType& ChangeIn)
diff --git a/Engine/Plugins/Experimental/GeometryProcessing/Source/DynamicMesh/Public/DynamicMeshEditor.h b/Engine/Plugins/Experimental/GeometryProcessing/Source/DynamicMesh/Public/DynamicMeshEditor.h
--- a/Engine/Plugins/Experimental/GeometryProcessing/Source/DynamicMesh/Public/DynamicMeshEditor.h
+++ b/Engine/Plugins/Experimental/GeometryProcessing/Source/DynamicMesh/Public/DynamicMeshEditor.h
@@ -39,6 +39,9 @@ struct DYNAMICMESH_API FDynamicMeshEditResult
 	/** New normal overlay elements */
 	TArray<TArray<int32>> NewNormalOverlayElements;
 
+	/** New color overlay elements */
+	TArray<int32> NewColorOverlayElements;
+
 	/** clear this data structure */
 	void Reset()
 	{
@@ -48,6 +51,7 @@ struct DYNAMICMESH_API FDynamicMeshEditResult
 		NewPolygons.Reset();
 		NewGroups.Reset();
 		NewNormalOverlayElements.Reset();
+		NewColorOverlayElements.Reset();
 	}
 
 	/** Flatten the triangle/quad/polygon lists into a single list of all triangles */
@@ -408,7 +412,15 @@ public:
 	 */
 	int FindOrCreateDuplicateNormal(int ElementID, int NormalLayerIndex, FMeshIndexMappings& IndexMaps, FDynamicMeshEditResult* ResultOut = nullptr);
 
-
+	/**
+	 * Find "new" color for input color element under Index mapping, or create new if missing
+	 * @param ElementID the source color we want a duplicate of
+	 * @param IndexMaps source/destination mapping of already-duplicated colors
+	 * @param ResultOut any newly created element indices are stored in NewColorOverlayElements here. 
+	 * @return index of duplicate color in given color layer
+	 */
+	int FindOrCreateDuplicateColor(int ElementID, FMeshIndexMappings& IndexMaps, FDynamicMeshEditResult* ResultOut);
+	
 	/**
 	 * Copy all attribute-layer values from one triangle to another, using the IndexMaps to track and re-use shared attribute values.
 	 * @param FromTriangleID source triangle
@@ -468,7 +480,21 @@ public:
 		const FIndexMapi& VertexMap, const FIndexMapi& TriangleMap,
 		FIndexMapi& UVMapOut);
 
-
+	/**
+	* Append Colors from one attribute overlay to another.
+	* Assumes that AppendMesh has already been appended to Mesh.
+	* Note that this function has no dependency on .Mesh, it could be static
+	* @param AppendMesh mesh that owns FromOverlay attribute overlay
+	* @param FromOverlay Color overlay we want to append from (owned by AppendMesh)
+	* @param ToOverlay Color overlay we want to append to (owned by Mesh)
+	* @param VertexMap map from AppendMesh vertex IDs to vertex IDs applicable to ToOverlay (ie of .Mesh)
+	* @param TriangleMap map from AppendMesh triangle IDs to triangle IDs applicable to ToOverlay (ie of .Mesh)
+	* @param ColorMapOut Mapping from element IDs of FromUVs to new element IDs in ToOverlay
+	*/
+	void AppendColors(const FDynamicMesh3* AppendMesh,
+		const FDynamicMeshColorOverlay* FromOverlay, FDynamicMeshColorOverlay* ToOverlay,
+		const FIndexMapi& VertexMap, const FIndexMapi& TriangleMap,
+		FIndexMapi& ColorMapOut);
 
 
 	/**
diff --git a/Engine/Plugins/Experimental/GeometryProcessing/Source/DynamicMesh/Public/MeshIndexMappings.h b/Engine/Plugins/Experimental/GeometryProcessing/Source/DynamicMesh/Public/MeshIndexMappings.h
--- a/Engine/Plugins/Experimental/GeometryProcessing/Source/DynamicMesh/Public/MeshIndexMappings.h
+++ b/Engine/Plugins/Experimental/GeometryProcessing/Source/DynamicMesh/Public/MeshIndexMappings.h
@@ -22,6 +22,7 @@ protected:
 	FIndexMapi TriangleMap;
 	FIndexMapi GroupMap;
 
+	FIndexMapi ColorMap;
 	TArray<FIndexMapi> UVMaps;
 	TArray<FIndexMapi> NormalMaps;
 
@@ -37,6 +38,7 @@ public:
 		VertexMap.Reset();
 		TriangleMap.Reset();
 		GroupMap.Reset();
+		ColorMap.Reset();
 		for (FIndexMapi& UVMap : UVMaps)
 		{
 			UVMap.Reset();
@@ -79,6 +81,11 @@ public:
 	int GetNewNormal(int NormalLayer, int FromID) const { return NormalMaps[NormalLayer].GetTo(FromID); }
 	inline bool ContainsNormal(int NormalLayer, int FromID) const { return NormalMaps[NormalLayer].ContainsFrom(FromID); }
 
+	FIndexMapi& GetColorMap() { return ColorMap; }
+	void SetColor( int FromID, int ToID) { ColorMap.Add(FromID, ToID); }
+	int GetNewColor(int FromID) const { return ColorMap.GetTo(FromID); }
+	inline bool ContainsColor(int FromID) const { return ColorMap.ContainsFrom(FromID); }
+
 };
 
 } // end namespace UE::Geometry
diff --git a/Engine/Plugins/Experimental/GeometryProcessing/Source/MeshConversion/Private/DynamicMeshToMeshDescription.cpp b/Engine/Plugins/Experimental/GeometryProcessing/Source/MeshConversion/Private/DynamicMeshToMeshDescription.cpp
--- a/Engine/Plugins/Experimental/GeometryProcessing/Source/MeshConversion/Private/DynamicMeshToMeshDescription.cpp
+++ b/Engine/Plugins/Experimental/GeometryProcessing/Source/MeshConversion/Private/DynamicMeshToMeshDescription.cpp
@@ -24,8 +24,8 @@ namespace DynamicMeshToMeshDescriptionConversionHelper
 		{
 			TArrayView<const FVertexInstanceID> InstanceTri = MeshOutArg.GetTriangleVertexInstances(TriangleID);
 
-			int32 MeshInTriIdx = TriangleID.GetValue();   
-
+			int32 MeshInTriIdx = TriangleID.GetValue();
+			
 			FIndex3i OverlayVertIndices = Overlay->GetTriangle(MeshInTriIdx);
 			InstanceAttrib.Set(InstanceTri[0], AttribIndex, OutAttributeType(Overlay->GetElement(OverlayVertIndices.A)));
 			InstanceAttrib.Set(InstanceTri[1], AttribIndex, OutAttributeType(Overlay->GetElement(OverlayVertIndices.B)));
@@ -292,24 +292,34 @@ void FDynamicMeshToMeshDescription::UpdateTangents(const FDynamicMesh3* MeshIn,
 
 void FDynamicMeshToMeshDescription::UpdateVertexColors(const FDynamicMesh3* MeshIn, FMeshDescription& MeshOut)
 {
-	check(MeshIn->IsCompactV() && MeshIn->HasVertexColors());
 
 	FStaticMeshAttributes Attributes(MeshOut);
-	TVertexInstanceAttributesRef<FVector4> InstanceColors = Attributes.GetVertexInstanceColors();
-	bool bIsValidDst = InstanceColors.IsValid();
-	ensureMsgf(bIsValidDst, TEXT("Trying to update colors on a MeshDescription that has no color attributes"));
-	if (bIsValidDst)
+	TVertexInstanceAttributesRef<FVector4> InstanceColorsAttrib = Attributes.GetVertexInstanceColors();
+
+	if (!ensureMsgf(MeshIn->IsCompactT(), TEXT("Trying to update MeshDescription Colors from a non-compact DynamicMesh"))) return;
+	if (!ensureMsgf(MeshIn->TriangleCount() == MeshOut.Triangles().Num(), TEXT("Trying to update MeshDescription Colors from Mesh that does not have same triangle count"))) return;
+	if (!ensureMsgf(MeshIn->HasAttributes() && MeshIn->Attributes()->HasPrimaryColors(), TEXT("Trying to update MeshDescription Colors from a DynamicMesh that has no attribute set at all or has no color data in its attribute set"))) return;
+	if (!ensureMsgf(InstanceColorsAttrib.IsValid(), TEXT("Trying to update colors on a MeshDescription that has no color attributes"))) return;
+	
+	const FDynamicMeshColorOverlay* ColorOverlay = MeshIn->Attributes()->PrimaryColors();
+	
+	const int32 NumTriangles = MeshIn->TriangleCount();
+	for (int32 t = 0; t < NumTriangles; ++t)
 	{
-		check(MeshIn->VertexCount() == MeshOut.Vertices().Num());
-		for (int VertID : MeshIn->VertexIndicesItr())
+		if (!ColorOverlay->IsSetTriangle(t))
 		{
-			FVector3f Color3f = MeshIn->GetVertexColor(VertID);
-			FVector4 Color4(Color3f.X, Color3f.Y, Color3f.Z, 1.0f);
-			for (FVertexInstanceID InstanceID : MeshOut.GetVertexVertexInstanceIDs(FVertexID(VertID)))
-			{
-				InstanceColors.Set(InstanceID, 0, Color4);
-			}
+			continue;
 		}
+	
+		FVector4f TriColors[3];
+		ColorOverlay->GetTriElements(t, TriColors[0], TriColors[1], TriColors[2]);
+		TArrayView<const FVertexInstanceID> InstanceIDs = MeshOut.GetTriangleVertexInstances(FTriangleID(t));
+		for (int32 i = 0; i < 3; ++i)
+		{
+			FVector4 InstanceColor4(TriColors[i].X, TriColors[i].Y, TriColors[i].Z, TriColors[i].W);
+			InstanceColorsAttrib.Set(InstanceIDs[i], 0, InstanceColor4);
+		}
+
 	}
 }
 
@@ -581,10 +591,11 @@ void FDynamicMeshToMeshDescription::Convert_NoSharedInstances(const FDynamicMesh
 	// check if we have per-triangle material ID
 	const FDynamicMeshMaterialAttribute* MaterialIDAttrib = (bHasAttributes && MeshIn->Attributes()->HasMaterialID()) ? MeshIn->Attributes()->GetMaterialID() : nullptr;
 
-	// cache tangent space and UV overlay info
+	// cache tangent space and UV and color overlay info
 	const FDynamicMeshNormalOverlay* NormalOverlay = bHasAttributes ? MeshIn->Attributes()->PrimaryNormals() : nullptr;
 	const FDynamicMeshNormalOverlay* TangentOverlay = bHasAttributes ? MeshIn->Attributes()->PrimaryTangents() : nullptr;
 	const FDynamicMeshNormalOverlay* BiTangentOverlay = bHasAttributes ? MeshIn->Attributes()->PrimaryBiTangents() : nullptr;
+	const FDynamicMeshColorOverlay* ColorOverlay = bHasAttributes ? MeshIn->Attributes()->PrimaryColors() : nullptr;
 
 	const int32 NumUVLayers = bHasAttributes ? MeshIn->Attributes()->NumUVLayers() : 0;
 	
@@ -608,7 +619,8 @@ void FDynamicMeshToMeshDescription::Convert_NoSharedInstances(const FDynamicMesh
 		bCopyGroupToPolyGroup = true;
 	}
 
-	bool bCopyVertexColors = MeshIn->HasVertexColors();		// always copy when we are baking new mesh? should this be a config option?
+	// always copy when we are baking new mesh? should this be a config option?
+	bool bCopyInstanceColors = (ColorOverlay != nullptr); 
 
 	// disable indexing during the full build of the mesh
 	Builder.SuspendMeshDescriptionIndexing();
@@ -696,7 +708,29 @@ void FDynamicMeshToMeshDescription::Convert_NoSharedInstances(const FDynamicMesh
 	{
 		TangetSpaceInstanceSetter = [](int TriID, FVertexInstanceID TriVertInstances[3]){ return;};
 	}
-
+	TFunction<void(int TriID, FVertexInstanceID TriVertInstances[3])> ColorInstanceSetter;
+	if (bCopyInstanceColors)
+	{
+		ColorInstanceSetter = [ColorOverlay, &Builder](int TriID, FVertexInstanceID TriVertInstances[3])
+		{
+			FIndex3i ColorTri = ColorOverlay->GetTriangle(TriID);
+			for (int32 j = 0; j < 3; ++j)
+			{
+				FVector4 DstColor(1,1,1,1);
+				const FVertexInstanceID CornerInstanceID = TriVertInstances[j];
+				if (ColorOverlay->IsElement(ColorTri[j]))
+				{
+					FVector4f TriVertColor4 = ColorOverlay->GetElement(ColorTri[j]);
+					DstColor = FVector4(TriVertColor4.X, TriVertColor4.Y, TriVertColor4.Z, TriVertColor4.W);
+				}
+				Builder.SetInstanceColor(CornerInstanceID, DstColor);
+			}
+		};
+	}
+	else
+	{
+		ColorInstanceSetter = [](int TriID, FVertexInstanceID TriVertInstances[3]) { return; };
+	}
 
 	// need to know max material index value so we can reserve groups in MeshDescription
 	int32 MaxPolygonGroupID = 0;
@@ -741,15 +775,6 @@ void FDynamicMeshToMeshDescription::Convert_NoSharedInstances(const FDynamicMesh
 			TriVertInstances[j] = Builder.AppendInstance(TriVertex);
 		}
 
-		// copy vertex colors. This writes each color multiple times, but that is not expensive in this context
-		if (bCopyVertexColors)
-		{
-			for (int32 j = 0; j < 3; ++j)
-			{
-				FVector3f Color = MeshIn->GetVertexColor(Triangle[j]);
-				Builder.SetInstanceColor(TriVertInstances[j], FVector4(Color.X, Color.Y, Color.Z, 1.0f));
-			}
-		}
 
 		// transfer material index to MeshDescription polygon group (by convention)
 		FPolygonGroupID UsePolygonGroupID = ZeroPolygonGroupID;
@@ -806,6 +831,10 @@ void FDynamicMeshToMeshDescription::Convert_NoSharedInstances(const FDynamicMesh
 		// NB: will need to be updated to follow the pattern used in UVs above when MeshDescription supports shared tangent space elements. 
 		TangetSpaceInstanceSetter(TriID, TriVertInstances);
 
+		// transfer the color overlay to per-instance colors.
+		// NB: will need to be updated to follow the pattern used in UVs above if MeshDescription supports shared color elements
+		ColorInstanceSetter(TriID, TriVertInstances);
+
 
 		if (bCopyGroupToPolyGroup)
 		{
diff --git a/Engine/Plugins/Experimental/GeometryProcessing/Source/MeshConversion/Private/MeshDescriptionToDynamicMesh.cpp b/Engine/Plugins/Experimental/GeometryProcessing/Source/MeshConversion/Private/MeshDescriptionToDynamicMesh.cpp
--- a/Engine/Plugins/Experimental/GeometryProcessing/Source/MeshConversion/Private/MeshDescriptionToDynamicMesh.cpp
+++ b/Engine/Plugins/Experimental/GeometryProcessing/Source/MeshConversion/Private/MeshDescriptionToDynamicMesh.cpp
@@ -113,7 +113,50 @@ public:
 	}
 };
 
+struct FVertexColor
+{
+	int vid;
+	FVector4f Color;
+	bool operator==(const FVertexColor& o) const
+	{
+		return vid == o.vid && Color == o.Color;
+	}
+};
+FORCEINLINE uint32 GetTypeHash(const FVertexColor& Vector)
+{
+	return FCrc::MemCrc32(&Vector, sizeof(Vector));
+}
 
+class FColorWelder
+{
+public:
+	TMap<FVertexColor, int> UniqueVertexColors;
+	FDynamicMeshColorOverlay* ColorOverlay;
+
+	FColorWelder() : ColorOverlay(nullptr)
+	{
+	}
+
+	FColorWelder(FDynamicMeshColorOverlay* ColorOverlayIn)
+	{
+		check(ColorOverlayIn);
+		ColorOverlay = ColorOverlayIn;
+	}
+	int FindOrAddUnique(const FVector4f& Color, int VertexID)
+	{
+		FVertexColor VertColor = { VertexID, Color };
+
+		const int32* FoundIndex = UniqueVertexColors.Find(VertColor);
+		if (FoundIndex != nullptr)
+		{
+			return *FoundIndex;
+		}
+
+		int32 NewIndex = ColorOverlay->AppendElement(&Color.X);
+		UniqueVertexColors.Add(VertColor, NewIndex);
+		return NewIndex;
+	}
+};
 
 
 
@@ -126,7 +169,7 @@ public:
 
 
 void FMeshDescriptionToDynamicMesh::Convert(const FMeshDescription* MeshIn, FDynamicMesh3& MeshOut, bool bCopyTangents )
-{
+{ 
 	TriIDMap.Reset();
 	VertIDMap.Reset();
 
@@ -153,15 +196,7 @@ void FMeshDescriptionToDynamicMesh::Convert(const FMeshDescription* MeshIn, FDyn
 	TVertexInstanceAttributesConstRef<float> InstanceBiTangentSign = Attributes.GetVertexInstanceBinormalSigns();
 
 	TVertexInstanceAttributesConstRef<FVector4> InstanceColors = Attributes.GetVertexInstanceColors();
-	if (bEnableOutputVertexColors && InstanceColors.IsValid())
-	{
-		MeshOut.EnableVertexColors(FVector3f::One());
-	}
-	else
-	{
-		bEnableOutputVertexColors = false;
-	}
-	bool bFoundNonDefaultVertexColor = false;
+
 
 	TPolygonAttributesConstRef<int> PolyGroups =
 		MeshIn->PolygonAttributes().GetAttributesRef<int>(ExtendedMeshAttribute::PolyTriGroups);
@@ -310,18 +345,6 @@ void FMeshDescriptionToDynamicMesh::Convert(const FMeshDescription* MeshIn, FDyn
 		
 		TriIDMap[NewTriangleID] = TriangleID;
 	
-		// copy triangle instance colors to per-vertex colors. Currently this just uses last-instance as
-		// we do not support per-triangle-vertex colors.
-		if (bEnableOutputVertexColors)
-		{
-			for (int32 j = 0; j < 3; ++j)
-			{
-				FVector4 InstanceColor4 = InstanceColors.Get(TriData.TriInstances[j], 0);
-				FVector3f InstanceColor3(InstanceColor4.X, InstanceColor4.Y, InstanceColor4.Z);
-				bFoundNonDefaultVertexColor |= (InstanceColor3 != FVector3f::One());
-				MeshOut.SetVertexColor(VertexIDs[j], InstanceColor3);
-			}
-		}
 	}
 
 	int32 MaxTriID = MeshOut.MaxTriangleID(); // really MaxTriID+1
@@ -330,11 +353,6 @@ void FMeshDescriptionToDynamicMesh::Convert(const FMeshDescription* MeshIn, FDyn
 	
 
 
-	if (bFoundNonDefaultVertexColor == false)
-	{
-		MeshOut.DiscardVertexColors();
-	}
-
 	FDateTime Time_AfterTriangles = FDateTime::Now();
 
 
@@ -363,12 +381,11 @@ void FMeshDescriptionToDynamicMesh::Convert(const FMeshDescription* MeshIn, FDyn
 		}
 	}
 
-
+	FDynamicMeshColorOverlay* ColorOverlay = nullptr;
 	FDynamicMeshNormalOverlay* NormalOverlay = nullptr;
 	FDynamicMeshNormalOverlay* TangentOverlay = nullptr;
 	FDynamicMeshNormalOverlay* BiTangentOverlay = nullptr;
 	FDynamicMeshMaterialAttribute* MaterialIDAttrib = nullptr;
-
 	
 
 	if (!bDisableAttributes)
@@ -402,7 +419,13 @@ void FMeshDescriptionToDynamicMesh::Convert(const FMeshDescription* MeshIn, FDyn
 		{
 			MeshOut.Attributes()->GetUVLayer(i)->InitializeTriangles(MeshOut.MaxTriangleID());
 		}
-	
+
+		if (InstanceColors.IsValid())
+		{
+			MeshOut.Attributes()->EnablePrimaryColors();
+			ColorOverlay = MeshOut.Attributes()->PrimaryColors();
+		}
+
 
 		// always enable Material ID if there are any attributes
 		MeshOut.Attributes()->EnableMaterialID();
@@ -429,6 +452,8 @@ void FMeshDescriptionToDynamicMesh::Convert(const FMeshDescription* MeshIn, FDyn
 
 	if (!bDisableAttributes)
 	{
+		bool bFoundNonDefaultVertexInstanceColor = false;
+
 		// we will weld/populate all the attributes simultaneously, hold on to futures in this array and then Wait for them at the end
 		TArray<TFuture<void>> Pending;
 
@@ -682,6 +707,33 @@ void FMeshDescriptionToDynamicMesh::Convert(const FMeshDescription* MeshIn, FDyn
 			Pending.Add(MoveTemp(BiTangentFuture));
 		}
 
+		if (ColorOverlay != nullptr)
+		{
+			auto ColorFuture = Async(EAsyncExecution::ThreadPool, [&]()
+			{
+				const FVector4 DefaultColor4 = InstanceColors.GetDefaultValue();
+
+				FColorWelder ColorWelder(ColorOverlay);
+				for (int32 TriangleID : MeshOut.TriangleIndicesItr())
+				{
+					const FIndex3i Tri = MeshOut.GetTriangle(TriangleID);
+					const FTriData& TriData = AddedTriangles[TriangleID];
+					FIndex3i TriVector;
+					for (int j = 0; j < 3; ++j)
+					{
+						const FVector4 InstanceColor4 = InstanceColors.Get(TriData.TriInstances[j], 0);
+						const FVector4f OverlayColor(InstanceColor4.X, InstanceColor4.Y, InstanceColor4.Z, InstanceColor4.W);
+						TriVector[j] = ColorWelder.FindOrAddUnique(OverlayColor, Tri[j]);
+
+						// need to detect if the vertex instance color actually held any real data..
+						bFoundNonDefaultVertexInstanceColor |= (InstanceColor4 != DefaultColor4);
+					}
+					ColorOverlay->SetTriangle(TriangleID, TriVector);
+				}
+			});
+			Pending.Add(MoveTemp(ColorFuture));
+		}
+
 		if (MaterialIDAttrib != nullptr)
 		{
 			auto MaterialFuture = Async(EAsyncExecution::ThreadPool, [&]()
@@ -721,6 +773,11 @@ void FMeshDescriptionToDynamicMesh::Convert(const FMeshDescription* MeshIn, FDyn
 		{
 			Future.Wait();
 		}
+
+		if (!bFoundNonDefaultVertexInstanceColor)
+		{
+			MeshOut.Attributes()->DisablePrimaryColors();
+		}
 	}
 
 	// free maps if no longer needed
diff --git a/Engine/Plugins/Experimental/GeometryProcessing/Source/MeshConversion/Private/SkeletalMeshLODModelToDynamicMesh.cpp b/Engine/Plugins/Experimental/GeometryProcessing/Source/MeshConversion/Private/SkeletalMeshLODModelToDynamicMesh.cpp
--- a/Engine/Plugins/Experimental/GeometryProcessing/Source/MeshConversion/Private/SkeletalMeshLODModelToDynamicMesh.cpp
+++ b/Engine/Plugins/Experimental/GeometryProcessing/Source/MeshConversion/Private/SkeletalMeshLODModelToDynamicMesh.cpp
@@ -19,6 +19,7 @@ public:
 
 	typedef int32 UVIDType;
 	typedef int32 NormalIDType;
+	typedef int32 ColorIDType;
 
 	FSkeletalMeshLODModelWrapper(const  FSkeletalMeshLODModel& MeshIn, bool bUseDisabledSections) :
 		bIncludeDisabledSections(bUseDisabledSections),
@@ -181,6 +182,11 @@ public:
 		return true;
 	}
 
+	bool HasColors() const
+	{
+		return true;
+	}
+
 	//-- Access to per-wedge attributes --//
 	void GetWedgeIDs(const TriIDType& TriID, WedgeIDType& WID0, WedgeIDType& WID1, WedgeIDType& WID2) const
 	{
@@ -211,6 +217,11 @@ public:
 		return FVector3f(GetWedgeVertexInstance(WID).TangentY);
 	}
 	
+	FVector4f GetWedgeColor(WedgeIDType WID) const
+	{
+		FLinearColor LinearColor = GetWedgeVertexInstance(WID).Color.ReinterpretAsLinear();
+		return FVector4f(LinearColor.R, LinearColor.G, LinearColor.B, LinearColor.A);
+	}
 
 
 	int32 GetMaterialIndex(TriIDType TriID) const
@@ -240,6 +251,10 @@ public:
 	FVector3f GetBiTangent(NormalIDType ID) const {check(0); return FVector3f(); }
 	bool GetBiTangentTri(const TriIDType&, NormalIDType& ID0, NormalIDType& ID1, NormalIDType& ID2) const { ID0 = ID1 = ID2 = NormalIDType(-1); return false; }
 
+	const TArray<int32>& GetColorIDs() const { return EmptyArray; }
+	FVector4f GetColor(ColorIDType ID) const { check(0); return FVector4f(); }
+	bool GetColorTri(const TriIDType&, ColorIDType& ID0, ColorIDType& ID1, ColorIDType& ID2) const { ID0 = ID1 = ID2 = ColorIDType(-1); return false; }
+
 	
 	// -- additional methods, not required by the conversion interface
 	const  FSkeletalMeshLODModel& GetSrcMesh() const
@@ -247,12 +262,6 @@ public:
 		return *Mesh;
 	}
 
-	
-	FLinearColor GetWedgeColor(const WedgeIDType WID) const
-	{
-		const FColor& Color = GetWedgeVertexInstance(WID).Color;
-		return Color.ReinterpretAsLinear();
-	}
 
 private:
 
@@ -319,33 +328,7 @@ void FSkeletalMeshLODModelToDynamicMesh::Convert(const  FSkeletalMeshLODModel* M
 		auto TriToMaterialID = [&ModelWrapper](const int32& SrcTriID)->int32 { return ModelWrapper.GetMaterialIndex(SrcTriID); };
 		SkeletalToDynamicMesh.Convert(MeshOut, ModelWrapper, TriToGroupID, TriToMaterialID, bCopyTangents);						
 	}
-
-	// Special code for vertex colors.  Currently DynamicMesh3 doesn't use an overlay for vertex colors 
-	// thus we do not support per-triangle-vertex colors. 
-	// NB: This just uses last-seen wedge color to define the vertex color.
-	// NB: The "alpha" channel is dropped.
-	if (bEnableOutputVertexColors)
-	{
-		MeshOut.EnableVertexColors(FVector3f::One());
-		bool bFoundNonDefaultVertexColor = false;
-		for (int32 TriangleID : MeshOut.TriangleIndicesItr())
-		{
-			int32 SrcWIDs[3];
-			ModelWrapper.GetWedgeIDs(SkeletalToDynamicMesh.ToSrcTriIDMap[TriangleID], SrcWIDs[0], SrcWIDs[1], SrcWIDs[2]);
-			FIndex3i VIDs = MeshOut.GetTriangle(TriangleID);
-			for (int32 j = 0; j < 3; ++j)
-			{ 
-				FLinearColor WedgeColor = ModelWrapper.GetWedgeColor(SrcWIDs[j]);
-				FVector3f WedgeColor3(WedgeColor); // Color.A is lost here.
-				bFoundNonDefaultVertexColor |= (WedgeColor3 != FVector3f::One());
-				MeshOut.SetVertexColor(VIDs[j], WedgeColor3);
-			}
-		}
-		if (bFoundNonDefaultVertexColor == false)
-		{
-			MeshOut.DiscardVertexColors();
-		}
-	}
+	
 
 	if (!bEnableOutputGroups)
 	{
@@ -389,4 +372,4 @@ void  FSkeletalMeshLODModelToDynamicMesh::Convert(const  FSkeletalMeshLODModel*
 	check(0);
 }
 
-#endif  // end with editor
\ No newline at end of file
+#endif  // end with editor
diff --git a/Engine/Plugins/Experimental/GeometryProcessing/Source/MeshConversion/Private/ToDynamicMesh.h b/Engine/Plugins/Experimental/GeometryProcessing/Source/MeshConversion/Private/ToDynamicMesh.h
--- a/Engine/Plugins/Experimental/GeometryProcessing/Source/MeshConversion/Private/ToDynamicMesh.h
+++ b/Engine/Plugins/Experimental/GeometryProcessing/Source/MeshConversion/Private/ToDynamicMesh.h
@@ -197,6 +197,7 @@ FORCEINLINE uint32 GetTypeHash(const TVertexAttr<AttrType>& Vector)
 template <typename AttrType> struct TOverlayTraits {};
 template<> struct TOverlayTraits<FVector3f> { typedef FDynamicMeshNormalOverlay OverlayType; };
 template<> struct TOverlayTraits<FVector2f> { typedef FDynamicMeshUVOverlay OverlayType; };
+template<> struct TOverlayTraits<FVector4f> { typedef FDynamicMeshColorOverlay OverlayType; };
 
 // Welder used for exact attribute value welding when constructing overlay
 template <typename AttrType>
@@ -246,6 +247,7 @@ public:
 *		typedef SrcTriIDType     TriIDType;
 *		typedef SrcUVIDType      UVIDType;
 *		typedef SrcNormalIDType  NormalIDType;
+*		typedef SrcColorIDType   ColorIDType
 *       typedef SrcWedgeIDType   WedgeIDType;
 *
 *		// accounting.
@@ -275,6 +277,7 @@ public:
 *		FVector3f GetWedgeNormal(WedgeIDType WID) const;
 *		FVector3f GetWedgeTangent(WedgeIDType WID) const;
 *		FVector3f GetWedgeBiTangent(WedgeIDType WID) const;
+*       FVector4f GetWedgeColor(WedgeIDType WID) const;
 * 
 *		// attribute access that exploits shared attributes. 
 *		// each group of shared attributes presents itself as a mesh with its own attribute vertex buffer.
@@ -297,6 +300,9 @@ public:
 *		FVector3f GetBiTangent(NormalIDType ID) const;
 *       bool GetBiTangentTri(const TriIDType& TID, NormalIDType& NID0, NormalIDType& NID1, NormalIDType& NID2) const;
 * 
+*		const TArray<ColorIDType>& GetColorIDs() const;
+*		FVector4f GetColor(ColorIDType ID) const;
+*		bool GetColorTri(const TriIDType& TID, ColorIDType& NID0, ColorIDType& NID1, ColorIDType& NID2) const;
 *	};
 *
 */
@@ -309,6 +315,7 @@ public:
 	using SrcVertIDType   = typename MyBase::SrcVertIDType;
 	using SrcUVIDType     = typename SrcMeshType::UVIDType;
 	using SrcNormalIDType = typename SrcMeshType::NormalIDType;
+	using SrcColorIDType  = typename SrcMeshType::ColorIDType;
 	using SrcWedgeIDType  = typename SrcMeshType::WedgeIDType;
 
 	TToDynamicMesh() :MyBase() {}
@@ -339,7 +346,7 @@ public:
 protected:
 
 
-	// Populates overlays for UVs, Normal, Material ID. Also Tangent and BiTangent if bCopyTangents == true.  
+	// Populates overlays for UVs, Normal, Color, Material ID. Also Tangent and BiTangent if bCopyTangents == true.  
 	// NB: This does not copy any polygroup information.  
 	void PopulateOverlays(FDynamicMesh3& MeshOut, const SrcMeshType& MeshIn, TFunctionRef<int32(const SrcTriIDType& SrcTrID)> MaterialIDFunction, bool bCopyTangents)
 	{
@@ -347,6 +354,7 @@ protected:
 		FDynamicMeshNormalOverlay* NormalOverlay = nullptr;
 		FDynamicMeshNormalOverlay* TangentOverlay = nullptr;
 		FDynamicMeshNormalOverlay* BiTangentOverlay = nullptr;
+		FDynamicMeshColorOverlay*  ColorOverlay = nullptr;
 		FDynamicMeshMaterialAttribute* MaterialIDAttrib = nullptr;
 
 		// only copy tangents if they exist
@@ -383,6 +391,9 @@ protected:
 			MeshOut.Attributes()->GetUVLayer(i)->InitializeTriangles(MeshOut.MaxTriangleID());
 		}
 
+		// create color overlay
+		MeshOut.Attributes()->EnablePrimaryColors();
+		ColorOverlay = MeshOut.Attributes()->PrimaryColors();
 
 		// always enable Material ID if there are any attributes
 		MeshOut.Attributes()->EnableMaterialID();
@@ -451,7 +462,7 @@ protected:
 				{
 					PopulateOverlay<FVector3f, SrcNormalIDType>(
 						BiTangentOverlay,
-						MeshIn.GetTangentIDs(),
+						MeshIn.GetBiTangentIDs(),
 						[&MeshIn](const SrcNormalIDType& NID)->FVector3f { return MeshIn.GetBiTangent(NID); },
 						[&MeshIn](const SrcTriIDType& TriID, SrcNormalIDType& N0, SrcNormalIDType& N1, SrcNormalIDType& N2)->bool {return MeshIn.GetBiTangentTri(TriID, N0, N1, N2); },
 						[&MeshIn](const SrcWedgeIDType& WID)->FVector3f {return MeshIn.GetWedgeBiTangent(WID); },
@@ -461,6 +472,23 @@ protected:
 				});
 			Pending.Add(MoveTemp(BiTangentFuture));
 		}
+		if (ColorOverlay != nullptr)
+		{
+			auto ColorFuture = Async(EAsyncExecution::ThreadPool, [&]()
+				{
+					PopulateOverlay<FVector4f, SrcColorIDType>(
+						ColorOverlay,
+						MeshIn.GetColorIDs(),
+						[&MeshIn](const SrcColorIDType& CID)->FVector4f { return MeshIn.GetColor(CID); },
+						[&MeshIn](const SrcTriIDType& TriID, SrcColorIDType& C0, SrcColorIDType& C1, SrcColorIDType& C2)->bool {return MeshIn.GetColorTri(TriID, C0, C1, C2); },
+						[&MeshIn](const SrcWedgeIDType& WID)->FVector4f {return MeshIn.GetWedgeColor(WID); },
+						[&MeshIn](const SrcTriIDType& TriID, SrcWedgeIDType& WID0, SrcWedgeIDType& WID1, SrcWedgeIDType& WID2)->void { MeshIn.GetWedgeIDs(TriID, WID0, WID1, WID2); }
+					);
+
+				});
+			Pending.Add(MoveTemp(ColorFuture));
+		}
+
 		//populate MaterialID overlay
 		if (MaterialIDAttrib != nullptr)
 		{
diff --git a/Engine/Plugins/Experimental/GeometryProcessing/Source/MeshConversion/Public/MeshDescriptionToDynamicMesh.h b/Engine/Plugins/Experimental/GeometryProcessing/Source/MeshConversion/Public/MeshDescriptionToDynamicMesh.h
--- a/Engine/Plugins/Experimental/GeometryProcessing/Source/MeshConversion/Public/MeshDescriptionToDynamicMesh.h
+++ b/Engine/Plugins/Experimental/GeometryProcessing/Source/MeshConversion/Public/MeshDescriptionToDynamicMesh.h
@@ -27,13 +27,10 @@ public:
 	/** Should we initialize triangle groups on output mesh */
 	bool bEnableOutputGroups = true;
 
-	/** Should we initialize vertex colors on output mesh */
-	bool bEnableOutputVertexColors = true;
-
 	/** Should we calculate conversion index maps */
 	bool bCalculateMaps = true;
 
-	/** Ignore all mesh attributes (e.g. UV/Normal layers, material groups) */
+	/** Ignore all mesh attributes (e.g. UV/Normal layers, color layer, material groups) */
 	bool bDisableAttributes = false;
 
 
diff --git a/Engine/Plugins/Experimental/GeometryProcessing/Source/MeshConversion/Public/SkeletalMeshLODModelToDynamicMesh.h b/Engine/Plugins/Experimental/GeometryProcessing/Source/MeshConversion/Public/SkeletalMeshLODModelToDynamicMesh.h
--- a/Engine/Plugins/Experimental/GeometryProcessing/Source/MeshConversion/Public/SkeletalMeshLODModelToDynamicMesh.h
+++ b/Engine/Plugins/Experimental/GeometryProcessing/Source/MeshConversion/Public/SkeletalMeshLODModelToDynamicMesh.h
@@ -23,9 +23,6 @@ public:
 	/** Should we initialize triangle groups on output mesh */
 	bool bEnableOutputGroups = true;
 
-	/** Should we initialize vertex colors on output mesh */
-	bool bEnableOutputVertexColors = true;
-
 	/** Should we calculate conversion index maps */
 	bool bCalculateMaps = true;
 
diff --git a/Engine/Plugins/Experimental/MeshModelingToolset/Source/MeshModelingTools/Private/CombineMeshesTool.cpp b/Engine/Plugins/Experimental/MeshModelingToolset/Source/MeshModelingTools/Private/CombineMeshesTool.cpp
--- a/Engine/Plugins/Experimental/MeshModelingToolset/Source/MeshModelingTools/Private/CombineMeshesTool.cpp
+++ b/Engine/Plugins/Experimental/MeshModelingToolset/Source/MeshModelingTools/Private/CombineMeshesTool.cpp
@@ -229,6 +229,7 @@ void UCombineMeshesTool::CreateNewAsset()
 	AccumulateDMesh.EnableTriangleGroups();
 	AccumulateDMesh.EnableAttributes();
 	AccumulateDMesh.Attributes()->EnableMaterialID();
+	AccumulateDMesh.Attributes()->EnablePrimaryColors();
 	FTransform AccumToWorld(Box.GetCenter());
 	FTransform ToAccum(-Box.GetCenter());
 
@@ -241,7 +242,7 @@ void UCombineMeshesTool::CreateNewAsset()
 			LOCTEXT("DuplicateMeshBuild", "Building duplicate mesh ...") :
 			LOCTEXT("CombineMeshesBuild", "Building combined mesh ..."));
 		SlowTask.MakeDialog();
-
+		bool bNeedColorAttr = false;
 		int MatIndexBase = 0;
 		for (int32 ComponentIdx = 0; ComponentIdx < Targets.Num(); ComponentIdx++)
 		{
@@ -253,6 +254,7 @@ void UCombineMeshesTool::CreateNewAsset()
 			FDynamicMesh3 ComponentDMesh;
 			const FMeshDescription* MeshDescription = MeshDescriptions[ComponentIdx];
 			Converter.Convert(MeshDescription, ComponentDMesh);
+			bNeedColorAttr = bNeedColorAttr || (ComponentDMesh.HasAttributes() && ComponentDMesh.Attributes()->HasPrimaryColors());
 
 			UE::Geometry::FTransform3d XF = (UE::Geometry::FTransform3d)(TargetComponentInterface(ComponentIdx)->GetWorldTransform() * ToAccum);
 			if (XF.GetDeterminant() < 0)
@@ -286,6 +288,11 @@ void UCombineMeshesTool::CreateNewAsset()
 			MatIndexBase += TargetMaterialInterface(ComponentIdx)->GetNumMaterials();
 		}
 
+		if (!bNeedColorAttr)
+		{
+			AccumulateDMesh.Attributes()->DisablePrimaryColors();
+		}
+
 		SlowTask.EnterProgressFrame(1);
 
 		if (bDuplicateMode)
@@ -398,6 +405,7 @@ void UCombineMeshesTool::UpdateExistingAsset()
 	AccumulateDMesh.EnableTriangleGroups();
 	AccumulateDMesh.EnableAttributes();
 	AccumulateDMesh.Attributes()->EnableMaterialID();
+	AccumulateDMesh.Attributes()->EnablePrimaryColors();
 
 	int32 SkipIndex = (BasicProperties->WriteOutputTo == ECombineTargetType::FirstInputAsset) ? 0 : (Targets.Num() - 1);
 	IPrimitiveComponentBackedTarget* UpdateTarget = TargetComponentInterface(SkipIndex);
@@ -419,7 +427,7 @@ void UCombineMeshesTool::UpdateExistingAsset()
 			LOCTEXT("DuplicateMeshBuild", "Building duplicate mesh ...") :
 			LOCTEXT("CombineMeshesBuild", "Building combined mesh ..."));
 		SlowTask.MakeDialog();
-
+		bool bNeedColorAttr = false;
 		int MatIndexBase = 0;
 		for (int32 ComponentIdx = 0; ComponentIdx < Targets.Num(); ComponentIdx++)
 		{
@@ -430,6 +438,7 @@ void UCombineMeshesTool::UpdateExistingAsset()
 			FMeshDescriptionToDynamicMesh Converter;
 			FDynamicMesh3 ComponentDMesh;
 			Converter.Convert(MeshDescriptions[ComponentIdx], ComponentDMesh);
+			bNeedColorAttr = bNeedColorAttr || (ComponentDMesh.HasAttributes() && ComponentDMesh.Attributes()->HasPrimaryColors());
 
 			// update material IDs to account for combined material set
 			FDynamicMeshMaterialAttribute* MatAttrib = ComponentDMesh.Attributes()->GetMaterialID();
@@ -467,6 +476,11 @@ void UCombineMeshesTool::UpdateExistingAsset()
 			Editor.AppendMesh(&ComponentDMesh, IndexMapping);
 		}
 
+		if (!bNeedColorAttr)
+		{
+			AccumulateDMesh.Attributes()->DisablePrimaryColors();
+		}
+
 		SlowTask.EnterProgressFrame(1);
 
 		UpdateTargetCommitter->CommitMeshDescription([&](const IMeshDescriptionCommitter::FCommitterParams& CommitParams)
diff --git a/Engine/Plugins/Experimental/MeshModelingToolset/Source/MeshModelingTools/Private/SelfUnionMeshesTool.cpp b/Engine/Plugins/Experimental/MeshModelingToolset/Source/MeshModelingTools/Private/SelfUnionMeshesTool.cpp
--- a/Engine/Plugins/Experimental/MeshModelingToolset/Source/MeshModelingTools/Private/SelfUnionMeshesTool.cpp
+++ b/Engine/Plugins/Experimental/MeshModelingToolset/Source/MeshModelingTools/Private/SelfUnionMeshesTool.cpp
@@ -94,14 +94,16 @@ void USelfUnionMeshesTool::ConvertInputsAndSetPreviewMaterials(bool bSetPreviewM
 	CombinedSourceMeshes->EnableAttributes();
 	CombinedSourceMeshes->EnableTriangleGroups(0);
 	CombinedSourceMeshes->Attributes()->EnableMaterialID();
+	CombinedSourceMeshes->Attributes()->EnablePrimaryColors();
 	FDynamicMeshEditor AppendEditor(CombinedSourceMeshes.Get());
 
+	bool bNeedColorAttr = false;
 	for (int ComponentIdx = 0; ComponentIdx < Targets.Num(); ComponentIdx++)
 	{
 		FDynamicMesh3 ComponentMesh;
 		FMeshDescriptionToDynamicMesh Converter;
 		Converter.Convert(TargetMeshProviderInterface(ComponentIdx)->GetMeshDescription(), ComponentMesh);
-
+		bNeedColorAttr = bNeedColorAttr || (ComponentMesh.Attributes()->HasPrimaryColors());
 		// ensure materials and attributes are always enabled
 		ComponentMesh.EnableAttributes();
 		ComponentMesh.Attributes()->EnableMaterialID();
@@ -128,6 +130,10 @@ void USelfUnionMeshesTool::ConvertInputsAndSetPreviewMaterials(bool bSetPreviewM
 			}
 			);
 	}
+	if (!bNeedColorAttr)
+	{
+		CombinedSourceMeshes->Attributes()->DisablePrimaryColors();
+	}
 
 	Preview->ConfigureMaterials(AllMaterialSet.Materials, ToolSetupUtil::GetDefaultWorkingMaterial(GetToolManager()));
 
diff --git a/Engine/Plugins/Experimental/MeshModelingToolset/Source/ModelingComponents/Private/BaseDynamicMeshSceneProxy.h b/Engine/Plugins/Experimental/MeshModelingToolset/Source/ModelingComponents/Private/BaseDynamicMeshSceneProxy.h
--- a/Engine/Plugins/Experimental/MeshModelingToolset/Source/ModelingComponents/Private/BaseDynamicMeshSceneProxy.h
+++ b/Engine/Plugins/Experimental/MeshModelingToolset/Source/ModelingComponents/Private/BaseDynamicMeshSceneProxy.h
@@ -13,6 +13,7 @@ using UE::Geometry::FDynamicMesh3;
 using UE::Geometry::FDynamicMeshAttributeSet;
 using UE::Geometry::FDynamicMeshUVOverlay;
 using UE::Geometry::FDynamicMeshNormalOverlay;
+using UE::Geometry::FDynamicMeshColorOverlay;
 using UE::Geometry::FDynamicMeshMaterialAttribute;
 
 /**
@@ -410,13 +411,14 @@ public:
 		int NumTriangles, TriangleEnumerable Enumerable,
 		const FDynamicMeshUVOverlay* UVOverlay,
 		const FDynamicMeshNormalOverlay* NormalOverlay,
+		const FDynamicMeshColorOverlay* ColorOverlay,
 		TFunction<void(int, int, int, FVector3f&, FVector3f&)> TangentsFunc = nullptr,
 		bool bTrackTriangles = false)
 	{
 		TArray<const FDynamicMeshUVOverlay*> UVOverlays;
 		UVOverlays.Add(UVOverlay);
 		InitializeBuffersFromOverlays(RenderBuffers, Mesh, NumTriangles, Enumerable,
-			UVOverlays, NormalOverlay, TangentsFunc, bTrackTriangles);
+			UVOverlays, NormalOverlay, ColorOverlay, TangentsFunc, bTrackTriangles);
 	}
 
 
@@ -432,6 +434,7 @@ public:
 		int NumTriangles, TriangleEnumerable Enumerable,
 		const TArray<const FDynamicMeshUVOverlay*>& UVOverlays,
 		const FDynamicMeshNormalOverlay* NormalOverlay,
+		const FDynamicMeshColorOverlay* ColorOverlay,
 		TFunction<void(int, int, int, FVector3f&, FVector3f&)> TangentsFunc = nullptr,
 		bool bTrackTriangles = false)
 	{
@@ -441,7 +444,7 @@ public:
 			return;
 		}
 
-		bool bHaveColors = Mesh->HasVertexColors() && (bIgnoreVertexColors == false);
+		bool bHaveColors = (ColorOverlay != nullptr) && (bIgnoreVertexColors == false);
 
 		int NumVertices = NumTriangles * 3;
 		int NumUVOverlays = UVOverlays.Num();
@@ -473,11 +476,12 @@ public:
 				UVTriangles[k] = (k < NumUVOverlays && UVOverlays[k] != nullptr) ? UVOverlays[k]->GetTriangle(TriangleID) : FIndex3i::Invalid();
 			}
 			FIndex3i TriNormal = (NormalOverlay != nullptr) ? NormalOverlay->GetTriangle(TriangleID) : FIndex3i::Zero();
+			FIndex3i TriColor = (ColorOverlay != nullptr) ? ColorOverlay->GetTriangle(TriangleID) : FIndex3i::Zero();
 
-			FColor TriColor = ConstantVertexColor;
+			FColor UniformTriColor = ConstantVertexColor;
 			if (bUsePerTriangleColor && PerTriangleColorFunc != nullptr)
 			{
-				TriColor = PerTriangleColorFunc(Mesh, TriangleID);
+				UniformTriColor = PerTriangleColorFunc(Mesh, TriangleID);
 				bHaveColors = false;
 			}
 
@@ -487,7 +491,7 @@ public:
 
 				FVector3f Normal = (NormalOverlay != nullptr && TriNormal[j] != FDynamicMesh3::InvalidID) ?
 					NormalOverlay->GetElement(TriNormal[j]) : Mesh->GetVertexNormal(Tri[j]);
-
+				
 				// either request known tangent, or calculate a nonsense one
 				if (TangentsFunc != nullptr)
 				{
@@ -506,8 +510,10 @@ public:
 					RenderBuffers->StaticMeshVertexBuffer.SetVertexUV(VertIdx, k, (FVector2D)UV);
 				}
 
-				RenderBuffers->ColorVertexBuffer.VertexColor(VertIdx) = (bHaveColors) ?
-					UE::Geometry::ToFColor(Mesh->GetVertexColor(Tri[j])) : TriColor;
+				FColor VertexFColor = (bHaveColors && TriColor[j] != FDynamicMesh3::InvalidID) ?
+					FColor(ColorOverlay->GetElement(TriColor[j])) : UniformTriColor;
+
+				RenderBuffers->ColorVertexBuffer.VertexColor(VertIdx) = VertexFColor;
 
 				RenderBuffers->IndexBuffer.Indices[TriIdx++] = VertIdx;		// currently TriIdx == VertIdx so we don't really need both...
 				VertIdx++;
@@ -628,6 +634,7 @@ public:
 		const FDynamicMesh3* Mesh,
 		int NumTriangles, TriangleEnumerable Enumerable,
 		FDynamicMeshNormalOverlay* NormalOverlay,
+		FDynamicMeshColorOverlay* ColorOverlay,
 		TFunctionRef<void(int, int, int, const FVector3f&, FVector3f&, FVector3f&)> TangentsFunc,
 		bool bUpdatePositions = true,
 		bool bUpdateNormals = false,
@@ -638,7 +645,7 @@ public:
 			return;
 		}
 
-		bool bHaveColors = Mesh->HasVertexColors() && (bIgnoreVertexColors == false);
+		bool bHaveColors = (ColorOverlay != nullptr) && (bIgnoreVertexColors == false);
 
 		int NumVertices = NumTriangles * 3;
 		check(RenderBuffers->PositionVertexBuffer.GetNumVertices() == NumVertices);
@@ -658,11 +665,12 @@ public:
 			FIndex3i Tri = Mesh->GetTriangle(TriangleID);
 
 			FIndex3i TriNormal = (bUpdateNormals && NormalOverlay != nullptr) ? NormalOverlay->GetTriangle(TriangleID) : FIndex3i::Zero();
+			FIndex3i TriColor = (bUpdateColors && ColorOverlay != nullptr) ? ColorOverlay->GetTriangle(TriangleID) : FIndex3i::Zero();
 
-			FColor TriColor = ConstantVertexColor;
+			FColor UniformTriColor = ConstantVertexColor;
 			if (bUpdateColors && bUsePerTriangleColor && PerTriangleColorFunc != nullptr)
 			{
-				TriColor = PerTriangleColorFunc(Mesh, TriangleID);
+				UniformTriColor = PerTriangleColorFunc(Mesh, TriangleID);
 				bHaveColors = false;
 			}
 
@@ -685,8 +693,9 @@ public:
 
 				if (bUpdateColors)
 				{
-					RenderBuffers->ColorVertexBuffer.VertexColor(VertIdx) = (bHaveColors) ?
-						UE::Geometry::ToFColor(Mesh->GetVertexColor(Tri[j])) : TriColor;
+					FColor VertexFColor = (bHaveColors  && TriColor[j] != FDynamicMesh3::InvalidID) ?
+						FColor(ColorOverlay->GetElement(TriColor[j])) : UniformTriColor;
+
 				}
 
 				VertIdx++;
diff --git a/Engine/Plugins/Experimental/MeshModelingToolset/Source/ModelingComponents/Private/Drawing/MeshElementsVisualizer.cpp b/Engine/Plugins/Experimental/MeshModelingToolset/Source/ModelingComponents/Private/Drawing/MeshElementsVisualizer.cpp
--- a/Engine/Plugins/Experimental/MeshModelingToolset/Source/ModelingComponents/Private/Drawing/MeshElementsVisualizer.cpp
+++ b/Engine/Plugins/Experimental/MeshModelingToolset/Source/ModelingComponents/Private/Drawing/MeshElementsVisualizer.cpp
@@ -59,8 +59,8 @@ public:
 		}
 		if (Mesh->HasAttributes())
 		{
-			bool bIsUVSeam = false, bIsNormalSeam = false;
-			if (Mesh->Attributes()->IsSeamEdge(EdgeIndex, bIsUVSeam, bIsNormalSeam))
+			bool bIsUVSeam = false, bIsNormalSeam = false, bIsColorSeam = false;
+			if (Mesh->Attributes()->IsSeamEdge(EdgeIndex, bIsUVSeam, bIsNormalSeam, bIsColorSeam))
 			{
 				if (bIsUVSeam)
 				{
@@ -70,6 +70,10 @@ public:
 				{
 					EdgeType |= (int32)EMeshEdgeType::NormalSeam;
 				}
+				if (bIsColorSeam)
+				{
+					EdgeType |= (int32)EMeshEdgeType::ColorSeam;
+				}
 			}
 		}
 		TypeOut = (EMeshEdgeType)EdgeType;
@@ -108,12 +112,14 @@ void UMeshElementsVisualizer::OnCreated()
 	Settings->WatchProperty(Settings->bShowBorders, [this](bool) { bSettingsModified = true; });
 	Settings->WatchProperty(Settings->bShowUVSeams, [this](bool) { bSettingsModified = true; });
 	Settings->WatchProperty(Settings->bShowNormalSeams, [this](bool) { bSettingsModified = true; });
+	Settings->WatchProperty(Settings->bShowColorSeams, [this](bool){ bSettingsModified = true; });
 	Settings->WatchProperty(Settings->ThicknessScale, [this](float) { bSettingsModified = true; });
 	Settings->WatchProperty(Settings->DepthBias, [this](float) { bSettingsModified = true; });
 	Settings->WatchProperty(Settings->WireframeColor, [this](FColor) { bSettingsModified = true; });
 	Settings->WatchProperty(Settings->BoundaryEdgeColor, [this](FColor) { bSettingsModified = true; });
 	Settings->WatchProperty(Settings->UVSeamColor, [this](FColor) { bSettingsModified = true; });
 	Settings->WatchProperty(Settings->NormalSeamColor, [this](FColor) { bSettingsModified = true; });
+	Settings->WatchProperty(Settings->ColorSeamColor, [this](FColor) { bSettingsModified = true; });
 	bSettingsModified = false;
 
 	WireframeComponent = NewObject<UMeshWireframeComponent>(GetActor());
@@ -165,11 +171,13 @@ void UMeshElementsVisualizer::UpdateVisibility()
 	WireframeComponent->bEnableBoundaryEdges = Settings->bShowBorders;
 	WireframeComponent->bEnableUVSeams = Settings->bShowUVSeams;
 	WireframeComponent->bEnableNormalSeams = Settings->bShowNormalSeams;
+	WireframeComponent->bEnableColorSeams = Settings->bShowColorSeams;
 
 	WireframeComponent->WireframeColor = Settings->WireframeColor;
 	WireframeComponent->BoundaryEdgeColor = Settings->BoundaryEdgeColor;
 	WireframeComponent->UVSeamColor = Settings->UVSeamColor;
 	WireframeComponent->NormalSeamColor = Settings->NormalSeamColor;
+	WireframeComponent->ColorSeamColor = Settings->ColorSeamColor;
 
 	WireframeComponent->MarkRenderStateDirty();
 }
diff --git a/Engine/Plugins/Experimental/MeshModelingToolset/Source/ModelingComponents/Private/Drawing/MeshWireframeComponent.cpp b/Engine/Plugins/Experimental/MeshModelingToolset/Source/ModelingComponents/Private/Drawing/MeshWireframeComponent.cpp
--- a/Engine/Plugins/Experimental/MeshModelingToolset/Source/ModelingComponents/Private/Drawing/MeshWireframeComponent.cpp
+++ b/Engine/Plugins/Experimental/MeshModelingToolset/Source/ModelingComponents/Private/Drawing/MeshWireframeComponent.cpp
@@ -74,6 +74,10 @@ public:
 			{
 				bEdgeIsVisible = true;
 			}
+			else if (((int)EdgeType & (int)IMeshWireframeSource::EMeshEdgeType::ColorSeam) != 0 && Component->bEnableColorSeams)
+			{
+				bEdgeIsVisible = true;
+			}
 
 			if (bEdgeIsVisible)
 			{
@@ -118,6 +122,8 @@ public:
 		float UVSeamThickness = Component->ThicknessScale * Component->UVSeamThickness;
 		FColor NormalSeamColor = FLinearColor::FromSRGBColor(Component->NormalSeamColor).ToFColor(false);
 		float NormalSeamThickness = Component->ThicknessScale * Component->NormalSeamThickness;
+		FColor ColorSeamColor = FLinearColor::FromSRGBColor(Component->ColorSeamColor).ToFColor(false);
+		float ColorSeamThickness = Component->ThicknessScale * Component->ColorSeamThickness;
 
 		float LineDepthBias = Component->LineDepthBias * Component->LineDepthBiasSizeScale;
 
@@ -154,6 +160,11 @@ public:
 					UseThickness = (bIsBoundaryEdge) ? BoundaryEdgeThickness : NormalSeamThickness;
 					UseColor = NormalSeamColor;
 				}
+				else if (((int)EdgeType & (int)IMeshWireframeSource::EMeshEdgeType::ColorSeam) != 0 && Component->bEnableColorSeams)
+				{
+					UseThickness = (bIsBoundaryEdge) ? BoundaryEdgeThickness : ColorSeamThickness;
+					UseColor = ColorSeamColor;
+				}
 			}
 
 			const FVector A = WireSource->GetVertex(EdgeInfo.B);
diff --git a/Engine/Plugins/Experimental/MeshModelingToolset/Source/ModelingComponents/Private/OctreeDynamicMeshSceneProxy.h b/Engine/Plugins/Experimental/MeshModelingToolset/Source/ModelingComponents/Private/OctreeDynamicMeshSceneProxy.h
--- a/Engine/Plugins/Experimental/MeshModelingToolset/Source/ModelingComponents/Private/OctreeDynamicMeshSceneProxy.h
+++ b/Engine/Plugins/Experimental/MeshModelingToolset/Source/ModelingComponents/Private/OctreeDynamicMeshSceneProxy.h
@@ -78,15 +78,17 @@ public:
 		// find suitable overlays
 		FDynamicMeshUVOverlay* UVOverlay = nullptr;
 		FDynamicMeshNormalOverlay* NormalOverlay = nullptr;
+		FDynamicMeshColorOverlay* ColorOverlay = nullptr;
 		if (Mesh->HasAttributes())
 		{
 			UVOverlay = Mesh->Attributes()->PrimaryUV();
 			NormalOverlay = Mesh->Attributes()->PrimaryNormals();
+			ColorOverlay = Mesh->Attributes()->PrimaryColors();
 		}
 
 		InitializeBuffersFromOverlays(RenderBuffers, Mesh,
 			Mesh->TriangleCount(), Mesh->TriangleIndicesItr(),
-			UVOverlay, NormalOverlay);
+			UVOverlay, NormalOverlay, ColorOverlay);
 
 		ENQUEUE_RENDER_COMMAND(FOctreeDynamicMeshSceneProxyInitializeSingle)(
 			[this, RenderBuffers](FRHICommandListImmediate& RHICmdList)
@@ -107,10 +109,12 @@ public:
 		// find suitable overlays
 		FDynamicMeshUVOverlay* UVOverlay = nullptr;
 		FDynamicMeshNormalOverlay* NormalOverlay = nullptr;
+		FDynamicMeshColorOverlay* ColorOverlay = nullptr;
 		if (Mesh->HasAttributes())
 		{
 			UVOverlay = Mesh->Attributes()->PrimaryUV();
 			NormalOverlay = Mesh->Attributes()->PrimaryNormals();
+			ColorOverlay = Mesh->Attributes()->PrimaryColors();
 		}
 
 		const TArray<int32>& SetIDs = Decomposition.GetIndexSetIDs();
@@ -123,7 +127,7 @@ public:
 
 			InitializeBuffersFromOverlays(RenderBuffers, Mesh,
 				Tris.Num(), Tris,
-				UVOverlay, NormalOverlay);
+				UVOverlay, NormalOverlay, ColorOverlay);
 
 			ENQUEUE_RENDER_COMMAND(FOctreeDynamicMeshSceneProxyInitializeFromDecomposition)(
 				[this, SetID, RenderBuffers](FRHICommandListImmediate& RHICmdList)
@@ -164,10 +168,12 @@ public:
 		// find suitable overlays
 		FDynamicMeshUVOverlay* UVOverlay = nullptr;
 		FDynamicMeshNormalOverlay* NormalOverlay = nullptr;
+		FDynamicMeshColorOverlay* ColorOverlay = nullptr;
 		if (Mesh->HasAttributes())
 		{
 			UVOverlay = Mesh->Attributes()->PrimaryUV();
 			NormalOverlay = Mesh->Attributes()->PrimaryNormals();
+			ColorOverlay = Mesh->Attributes()->PrimaryColors();
 		}
 
 		{
@@ -183,7 +189,7 @@ public:
 
 				InitializeBuffersFromOverlays(RenderBuffers, Mesh,
 					Tris.Num(), Tris,
-					UVOverlay, NormalOverlay);
+					UVOverlay, NormalOverlay, ColorOverlay);
 
 				ENQUEUE_RENDER_COMMAND(FOctreeDynamicMeshSceneProxyUpdateAddOne)(
 					[this, SetID, RenderBuffers](FRHICommandListImmediate& RHICmdList)
diff --git a/Engine/Plugins/Experimental/MeshModelingToolset/Source/ModelingComponents/Private/SimpleDynamicMeshSceneProxy.h b/Engine/Plugins/Experimental/MeshModelingToolset/Source/ModelingComponents/Private/SimpleDynamicMeshSceneProxy.h
--- a/Engine/Plugins/Experimental/MeshModelingToolset/Source/ModelingComponents/Private/SimpleDynamicMeshSceneProxy.h
+++ b/Engine/Plugins/Experimental/MeshModelingToolset/Source/ModelingComponents/Private/SimpleDynamicMeshSceneProxy.h
@@ -108,6 +108,7 @@ public:
 		// find suitable overlays
 		FDynamicMeshUVOverlay* UVOverlay = Mesh->Attributes()->PrimaryUV();
 		FDynamicMeshNormalOverlay* NormalOverlay = Mesh->Attributes()->PrimaryNormals();
+		FDynamicMeshColorOverlay* ColorOverlay = Mesh->Attributes()->PrimaryColors();
 
 		TFunction<void(int, int, int, FVector3f&, FVector3f&)> TangentsFunc = nullptr;
 		const UE::Geometry::FMeshTangentsf* Tangents = ParentComponent->GetTangents();
@@ -130,7 +131,7 @@ public:
 				RenderBuffers->Triangles = Triangles;
 				InitializeBuffersFromOverlays(RenderBuffers, Mesh,
 					Triangles.Num(), Triangles,
-					UVOverlay, NormalOverlay, TangentsFunc);
+					UVOverlay, NormalOverlay, ColorOverlay, TangentsFunc);
 
 				ENQUEUE_RENDER_COMMAND(FSimpleDynamicMeshSceneProxyInitializeFromDecomposition)(
 					[RenderBuffers](FRHICommandListImmediate& RHICmdList)
@@ -158,6 +159,7 @@ public:
 		// find suitable overlays
 		TArray<const FDynamicMeshUVOverlay*> UVOverlays;
 		const FDynamicMeshNormalOverlay* NormalOverlay = nullptr;
+		const FDynamicMeshColorOverlay* ColorOverlay = nullptr;
 		if (Mesh->HasAttributes())
 		{
 			const FDynamicMeshAttributeSet* Attributes = Mesh->Attributes();
@@ -167,6 +169,7 @@ public:
 			{
 				UVOverlays[k] = Attributes->GetUVLayer(k);
 			}
+			ColorOverlay = Attributes->PrimaryColors();
 		}
 
 		TFunction<void(int, int, int, FVector3f&, FVector3f&)> TangentsFunc = nullptr;
@@ -181,7 +184,7 @@ public:
 
 		InitializeBuffersFromOverlays(RenderBuffers, Mesh,
 			Mesh->TriangleCount(), Mesh->TriangleIndicesItr(),
-			UVOverlays, NormalOverlay, TangentsFunc);
+			UVOverlays, NormalOverlay, ColorOverlay, TangentsFunc);
 
 		ENQUEUE_RENDER_COMMAND(FSimpleDynamicMeshSceneProxyInitializeSingle)(
 			[RenderBuffers](FRHICommandListImmediate& RHICmdList)
@@ -205,6 +208,7 @@ public:
 		// find suitable overlays
 		const FDynamicMeshMaterialAttribute* MaterialID = Attributes->GetMaterialID();
 		const FDynamicMeshNormalOverlay* NormalOverlay = Mesh->Attributes()->PrimaryNormals();
+		const FDynamicMeshColorOverlay* ColorOverlay = Mesh->Attributes()->PrimaryColors();
 
 		TArray<const FDynamicMeshUVOverlay*> UVOverlays;
 		UVOverlays.SetNum(Attributes->NumUVLayers());
@@ -271,7 +275,7 @@ public:
 
 				InitializeBuffersFromOverlays(RenderBuffers, Mesh,
 					Triangles.Num(), Triangles,
-					UVOverlays, NormalOverlay, TangentsFunc);
+					UVOverlays, NormalOverlay, ColorOverlay, TangentsFunc);
 
 				RenderBuffers->Triangles = Triangles;
 
@@ -383,7 +387,13 @@ public:
 			check(Mesh->HasAttributes());
 			UVOVerlay = Mesh->Attributes()->PrimaryUV();
 		}
-
+		FDynamicMeshColorOverlay* ColorOverlay = nullptr;
+		if (bColors)
+		{
+			check(Mesh->HasAttributes());
+			ColorOverlay = Mesh->Attributes()->PrimaryColors();
+					
+		}
 		if (bIsSingleBuffer)
 		{
 			check(RenderBufferSets.Num() == 1);
@@ -392,7 +402,7 @@ public:
 			{
 				UpdateVertexBuffersFromOverlays(Buffers, Mesh,
 					Mesh->TriangleCount(), Mesh->TriangleIndicesItr(),
-					NormalOverlay, TangentsFunc,
+					NormalOverlay, ColorOverlay, TangentsFunc,
 					bPositions, bNormals, bColors);
 			}
 			if (bUVs)
@@ -422,7 +432,7 @@ public:
 				{
 					UpdateVertexBuffersFromOverlays(Buffers, Mesh,
 						Buffers->Triangles->Num(), Buffers->Triangles.GetValue(),
-						NormalOverlay, TangentsFunc,
+						NormalOverlay, ColorOverlay, TangentsFunc,
 						bPositions, bNormals, bColors);
 				}
 				if (bUVs)
@@ -484,7 +494,13 @@ public:
 			check(Mesh->HasAttributes());
 			UVOVerlay = Mesh->Attributes()->PrimaryUV();
 		}
+		FDynamicMeshColorOverlay* ColorOverlay = nullptr;
+		if (bColors)
+		{
+			check(Mesh->HasAttributes());
+			ColorOverlay = Mesh->Attributes()->PrimaryColors();
 
+		}
 		ParallelFor(WhichBuffers.Num(), [&](int idx)
 		{
 			int32 BufferIndex = WhichBuffers[idx];
@@ -498,7 +514,7 @@ public:
 			{
 				UpdateVertexBuffersFromOverlays(Buffers, Mesh,
 					Buffers->Triangles->Num(), Buffers->Triangles.GetValue(),
-					NormalOverlay, TangentsFunc,
+					NormalOverlay, ColorOverlay, TangentsFunc,
 					bPositions, bNormals, bColors);
 			}
 			if (bUVs)
diff --git a/Engine/Plugins/Experimental/MeshModelingToolset/Source/ModelingComponents/Public/Drawing/MeshElementsVisualizer.h b/Engine/Plugins/Experimental/MeshModelingToolset/Source/ModelingComponents/Public/Drawing/MeshElementsVisualizer.h
--- a/Engine/Plugins/Experimental/MeshModelingToolset/Source/ModelingComponents/Public/Drawing/MeshElementsVisualizer.h
+++ b/Engine/Plugins/Experimental/MeshModelingToolset/Source/ModelingComponents/Public/Drawing/MeshElementsVisualizer.h
@@ -39,6 +39,10 @@ public:
 	UPROPERTY(EditAnywhere, Category = MeshElementVisualization)
 	bool bShowNormalSeams = true;
 
+	/** Should mesh color seam edges be shown */
+	UPROPERTY(EditAnywhere, Category = MeshElementVisualization)
+	bool bShowColorSeams = true;
+
 	/** multiplier on edge thicknesses */
 	UPROPERTY(EditAnywhere, Category = MeshElementVisualization, AdvancedDisplay, meta = (UIMin = 0.1, UIMax = 10.0))
 	float ThicknessScale = 1.0;
@@ -59,6 +63,10 @@ public:
 	UPROPERTY(EditAnywhere, Category = MeshElementVisualization, AdvancedDisplay)
 	FColor NormalSeamColor = FColor(128, 128, 240);
 
+	/** Color of mesh color seam edges */
+	UPROPERTY(EditAnywhere, Category = MeshElementVisualization, AdvancedDisplay)
+	FColor ColorSeamColor = FColor(46, 204, 113);
+
 	/** depth bias used to slightly shift depth of lines */
 	UPROPERTY(EditAnywhere, Category = MeshElementVisualization, AdvancedDisplay, meta = (UIMin = -2.0, UIMax = 2.0))
 	float DepthBias = 0.2;
@@ -67,7 +75,7 @@ public:
 
 /**
  * UMeshElementsVisualizer is a subclass of UPreviewGeometry that displays mesh elements.
- * Currently supports wireframe, boundary edges, UV seams, and Normal seams.
+ * Currently supports wireframe, boundary edges, UV seams, Normal seams and Color seams.
  *
  * UMeshElementsVisualizer initializes an instance of UMeshElementsVisualizerProperties
  * as its .Settings value, and will watch for changes in these properties.
@@ -107,7 +115,7 @@ public:
 	UPROPERTY()
 	UMeshElementsVisualizerProperties* Settings;
 
-	/** Mesh Wireframe component, draws wireframe, boundaries, UV seams, normal seams */
+	/** Mesh Wireframe component, draws wireframe, boundaries, UV seams, normal seams, color seams */
 	UPROPERTY()
 	UMeshWireframeComponent* WireframeComponent;
 
diff --git a/Engine/Plugins/Experimental/MeshModelingToolset/Source/ModelingComponents/Public/Drawing/MeshWireframeComponent.h b/Engine/Plugins/Experimental/MeshModelingToolset/Source/ModelingComponents/Public/Drawing/MeshWireframeComponent.h
--- a/Engine/Plugins/Experimental/MeshModelingToolset/Source/ModelingComponents/Public/Drawing/MeshWireframeComponent.h
+++ b/Engine/Plugins/Experimental/MeshModelingToolset/Source/ModelingComponents/Public/Drawing/MeshWireframeComponent.h
@@ -29,7 +29,8 @@ public:
 		MeshBoundary = 1 << 2,
 		UVSeam = 1 << 3,
 		NormalSeam = 1 << 4,
-		GroupBoundary = 1 << 5
+		GroupBoundary = 1 << 5,
+		ColorSeam = 1 << 6
 	};
 
 public:
@@ -80,6 +81,7 @@ public:
  *    - boundary edges
  *    - UV seam edges
  *    - Normal seam edges
+ *    - Color seam edges
  *
  * Client must provide a IMeshWireframeSourceProvider implementation that provides the 
  * edge set, vertices, and edge type
@@ -170,6 +172,18 @@ public:
 	UPROPERTY(EditAnywhere, Category = MeshWireframe)
 	float NormalSeamThickness = 2.0f;
 
+
+	// color seam properties
+
+	UPROPERTY(EditAnywhere, Category = MeshWireframe)
+	bool bEnableColorSeams = true;
+
+	UPROPERTY(EditAnywhere, Category = MeshWireframe)
+	FColor ColorSeamColor = FColor(46, 204, 113);
+
+	UPROPERTY(EditAnywhere, Category = MeshWireframe)
+	float ColorSeamThickness = 2.0f;
+
 private:
 
 	//~ Begin UPrimitiveComponent Interface.
