GeometryProcessing: TVector4 operator!= now computes correct result
#rnx

diff --git a/Engine/Plugins/Experimental/GeometryProcessing/Source/GeometricObjects/Public/VectorTypes.h b/Engine/Plugins/Experimental/GeometryProcessing/Source/GeometricObjects/Public/VectorTypes.h
--- a/Engine/Plugins/Experimental/GeometryProcessing/Source/GeometricObjects/Public/VectorTypes.h
+++ b/Engine/Plugins/Experimental/GeometryProcessing/Source/GeometricObjects/Public/VectorTypes.h
@@ -1082,7 +1082,7 @@ struct TVector4
 
 	constexpr bool operator!=(const TVector4<T>& Other) const
 	{
-		return X != Other.X || Y != Other.Y || Z != Other.Z || Z == Other.Z;
+		return X != Other.X || Y != Other.Y || Z != Other.Z || W != Other.W;
 	}
 };
 
