Geometry Processing:  MeshGroupPaintTool.  BaseDynamicMeshSceneProxy Fix bug with visibility introduced by change to using Color Overlay.

diff --git a/Engine/Plugins/Experimental/GeometryProcessing/Source/DynamicMesh/Private/DynamicMeshAttributeSet.cpp b/Engine/Plugins/Experimental/GeometryProcessing/Source/DynamicMesh/Private/DynamicMeshAttributeSet.cpp
--- a/Engine/Plugins/Experimental/GeometryProcessing/Source/DynamicMesh/Private/DynamicMeshAttributeSet.cpp
+++ b/Engine/Plugins/Experimental/GeometryProcessing/Source/DynamicMesh/Private/DynamicMeshAttributeSet.cpp
@@ -25,6 +25,10 @@ void FDynamicMeshAttributeSet::Copy(const FDynamicMeshAttributeSet& Copy)
 		EnablePrimaryColors();
 		ColorLayer->Copy(*(Copy.ColorLayer));
 	}
+	else
+	{
+		DisablePrimaryColors();
+	}
 	if (Copy.MaterialIDAttrib)
 	{
 		EnableMaterialID();
@@ -99,6 +103,10 @@ void FDynamicMeshAttributeSet::CompactCopy(const FCompactMaps& CompactMaps, cons
 		EnablePrimaryColors();
 		ColorLayer->CompactCopy(CompactMaps, *(Copy.ColorLayer));
 	}
+	else
+	{
+		DisablePrimaryColors();
+	}
 
 	if (Copy.MaterialIDAttrib)
 	{
diff --git a/Engine/Plugins/Experimental/GeometryProcessing/Source/DynamicMesh/Public/DynamicMeshAttributeSet.h b/Engine/Plugins/Experimental/GeometryProcessing/Source/DynamicMesh/Public/DynamicMeshAttributeSet.h
--- a/Engine/Plugins/Experimental/GeometryProcessing/Source/DynamicMesh/Public/DynamicMeshAttributeSet.h
+++ b/Engine/Plugins/Experimental/GeometryProcessing/Source/DynamicMesh/Public/DynamicMeshAttributeSet.h
@@ -405,6 +405,10 @@ protected:
 			bValid = GetUVLayer(UVLayerIndex)->CheckValidity(bAllowNonmanifold, FailMode) && bValid;
 		}
 		bValid = PrimaryNormals()->CheckValidity(bAllowNonmanifold, FailMode) && bValid;
+		if (ColorLayer)
+		{
+			bValid = ColorLayer->CheckValidity(bAllowNonmanifold, FailMode) && bValid;
+		}
 		if (MaterialIDAttrib)
 		{
 			bValid = MaterialIDAttrib->CheckValidity(bAllowNonmanifold, FailMode) && bValid;
diff --git a/Engine/Plugins/Experimental/MeshModelingToolset/Source/ModelingComponents/Private/BaseDynamicMeshSceneProxy.h b/Engine/Plugins/Experimental/MeshModelingToolset/Source/ModelingComponents/Private/BaseDynamicMeshSceneProxy.h
--- a/Engine/Plugins/Experimental/MeshModelingToolset/Source/ModelingComponents/Private/BaseDynamicMeshSceneProxy.h
+++ b/Engine/Plugins/Experimental/MeshModelingToolset/Source/ModelingComponents/Private/BaseDynamicMeshSceneProxy.h
@@ -695,7 +695,7 @@ public:
 				{
 					FColor VertexFColor = (bHaveColors  && TriColor[j] != FDynamicMesh3::InvalidID) ?
 						FColor(ColorOverlay->GetElement(TriColor[j])) : UniformTriColor;
-
+					RenderBuffers->ColorVertexBuffer.VertexColor(VertIdx) = VertexFColor;
 				}
 
 				VertIdx++;
