GeometryProcessing: Fix clear color for bent normal builder.

#rnx

diff --git a/Engine/Plugins/Experimental/GeometryProcessing/Source/DynamicMesh/Private/Sampling/MeshOcclusionMapBaker.cpp b/Engine/Plugins/Experimental/GeometryProcessing/Source/DynamicMesh/Private/Sampling/MeshOcclusionMapBaker.cpp
--- a/Engine/Plugins/Experimental/GeometryProcessing/Source/DynamicMesh/Private/Sampling/MeshOcclusionMapBaker.cpp
+++ b/Engine/Plugins/Experimental/GeometryProcessing/Source/DynamicMesh/Private/Sampling/MeshOcclusionMapBaker.cpp
@@ -194,7 +194,17 @@ void FMeshOcclusionMapBaker::Bake()
 	{
 		NormalBuilder = MakeUnique<TImageBuilder<FVector3f>>();
 		NormalBuilder->SetDimensions(BakeCache->GetDimensions());
-		NormalBuilder->Clear(FVector3f::Zero());
+		FVector3f DefaultNormal = FVector3f::UnitZ();
+		switch (NormalSpace)
+		{
+		case ESpace::Tangent:
+			break;
+		case ESpace::Object:
+			DefaultNormal = FVector3f::Zero();
+			break;
+		}
+		FVector3f DefaultMapNormal = (DefaultNormal + FVector3f::One()) * 0.5f;
+		NormalBuilder->Clear(DefaultMapNormal);
 	}
 
 	BakeCache->EvaluateSamples([&](const FVector2i& Coords, const FMeshImageBakingCache::FCorrespondenceSample& Sample)
