MeshModelingTools: Adding current version of UV editor.

#rnx

diff --git a/Engine/Plugins/Experimental/GeometryProcessing/Source/MeshConversion/Private/MeshDescriptionUVsToDynamicMesh.cpp b/Engine/Plugins/Experimental/GeometryProcessing/Source/MeshConversion/Private/MeshDescriptionUVsToDynamicMesh.cpp
new file mode 100644
--- /dev/null
+++ b/Engine/Plugins/Experimental/GeometryProcessing/Source/MeshConversion/Private/MeshDescriptionUVsToDynamicMesh.cpp
@@ -0,0 +1,351 @@
+// Copyright Epic Games, Inc. All Rights Reserved.
+
+#include "MeshDescriptionUVsToDynamicMesh.h"
+
+#include "DynamicMesh3.h"
+#include "MeshDescription.h"
+#include "StaticMeshAttributes.h"
+
+#include "ExplicitUseGeometryMathTypes.h"		// using UE::Geometry::(math types)
+using namespace UE::Geometry;
+
+// TODO: Need to test both conversions some more with meshes that only have instanced UVs.
+// Also, we should probably keep maps for elements we split due to manifoldness issues, etc,
+// rather than using all those heuristics to try to weld things back together on the conversion
+// back.
+
+namespace MeshDescriptionUVsToDynamicMeshLocals
+{
+	bool ShouldConversionUseSharedUVs(const FMeshDescription* MeshDescription, int32 LayerIndex)
+	{
+		const int32 NumSharedUVLayers = MeshDescription->GetNumUVElementChannels();
+		return NumSharedUVLayers > LayerIndex && MeshDescription->UVs(LayerIndex).GetArraySize() != 0;
+	}
+
+	// This class is copied from MeshDescriptionToDynamicMesh.cpp. We could consider putting it
+	// somewhere common.
+	struct FVertexUV
+	{
+		int vid;
+		float x;
+		float y;
+		bool operator==(const FVertexUV& o) const
+		{
+			return vid == o.vid && x == o.x && y == o.y;
+		}
+	};
+	FORCEINLINE uint32 GetTypeHash(const FVertexUV& Vector)
+	{
+		// ugh copied from FVector clearly should not be using CRC for hash!!
+		return FCrc::MemCrc32(&Vector, sizeof(Vector));
+	}
+
+	// This class is modeled on FUVWelder in MeshDescriptionToDynamicMesh.cpp,
+	// except it adds to the dynamic mesh vertices instead of the mesh UV's.
+	// Note: we could consider welding UV's differently by sorting all the verts and welding in one pass.
+	class FUVMeshWelder
+	{
+	public:
+
+		FUVMeshWelder(FDynamicMesh3* MeshToAddTo, double ScaleFactorIn)
+			: Mesh(MeshToAddTo)
+			, ScaleFactor(ScaleFactorIn)
+		{
+			check(MeshToAddTo && ScaleFactorIn != 0);
+		}
+
+		int FindOrAddUnique(const FVector2D& UV, int VertexID)
+		{
+			FVertexUV VertUV = { VertexID, UV.X, UV.Y };
+
+			const int32* FoundIndex = UniqueVertexUVs.Find(VertUV);
+			if (FoundIndex != nullptr)
+			{
+				return *FoundIndex;
+			}
+
+			int32 NewIndex = Mesh->AppendVertex(FVector3d(UV.Y * ScaleFactor, UV.X * ScaleFactor, 0));
+			UniqueVertexUVs.Add(VertUV, NewIndex);
+			return NewIndex;
+		}
+
+	protected:
+		FDynamicMesh3* Mesh;
+		double ScaleFactor;
+
+		TMap<FVertexUV, int> UniqueVertexUVs;
+	};
+}
+
+
+int32 FMeshDescriptionUVsToDynamicMesh::GetNumUVLayers(const FMeshDescription* MeshDescription) const
+{
+	FStaticMeshConstAttributes Attributes(*MeshDescription);
+	return Attributes.GetVertexInstanceUVs().GetNumChannels();
+}
+
+TSharedPtr<FDynamicMesh3> FMeshDescriptionUVsToDynamicMesh::GetUVMesh(const FMeshDescription* MeshDescription)
+{
+
+	using namespace MeshDescriptionUVsToDynamicMeshLocals;
+
+	check(ScaleFactor != 0);
+
+	bool bUseSharedUVs = ShouldConversionUseSharedUVs(MeshDescription, UVLayerIndex);
+
+	// The output that we'll return
+	TSharedPtr<FDynamicMesh3> MeshOut = MakeShared<FDynamicMesh3>();
+
+	if (bUseSharedUVs)
+	{
+		// Create vertices for the UV elements.
+		const FUVArray& SharedUVs = MeshDescription->UVs(UVLayerIndex);
+		TUVAttributesRef<const FVector2D> SharedUVCoordinates = SharedUVs.GetAttributes().GetAttributesRef<FVector2D>(MeshAttribute::UV::UVCoordinate);
+
+		MeshOut->BeginUnsafeVerticesInsert();
+		for (FUVID UVID : SharedUVs.GetElementIDs())
+		{
+			const FVector2D UV = SharedUVCoordinates[UVID];
+			MeshOut->InsertVertex(UVID.GetValue(), FVector3d(UV.Y * ScaleFactor, UV.X * ScaleFactor, 0), true);
+		}
+		MeshOut->EndUnsafeTrianglesInsert();
+	}
+
+	// Instance UV's and a welder are used if bUseSharedUVs is false, or if it is true but some 
+	// verts are not initialized.
+	FStaticMeshConstAttributes Attributes(*MeshDescription);
+	TVertexInstanceAttributesConstRef<FVector2D> InstanceUVs = Attributes.GetVertexInstanceUVs();
+	FUVMeshWelder UVMeshWelder(MeshOut.Get(), ScaleFactor);
+
+	// Go through triangles and build our mesh
+	MeshOut->BeginUnsafeTrianglesInsert();
+	for (const FTriangleID TriangleID : MeshDescription->Triangles().GetElementIDs())
+	{
+		// The Tid must stay the same in the dynamic mesh for us to be able to bake things back
+		int32 Tid = TriangleID.GetValue();
+
+		// Will be populated with Vids in the dynamic mesh
+		FIndex3i TriangleToInsert;
+
+		// Used on instance UVs to map them to shared UV elements (vertices) in the dynamic mesh. 
+		auto MakeWeldedUVsForTri = [this, MeshDescription, &UVMeshWelder, &InstanceUVs](const FTriangleID& TriangleIDIn, FIndex3i& TriangleToInsertOut) {
+			
+			TArrayView<const FVertexID> SourceTriangleVids = MeshDescription->GetTriangleVertices(TriangleIDIn);
+			TArrayView<const FVertexInstanceID> SourceTriangleInstanceIds = MeshDescription->GetTriangleVertexInstances(TriangleIDIn);
+			for (int32 i = 0; i < 3; ++i)
+			{
+				FVector2D UV = InstanceUVs.Get(SourceTriangleInstanceIds[i], UVLayerIndex);
+				TriangleToInsertOut[i] = UVMeshWelder.FindOrAddUnique(UV, SourceTriangleVids[i].GetValue());
+			}
+		};
+
+		if (!bUseSharedUVs)
+		{
+			// If we're not using shared UV's, each vertex instance has its own UV element. We weld them per vertex
+			// if the UV's match.
+			MakeWeldedUVsForTri(TriangleID, TriangleToInsert);
+		}
+		else
+		{
+			TArrayView<FUVID> TriUVIndices = MeshDescription->GetTriangleUVIndices(TriangleID, UVLayerIndex);
+
+			// If some of the mesh description triangles were missing shared UVs,
+			// use the per-vertex UVs on the mesh description and weld these.
+			// NB: This can happen when multiple meshes with different UV layer count are "combined" during import
+			// NB: these mesh description per-vertex UVs always exist, but may default to zero
+			FUVID InvalidUVID(INDEX_NONE);
+			if (TriUVIndices[0] == InvalidUVID || TriUVIndices[1] == InvalidUVID || TriUVIndices[2] == InvalidUVID)
+			{
+				MakeWeldedUVsForTri(TriangleID, TriangleToInsert);
+			}
+			else
+			{
+				// Triangle seems ok so far. We already added vertices for it previously, so set up the tri to add.
+				// We'll be doing more checks for degeneracy and such before inserting.
+				for (int32 i = 0; i < 3; ++i)
+				{
+					TriangleToInsert[i] = TriUVIndices[i].GetValue();
+				}
+			}
+		}
+
+		// Our splitter, returns the new vertex id
+		auto SplitDynamicMeshVertex = [MeshOut](int32 Vid) {
+			check(MeshOut->IsVertex(Vid));
+			const FVector3d Position = MeshOut->GetVertex(Vid);
+			return MeshOut->AppendVertex(Position);
+		};
+
+		// Split vertices if our triangle turned out to be degenerate.
+		if (TriangleToInsert[0] == TriangleToInsert[1] || TriangleToInsert[0] == TriangleToInsert[2])
+		{
+			TriangleToInsert[0] = SplitDynamicMeshVertex(TriangleToInsert[0]);
+		}
+		if (TriangleToInsert[1] == TriangleToInsert[2])
+		{
+			TriangleToInsert[1] = SplitDynamicMeshVertex(TriangleToInsert[1]);
+		}
+
+		// Attempt triangle insertion
+		EMeshResult Result = MeshOut->InsertTriangle(Tid, TriangleToInsert, 0, true);
+
+		// Duplicate vertices if insertion failed due to a non-manifold edge
+		if (Result == EMeshResult::Failed_WouldCreateNonmanifoldEdge)
+		{
+			bool bDuplicate[3] = { false, false, false };
+
+			int e0 = MeshOut->FindEdge(TriangleToInsert[0], TriangleToInsert[1]);
+			int e1 = MeshOut->FindEdge(TriangleToInsert[1], TriangleToInsert[2]);
+			int e2 = MeshOut->FindEdge(TriangleToInsert[2], TriangleToInsert[0]);
+
+			// determine which verts need to be duplicated by seeing whether existing edges already
+			// have two triangles.
+			if (e0 != FDynamicMesh3::InvalidID && MeshOut->IsBoundaryEdge(e0) == false)
+			{
+				bDuplicate[0] = true;
+				bDuplicate[1] = true;
+			}
+			if (e1 != FDynamicMesh3::InvalidID && MeshOut->IsBoundaryEdge(e1) == false)
+			{
+				bDuplicate[1] = true;
+				bDuplicate[2] = true;
+			}
+			if (e2 != FDynamicMesh3::InvalidID && MeshOut->IsBoundaryEdge(e2) == false)
+			{
+				bDuplicate[2] = true;
+				bDuplicate[0] = true;
+			}
+
+			// Do the actual duplication
+			for (int32 i = 0; i < 3; ++i)
+			{
+				if (bDuplicate[i])
+				{
+					TriangleToInsert[i] = SplitDynamicMeshVertex(TriangleToInsert[i]);
+				}
+			}
+
+			// Reinsert the triangle
+			Result = MeshOut->InsertTriangle(Tid, TriangleToInsert, 0, true);
+		}
+
+		check(Result == EMeshResult::Ok || Result == EMeshResult::Failed_TriangleAlreadyExists);
+	}//end iterating through triangles
+
+	MeshOut->EndUnsafeTrianglesInsert();
+
+	return MeshOut;
+}
+
+void FMeshDescriptionUVsToDynamicMesh::BakeBackUVsFromUVMesh(const FDynamicMesh3* DynamicMesh, FMeshDescription* MeshDescription) const
+{
+	using namespace MeshDescriptionUVsToDynamicMeshLocals;
+
+	check(ScaleFactor != 0);
+	check(DynamicMesh->TriangleCount() == MeshDescription->Triangles().Num())
+
+	// Regardless of whether we have shared UVs or not, we are supposed to update the instance UVs.
+	// This is straightforward, since any welding or splitting does not change anything.
+	FStaticMeshAttributes Attributes(*MeshDescription);
+	TVertexInstanceAttributesRef<FVector2D> InstanceUVs = Attributes.GetVertexInstanceUVs();
+	for (int32 Tid : DynamicMesh->TriangleIndicesItr())
+	{
+		TArrayView<const FVertexInstanceID> MeshDescriptionTriInstanceIds = MeshDescription->GetTriangleVertexInstances(Tid);
+		FIndex3i DynamicMeshTriVids = DynamicMesh->GetTriangle(Tid);
+		for (int i = 0; i < 3; ++i)
+		{
+			FVector3d VertPosition = DynamicMesh->GetVertex(DynamicMeshTriVids[i]);
+			InstanceUVs.Set(MeshDescriptionTriInstanceIds[i], UVLayerIndex, FVector2D(VertPosition.Y / ScaleFactor, VertPosition.X / ScaleFactor));
+		}
+	}
+
+	if (ShouldConversionUseSharedUVs(MeshDescription, UVLayerIndex))
+	{
+		// For shared UVs we may have to deal with splitting and welding, which makes things a bit more
+		// complicated.
+
+		// VidToUVElement is able to map vids back to newly created or merged UV elements. It
+		// also lets us know if we've dealt with a specific vid before.
+		TArray<int32> VidToUVElement;
+		VidToUVElement.SetNum(DynamicMesh->MaxVertexID());
+		for (int32& Value : VidToUVElement)
+		{
+			Value = INDEX_NONE;
+		}
+
+		FUVArray& SharedUvs = MeshDescription->UVs(UVLayerIndex);
+		TUVAttributesRef<FVector2D> SharedUVCoordinates = SharedUvs.GetAttributes().GetAttributesRef<FVector2D>(MeshAttribute::UV::UVCoordinate);
+
+		MeshDescription->SuspendUVIndexing();
+
+		for (int32 Tid : DynamicMesh->TriangleIndicesItr())
+		{
+			// Look at the UV ID's this triangle had, and the vertex ID's it now has. Vertex ID's correspond
+			// to UV element ID's except in cases where a new UV element needs to be created.
+			TArrayView<FUVID> TriUVIDs = MeshDescription->GetTriangleUVIndices(FTriangleID(Tid), UVLayerIndex);
+			FIndex3i TriVids = DynamicMesh->GetTriangle(Tid);
+
+			bool bTriUVIDsChanged = false;
+
+			// Examine the triangle vertices
+			for (int i = 0; i < 3; ++i)
+			{
+				// By examining the current and expected Vids, we can tell whether the UV connectivity changed
+				int32 Vid = TriVids[i];
+				int32 ExpectedVid = TriUVIDs[i].GetValue();
+
+				// First check whether we've dealt with this vert before
+				if (VidToUVElement[Vid] != INDEX_NONE)
+				{
+					// We still need to update triangle connectivity unless it corresponds to the existing expected vert
+					// (this includes verts that we split for conversion only, since those will map to the expected vert)
+					if (VidToUVElement[Vid] != ExpectedVid)
+					{
+						TriUVIDs[i] = VidToUVElement[Vid];
+						bTriUVIDsChanged = true;
+					}
+					continue;
+				}
+
+				// If we haven't dealt with this vert, we will need to update its entry in SharedUVCoordinates.
+
+				FVector3d VertPosition = DynamicMesh->GetVertex(Vid);
+				FVector2D NewUVValue = FVector2D(VertPosition.Y / ScaleFactor, VertPosition.X / ScaleFactor);
+				
+				// See if the vert is the same one that we were expecting (including one that we may have split off just for
+				// the conversion, in which case the location will match).
+				if (Vid == ExpectedVid ||
+					(DynamicMesh->IsVertex(ExpectedVid) && DynamicMesh->GetVertex(ExpectedVid) == VertPosition))
+				{
+					VidToUVElement[Vid] = ExpectedVid;
+					SharedUVCoordinates.Set(TriUVIDs[i], NewUVValue);
+					// No need to update TriUVIDs
+					continue;
+				}
+
+				// See if the vert corresponds to an element that exists in SharedUVCoordinates
+				if (MeshDescription->IsUVValid(FUVID(Vid), UVLayerIndex))
+				{
+					VidToUVElement[Vid] = Vid;
+					SharedUVCoordinates.Set(FUVID(Vid), NewUVValue);
+					TriUVIDs[i] = FUVID(Vid);
+					bTriUVIDsChanged = true;
+					continue;
+				}
+
+				// Otherwise, this is a new element we need to create, and we need to update TriUVIDs
+				TriUVIDs[i] = MeshDescription->CreateUV(UVLayerIndex);
+				VidToUVElement[Vid] = TriUVIDs[i].GetValue();
+				SharedUVCoordinates.Set(TriUVIDs[i], NewUVValue);
+				bTriUVIDsChanged = true;
+			}//end dealing with tri verts
+
+			if (bTriUVIDsChanged)
+			{
+				MeshDescription->SetTriangleUVIndices(Tid, TriUVIDs, UVLayerIndex);
+			}
+		}//end iterating through triangles
+
+		MeshDescription->ResumeUVIndexing();
+	}
+}
\ No newline at end of file
diff --git a/Engine/Plugins/Experimental/GeometryProcessing/Source/MeshConversion/Public/MeshDescriptionUVsToDynamicMesh.h b/Engine/Plugins/Experimental/GeometryProcessing/Source/MeshConversion/Public/MeshDescriptionUVsToDynamicMesh.h
new file mode 100644
--- /dev/null
+++ b/Engine/Plugins/Experimental/GeometryProcessing/Source/MeshConversion/Public/MeshDescriptionUVsToDynamicMesh.h
@@ -0,0 +1,46 @@
+// Copyright Epic Games, Inc. All Rights Reserved.
+
+#pragma once
+
+#include "CoreMinimal.h"
+#include "GeometryBase.h"
+
+struct FMeshDescription;
+
+namespace UE {
+namespace Geometry {
+
+class FDynamicMesh3;
+
+/**
+ * Converter that gives a specific UV layer of a provided mesh description and turns it into a flat
+ * FDynamicMesh, with vertices corresponding to UV elements, and their positions corresponding to
+ * those elements' UV coordinate values. The triangle id's in the dynamic mesh will correspond to
+ * original triangles of the mesh.
+ * 
+ * The converter can also bake back such a dynamic mesh to an existing mesh description UV layer
+ * after the dynamic mesh has been modified. The important part is that the triangle IDs continue
+ * to correspond to each other. This means that while welding/splitting vertices in the dynamic
+ * mesh is ok, adding or removing triangles is not.
+ * 
+ * The converter prefers shared UVs when they are available, though it will populate instanced
+ * UVs as well on the way back.
+ */
+class MESHCONVERSION_API FMeshDescriptionUVsToDynamicMesh
+{
+public:
+	int32 UVLayerIndex = 0;
+
+	// The [0,1] UV range will be scaled to [0, ScaleFactor]
+	double ScaleFactor = 1000;
+
+	int32 GetNumUVLayers(const FMeshDescription* InputMeshDescription) const;
+
+	// TODO: The conversion function forward to a dynamic mesh should be made to produce maps 
+	// similar to MeshDescriptionToDynamicMesh, which is why this is not a const function.
+	TSharedPtr<FDynamicMesh3> GetUVMesh(const FMeshDescription* InputMeshDescription);
+
+	void BakeBackUVsFromUVMesh(const FDynamicMesh3* DynamicUVMesh, FMeshDescription* OutputMeshDescription) const;
+};
+
+}}//end namespace UE::Geometry
\ No newline at end of file
diff --git a/Engine/Plugins/Experimental/MeshModelingToolset/Source/ModelingComponents/Private/Drawing/MeshElementsVisualizer.cpp b/Engine/Plugins/Experimental/MeshModelingToolset/Source/ModelingComponents/Private/Drawing/MeshElementsVisualizer.cpp
--- a/Engine/Plugins/Experimental/MeshModelingToolset/Source/ModelingComponents/Private/Drawing/MeshElementsVisualizer.cpp
+++ b/Engine/Plugins/Experimental/MeshModelingToolset/Source/ModelingComponents/Private/Drawing/MeshElementsVisualizer.cpp
@@ -115,6 +115,9 @@ void UMeshElementsVisualizer::OnCreated()
 	Settings->WatchProperty(Settings->bShowColorSeams, [this](bool){ bSettingsModified = true; });
 	Settings->WatchProperty(Settings->ThicknessScale, [this](float) { bSettingsModified = true; });
 	Settings->WatchProperty(Settings->DepthBias, [this](float) { bSettingsModified = true; });
+	Settings->WatchProperty(Settings->bAdjustDepthBiasUsingMeshSize, [this](bool) {
+		UpdateLineDepthBiasScale(); // A little expensive, so only want to do this when necessary.
+		bSettingsModified = true; }); // Still mark this so we mark the wireframe as dirty.
 	Settings->WatchProperty(Settings->WireframeColor, [this](FColor) { bSettingsModified = true; });
 	Settings->WatchProperty(Settings->BoundaryEdgeColor, [this](FColor) { bSettingsModified = true; });
 	Settings->WatchProperty(Settings->UVSeamColor, [this](FColor) { bSettingsModified = true; });
@@ -131,18 +134,28 @@ void UMeshElementsVisualizer::OnCreated()
 
 void UMeshElementsVisualizer::SetMeshAccessFunction(TUniqueFunction<const FDynamicMesh3* (void)>&& MeshAccessFunctionIn)
 {
-	const FDynamicMesh3* Mesh = MeshAccessFunctionIn();
-	FAxisAlignedBox3d Bounds = Mesh->GetBounds();
-	float DepthBiasDim = Bounds.DiagonalLength() * 0.01;
-	WireframeComponent->LineDepthBiasSizeScale = DepthBiasDim;
-
 	WireframeSourceProvider = MakeShared<FDynamicMeshWireframeSourceProvider>(MoveTemp(MeshAccessFunctionIn));
 
+	UpdateLineDepthBiasScale();
+
 	WireframeComponent->SetWireframeSourceProvider(WireframeSourceProvider);
 }
 
-
-
+void UMeshElementsVisualizer::UpdateLineDepthBiasScale()
+{
+	if (Settings->bAdjustDepthBiasUsingMeshSize && WireframeSourceProvider)
+	{
+		WireframeSourceProvider->AccessMesh(
+			[this](const IMeshWireframeSource& WireframeSource) {
+				// Scale is 0.01 of diameter of bounding sphere (diagonal of box)
+				WireframeComponent->LineDepthBiasSizeScale = WireframeSource.GetBounds().SphereRadius * 2 * 0.01;
+			});
+	}
+	else
+	{
+		WireframeComponent->LineDepthBiasSizeScale = 1;
+	}
+}
 
 void UMeshElementsVisualizer::OnTick(float DeltaTime)
 {
diff --git a/Engine/Plugins/Experimental/MeshModelingToolset/Source/ModelingComponents/Private/Selection/MeshSelectionMechanic.cpp b/Engine/Plugins/Experimental/MeshModelingToolset/Source/ModelingComponents/Private/Selection/MeshSelectionMechanic.cpp
new file mode 100644
--- /dev/null
+++ b/Engine/Plugins/Experimental/MeshModelingToolset/Source/ModelingComponents/Private/Selection/MeshSelectionMechanic.cpp
@@ -0,0 +1,163 @@
+// Copyright Epic Games, Inc. All Rights Reserved.
+
+#include "Selection/MeshSelectionMechanic.h"
+
+#include "BaseBehaviors/BehaviorTargetInterfaces.h"
+#include "BaseBehaviors/SingleClickBehavior.h"
+#include "InteractiveToolManager.h"
+#include "Polyline3.h"
+#include "Selections/MeshConnectedComponents.h"
+#include "SceneManagement.h"
+#include "Spatial/GeometrySet3.h"
+#include "ToolSceneQueriesUtil.h"
+
+using namespace UE::Geometry;
+
+#define LOCTEXT_NAMESPACE "UMeshSelectionMechanic"
+
+void UMeshSelectionMechanic::Setup(UInteractiveTool* ParentToolIn)
+{
+	UInteractionMechanic::Setup(ParentToolIn);
+
+	// TODO: Add shift/ctrl modifiers
+	USingleClickInputBehavior* ClickBehavior = NewObject<USingleClickInputBehavior>();
+	ClickBehavior->Initialize(this);
+	ParentTool->AddInputBehavior(ClickBehavior);
+}
+
+void UMeshSelectionMechanic::Shutdown()
+{
+	// TODO: Will need to destroy preview actors here once using those for rendering
+}
+
+void UMeshSelectionMechanic::AddSpatial(TSharedPtr<FDynamicMeshAABBTree3> SpatialIn)
+{
+	MeshSpatials.Add(SpatialIn);
+}
+
+const FDynamicMeshSelection& UMeshSelectionMechanic::GetCurrentSelection() const
+{
+	return CurrentSelection;
+}
+
+void UMeshSelectionMechanic::SetSelection(const FDynamicMeshSelection& Selection, bool bBroadcast, bool bEmitChange)
+{
+	CurrentSelection = Selection;
+	if (bBroadcast)
+	{
+		OnSelectionChanged.Broadcast();
+	}
+	// TODO: Undo/redo
+}
+
+void UMeshSelectionMechanic::Render(IToolsContextRenderAPI* RenderAPI)
+{
+	// TODO:We do this in other places, but should CameraState be cached somewhere else?
+	GetParentTool()->GetToolManager()->GetContextQueriesAPI()->GetCurrentViewState(CameraState);
+
+	if (CurrentSelection.IsEmpty())
+	{
+		return;
+	}
+
+	// TODO: We should do all the selection with line/triangle components instead of PDI
+	FPrimitiveDrawInterface* PDI = RenderAPI->GetPrimitiveDrawInterface();
+	float PDIScale = RenderAPI->GetCameraState().GetPDIScalingFactor();
+	float LineThickness = 3.0 * PDIScale;
+	float DepthBias = 3.0f;
+
+	if (CurrentSelection.Type == FDynamicMeshSelection::EType::Triangle)
+	{
+		for (int32 Tid : CurrentSelection.SelectedIDs)
+		{
+			for (int i = 0; i < 3; ++i)
+			{
+				int NextIndex = (i + 1) % 3;
+				PDI->DrawLine((FVector)CurrentSelection.Mesh->GetTriVertex(Tid, i),
+					(FVector)CurrentSelection.Mesh->GetTriVertex(Tid, NextIndex), FLinearColor::Yellow,
+					SDPG_Foreground, LineThickness, DepthBias, true);
+			}
+		}
+	}
+	else if (CurrentSelection.Type == FDynamicMeshSelection::EType::Edge)
+	{
+		for (int32 Eid : CurrentSelection.SelectedIDs)
+		{
+			FIndex2i EdgeVids = CurrentSelection.Mesh->GetEdgeV(Eid);
+			PDI->DrawLine((FVector)CurrentSelection.Mesh->GetVertex(EdgeVids.A),
+				(FVector)CurrentSelection.Mesh->GetVertex(EdgeVids.B), FLinearColor::Yellow,
+				SDPG_Foreground, LineThickness, DepthBias, true);
+		}
+	}
+}
+
+FInputRayHit UMeshSelectionMechanic::IsHitByClick(const FInputDeviceRay& ClickPos)
+{
+	// We grab any clicks we are offered because even if we miss the mesh, we want to be able to clear selection.
+	FInputRayHit Hit; 
+	Hit.bHit = true;
+	return Hit;
+}
+
+void UMeshSelectionMechanic::OnClicked(const FInputDeviceRay& ClickPos)
+{
+	FDynamicMeshSelection OriginalSelection = CurrentSelection;
+
+	CurrentSelection.SelectedIDs.Empty();
+	CurrentSelection.Mesh = nullptr;
+
+	int32 HitTid = IndexConstants::InvalidID;
+	double RayT = 0; 
+	for (TSharedPtr<FDynamicMeshAABBTree3> Spatial : MeshSpatials)
+	{
+		if (Spatial->FindNearestHitTriangle(ClickPos.WorldRay, RayT, HitTid))
+		{
+			CurrentSelection.Mesh = Spatial->GetMesh();
+			break;
+		}
+	}
+
+	if (HitTid != IndexConstants::InvalidID)
+	{
+		if (SelectionMode == EMeshSelectionMechanicMode::Component)
+		{
+			CurrentSelection.SelectedIDs.Add(HitTid);
+			FMeshConnectedComponents MeshSelectedComponent(CurrentSelection.Mesh);
+			MeshSelectedComponent.FindTrianglesConnectedToSeeds(CurrentSelection.SelectedIDs.Array());
+			CurrentSelection.SelectedIDs = TSet(MeshSelectedComponent.Components[0].Indices);
+			CurrentSelection.Type = FDynamicMeshSelection::EType::Triangle;
+		}
+		// TODO: We'll need the ability to hit occluded triangles to see if there is a better edge
+		// to snap to.
+		else if (SelectionMode == EMeshSelectionMechanicMode::Edge)
+		{
+			// Try to snap to one of the edges.
+			FIndex3i Eids = CurrentSelection.Mesh->GetTriEdges(HitTid);
+
+			FGeometrySet3 GeometrySet;
+			for (int i = 0; i < 3; ++i)
+			{
+				FIndex2i Vids = CurrentSelection.Mesh->GetEdgeV(Eids[i]);
+				FPolyline3d Polyline(CurrentSelection.Mesh->GetVertex(Vids.A), CurrentSelection.Mesh->GetVertex(Vids.B));
+				GeometrySet.AddCurve(Eids[i], Polyline);
+			}
+			FGeometrySet3::FNearest Result;
+			if (GeometrySet.FindNearestCurveToRay(ClickPos.WorldRay, Result, 
+				[this](const FVector3d& Position1, const FVector3d& Position2) {
+					return ToolSceneQueriesUtil::PointSnapQuery(CameraState,
+						Position1, Position2,
+						ToolSceneQueriesUtil::GetDefaultVisualAngleSnapThreshD()); }))
+			{
+				CurrentSelection.SelectedIDs.Add(Result.ID);
+				CurrentSelection.Type = FDynamicMeshSelection::EType::Edge;
+			}
+		}
+	}
+
+	if (!(OriginalSelection == CurrentSelection))
+	{
+		OnSelectionChanged.Broadcast();
+	}
+}
+
+#undef LOCTEXT_NAMESPACE
diff --git a/Engine/Plugins/Experimental/MeshModelingToolset/Source/ModelingComponents/Private/ToolSetupUtil.cpp b/Engine/Plugins/Experimental/MeshModelingToolset/Source/ModelingComponents/Private/ToolSetupUtil.cpp
--- a/Engine/Plugins/Experimental/MeshModelingToolset/Source/ModelingComponents/Private/ToolSetupUtil.cpp
+++ b/Engine/Plugins/Experimental/MeshModelingToolset/Source/ModelingComponents/Private/ToolSetupUtil.cpp
@@ -198,6 +198,20 @@ UMaterialInstanceDynamic* ToolSetupUtil::GetSimpleCustomMaterial(UInteractiveToo
 	return nullptr;
 }
 
+UMaterialInstanceDynamic* ToolSetupUtil::GetCustomTwoSidedDepthOffsetMaterial(UInteractiveToolManager* ToolManager, const FLinearColor& Color, float PercentDepthOffset, float Opacity)
+{
+	UMaterialInterface* Material = LoadObject<UMaterial>(nullptr, TEXT("/MeshModelingToolset/Materials/SimpleTwoSidedOffsetMaterial"));
+	if (Material != nullptr)
+	{
+		UMaterialInstanceDynamic* MatInstance = UMaterialInstanceDynamic::Create(Material, ToolManager);
+		MatInstance->SetVectorParameterValue(TEXT("Color"), Color);
+		MatInstance->SetScalarParameterValue(TEXT("PercentDepthOffset"), PercentDepthOffset);
+		MatInstance->SetScalarParameterValue(TEXT("Opacity"), Opacity);
+		return MatInstance;
+	}
+	return nullptr;
+}
+
 UMaterialInstanceDynamic* ToolSetupUtil::GetCustomDepthOffsetMaterial(UInteractiveToolManager* ToolManager, const FLinearColor& Color, float PercentDepthOffset, float Opacity)
 {
 	UMaterialInterface* Material = LoadObject<UMaterial>(nullptr, TEXT("/MeshModelingToolset/Materials/SimpleOffsetMaterial"));
diff --git a/Engine/Plugins/Experimental/MeshModelingToolset/Source/ModelingComponents/Public/Drawing/MeshElementsVisualizer.h b/Engine/Plugins/Experimental/MeshModelingToolset/Source/ModelingComponents/Public/Drawing/MeshElementsVisualizer.h
--- a/Engine/Plugins/Experimental/MeshModelingToolset/Source/ModelingComponents/Public/Drawing/MeshElementsVisualizer.h
+++ b/Engine/Plugins/Experimental/MeshModelingToolset/Source/ModelingComponents/Public/Drawing/MeshElementsVisualizer.h
@@ -70,6 +70,10 @@ public:
 	/** depth bias used to slightly shift depth of lines */
 	UPROPERTY(EditAnywhere, Category = MeshElementVisualization, AdvancedDisplay, meta = (UIMin = -2.0, UIMax = 2.0))
 	float DepthBias = 0.2;
+
+	/** If true, the depth bias will be adjusted depending on the diagonal length of the mesh bounding box (smaller for smaller meshes). */
+	UPROPERTY(EditAnywhere, Category = MeshElementVisualization, AdvancedDisplay)
+	bool bAdjustDepthBiasUsingMeshSize = true;
 };
 
 
@@ -126,6 +130,7 @@ protected:
 	bool bSettingsModified = false;
 
 	void UpdateVisibility();
+	void UpdateLineDepthBiasScale();
 
 	TSharedPtr<IMeshWireframeSourceProvider> WireframeSourceProvider;
 };
diff --git a/Engine/Plugins/Experimental/MeshModelingToolset/Source/ModelingComponents/Public/Selection/DynamicMeshSelection.h b/Engine/Plugins/Experimental/MeshModelingToolset/Source/ModelingComponents/Public/Selection/DynamicMeshSelection.h
new file mode 100644
--- /dev/null
+++ b/Engine/Plugins/Experimental/MeshModelingToolset/Source/ModelingComponents/Public/Selection/DynamicMeshSelection.h
@@ -0,0 +1,49 @@
+// Copyright Epic Games, Inc. All Rights Reserved.
+
+#pragma once
+
+#include "CoreMinimal.h"
+
+namespace UE {
+namespace Geometry {
+
+class FDynamicMesh3;
+
+/**
+ * Class that represents a selection in a dynamic mesh.
+ */
+class FDynamicMeshSelection
+{
+public:
+	enum class EType
+	{
+		Vertex,
+		Edge,
+		Triangle,
+		Group
+	};
+
+	const FDynamicMesh3* Mesh = nullptr;
+	TSet<int32> SelectedIDs;
+	EType Type = EType::Vertex;
+
+	// Not be relevant if the selection type is not Group
+	int32 GroupLayer = 0;
+
+	bool IsEmpty() const
+	{
+		return SelectedIDs.IsEmpty();
+	}
+
+	bool operator==(const FDynamicMeshSelection& Other)
+	{
+		return Mesh == Other.Mesh
+			&& Type == Other.Type
+			&& (Type != EType::Group || GroupLayer == Other.GroupLayer)
+			&& SelectedIDs.Num() == Other.SelectedIDs.Num()
+			&& SelectedIDs.Includes(Other.SelectedIDs);
+	}
+
+};
+
+}}// end namespace UE::Geometry
\ No newline at end of file
diff --git a/Engine/Plugins/Experimental/MeshModelingToolset/Source/ModelingComponents/Public/Selection/MeshSelectionMechanic.h b/Engine/Plugins/Experimental/MeshModelingToolset/Source/ModelingComponents/Public/Selection/MeshSelectionMechanic.h
new file mode 100644
--- /dev/null
+++ b/Engine/Plugins/Experimental/MeshModelingToolset/Source/ModelingComponents/Public/Selection/MeshSelectionMechanic.h
@@ -0,0 +1,75 @@
+// Copyright Epic Games, Inc. All Rights Reserved.
+
+#pragma once
+
+#include "CoreMinimal.h"
+#include "GeometryBase.h" // Predeclare macros
+
+#include "BaseBehaviors/BehaviorTargetInterfaces.h"
+#include "DynamicMeshAABBTree3.h"
+#include "InteractionMechanic.h"
+#include "InteractiveTool.h"
+#include "Selection/DynamicMeshSelection.h"
+#include "ToolContextInterfaces.h" //FViewCameraState
+
+#include "MeshSelectionMechanic.generated.h"
+
+enum class EMeshSelectionMechanicMode
+{
+	Component,
+	
+	// Not yet fully implemented for UV mesh purposes, since 
+	// we need to be able to select occluded edges
+	Edge
+};
+
+UCLASS()
+class MODELINGCOMPONENTS_API UMeshSelectionMechanicProperties : public UInteractiveToolPropertySet
+{
+	GENERATED_BODY()
+
+public:
+};
+
+/**
+ * Mechanic for selecting elements of a dynamic mesh.
+ * 
+ * TODO: Not finished. Needs undo/redo, currently only selects connected components or unoccluded edges.
+ */
+UCLASS()
+class MODELINGCOMPONENTS_API UMeshSelectionMechanic : public UInteractionMechanic, public IClickBehaviorTarget
+{
+	GENERATED_BODY()
+
+public:
+	using FDynamicMeshAABBTree3 = UE::Geometry::FDynamicMeshAABBTree3;
+	using FDynamicMeshSelection = UE::Geometry::FDynamicMeshSelection;
+
+	virtual ~UMeshSelectionMechanic() {}
+
+	virtual void Setup(UInteractiveTool* ParentTool) override;
+	virtual void Shutdown() override;
+
+	// Use this to initialize the meshes we want to hit test.
+	virtual void AddSpatial(TSharedPtr<FDynamicMeshAABBTree3> SpatialIn);
+
+	virtual const FDynamicMeshSelection& GetCurrentSelection() const;
+	virtual void SetSelection(const FDynamicMeshSelection& Selection, bool bBroadcast = false, bool bEmitChange = false);
+
+	virtual void Render(IToolsContextRenderAPI* RenderAPI) override;
+	//virtual void DrawHUD(FCanvas* Canvas, IToolsContextRenderAPI* RenderAPI); // Will need once we have marquee
+
+	// IClickBehaviorTarget implementation
+	virtual FInputRayHit IsHitByClick(const FInputDeviceRay& ClickPos) override;
+	virtual void OnClicked(const FInputDeviceRay& ClickPos) override;
+
+	FSimpleMulticastDelegate OnSelectionChanged;
+
+	EMeshSelectionMechanicMode SelectionMode;
+
+protected:
+	TArray<TSharedPtr<FDynamicMeshAABBTree3>> MeshSpatials;
+	FDynamicMeshSelection CurrentSelection;
+	FViewCameraState CameraState;
+};
+
diff --git a/Engine/Plugins/Experimental/MeshModelingToolset/Source/ModelingComponents/Public/ToolSetupUtil.h b/Engine/Plugins/Experimental/MeshModelingToolset/Source/ModelingComponents/Public/ToolSetupUtil.h
--- a/Engine/Plugins/Experimental/MeshModelingToolset/Source/ModelingComponents/Public/ToolSetupUtil.h
+++ b/Engine/Plugins/Experimental/MeshModelingToolset/Source/ModelingComponents/Public/ToolSetupUtil.h
@@ -104,7 +104,12 @@ namespace ToolSetupUtil
 	/**
 	 * @return Simple material with configurable depth offset, color, and opacity.
 	 */
-	MODELINGCOMPONENTS_API UMaterialInstanceDynamic* GetCustomDepthOffsetMaterial(UInteractiveToolManager* ToolManager, const FLinearColor& Color, float DepthOffset, float Opacity);
+	MODELINGCOMPONENTS_API UMaterialInstanceDynamic* GetCustomDepthOffsetMaterial(UInteractiveToolManager* ToolManager, const FLinearColor& Color, float PercentDepthOffset, float Opacity);
+
+	/**
+	 * @return Simple two-sided material with configurable depth offset, color, and opacity.
+	 */
+	MODELINGCOMPONENTS_API UMaterialInstanceDynamic* GetCustomTwoSidedDepthOffsetMaterial(UInteractiveToolManager* ToolManager, const FLinearColor& Color, float PercentDepthOffset, float Opacity);
 
 	/**
 	 * @return Material used when editing AVolume objects using our tools.
diff --git a/Engine/Plugins/Experimental/UVEditor/Source/UVEditor/Private/ToolTargets/StaticMeshUVMeshToolTarget.cpp b/Engine/Plugins/Experimental/UVEditor/Source/UVEditor/Private/ToolTargets/StaticMeshUVMeshToolTarget.cpp
new file mode 100644
--- /dev/null
+++ b/Engine/Plugins/Experimental/UVEditor/Source/UVEditor/Private/ToolTargets/StaticMeshUVMeshToolTarget.cpp
@@ -0,0 +1,76 @@
+// Copyright Epic Games, Inc. All Rights Reserved.
+
+#include "ToolTargets/StaticMeshUVMeshToolTarget.h"
+
+#include "DynamicMesh3.h"
+#include "DynamicMeshAttributeSet.h"
+#include "Engine/StaticMesh.h"
+#include "MeshDescriptionUVsToDynamicMesh.h"
+
+using namespace UE::Geometry;
+
+namespace StaticMeshUVMeshToolTargetLocals
+{
+	// TODO: We'd like to make this configurable somehow.
+	const int32 LODIndex = 0;
+
+	// This probably should be a visible static const. It represents the
+	// translation from UV coordinates to UV mesh vertex coordinates (i.e.,
+	// they will range [0,UVScalingFactor] instead of [0,1]).
+	const double UVScalingFactor = 1000;
+}
+
+int32 UStaticMeshUVMeshToolTarget::GetNumUVLayers()
+{
+	using namespace StaticMeshUVMeshToolTargetLocals;
+	return OriginalAsset->GetNumUVChannels(LODIndex);
+}
+
+void UStaticMeshUVMeshToolTarget::SaveBackToUVs(const FDynamicMesh3* MeshToSave, int32 LayerIndex)
+{
+	using namespace StaticMeshUVMeshToolTargetLocals;
+
+	FMeshDescriptionUVsToDynamicMesh Converter;
+	Converter.UVLayerIndex = LODIndex;
+	Converter.ScaleFactor = UVScalingFactor;
+
+	FMeshDescription* MeshDescription = OriginalAsset->GetMeshDescription(LODIndex);
+	Converter.BakeBackUVsFromUVMesh(MeshToSave, MeshDescription);
+
+	OriginalAsset->CommitMeshDescription(LODIndex);
+	OriginalAsset->PostEditChange();
+}
+
+TSharedPtr<FDynamicMesh3> UStaticMeshUVMeshToolTarget::GetMesh(int32 LayerIndex)
+{
+	using namespace StaticMeshUVMeshToolTargetLocals;
+
+	FMeshDescriptionUVsToDynamicMesh Converter;
+	Converter.UVLayerIndex = LODIndex;
+	Converter.ScaleFactor = UVScalingFactor;
+
+	FMeshDescription* MeshDescription = OriginalAsset->GetMeshDescription(LODIndex);
+	return Converter.GetUVMesh(MeshDescription);
+}
+
+
+// Factory
+
+bool UStaticMeshUVMeshToolTargetFactory::CanBuildTarget(UObject* SourceObject, const FToolTargetTypeRequirements& Requirements) const
+{
+	return Cast<UStaticMesh>(SourceObject) && Requirements.AreSatisfiedBy(UStaticMeshUVMeshToolTarget::StaticClass());
+}
+
+UToolTarget* UStaticMeshUVMeshToolTargetFactory::BuildTarget(UObject* SourceObject, const FToolTargetTypeRequirements& Requirements)
+{
+	UStaticMesh* StaticMesh = Cast<UStaticMesh>(SourceObject);
+	check(StaticMesh);
+	FMeshDescription* MeshDescription = StaticMesh->GetMeshDescription(0);
+	
+	// We need to set the outer here for the UV editor initialization to work properly, because GetPath uses it.
+	UStaticMeshUVMeshToolTarget* Target = NewObject<UStaticMeshUVMeshToolTarget>(StaticMesh);
+	Target->OriginalAsset = StaticMesh;
+
+	check(Requirements.AreSatisfiedBy(Target));
+	return Target;
+}
\ No newline at end of file
diff --git a/Engine/Plugins/Experimental/UVEditor/Source/UVEditor/Private/UVEditor.cpp b/Engine/Plugins/Experimental/UVEditor/Source/UVEditor/Private/UVEditor.cpp
new file mode 100644
--- /dev/null
+++ b/Engine/Plugins/Experimental/UVEditor/Source/UVEditor/Private/UVEditor.cpp
@@ -0,0 +1,45 @@
+// Copyright Epic Games, Inc. All Rights Reserved.
+
+#include "UVEditor.h"
+
+#include "AdvancedPreviewScene.h"
+#include "ToolContextInterfaces.h"
+#include "UVEditorToolkit.h"
+#include "UVEditorSubsystem.h"
+#include "EditorModeManager.h"
+#include "EdModeInteractiveToolsContext.h"
+
+#include "AssetEditorModeManager.h"
+#include "UVEditorMode.h"
+
+void UUVEditor::Initialize(const TArray<TObjectPtr<UObject>>& InObjects)
+{
+	// Make sure we have valid targets.
+	UUVEditorSubsystem* UVSubsystem = GEditor->GetEditorSubsystem<UUVEditorSubsystem>();
+	check(UVSubsystem && UVSubsystem->AreObjectsValidTargets(InObjects));
+
+	OriginalObjectsToEdit = InObjects;
+
+	// This will do a variety of things including registration of the asset editor, creation of the toolkit
+	// (via CreateToolkit()), and creation of the editor mode manager within the toolkit.
+	// The asset editor toolkit will do the rest of the necessary initialization inside its PostInitAssetEditor.
+	UAssetEditor::Initialize();
+}
+
+IAssetEditorInstance* UUVEditor::GetInstanceInterface() 
+{ 
+	return ToolkitInstance; 
+}
+
+void UUVEditor::GetObjectsToEdit(TArray<UObject*>& OutObjects)
+{
+	OutObjects.Append(OriginalObjectsToEdit);
+	check(OutObjects.Num() > 0);
+}
+
+TSharedPtr<FBaseAssetToolkit> UUVEditor::CreateToolkit()
+{
+	TSharedPtr<FUVEditorToolkit> Toolkit = MakeShared<FUVEditorToolkit>(this);
+
+	return Toolkit;
+}
diff --git a/Engine/Plugins/Experimental/UVEditor/Source/UVEditor/Private/UVEditor2DViewportBehaviorTargets.cpp b/Engine/Plugins/Experimental/UVEditor/Source/UVEditor/Private/UVEditor2DViewportBehaviorTargets.cpp
new file mode 100644
--- /dev/null
+++ b/Engine/Plugins/Experimental/UVEditor/Source/UVEditor/Private/UVEditor2DViewportBehaviorTargets.cpp
@@ -0,0 +1,117 @@
+// Copyright Epic Games, Inc. All Rights Reserved.
+
+#include "UVEditor2DViewportBehaviorTargets.h"
+
+#include "UVEditor2DViewportClient.h"
+
+// FUVEditor2DScrollBehaviorTarget
+
+FUVEditor2DScrollBehaviorTarget::FUVEditor2DScrollBehaviorTarget(FUVEditor2DViewportClient* ViewportClientIn)
+	: ViewportClient(ViewportClientIn)
+{
+}
+
+FInputRayHit FUVEditor2DScrollBehaviorTarget::CanBeginClickDragSequence(const FInputDeviceRay& PressPos)
+{
+	// Verify that ray is facing the proper direction
+	if (PressPos.WorldRay.Direction.Z * PressPos.WorldRay.Origin.Z < 0)
+	{
+		return FInputRayHit(-PressPos.WorldRay.Origin.Z / PressPos.WorldRay.Direction.Z);
+	}
+	return FInputRayHit();
+}
+
+void FUVEditor2DScrollBehaviorTarget::OnClickPress(const FInputDeviceRay& PressPos)
+{
+	if (ensure(PressPos.WorldRay.Direction.Z * PressPos.WorldRay.Origin.Z < 0))
+	{
+		// Intersect with XY plane
+		double T = -PressPos.WorldRay.Origin.Z / PressPos.WorldRay.Direction.Z;
+		DragStart = FVector3d(
+			PressPos.WorldRay.Origin.X + T * PressPos.WorldRay.Direction.X,
+			PressPos.WorldRay.Origin.Y + T * PressPos.WorldRay.Direction.Y,
+			0);
+
+		OriginalCameraLocation = (FVector3d)ViewportClient->GetViewLocation();
+	}
+}
+
+void FUVEditor2DScrollBehaviorTarget::OnClickDrag(const FInputDeviceRay& DragPos)
+{
+	if (ensure(DragPos.WorldRay.Direction.Z * DragPos.WorldRay.Origin.Z < 0))
+	{
+		// Intersect a ray starting from the original position and using the new
+		// ray direction. I.e., pretend the camera is not moving.
+
+		double T = -OriginalCameraLocation.Z / DragPos.WorldRay.Direction.Z;
+		FVector3d DragEnd = FVector3d(
+			OriginalCameraLocation.X + T * DragPos.WorldRay.Direction.X,
+			OriginalCameraLocation.Y + T * DragPos.WorldRay.Direction.Y,
+			0);
+
+		// We want to make it look like we are sliding the plane such that DragStart
+		// ends up on DragEnd. For that, our camera will be moving the opposite direction.
+		FVector3d CameraDisplacement = DragStart - DragEnd;
+		check(CameraDisplacement.Z == 0);
+		ViewportClient->SetViewLocation((FVector)(OriginalCameraLocation + CameraDisplacement));
+	}
+}
+
+void FUVEditor2DScrollBehaviorTarget::OnClickRelease(const FInputDeviceRay& ReleasePos)
+{
+}
+
+void FUVEditor2DScrollBehaviorTarget::OnTerminateDragSequence()
+{
+}
+
+
+
+// FUVEditor2DMouseWheelZoomBehaviorTarget
+
+FUVEditor2DMouseWheelZoomBehaviorTarget::FUVEditor2DMouseWheelZoomBehaviorTarget(FUVEditor2DViewportClient* ViewportClientIn)
+	: ViewportClient(ViewportClientIn)
+{
+	SetZoomAmount(DEFAULT_ZOOM_AMOUNT);
+}
+
+FInputRayHit FUVEditor2DMouseWheelZoomBehaviorTarget::ShouldRespondToMouseWheel(const FInputDeviceRay& CurrentPos)
+{
+	// Always allowed to zoom with mouse wheel
+	FInputRayHit ToReturn;
+	ToReturn.bHit = true;
+	return ToReturn;
+}
+
+void FUVEditor2DMouseWheelZoomBehaviorTarget::OnMouseWheelScrollUp(const FInputDeviceRay& CurrentPos)
+{
+	FVector OriginalLocation = ViewportClient->GetViewLocation();
+	double TToPlane = -OriginalLocation.Z / CurrentPos.WorldRay.Direction.Z;
+
+	FVector NewLocation = OriginalLocation + (ZoomInProportion * TToPlane * CurrentPos.WorldRay.Direction);
+
+	// Don't zoom in so far that the XY plane lies in front of our near clipping plane, or else everything
+	// will suddenly disappear.
+	if (NewLocation.Z > ViewportClient->GetNearClipPlane())
+	{
+		ViewportClient->SetViewLocation(NewLocation);
+	}
+}
+
+void FUVEditor2DMouseWheelZoomBehaviorTarget::OnMouseWheelScrollDown(const FInputDeviceRay& CurrentPos)
+{
+	FVector OriginalLocation = ViewportClient->GetViewLocation();
+	double TToPlane = -OriginalLocation.Z / CurrentPos.WorldRay.Direction.Z;
+	ViewportClient->SetViewLocation(OriginalLocation - (ZoomOutProportion * TToPlane * CurrentPos.WorldRay.Direction));
+}
+
+void FUVEditor2DMouseWheelZoomBehaviorTarget::SetZoomAmount(double PercentZoomIn)
+{
+	check(PercentZoomIn < 100 && PercentZoomIn >= 0);
+
+	ZoomInProportion = PercentZoomIn / 100;
+
+	// Set the zoom out proportion such that (1 + ZoomOutProportion)(1 - ZoomInProportion) = 1
+	// so that zooming in and then zooming out will return to the same zoom level.
+	ZoomOutProportion = ZoomInProportion / (1 - ZoomInProportion);
+}
\ No newline at end of file
diff --git a/Engine/Plugins/Experimental/UVEditor/Source/UVEditor/Private/UVEditor2DViewportClient.cpp b/Engine/Plugins/Experimental/UVEditor/Source/UVEditor/Private/UVEditor2DViewportClient.cpp
new file mode 100644
--- /dev/null
+++ b/Engine/Plugins/Experimental/UVEditor/Source/UVEditor/Private/UVEditor2DViewportClient.cpp
@@ -0,0 +1,91 @@
+// Copyright Epic Games, Inc. All Rights Reserved.
+
+#include "UVEditor2DViewportClient.h"
+
+#include "BaseBehaviors/ClickDragBehavior.h"
+#include "BaseBehaviors/MouseWheelBehavior.h"
+//#include "Drawing/DensityAdjustingGrid.h"
+#include "EditorModeManager.h"
+#include "EdModeInteractiveToolsContext.h"
+
+FUVEditor2DViewportClient::FUVEditor2DViewportClient(FEditorModeTools* InModeTools,
+	FPreviewScene* InPreviewScene, const TWeakPtr<SEditorViewport>& InEditorViewportWidget)
+	: FEditorViewportClient(InModeTools, InPreviewScene, InEditorViewportWidget)
+{
+	// Don't draw the little XYZ drawing in the corner.
+	bDrawAxes = false;
+
+	// We want our near clip plane to be quite close so that we can zoom in further.
+	OverrideNearClipPlane(KINDA_SMALL_NUMBER);
+
+	// Set up viewport manipulation behaviors:
+
+	BehaviorSet = NewObject<UInputBehaviorSet>();
+
+	// We'll have the priority of our viewport manipulation behaviors be lower (i.e. higher
+	// numerically) than both the gizmo default and the tool default.
+	static constexpr int DEFAULT_VIEWPORT_BEHAVIOR_PRIORITY = 150;
+
+	ScrollBehaviorTarget = MakeUnique<FUVEditor2DScrollBehaviorTarget>(this);
+	UClickDragInputBehavior* ScrollBehavior = NewObject<UClickDragInputBehavior>();
+	ScrollBehavior->Initialize(ScrollBehaviorTarget.Get());
+	ScrollBehavior->SetDefaultPriority(DEFAULT_VIEWPORT_BEHAVIOR_PRIORITY);
+	ScrollBehavior->SetUseRightMouseButton();
+	BehaviorSet->Add(ScrollBehavior);
+
+	ZoomBehaviorTarget = MakeUnique<FUVEditor2DMouseWheelZoomBehaviorTarget>(this);
+	UMouseWheelInputBehavior* ZoomBehavior = NewObject<UMouseWheelInputBehavior>();
+	ZoomBehavior->Initialize(ZoomBehaviorTarget.Get());
+	ZoomBehavior->SetDefaultPriority(DEFAULT_VIEWPORT_BEHAVIOR_PRIORITY);
+	BehaviorSet->Add(ZoomBehavior);
+
+	ModeTools->GetInteractiveToolsContext()->InputRouter->RegisterSource(this);
+}
+
+const UInputBehaviorSet* FUVEditor2DViewportClient::GetInputBehaviors() const
+{
+	return BehaviorSet;
+}
+
+// Collects UObjects that we don't want the garbage collecter to throw away under us
+void FUVEditor2DViewportClient::AddReferencedObjects(FReferenceCollector& Collector)
+{
+	FEditorViewportClient::AddReferencedObjects(Collector);
+
+	Collector.AddReferencedObject(BehaviorSet);
+	//Collector.AddReferencedObject(Grid);
+}
+
+bool FUVEditor2DViewportClient::InputKey(FViewport* InViewport, int32 ControllerId, FKey Key, EInputEvent Event, float/*AmountDepressed*/, bool/*Gamepad*/)
+{
+	// We'll support disabling input like our base class, even if it does not end up being used.
+	if (bDisableInput)
+	{
+		return true;
+	}
+
+	// Our viewport manipulation is placed in the input router that ModeTools manages
+	return ModeTools->InputKey(this, Viewport, Key, Event);
+}
+
+// Note that this function gets called from the super class Draw(FViewport*, FCanvas*) overload to draw the scene.
+// We don't override that top-level function so that it can do whatever view calculations it needs to do.
+void FUVEditor2DViewportClient::Draw(const FSceneView* View, FPrimitiveDrawInterface* PDI)
+{
+	// Draw the axes
+	const float AxisThickness = 1.0;
+	PDI->DrawLine(FVector(-HALF_WORLD_MAX, 0, 0), FVector(HALF_WORLD_MAX, 0, 0), FLinearColor::Red, SDPG_World, AxisThickness, 0, true);
+	PDI->DrawLine(FVector(0, -HALF_WORLD_MAX, 0), FVector(0, HALF_WORLD_MAX, 0), FLinearColor::Green, SDPG_World, AxisThickness, 0, true);
+
+	// TODO: Eventually we want a proper grid. For now, draw the unit boundary
+	const float UVScale = 1000;
+	PDI->DrawLine(FVector(UVScale, 0, 0), FVector(UVScale, UVScale, 0), FLinearColor::Gray, SDPG_World, AxisThickness, 0, true);
+	PDI->DrawLine(FVector(0, UVScale, 0), FVector(UVScale, UVScale, 0), FLinearColor::Gray, SDPG_World, AxisThickness, 0, true);
+	//Grid->Render(View, PDI);
+
+	// TODO: Draw a little UV axis thing in the lower left, like the XYZ things that normal viewports have.
+
+	// Calls ModeTools draw/render functions
+	FEditorViewportClient::Draw(View, PDI);
+}
+
diff --git a/Engine/Plugins/Experimental/UVEditor/Source/UVEditor/Private/UVEditorCommands.cpp b/Engine/Plugins/Experimental/UVEditor/Source/UVEditor/Private/UVEditorCommands.cpp
new file mode 100644
--- /dev/null
+++ b/Engine/Plugins/Experimental/UVEditor/Source/UVEditor/Private/UVEditorCommands.cpp
@@ -0,0 +1,30 @@
+// Copyright Epic Games, Inc. All Rights Reserved.
+
+#include "UVEditorCommands.h"
+
+#include "EditorStyleSet.h"
+#include "Framework/Commands/InputChord.h"
+
+#define LOCTEXT_NAMESPACE "FUVEditorCommands"
+	
+FUVEditorCommands::FUVEditorCommands()
+	: TCommands<FUVEditorCommands>("UVEditorCommands", 
+		NSLOCTEXT("Contexts", "UVEditorCommands", "UV Editor"), 
+		NAME_None, // Parent
+		FEditorStyle::GetStyleSetName() // TODO: What should go here?
+		)
+{
+}
+
+void FUVEditorCommands::RegisterCommands()
+{
+	// These are part of the asset editor UI
+	UI_COMMAND(OpenUVEditor, "Open UV Editor", "Open the UV Editor window", EUserInterfaceActionType::Button, FInputChord());
+	UI_COMMAND(ApplyChanges, "Apply Changes", "Apply changes without yet saving the assets", EUserInterfaceActionType::Button, FInputChord());
+
+	// These get linked to various tool buttons.
+	UI_COMMAND(BeginSelectTool, "Select", "Switch to select tool", EUserInterfaceActionType::ToggleButton, FInputChord());
+	UI_COMMAND(BeginTransformTool, "Transform", "Switch to transform tool", EUserInterfaceActionType::ToggleButton, FInputChord());
+}
+
+#undef LOCTEXT_NAMESPACE
diff --git a/Engine/Plugins/Experimental/UVEditor/Source/UVEditor/Private/UVEditorMode.cpp b/Engine/Plugins/Experimental/UVEditor/Source/UVEditor/Private/UVEditorMode.cpp
new file mode 100644
--- /dev/null
+++ b/Engine/Plugins/Experimental/UVEditor/Source/UVEditor/Private/UVEditorMode.cpp
@@ -0,0 +1,235 @@
+// Copyright Epic Games, Inc. All Rights Reserved.
+
+#include "UVEditorMode.h"
+
+#include "Algo/AnyOf.h"
+#include "Drawing/MeshElementsVisualizer.h"
+#include "EdModeInteractiveToolsContext.h" //ToolsContext
+#include "InteractiveTool.h"
+#include "TargetInterfaces/UVUnwrapDynamicMesh.h"
+#include "ToolTargetManager.h"
+#include "ToolTargets/StaticMeshUVMeshToolTarget.h"
+#include "ToolTargets/ToolTarget.h"
+#include "PreviewMesh.h"
+#include "UVEditorCommands.h"
+#include "UVSelectTool.h"
+#include "UVEditorModeToolkit.h"
+#include "UVEditorSubsystem.h"
+#include "ToolSetupUtil.h"
+#include "UVToolStateObjects.h"
+
+#define LOCTEXT_NAMESPACE "UUVEditorMode"
+
+const FEditorModeID UUVEditorMode::EM_UVEditorModeId = TEXT("EM_UVEditorMode");
+
+namespace UVEditorModeLocals
+{
+	// TODO: This is temporary. We need to be able to configure the layer we look at,
+	// and ideally show multiple layers at a time.
+	const int32 UVLayerIndex = 0;
+
+	FColor TriangleColor = FColor(50, 194, 219);
+	FColor WireframeColor = FColor(50, 100, 219);
+	FColor IslandBorderColor = FColor(103, 52, 235);
+}
+
+UUVEditorMode::UUVEditorMode()
+{
+	Info = FEditorModeInfo(
+		EM_UVEditorModeId,
+		LOCTEXT("UVEditorModeName", "UV"),
+		FSlateIcon(),
+		false);
+}
+
+void UUVEditorMode::Enter()
+{
+	Super::Enter();
+
+	StateObjectStore = NewObject<UUVToolStateObjectStore>(this);
+
+	RegisterTools();
+
+	ActivateDefaultTool();
+
+	// Add new target factories here and in UUVEditorSubsystem::Initialize() as they are developed.
+	ToolsContext->TargetManager->AddTargetFactory(NewObject<UStaticMeshUVMeshToolTargetFactory>(ToolsContext->TargetManager));
+}
+
+void UUVEditorMode::RegisterTools()
+{
+	const FUVEditorCommands& CommandInfos = FUVEditorCommands::Get();
+
+	// TODO: Other tool registrations go here
+	auto UVSelectToolBuilder = NewObject<UUVSelectToolBuilder>();
+	UVSelectToolBuilder->DisplayedMeshes = &DisplayedMeshes;
+	UVSelectToolBuilder->StateObjectStore = GetStateObjectStore();
+	RegisterTool(CommandInfos.BeginSelectTool, TEXT("UVSelectTool"), UVSelectToolBuilder);
+
+	auto UVTransformToolBuilder = NewObject<UUVSelectToolBuilder>();
+	UVTransformToolBuilder->DisplayedMeshes = &DisplayedMeshes;
+	UVTransformToolBuilder->StateObjectStore = GetStateObjectStore();
+	UVTransformToolBuilder->bGizmoEnabled = true;
+	RegisterTool(CommandInfos.BeginTransformTool, TEXT("UVTransformTool"), UVTransformToolBuilder);
+}
+
+void UUVEditorMode::CreateToolkit()
+{
+	Toolkit = MakeShared<FUVEditorModeToolkit>();
+}
+
+void UUVEditorMode::ActivateDefaultTool()
+{
+	ToolsContext->StartTool(TEXT("UVSelectTool"));
+}
+
+
+void UUVEditorMode::BindCommands()
+{
+	// We currently don't have mode-level commands (rather than tool level or asset
+	// editor level), but presumably they would go here if we had them.
+}
+
+void UUVEditorMode::Exit()
+{
+	StateObjectStore->Clear();
+	StateObjectStore = nullptr;
+
+	for (TObjectPtr<UPreviewMesh>& Preview : DisplayedMeshes)
+	{
+		Preview->Disconnect();
+	}
+	for (TObjectPtr<UMeshElementsVisualizer>& Wireframe : DisplayedWireframes)
+	{
+		Wireframe->Disconnect();
+	}
+
+	Super::Exit();
+}
+
+void UUVEditorMode::InitializeTargets(TArray<TObjectPtr<UObject>>& AssetsIn)
+{
+	using namespace UVEditorModeLocals;
+
+	OriginalObjectsToEdit = AssetsIn;
+
+	for (const TObjectPtr<UObject>& Object : AssetsIn)
+	{
+		UToolTarget* ToolTarget = ToolsContext->TargetManager->BuildTarget(Object, UUVEditorSubsystem::UVUnwrapMeshTargetRequirements);
+		IUVUnwrapDynamicMesh* UVMesh = Cast<IUVUnwrapDynamicMesh>(ToolTarget);
+		check(UVMesh);
+		UVUnwrapMeshTargets.Add(ToolTarget);
+
+		// Create in-viewport representations of the unwrapped mesh
+		UPreviewMesh* UVPreview = NewObject<UPreviewMesh>();
+		UVPreview->CreateInWorld(GetWorld(), FTransform::Identity);
+		UVPreview->UpdatePreview(UVMesh->GetMesh(UVLayerIndex).Get());
+		
+		// We use a two-sided material because the triangles of our unwrapped meshes may have different orientations
+		UVPreview->SetMaterial(0, ToolSetupUtil::GetCustomTwoSidedDepthOffsetMaterial(
+			GetToolManager(),
+			(FLinearColor)TriangleColor,
+			0, //depth offset
+			0.3)); //opacity
+
+		UVPreview->SetVisible(true);
+		DisplayedMeshes.Add(UVPreview);
+
+		// Set up the wireframe display of the unwrapped mesh.
+		UMeshElementsVisualizer* WireframeDisplay = NewObject<UMeshElementsVisualizer>(this);
+		WireframeDisplay->CreateInWorld(GetWorld(), FTransform::Identity);
+		check(WireframeDisplay->Settings);
+
+		WireframeDisplay->Settings->DepthBias = 1;
+		WireframeDisplay->Settings->bAdjustDepthBiasUsingMeshSize = false;
+		WireframeDisplay->Settings->bShowWireframe = true;
+		WireframeDisplay->Settings->bShowBorders = true;
+		WireframeDisplay->Settings->WireframeColor = WireframeColor;
+		WireframeDisplay->Settings->BoundaryEdgeColor = IslandBorderColor;
+		WireframeDisplay->Settings->bShowUVSeams = false;
+		WireframeDisplay->Settings->bShowNormalSeams = false;
+
+		// These are not exposed at the visualizer level yet
+		// TODO: Should they be?
+		WireframeDisplay->WireframeComponent->BoundaryEdgeThickness = 2;
+
+		// The settings object is not part of a tool, so it won't get ticked like its
+		// supposed to (to enable property watching), unless we add this here.
+		PropertyObjectsToTick.Add(WireframeDisplay->Settings);
+
+		// The wireframe will track the preview mesh
+		WireframeDisplay->SetMeshAccessFunction([UVPreview](void) { return UVPreview->GetMesh(); });
+		UVPreview->GetOnMeshChanged().AddWeakLambda(this, [WireframeDisplay]() 
+		{
+			WireframeDisplay->NotifyMeshChanged();
+		});
+		DisplayedWireframes.Add(WireframeDisplay);
+
+		// Initialize our timestamp so we can later detect changes
+		MeshChangeStamps.Add(UVPreview->GetMesh()->GetShapeTimestamp());
+	}
+}
+
+bool UUVEditorMode::HaveUnappliedChanges()
+{
+	for (int32 i = 0; i < DisplayedMeshes.Num(); ++i)
+	{
+		if (DisplayedMeshes[i]->GetMesh()->GetShapeTimestamp() != MeshChangeStamps[i])
+		{
+			return true;
+		}
+	}
+	return false;
+}
+
+void UUVEditorMode::GetAssetsWithUnappliedChanges(TArray<TObjectPtr<UObject>> UnappliedAssetsOut)
+{
+	for (int32 i = 0; i < DisplayedMeshes.Num(); ++i)
+	{
+		if (DisplayedMeshes[i]->GetMesh()->GetShapeTimestamp() != MeshChangeStamps[i])
+		{
+			UnappliedAssetsOut.Add(OriginalObjectsToEdit[i]);
+		}
+	}
+}
+
+void UUVEditorMode::ApplyChanges()
+{
+	using namespace UVEditorModeLocals;
+
+	for (int32 i = 0; i < DisplayedMeshes.Num(); ++i)
+	{
+		if (DisplayedMeshes[i]->GetMesh()->GetShapeTimestamp() != MeshChangeStamps[i])
+		{
+			UVUnwrapMeshTargets[i]->SaveBackToUVs(DisplayedMeshes[i]->GetMesh(), UVLayerIndex);
+			MeshChangeStamps[i] = DisplayedMeshes[i]->GetMesh()->GetShapeTimestamp();
+		}
+	}
+}
+
+void UUVEditorMode::ModeTick(float DeltaTime)
+{
+	Super::ModeTick(DeltaTime);
+
+	for (TObjectPtr<UInteractiveToolPropertySet>& Propset : PropertyObjectsToTick)
+	{
+		if (Propset)
+		{
+			if (Propset->IsPropertySetEnabled())
+			{
+				Propset->CheckAndUpdateWatched();
+			}
+			else
+			{
+				Propset->SilentUpdateWatched();
+			}
+		}
+	}
+
+	for (UMeshElementsVisualizer* WireframeDisplay : DisplayedWireframes)
+	{
+		WireframeDisplay->OnTick(DeltaTime);
+	}
+}
+
+#undef LOCTEXT_NAMESPACE
diff --git a/Engine/Plugins/Experimental/UVEditor/Source/UVEditor/Private/UVEditorModeToolkit.cpp b/Engine/Plugins/Experimental/UVEditor/Source/UVEditor/Private/UVEditorModeToolkit.cpp
new file mode 100644
--- /dev/null
+++ b/Engine/Plugins/Experimental/UVEditor/Source/UVEditor/Private/UVEditorModeToolkit.cpp
@@ -0,0 +1,102 @@
+// Copyright Epic Games, Inc. All Rights Reserved.
+
+#include "UVEditorModeToolkit.h"
+
+#include "UVEditorCommands.h"
+#include "Tools/UEdMode.h"
+
+#define LOCTEXT_NAMESPACE "FUVEditorModeToolkit"
+
+FUVEditorModeToolkit::FUVEditorModeToolkit()
+{
+	// Construct the panel that we will give in GetInlineContent().
+	// This could probably be done in Init() instead, but constructor
+	// makes it easy to guarantee that GetInlineContent() will always
+	// be ready to work.
+
+	SAssignNew(ToolkitWidget, SBorder)
+		.HAlign(HAlign_Fill)
+		.Padding(4)
+		[
+			SNew(SVerticalBox)
+			+ SVerticalBox::Slot()
+			.AutoHeight()
+			.HAlign(HAlign_Left)
+			.Padding(1.f)
+			[
+				SAssignNew(ToolButtonsContainer, SBorder)
+				.BorderImage(FEditorStyle::GetBrush("ToolPanel.GroupBorder"))
+				.Padding(FMargin(4, 2, 0, 0))
+			]
+
+			+ SVerticalBox::Slot()
+				.FillHeight(1.0f)
+				.Padding(2, 0, 0, 0)
+			[
+				SNew(SVerticalBox)
+				+ SVerticalBox::Slot()
+				.AutoHeight()
+				[
+					SAssignNew(ToolWarningArea, STextBlock)
+					.AutoWrapText(true)
+					.Font(FCoreStyle::GetDefaultFontStyle("Bold", 9))
+					.ColorAndOpacity(FSlateColor(FLinearColor(0.9f, 0.15f, 0.15f))) //TODO: This probably needs to not be hardcoded
+					.Text(FText::GetEmpty())
+					.Visibility(EVisibility::Collapsed)
+				]
+
+				+ SVerticalBox::Slot()
+				[
+					SAssignNew(ToolDetailsContainer, SBorder)
+					.BorderImage(FEditorStyle::GetBrush("NoBorder"))
+				]
+
+				+ SVerticalBox::Slot()
+				.AutoHeight()
+				[
+					SAssignNew(ToolMessageArea, STextBlock)
+					.AutoWrapText(true)
+					.Font(FCoreStyle::GetDefaultFontStyle("Bold", 9))
+					.Text(FText::GetEmpty())
+				]
+			]
+		];
+}
+
+FUVEditorModeToolkit::~FUVEditorModeToolkit()
+{
+}
+
+void FUVEditorModeToolkit::Init(const TSharedPtr<IToolkitHost>& InitToolkitHost, TWeakObjectPtr<UEdMode> InOwningMode)
+{
+	FModeToolkit::Init(InitToolkitHost, InOwningMode);	
+
+	// Build the tool palette
+	const FUVEditorCommands& CommandInfos = FUVEditorCommands::Get();
+	TSharedRef<FUICommandList> CommandList = GetToolkitCommands();
+	check(OwningEditorMode.IsValid());
+	FUniformToolBarBuilder ToolbarBuilder(CommandList, FMultiBoxCustomization(OwningEditorMode->GetModeInfo().ToolbarCustomizationName), TSharedPtr<FExtender>(), false);
+	ToolbarBuilder.SetStyle(&FEditorStyle::Get(), "PaletteToolBar");
+
+	// TODO: Add more tools here
+	ToolbarBuilder.AddToolBarButton(CommandInfos.BeginSelectTool);
+	ToolbarBuilder.AddToolBarButton(CommandInfos.BeginTransformTool);
+
+	// Hook in the tool palette
+	ToolButtonsContainer->SetContent(ToolbarBuilder.MakeWidget());
+
+	// Hook up the tool detail panel
+	ToolDetailsContainer->SetContent(DetailsView.ToSharedRef());
+}
+
+FName FUVEditorModeToolkit::GetToolkitFName() const
+{
+	return FName("UVEditorMode");
+}
+
+FText FUVEditorModeToolkit::GetBaseToolkitName() const
+{
+	return NSLOCTEXT("UVEditorModeToolkit", "DisplayName", "UVEditorMode");
+}
+
+#undef LOCTEXT_NAMESPACE
diff --git a/Engine/Plugins/Experimental/UVEditor/Source/UVEditor/Private/UVEditorModule.cpp b/Engine/Plugins/Experimental/UVEditor/Source/UVEditor/Private/UVEditorModule.cpp
new file mode 100644
--- /dev/null
+++ b/Engine/Plugins/Experimental/UVEditor/Source/UVEditor/Private/UVEditorModule.cpp
@@ -0,0 +1,103 @@
+// Copyright Epic Games, Inc. All Rights Reserved.
+
+#include "UVEditorModule.h"
+
+#include "ContentBrowserMenuContexts.h"
+#include "LevelEditorMenuContext.h"
+#include "UVEditor.h"
+#include "UVEditorCommands.h"
+#include "UVEditorMode.h"
+#include "UVEditorSubsystem.h"
+#include "EditorModeRegistry.h"
+
+#include "ToolMenus.h"
+
+#define LOCTEXT_NAMESPACE "FUVEditorModule"
+
+void FUVEditorModule::StartupModule()
+{
+	FUVEditorCommands::Register();
+
+	// Menus need to be registered in a callback to make sure the system is ready for them.
+	UToolMenus::RegisterStartupCallback(
+		FSimpleMulticastDelegate::FDelegate::CreateRaw(this, &FUVEditorModule::RegisterMenus));
+}
+
+void FUVEditorModule::ShutdownModule()
+{
+	// Clean up menu things
+	UToolMenus::UnRegisterStartupCallback(this);
+	UToolMenus::UnregisterOwner(this);
+
+	FUVEditorCommands::Unregister();
+
+	FEditorModeRegistry::Get().UnregisterMode(UUVEditorMode::EM_UVEditorModeId);
+}
+
+void FUVEditorModule::RegisterMenus()
+{
+	// Allows cleanup when module unloads.
+	FToolMenuOwnerScoped OwnerScoped(this);
+
+	// Extent the content browser context menu for static meshes
+	{
+		UToolMenu* Menu = UToolMenus::Get()->ExtendMenu("ContentBrowser.AssetContextMenu.StaticMesh");
+		FToolMenuSection& Section = Menu->FindOrAddSection("GetAssetActions");
+
+		Section.AddDynamicEntry("OpenUVEditor", FNewToolMenuSectionDelegate::CreateLambda(
+			[this](FToolMenuSection& Section) 
+			{
+				// We'll need to get the target objects out of the context
+				if (UContentBrowserAssetContextMenuContext* Context = Section.FindContext<UContentBrowserAssetContextMenuContext>())
+				{
+					TArray<TObjectPtr<UObject>> AssetsToEdit;
+					AssetsToEdit.Append(Context->GetSelectedObjects());
+
+					UUVEditorSubsystem* UVSubsystem = GEditor->GetEditorSubsystem<UUVEditorSubsystem>();
+					check(UVSubsystem);
+
+					if (UVSubsystem->AreObjectsValidTargets(AssetsToEdit))
+					{
+						TSharedPtr<class FUICommandList> CommandListToBind = MakeShared<FUICommandList>();
+						CommandListToBind->MapAction(
+							FUVEditorCommands::Get().OpenUVEditor,
+							FExecuteAction::CreateUObject(UVSubsystem, &UUVEditorSubsystem::StartUVEditor, AssetsToEdit));
+
+						Section.AddMenuEntryWithCommandList(FUVEditorCommands::Get().OpenUVEditor, CommandListToBind);
+					}
+				}
+			}));
+	}
+
+	// Extent the level editor context menu
+	{
+		UToolMenu* Menu = UToolMenus::Get()->ExtendMenu("LevelEditor.ActorContextMenu");
+		FToolMenuSection& Section = Menu->FindOrAddSection("ActorAsset");
+		Section.AddDynamicEntry("OpenUVEditor", FNewToolMenuSectionDelegate::CreateLambda(
+			[this](FToolMenuSection& Section) 
+			{
+				TArray<UObject*> ReferencedAssets;
+				GEditor->GetReferencedAssetsForEditorSelection(ReferencedAssets);
+
+				UUVEditorSubsystem* UVSubsystem = GEditor->GetEditorSubsystem<UUVEditorSubsystem>();
+				check(UVSubsystem);
+
+				if (UVSubsystem->AreObjectsValidTargets(ReferencedAssets))
+				{
+					TArray<TObjectPtr<UObject>> AssetsToEdit;
+					AssetsToEdit.Append(ReferencedAssets);
+
+					TSharedPtr<class FUICommandList> CommandListToBind = MakeShared<FUICommandList>();
+					CommandListToBind->MapAction(
+						FUVEditorCommands::Get().OpenUVEditor,
+						FExecuteAction::CreateUObject(UVSubsystem, &UUVEditorSubsystem::StartUVEditor, AssetsToEdit));
+
+					Section.AddMenuEntryWithCommandList(FUVEditorCommands::Get().OpenUVEditor, CommandListToBind);
+				}
+			}));
+	}
+}
+
+#undef LOCTEXT_NAMESPACE
+	
+IMPLEMENT_MODULE(FUVEditorModule, UVEditor)
\ No newline at end of file
diff --git a/Engine/Plugins/Experimental/UVEditor/Source/UVEditor/Private/UVEditorSubsystem.cpp b/Engine/Plugins/Experimental/UVEditor/Source/UVEditor/Private/UVEditorSubsystem.cpp
new file mode 100644
--- /dev/null
+++ b/Engine/Plugins/Experimental/UVEditor/Source/UVEditor/Private/UVEditorSubsystem.cpp
@@ -0,0 +1,86 @@
+// Copyright Epic Games, Inc. All Rights Reserved.
+
+#include "UVEditorSubsystem.h"
+
+#include "ToolTargetManager.h"
+#include "ToolTargets/StaticMeshUVMeshToolTarget.h"
+#include "TargetInterfaces/UVUnwrapDynamicMesh.h"
+#include "UVEditor.h"
+
+const FToolTargetTypeRequirements UUVEditorSubsystem::UVUnwrapMeshTargetRequirements = 
+	FToolTargetTypeRequirements(UUVUnwrapDynamicMesh::StaticClass());
+
+void UUVEditorSubsystem::Initialize(FSubsystemCollectionBase& Collection)
+{
+	// The subsystem has its own tool target manager because it must exist before any UV editors exist,
+	// to see if a UV editor can be started.
+	ToolTargetManager = NewObject<UToolTargetManager>(this);
+	ToolTargetManager->Initialize();
+
+	// Add new target factories here and in UUVEditorMode::Enter() as they are developed.
+	ToolTargetManager->AddTargetFactory(NewObject<UStaticMeshUVMeshToolTargetFactory>(ToolTargetManager));
+}
+
+void UUVEditorSubsystem::Deinitialize()
+{
+	ToolTargetManager->Shutdown();
+	ToolTargetManager = nullptr;
+}
+
+bool UUVEditorSubsystem::AreObjectsValidTargets(const TArray<UObject*>& InObjects) const
+{
+	if (InObjects.IsEmpty())
+	{
+		return false;
+	}
+
+	for (UObject* Object : InObjects)
+	{
+		if (!ToolTargetManager->CanBuildTarget(Object, UVUnwrapMeshTargetRequirements))
+		{
+			return false;
+		}
+	}
+
+	return true;
+}
+
+void UUVEditorSubsystem::StartUVEditor(TArray<TObjectPtr<UObject>> ObjectsToEdit)
+{
+	// We don't allow opening a new instance if any of the objects are already opened
+	// in an existing instance. Instead, we bring such an instance to the front.
+	// Note that the asset editor subsystem takes care of this for "primary" asset editors, 
+	// i.e., the editors that open when one double clicks an asset or selects "edit". Since
+	// the UV editor is not a "primary" asset editor for any asset type, we do this management 
+	// ourselves.
+	for (TObjectPtr<UObject>& Object : ObjectsToEdit)
+	{
+		if (OpenedEditorInstances.Contains(Object))
+		{
+			OpenedEditorInstances[Object]->GetInstanceInterface()->FocusWindow(Object);
+			return;
+		}
+	}
+
+	// If we got here, there's not an instance already opened.
+
+	UUVEditor* UVEditor = NewObject<UUVEditor>(this);
+
+	// Among other things, this call registers the UV editor with the asset editor subsystem,
+	// which will prevent it from being garbage collected.
+	UVEditor->Initialize(ObjectsToEdit);
+
+	for (TObjectPtr<UObject>& Object : ObjectsToEdit)
+	{
+		OpenedEditorInstances.Add(Object, UVEditor);
+	}
+}
+
+void UUVEditorSubsystem::NotifyThatUVEditorClosed(TArray<TObjectPtr<UObject>> ObjectsItWasEditing)
+{
+	for (TObjectPtr<UObject>& Object : ObjectsItWasEditing)
+	{
+		OpenedEditorInstances.Remove(Object);
+	}
+}
+
diff --git a/Engine/Plugins/Experimental/UVEditor/Source/UVEditor/Private/UVEditorToolkit.cpp b/Engine/Plugins/Experimental/UVEditor/Source/UVEditor/Private/UVEditorToolkit.cpp
new file mode 100644
--- /dev/null
+++ b/Engine/Plugins/Experimental/UVEditor/Source/UVEditor/Private/UVEditorToolkit.cpp
@@ -0,0 +1,286 @@
+// Copyright Epic Games, Inc. All Rights Reserved.
+
+#include "UVEditorToolkit.h"
+
+#include "PreviewScene.h"
+#include "AssetEditorModeManager.h"
+#include "SAssetEditorViewport.h"
+#include "ToolMenus.h"
+#include "UVEditor.h"
+#include "UVEditorMode.h"
+#include "UVEditorCommands.h"
+#include "UVEditorSubsystem.h"
+#include "UVEditor2DViewportClient.h"
+#include "Widgets/Docking/SDockTab.h"
+
+#include "EdModeInteractiveToolsContext.h"
+
+#define LOCTEXT_NAMESPACE "UVEditorToolkit"
+
+const FName FUVEditorToolkit::InteractiveToolsPanelTabID(TEXT("UVEditor_InteractiveToolsTab"));
+
+FUVEditorToolkit::FUVEditorToolkit(UAssetEditor* InOwningAssetEditor)
+	: FBaseAssetToolkit(InOwningAssetEditor)
+{
+	check(Cast<UUVEditor>(InOwningAssetEditor));
+
+	// TODO: Will eventually need another viewport/previewscene for a 3d preview of the mesh
+
+	// We will replace the StandaloneDefaultLayout that our parent class gave us with 
+	// one where the properties detail panel is a vertical column on the left, and the
+	// viewport is on the right.
+	StandaloneDefaultLayout = FTabManager::NewLayout(FName("UVEditorLayout"))
+		->AddArea
+		(
+			FTabManager::NewPrimaryArea()->SetOrientation(Orient_Vertical)
+			->Split
+			(
+				FTabManager::NewSplitter()->SetOrientation(Orient_Horizontal)
+				->Split
+				(
+					FTabManager::NewStack()
+					->SetSizeCoefficient(0.2f)
+					->AddTab(InteractiveToolsPanelTabID, ETabState::OpenedTab)
+					->SetHideTabWell(true)
+				)
+				->Split
+				(
+					FTabManager::NewStack()
+					->SetSizeCoefficient(0.7f)
+					->AddTab(ViewportTabID, ETabState::OpenedTab)
+					->SetHideTabWell(true)
+				)
+			)
+		);
+
+	// We could create the preview scene in CreateEditorViewportClient() the way that FBaseAssetToolkit
+	// does, but it seems more intuitive to create it right off the bat and pass it in later. The preview
+	// scene is important because it holds the world in which our mode will create any objects to display.
+	FPreviewScene::ConstructionValues PreviewSceneArgs;
+	PreviewScene = MakeUnique<FPreviewScene>(PreviewSceneArgs);
+}
+
+FUVEditorToolkit::~FUVEditorToolkit()
+{
+	// The UV subsystem is responsible for opening/focusing UV editor instances, so we should
+	// notify it that this one is closing.
+	UUVEditorSubsystem* UVSubsystem = GEditor->GetEditorSubsystem<UUVEditorSubsystem>();
+	if (UVSubsystem)
+	{
+		TArray<TObjectPtr<UObject>> ObjectsWeWereEditing;
+		OwningAssetEditor->GetObjectsToEdit(ObjectsWeWereEditing);
+		UVSubsystem->NotifyThatUVEditorClosed(ObjectsWeWereEditing);
+	}
+}
+
+// This gets used to label the editor's tab in the window that opens.
+FText FUVEditorToolkit::GetToolkitName() const
+{
+	const TArray<UObject*>* Objects = GetObjectsCurrentlyBeingEdited();
+	if (Objects->Num() == 1)
+	{
+		return FText::Format(LOCTEXT("UVEditorTabNameWithObject", "{0} UVs"), 
+			GetLabelForObject((*Objects)[0]));
+	}
+	return LOCTEXT("UVEditorSimpleTabName", "UV Editor");
+}
+
+// This gets used multiple places, most notably in GetToolMenuAppName, which gets
+// used to refer to menus/toolbars internally.
+FName FUVEditorToolkit::GetToolkitFName() const
+{
+	return FName("UVEditor");
+}
+
+// TODO: What is this actually used for?
+FText FUVEditorToolkit::GetBaseToolkitName() const
+{
+	return LOCTEXT("UVBaseToolkitName", "UV");
+}
+
+void FUVEditorToolkit::RegisterTabSpawners(const TSharedRef<FTabManager>& InTabManager)
+{
+	// Here we set up the tabs we referenced in StandaloneDefaultLayout (in the constructor).
+	// The base class already creates the viewport tab (and also a details tab that we don't want).
+
+	FBaseAssetToolkit::RegisterTabSpawners(InTabManager);
+	InTabManager->UnregisterTabSpawner(DetailsTabID);
+
+	InTabManager->RegisterTabSpawner(InteractiveToolsPanelTabID, FOnSpawnTab::CreateSP(this, &FUVEditorToolkit::SpawnTab_InteractiveToolsPanel))
+		.SetDisplayName(LOCTEXT("InteractiveToolsPanel", "Tools Panel"))
+		.SetGroup(AssetEditorTabsCategory.ToSharedRef())
+		.SetIcon(FSlateIcon(FEditorStyle::GetStyleSetName(), "LevelEditor.Tabs.Details"));
+}
+
+bool FUVEditorToolkit::OnRequestClose()
+{
+	// Note: This needs a bit of adjusting, because currently OnRequestClose seems to be 
+	// called multiple times when the editor itself is being closed. We can take the route 
+	// of NiagaraScriptToolkit and remember when changes are discarded, but this can cause
+	// issues if the editor close sequence is interrupted due to some other asset editor.
+
+	UUVEditorMode* UVMode = Cast<UUVEditorMode>(EditorModeManager->GetActiveScriptableMode(UUVEditorMode::EM_UVEditorModeId));
+	check(UVMode);
+
+	// Warn the user of any unapplied changes.
+	if (UVMode->HaveUnappliedChanges())
+	{
+		TArray<TObjectPtr<UObject>> UnappliedAssets;
+		UVMode->GetAssetsWithUnappliedChanges(UnappliedAssets);
+
+		EAppReturnType::Type YesNoCancelReply = FMessageDialog::Open(EAppMsgType::YesNoCancel,
+			NSLOCTEXT("UVEditor", "Prompt_UVEditorClose", "At least one of the assets has unapplied changes. Would you like to apply them? (Selecting 'No' will cause all changes to be lost!)"));
+
+		switch (YesNoCancelReply)
+		{
+		case EAppReturnType::Yes:
+			UVMode->ApplyChanges();
+			break;
+
+		case EAppReturnType::No:
+			// exit
+			break;
+
+		case EAppReturnType::Cancel:
+			// don't exit
+			return false;
+		}
+	}
+
+	return FAssetEditorToolkit::OnRequestClose();
+}
+
+// We override the "Save" button behavior slightly to apply our changes before saving the asset.
+void FUVEditorToolkit::SaveAsset_Execute()
+{
+	UUVEditorMode* UVMode = Cast<UUVEditorMode>(EditorModeManager->GetActiveScriptableMode(UUVEditorMode::EM_UVEditorModeId));
+	check(UVMode);
+	if (UVMode->HaveUnappliedChanges())
+	{
+		UVMode->ApplyChanges();
+	}
+
+	FAssetEditorToolkit::SaveAsset_Execute();
+}
+
+// This is bound in RegisterTabSpawners() to create the panel on the left. It's filled in by the mode.
+TSharedRef<SDockTab> FUVEditorToolkit::SpawnTab_InteractiveToolsPanel(const FSpawnTabArgs& Args)
+{
+	ToolsPanel = SNew(SDockTab)
+		.Icon(FEditorStyle::GetBrush("LevelEditor.Tabs.Details")) // We'll use the same icon as a details panel
+		.Label(LOCTEXT("UVToolPanelTitle", "UV Tools"));
+
+	return ToolsPanel.ToSharedRef();
+}
+
+// This is used to create the viewport itself
+AssetEditorViewportFactoryFunction FUVEditorToolkit::GetViewportDelegate()
+{
+	AssetEditorViewportFactoryFunction TempViewportDelegate = [this](FAssetEditorViewportConstructionArgs InArgs)
+	{
+		return SNew(SAssetEditorViewport, InArgs)
+			.EditorViewportClient(ViewportClient);
+	};
+
+	return TempViewportDelegate;
+}
+
+TSharedPtr<FEditorViewportClient> FUVEditorToolkit::CreateEditorViewportClient() const
+{
+	// Note that we can't reliably adjust the viewport client here because we will be passing it
+	// into the viewport created by the viewport delegate we get from GetViewportDelegate(), and
+	// that delegate may (will) affect the settings based on FAssetEditorViewportConstructionArgs,
+	// namely ViewportType.
+	// Instead, we do viewport client adjustment in PostInitAssetEditor().
+	check(EditorModeManager.IsValid());
+	return MakeShared<FUVEditor2DViewportClient>(EditorModeManager.Get(), PreviewScene.Get());
+}
+
+void FUVEditorToolkit::CreateWidgets()
+{
+	// This gets called during UAssetEditor::Init() after creation of the toolkit but before
+	// calling InitAssetEditor on it. If we have custom mode-level toolbars we want to add,
+	// they could potentially go here, but we still need to call the base CreateWidgets as well.
+
+	FBaseAssetToolkit::CreateWidgets();
+}
+
+
+void FUVEditorToolkit::CreateEditorModeManager()
+{
+	EditorModeManager = MakeShared<FAssetEditorModeManager>();
+
+	// The mode manager is the authority on what the world is for the mode and the tools context,
+	// and setting the preview scene here makes our GetWorld() function return the preview scene
+	// world instead of the normal level editor one. Important because that is where we create
+	// any preview meshes, gizmo actors, etc.
+	StaticCastSharedPtr<FAssetEditorModeManager>(EditorModeManager)->SetPreviewScene(PreviewScene.Get());
+}
+
+void FUVEditorToolkit::PostInitAssetEditor()
+{
+	// Currently, aside from setting up all the UI elements, the toolkit also kicks off the UV
+	// editor mode, which is the mode that the editor always works in (things are packaged into
+	// a mode so that they can be moved to another asset editor if necessary).
+	// We need the UV mode to be active to create the toolbox on the left.
+	check(EditorModeManager.IsValid());
+	EditorModeManager->ActivateMode(UUVEditorMode::EM_UVEditorModeId);
+	UUVEditorMode* UVMode = Cast<UUVEditorMode>(
+		EditorModeManager->GetActiveScriptableMode(UUVEditorMode::EM_UVEditorModeId));
+	check(UVMode);
+
+	// Initialize mode state.
+	TArray<TObjectPtr<UObject>> ObjectsToEdit;
+	OwningAssetEditor->GetObjectsToEdit(ObjectsToEdit);
+	UVMode->InitializeTargets(ObjectsToEdit);
+
+	// Plug in the mode tool panel
+	TSharedPtr<FModeToolkit> UVModeToolkit = UVMode->GetToolkit().Pin();
+	check(UVModeToolkit.IsValid());
+	ToolsPanel->SetContent(UVModeToolkit->GetInlineContent().ToSharedRef());
+
+	// Add the "Apply Changes" button. It should actually be safe to do this almost
+	// any time, even before that toolbar's registration, but it's easier to put most
+	// things into PostInitAssetEditor().
+
+	// TODO: Use the icon that material editor uses. Also, the space between the button
+	// sections is too large, but the button can't go into the default toolbar section or
+	// else it will show up in other asset editors where the section is used.
+	// TODO: We may consider putting actions like these, which are tied to a mode, into
+	// some list of mode actions, and then letting the mode supply them to the owning
+	// asset editor on enter/exit. Revisit when/if this becomes easier to do.
+	ToolkitCommands->MapAction(
+		FUVEditorCommands::Get().ApplyChanges,
+		FExecuteAction::CreateUObject(UVMode, &UUVEditorMode::ApplyChanges),
+		FCanExecuteAction::CreateUObject(UVMode, &UUVEditorMode::HaveUnappliedChanges));
+	FName ParentToolbarName;
+	const FName ToolBarName = GetToolMenuToolbarName(ParentToolbarName);
+	UToolMenu* AssetToolbar = UToolMenus::Get()->ExtendMenu(ToolBarName);
+	FToolMenuSection& Section = AssetToolbar->FindOrAddSection("UVAsset");
+	Section.AddEntry(FToolMenuEntry::InitToolBarButton(FUVEditorCommands::Get().ApplyChanges));
+
+
+	// Adjust the viewport:
+
+	// Ortho has too many problems with rendering things, unfortunately, so we should use perspective.
+	ViewportClient->SetViewportType(ELevelViewportType::LVT_Perspective);
+
+	// Lit gives us the most options in terms of the materials we can use.
+	ViewportClient->SetViewMode(EViewModeIndex::VMI_Lit);
+
+	// TODO: This should be tied to wherever we store the settings for UV scaling. Right now we
+	// scale [0,1] to [0,1000], hence these settings.
+	ViewportClient->SetViewLocation(FVector(500, 500, 1000));
+	ViewportClient->SetViewRotation(FRotator(-90, 0, 0));
+
+	// If exposure isn't set to fixed, it will flash as we stare into the void
+	ViewportClient->ExposureSettings.bFixed = true;
+
+	// TODO: Disable temporal aa or whatever else is blurring the lines as the camera moves.
+
+	// We need the viewport client to start out focused, or else it won't get ticked until
+	// we click inside it.
+	ViewportClient->ReceivedFocus(ViewportClient->Viewport);
+}
+
+#undef LOCTEXT_NAMESPACE
\ No newline at end of file
diff --git a/Engine/Plugins/Experimental/UVEditor/Source/UVEditor/Public/ToolTargets/StaticMeshUVMeshToolTarget.h b/Engine/Plugins/Experimental/UVEditor/Source/UVEditor/Public/ToolTargets/StaticMeshUVMeshToolTarget.h
new file mode 100644
--- /dev/null
+++ b/Engine/Plugins/Experimental/UVEditor/Source/UVEditor/Public/ToolTargets/StaticMeshUVMeshToolTarget.h
@@ -0,0 +1,55 @@
+// Copyright Epic Games, Inc. All Rights Reserved.
+
+#pragma once
+
+#include "CoreMinimal.h"
+
+#include "GeometryBase.h"
+#include "ToolTargets/ToolTarget.h"
+#include "TargetInterfaces/UVUnwrapDynamicMesh.h"
+
+#include "StaticMeshUVMeshToolTarget.generated.h"
+
+PREDECLARE_USE_GEOMETRY_CLASS(FDynamicMesh3);
+PREDECLARE_USE_GEOMETRY_CLASS(FMeshDescriptionUVsToDynamicMesh);
+struct FMeshDescription;
+class UStaticMesh;
+
+/**
+ * Uses the UV's inside a mesh description provided by a static mesh to create a
+ * flat FDynamicMesh that can be manipulated and saved back to the static mesh.
+ * 
+ * Does not do any caching of any produced meshes, so GetMesh() is not a trivial call.
+ */
+UCLASS(Transient)
+class UVEDITOR_API UStaticMeshUVMeshToolTarget : public UToolTarget,
+	public IUVUnwrapDynamicMesh
+{
+	GENERATED_BODY()
+
+public:
+
+	// IUVUnwrapDynamicMesh implementation
+	virtual int32 GetNumUVLayers() override;
+	virtual TSharedPtr<FDynamicMesh3> GetMesh(int32 LayerIndex) override;
+	virtual void SaveBackToUVs(const FDynamicMesh3* MeshToSave, int32 LayerIndex) override;
+
+protected:
+	UPROPERTY()
+	UStaticMesh* OriginalAsset = nullptr;
+
+	friend class UStaticMeshUVMeshToolTargetFactory;
+};
+
+/** Factory for UStaticMeshUVMeshToolTarget to be used by the target manager. */
+UCLASS(Transient)
+class UVEDITOR_API UStaticMeshUVMeshToolTargetFactory : public UToolTargetFactory
+{
+	GENERATED_BODY()
+
+public:
+
+	virtual bool CanBuildTarget(UObject* SourceObject, const FToolTargetTypeRequirements& TargetTypeInfo) const override;
+
+	virtual UToolTarget* BuildTarget(UObject* SourceObject, const FToolTargetTypeRequirements& TargetTypeInfo) override;
+};
\ No newline at end of file
diff --git a/Engine/Plugins/Experimental/UVEditor/Source/UVEditor/Public/UVEditor.h b/Engine/Plugins/Experimental/UVEditor/Source/UVEditor/Public/UVEditor.h
new file mode 100644
--- /dev/null
+++ b/Engine/Plugins/Experimental/UVEditor/Source/UVEditor/Public/UVEditor.h
@@ -0,0 +1,58 @@
+// Copyright Epic Games, Inc. All Rights Reserved.
+
+#pragma once
+
+#include "CoreMinimal.h"
+
+#include "Tools/UAssetEditor.h"
+
+#include "UVEditor.generated.h"
+
+
+class FUICommandList;
+class UEdModeInteractiveToolsContext;
+class UInteractiveToolBuilder;
+class UPreviewMesh;
+class UUVEditorToolTargetManager;
+class UToolTarget;
+class UUVEditorMode;
+class IUVUnwrapDynamicMesh;
+class IAssetEditorInstance;
+
+/** 
+ * The actual asset editor class doesn't have that much in it, intentionally. 
+ * 
+ * Our current asset editor guidelines ask us to place as little business logic as possible
+ * into the class, instead putting as much of the non-UI code into the subsystem as possible,
+ * and the UI code into the toolkit (which this class owns).
+ *
+ * However, since we're using a mode and the Interactive Tools Framework, a lot of our business logic
+ * ends up inside the mode and the tools, not the subsystem. The front-facing code is mostly in
+ * the asset editor toolkit, though the mode toolkit has most of the things that deal with the toolbar
+ * on the left.
+ */
+UCLASS()
+class UVEDITOR_API UUVEditor : public UAssetEditor
+{
+	GENERATED_BODY()
+
+public:
+	void Initialize(const TArray<TObjectPtr<UObject>>& InObjects);
+
+	/** Returns the asset editor instance interface, so that its window can be focused, for example. */
+	IAssetEditorInstance* GetInstanceInterface();
+
+	// UAssetEditor overrides
+	void GetObjectsToEdit(TArray<UObject*>& OutObjects) override;
+	virtual TSharedPtr<FBaseAssetToolkit> CreateToolkit() override;
+
+	// Note that things like CloseWindow, FocusWindow, and IsPrimaryEditor seem to be called
+	// on the toolkit, not the editor class, so they are overriden there.
+
+protected:
+	
+	UPROPERTY()
+	TArray<TObjectPtr<UObject>> OriginalObjectsToEdit;
+
+	friend class UUVEditorSubsystem;
+};
diff --git a/Engine/Plugins/Experimental/UVEditor/Source/UVEditor/Public/UVEditor2DViewportBehaviorTargets.h b/Engine/Plugins/Experimental/UVEditor/Source/UVEditor/Public/UVEditor2DViewportBehaviorTargets.h
new file mode 100644
--- /dev/null
+++ b/Engine/Plugins/Experimental/UVEditor/Source/UVEditor/Public/UVEditor2DViewportBehaviorTargets.h
@@ -0,0 +1,64 @@
+// Copyright Epic Games, Inc. All Rights Reserved.
+
+#pragma once
+
+#include "CoreMinimal.h"
+
+#include "BaseBehaviors/BehaviorTargetInterfaces.h"
+#include "VectorTypes.h"
+
+class FUVEditor2DViewportClient;
+
+/**
+ * Allows click-dragging to move the camera in the XY plane.
+ */
+class FUVEditor2DScrollBehaviorTarget : public IClickDragBehaviorTarget
+{
+public:
+
+	FUVEditor2DScrollBehaviorTarget(FUVEditor2DViewportClient* ViewportClientIn);
+
+	// IClickDragBehaviorTarget
+	virtual FInputRayHit CanBeginClickDragSequence(const FInputDeviceRay& PressPos) override;
+	virtual void OnClickPress(const FInputDeviceRay& PressPos) override;
+	virtual void OnClickDrag(const FInputDeviceRay& DragPos) override;
+	virtual void OnClickRelease(const FInputDeviceRay& ReleasePos) override;
+	virtual void OnTerminateDragSequence() override;
+
+protected:
+	FUVEditor2DViewportClient* ViewportClient = nullptr;
+	FVector3d OriginalCameraLocation;
+	FVector3d DragStart;
+};
+
+/**
+ * Allows the mouse wheel to move the camera forwards/backwards relative to the XY plane,
+ * in the direction pointed to by the mouse.
+ */
+class FUVEditor2DMouseWheelZoomBehaviorTarget : public IMouseWheelBehaviorTarget
+{
+public:
+
+	FUVEditor2DMouseWheelZoomBehaviorTarget(FUVEditor2DViewportClient* ViewportClientIn);
+
+	// IMouseWheelBehaviorTarget
+	virtual FInputRayHit ShouldRespondToMouseWheel(const FInputDeviceRay& CurrentPos);
+	virtual void OnMouseWheelScrollUp(const FInputDeviceRay& CurrentPos);
+	virtual void OnMouseWheelScrollDown(const FInputDeviceRay& CurrentPos);
+	
+	/**
+	 * @param PercentZoomIn How much to move forward on each mouse wheel forward scroll. For instance,
+	 *  passing 20 here will decrease the distance to the XY plane by 20% each time. The zoom out amount
+	 *  will be set in such a way that it undoes the same effect. For instance if 20 was passed here,
+	 *  zoom out amount will be 25% since zooming in by 20% and then out by 25% of the result gets you
+	 *  back to where you were.
+	 */
+	virtual void SetZoomAmount(double PercentZoomIn);
+
+	inline static const double DEFAULT_ZOOM_AMOUNT = 20;
+protected:
+	FUVEditor2DViewportClient* ViewportClient = nullptr;
+	double ZoomInProportion;
+	double ZoomOutProportion;
+	double ZoomInLimit;
+};
\ No newline at end of file
diff --git a/Engine/Plugins/Experimental/UVEditor/Source/UVEditor/Public/UVEditor2DViewportClient.h b/Engine/Plugins/Experimental/UVEditor/Source/UVEditor/Public/UVEditor2DViewportClient.h
new file mode 100644
--- /dev/null
+++ b/Engine/Plugins/Experimental/UVEditor/Source/UVEditor/Public/UVEditor2DViewportClient.h
@@ -0,0 +1,45 @@
+// Copyright Epic Games, Inc. All Rights Reserved.
+
+#pragma once
+
+#include "CoreMinimal.h"
+
+#include "EditorViewportClient.h"
+#include "InputBehaviorSet.h"
+#include "UVEditor2DViewportBehaviorTargets.h" // FUVEditor2DScrollBehaviorTarget, FUVEditor2DMouseWheelZoomBehaviorTarget
+
+//class UDensityAdjustingGrid;
+
+/**
+ * Client used to display a 2D view of the UV's, implemented by using a perspective viewport with a locked
+ * camera.
+ */
+class UVEDITOR_API FUVEditor2DViewportClient : public FEditorViewportClient, public IInputBehaviorSource
+{
+public:
+	FUVEditor2DViewportClient(FEditorModeTools* InModeTools, FPreviewScene* InPreviewScene = nullptr,
+		const TWeakPtr<SEditorViewport>& InEditorViewportWidget = nullptr);
+
+	virtual ~FUVEditor2DViewportClient() {}
+
+	// FEditorViewportClient
+	virtual bool InputKey(FViewport* InViewport, int32 ControllerId, FKey Key, EInputEvent Event, float/*AmountDepressed*/, bool/*Gamepad*/) override;
+	virtual void Draw(const FSceneView* View, FPrimitiveDrawInterface* PDI) override;
+
+	// IInputBehaviorSource
+	virtual const UInputBehaviorSet* GetInputBehaviors() const override;
+
+	// FGCObject
+	virtual void AddReferencedObjects(FReferenceCollector& Collector) override;
+
+protected:
+	// These get added in AddReferencedObjects for memory management
+	UInputBehaviorSet* BehaviorSet;
+	//UDensityAdjustingGrid* Grid;
+
+	// Note that it's generally less hassle if the unique ptr types are complete here,
+	// not forward declared, else we get compile errors if their destruction shows up
+	// anywhere in the header.
+	TUniquePtr<FUVEditor2DScrollBehaviorTarget> ScrollBehaviorTarget;
+	TUniquePtr<FUVEditor2DMouseWheelZoomBehaviorTarget> ZoomBehaviorTarget;
+};
\ No newline at end of file
diff --git a/Engine/Plugins/Experimental/UVEditor/Source/UVEditor/Public/UVEditorCommands.h b/Engine/Plugins/Experimental/UVEditor/Source/UVEditor/Public/UVEditorCommands.h
new file mode 100644
--- /dev/null
+++ b/Engine/Plugins/Experimental/UVEditor/Source/UVEditor/Public/UVEditorCommands.h
@@ -0,0 +1,23 @@
+// Copyright Epic Games, Inc. All Rights Reserved.
+
+#pragma once
+
+#include "CoreMinimal.h"
+#include "Framework/Commands/Commands.h"
+
+class UVEDITOR_API FUVEditorCommands : public TCommands<FUVEditorCommands>
+{
+public:
+
+	FUVEditorCommands();
+
+	// TCommands<> interface
+	virtual void RegisterCommands() override;
+
+public:
+	TSharedPtr<FUICommandInfo> OpenUVEditor;
+	TSharedPtr<FUICommandInfo> ApplyChanges;
+
+	TSharedPtr<FUICommandInfo> BeginSelectTool;
+	TSharedPtr<FUICommandInfo> BeginTransformTool;
+};
\ No newline at end of file
diff --git a/Engine/Plugins/Experimental/UVEditor/Source/UVEditor/Public/UVEditorMode.h b/Engine/Plugins/Experimental/UVEditor/Source/UVEditor/Public/UVEditorMode.h
new file mode 100644
--- /dev/null
+++ b/Engine/Plugins/Experimental/UVEditor/Source/UVEditor/Public/UVEditorMode.h
@@ -0,0 +1,104 @@
+// Copyright Epic Games, Inc. All Rights Reserved.
+
+#pragma once
+
+#include "CoreMinimal.h"
+
+#include "Tools/UEdMode.h"
+
+#include "UVEditorMode.generated.h"
+
+
+class IUVUnwrapDynamicMesh;
+class UMeshElementsVisualizer;
+class UPreviewMesh;
+class UWorld;
+class UInteractiveToolPropertySet;
+class UUVToolStateObjectStore;
+
+/**
+ * The UV editor mode is the mode used in the UV asset editor. It holds most of the inter-tool state.
+ * We put things in a mode instead of directly into the asset editor in case we want to someday use the mode
+ * in multiple asset editors.
+ */
+UCLASS(Transient)
+class UUVEditorMode : public UEdMode
+{
+	GENERATED_BODY()
+public:
+	const static FEditorModeID EM_UVEditorModeId;
+
+	UUVEditorMode();
+
+	void RegisterTools();
+	void InitializeTargets(TArray<TObjectPtr<UObject>>& UVUnwrapMeshTargetsIn);
+
+	TObjectPtr<UUVToolStateObjectStore> GetStateObjectStore() { return StateObjectStore; }
+
+	// Asset management
+	bool HaveUnappliedChanges();
+	void GetAssetsWithUnappliedChanges(TArray<TObjectPtr<UObject>> UnappliedAssetsOut);
+	void ApplyChanges();
+
+	// UEdMode overrides
+	virtual void Enter() override;
+	virtual void Exit() override;
+	virtual void ModeTick(float DeltaTime) override;
+
+	// We don't actually override MouseEnter, etc, because things get forwarded to the input
+	// router via FEditorModeTools, and we don't have any additional input handling to do at the mode level.
+
+protected:
+
+	// UEdMode overrides
+	virtual void CreateToolkit() override;
+	virtual void ActivateDefaultTool() override;
+	// Not sure whether we need these yet
+	virtual void BindCommands() override;
+	virtual void OnToolStarted(UInteractiveToolManager* Manager, UInteractiveTool* Tool) {}
+	virtual void OnToolEnded(UInteractiveToolManager* Manager, UInteractiveTool* Tool) {}
+
+	/**
+	 * This stores the original objects that the UV editor was asked to operate on (for
+	 * instance UStaticMesh pointers).
+	 */
+	UPROPERTY()
+	TArray<TObjectPtr<UObject>> OriginalObjectsToEdit;
+
+	/**
+	 * This stores the targets generated by the tool target manager that give us the unwrapped UV meshes to work with.
+	 */
+	UPROPERTY()
+	TArray<TScriptInterface<IUVUnwrapDynamicMesh>> UVUnwrapMeshTargets;
+
+	/**
+	 * This is the main storage for the UV meshes, wrapped inside preview objects. Tools will
+	 * get a pointer to these so that they can hide and alter the meshes.
+	 */
+	UPROPERTY()
+	TArray<TObjectPtr<UPreviewMesh>> DisplayedMeshes;
+
+	/**
+	 * These wireframes track the meshes inside DisplayedMeshes
+	 */
+	UPROPERTY()
+	TArray<TObjectPtr<UMeshElementsVisualizer>> DisplayedWireframes;
+
+	/**
+	 * This lets the tools cache arbitrary state or see if it has been cached by another tool. This
+	 * can store selection, spatial data structures, etc.
+	 */
+	UPROPERTY()
+	TObjectPtr<UUVToolStateObjectStore> StateObjectStore = nullptr;
+
+	/**
+	 * If we want to have mode-level property objects, whether user-visible or not,
+	 * they can be put here to be ticked.
+	 */
+	UPROPERTY()
+	TArray<TObjectPtr<UInteractiveToolPropertySet>> PropertyObjectsToTick;
+
+	// This keeps track of whether the meshes inside DisplayedMeshes have been changed since
+	// the last time that they were baked back to their targets.
+	TArray<int32> MeshChangeStamps;
+};
diff --git a/Engine/Plugins/Experimental/UVEditor/Source/UVEditor/Public/UVEditorModeToolkit.h b/Engine/Plugins/Experimental/UVEditor/Source/UVEditor/Public/UVEditorModeToolkit.h
new file mode 100644
--- /dev/null
+++ b/Engine/Plugins/Experimental/UVEditor/Source/UVEditor/Public/UVEditorModeToolkit.h
@@ -0,0 +1,43 @@
+// Copyright Epic Games, Inc. All Rights Reserved.
+
+#pragma once
+
+#include "CoreMinimal.h"
+
+#include "Toolkits/BaseToolkit.h"
+
+/**
+ * The UV editor mode toolkit is responsible for the panel on the side in the UV editor
+ * that shows mode and tool properties. Tool buttons would go in Init().
+ */
+class FUVEditorModeToolkit : public FModeToolkit
+{
+public:
+	FUVEditorModeToolkit();
+	~FUVEditorModeToolkit();
+
+	/** FModeToolkit interface */
+	virtual void Init(const TSharedPtr<IToolkitHost>& InitToolkitHost, TWeakObjectPtr<UEdMode> InOwningMode) override;
+
+	/** IToolkit interface */
+	virtual FName GetToolkitFName() const override;
+	virtual FText GetBaseToolkitName() const override;
+	virtual TSharedPtr<class SWidget> GetInlineContent() const override { return ToolkitWidget; }
+
+protected:
+	// The mode's entire toolbox, which gets returned by GetInlineContent()
+	TSharedPtr<SWidget> ToolkitWidget;
+
+	// The portion of the toolbox composed of buttons that activate tools
+	TSharedPtr<SBorder> ToolButtonsContainer;
+
+	// A place for tools to write out any warnings
+	TSharedPtr<STextBlock> ToolWarningArea;
+
+	// A container for the tool settings that is populated by the DetailsView managed
+	// in FModeToolkit
+	TSharedPtr<SBorder> ToolDetailsContainer;
+
+	// A place for tools to write out any instructions
+	TSharedPtr<STextBlock> ToolMessageArea;
+};
\ No newline at end of file
diff --git a/Engine/Plugins/Experimental/UVEditor/Source/UVEditor/Public/UVEditorModule.h b/Engine/Plugins/Experimental/UVEditor/Source/UVEditor/Public/UVEditorModule.h
new file mode 100644
--- /dev/null
+++ b/Engine/Plugins/Experimental/UVEditor/Source/UVEditor/Public/UVEditorModule.h
@@ -0,0 +1,22 @@
+// Copyright Epic Games, Inc. All Rights Reserved.
+
+#pragma once
+
+#include "CoreMinimal.h"
+#include "Modules/ModuleManager.h"
+
+/**
+ * Besides the normal module things, the module class is also responsible for hooking the 
+ * UV editor into existing menus.
+ */
+class UVEDITOR_API FUVEditorModule : public IModuleInterface
+{
+public:
+
+	/** IModuleInterface implementation */
+	virtual void StartupModule() override;
+	virtual void ShutdownModule() override;
+
+protected:
+	void RegisterMenus();
+};
diff --git a/Engine/Plugins/Experimental/UVEditor/Source/UVEditor/Public/UVEditorSubsystem.h b/Engine/Plugins/Experimental/UVEditor/Source/UVEditor/Public/UVEditorSubsystem.h
new file mode 100644
--- /dev/null
+++ b/Engine/Plugins/Experimental/UVEditor/Source/UVEditor/Public/UVEditorSubsystem.h
@@ -0,0 +1,62 @@
+// Copyright Epic Games, Inc. All Rights Reserved.
+
+#pragma once
+
+#include "CoreMinimal.h"
+
+#include "EditorSubsystem.h"
+#include "ToolTargets/ToolTarget.h" //FToolTargetTypeRequirements
+
+#include "UVEditorSubsystem.generated.h"
+
+class UToolTargetManager;
+class UUVEditor;
+class UToolTarget;
+
+/**
+ * The subsystem is meant to hold any UV editor operations that are not UI related (those are
+ * handled by the toolkit); however, in our case, most of our operations are wrapped up inside
+ * tools and the UV mode. 
+ * Instead, the subsystem deals with some global UV asset editor things- it manages existing
+ * instances (we can't rely on the asset editor subsystem for this because the UV editor is
+ * not the primary editor for meshes), and figures out what we can launch the editor on.
+ */
+UCLASS()
+class UVEDITOR_API UUVEditorSubsystem : public UEditorSubsystem
+{
+	GENERATED_BODY()
+
+public:
+	virtual void Initialize(FSubsystemCollectionBase& Collection) override;
+	virtual void Deinitialize() override;
+
+	/** Checks that all of the objects are valid targets for a UV editor session. */
+	virtual bool AreObjectsValidTargets(const TArray<UObject*>& InObjects) const;
+
+	/**
+	 * Either brings to the front an existing UV editor instance that is editing one of
+	 * these objects, if one exists, or starts up a new instance editing all of these 
+	 * objects.
+	 */
+	virtual void StartUVEditor(TArray<TObjectPtr<UObject>> ObjectsToEdit);
+
+	/** 
+	 * Needs to be called when a UV editor instance is closed so that the subsystem knows
+	 * to create a new one for these objects if they are opened again.
+	 */
+	virtual void NotifyThatUVEditorClosed(TArray<TObjectPtr<UObject>> ObjectsItWasEditing);
+
+	static const FToolTargetTypeRequirements UVUnwrapMeshTargetRequirements;
+
+protected:
+
+	/**
+	 * Used to let the subsystem figure out whether targets are valid. New factories should be
+	 * added here in Initialize() as they are developed, as well as to the tool manager used
+	 * by a UV mode.
+	 */
+	UPROPERTY()
+	UToolTargetManager* ToolTargetManager = nullptr;
+
+	TMap<TObjectPtr<UObject>, TObjectPtr<UUVEditor>> OpenedEditorInstances;
+};
\ No newline at end of file
diff --git a/Engine/Plugins/Experimental/UVEditor/Source/UVEditor/Public/UVEditorToolkit.h b/Engine/Plugins/Experimental/UVEditor/Source/UVEditor/Public/UVEditorToolkit.h
new file mode 100644
--- /dev/null
+++ b/Engine/Plugins/Experimental/UVEditor/Source/UVEditor/Public/UVEditorToolkit.h
@@ -0,0 +1,65 @@
+// Copyright Epic Games, Inc. All Rights Reserved.
+
+#pragma once
+
+#include "CoreMinimal.h"
+
+#include "Tools/BaseAssetToolkit.h"
+
+class FAdvancedPreviewScene;
+class IDetailsView;
+class SDockTab;
+class SBorder;
+class UInteractiveToolsContext;
+
+/**
+ * The toolkit is supposed to act as the UI manager for the asset editor. It's responsible 
+ * for setting up viewports and most toolbars, except for the internals of the mode panel.
+ * However, because the toolkit also sets up the mode manager, and much of the important
+ * state is held in the UUVEditorMode managed by the mode manager, the toolkit also ends up
+ * initializing the UV mode.
+ * Thus, the FUVEdiotrToolkit ends up being the central place for the UV Asset editor setup.
+ */
+class UVEDITOR_API FUVEditorToolkit : public FBaseAssetToolkit
+{
+public:
+	FUVEditorToolkit(UAssetEditor* InOwningAssetEditor);
+	virtual ~FUVEditorToolkit();
+
+	static const FName InteractiveToolsPanelTabID;
+
+	FPreviewScene* GetPreviewScene() { return PreviewScene.Get(); }
+
+	// FBaseAssetToolkit
+	virtual void CreateWidgets() override;
+
+	// FAssetEditorToolkit
+	virtual void CreateEditorModeManager() override;
+	virtual FText GetToolkitName() const override;
+	virtual FName GetToolkitFName() const override;
+	virtual FText GetBaseToolkitName() const override;
+	virtual void RegisterTabSpawners(const TSharedRef<FTabManager>& InTabManager) override;
+	virtual bool OnRequestClose() override;
+	virtual void SaveAsset_Execute() override;
+
+	// IAssetEditorInstance
+	// This is important because if this returns true, attempting to edit a static mesh that is
+	// open in the UV editor may open the UV editor instead of opening the static mesh editor.
+	virtual bool IsPrimaryEditor() const override { return false; };
+
+protected:
+
+	TSharedRef<SDockTab> SpawnTab_InteractiveToolsPanel(const FSpawnTabArgs& Args);
+
+	// FBaseAssetToolkit
+	virtual AssetEditorViewportFactoryFunction GetViewportDelegate() override;
+	virtual TSharedPtr<FEditorViewportClient> CreateEditorViewportClient() const override;
+
+	// FAssetEditorToolkit
+	virtual void PostInitAssetEditor() override;
+
+	/** Inline content area for the UV mode's content (gotten from UVEditorModeToolkit) */
+	TSharedPtr<SDockTab> ToolsPanel;
+
+	TUniquePtr<FPreviewScene> PreviewScene;
+};
\ No newline at end of file
diff --git a/Engine/Plugins/Experimental/UVEditor/Source/UVEditor/UVEditor.Build.cs b/Engine/Plugins/Experimental/UVEditor/Source/UVEditor/UVEditor.Build.cs
new file mode 100644
--- /dev/null
+++ b/Engine/Plugins/Experimental/UVEditor/Source/UVEditor/UVEditor.Build.cs
@@ -0,0 +1,73 @@
+// Copyright Epic Games, Inc. All Rights Reserved.
+
+using UnrealBuildTool;
+
+public class UVEditor : ModuleRules
+{
+	public UVEditor(ReadOnlyTargetRules Target) : base(Target)
+	{
+		PCHUsage = ModuleRules.PCHUsageMode.UseExplicitOrSharedPCHs;
+		
+		PublicIncludePaths.AddRange(
+			new string[] {
+				// ... add public include paths required here ...
+			}
+			);
+				
+		
+		PrivateIncludePaths.AddRange(
+			new string[] {
+				// ... add other private include paths required here ...
+			}
+			);
+			
+		
+		PublicDependencyModuleNames.AddRange(
+			new string[]
+			{
+				"Core",
+				
+				"UnrealEd",
+				"InputCore",
+				"InteractiveToolsFramework",
+				
+				// ... add other public dependencies that you statically link with here ...
+			}
+			);
+			
+		
+		PrivateDependencyModuleNames.AddRange(
+			new string[]
+			{
+				"CoreUObject",
+				"Engine",
+				"Slate",
+				"SlateCore",
+				
+				"AdvancedPreviewScene",
+				"ContentBrowser",
+				"DynamicMesh",
+				"EditorFramework", // FEditorModeInfo
+				"EditorInteractiveToolsFramework",
+				"EditorStyle",
+				"EditorSubsystem",
+				"LevelEditor",
+				"ModelingComponents",
+				"ToolMenus",
+				"UVEditorTools",
+				
+				// TODO: Remove?
+				"MeshConversion",
+				// ... add private dependencies that you statically link with here ...	
+			}
+			);
+		
+		
+		DynamicallyLoadedModuleNames.AddRange(
+			new string[]
+			{
+				// ... add any modules that your module loads dynamically here ...
+			}
+			);
+	}
+}
diff --git a/Engine/Plugins/Experimental/UVEditor/Source/UVEditorTools/Private/UVEditorToolsModule.cpp b/Engine/Plugins/Experimental/UVEditor/Source/UVEditorTools/Private/UVEditorToolsModule.cpp
new file mode 100644
--- /dev/null
+++ b/Engine/Plugins/Experimental/UVEditor/Source/UVEditorTools/Private/UVEditorToolsModule.cpp
@@ -0,0 +1,20 @@
+// Copyright Epic Games, Inc. All Rights Reserved.
+
+#include "UVEditorToolsModule.h"
+
+#define LOCTEXT_NAMESPACE "FUVEditorToolsModule"
+
+void FUVEditorToolsModule::StartupModule()
+{
+	// This code will execute after your module is loaded into memory; the exact timing is specified in the .uplugin file per-module
+}
+
+void FUVEditorToolsModule::ShutdownModule()
+{
+	// This function may be called during shutdown to clean up your module.  For modules that support dynamic reloading,
+	// we call this function before unloading the module.
+}
+
+#undef LOCTEXT_NAMESPACE
+
+IMPLEMENT_MODULE(FUVEditorToolsModule, UVEditorTools)
diff --git a/Engine/Plugins/Experimental/UVEditor/Source/UVEditorTools/Private/UVSelectTool.cpp b/Engine/Plugins/Experimental/UVEditor/Source/UVEditorTools/Private/UVSelectTool.cpp
new file mode 100644
--- /dev/null
+++ b/Engine/Plugins/Experimental/UVEditor/Source/UVEditorTools/Private/UVSelectTool.cpp
@@ -0,0 +1,310 @@
+// Copyright Epic Games, Inc. All Rights Reserved.
+
+#include "UVSelectTool.h"
+
+#include "BaseGizmos/TransformGizmo.h"
+#include "TargetInterfaces/UVUnwrapDynamicMesh.h"
+#include "InteractiveToolManager.h"
+#include "PreviewMesh.h"
+#include "Selection/MeshSelectionMechanic.h"
+#include "Selection/DynamicMeshSelection.h"
+#include "UVToolStateObjects.h"
+
+#include "ToolTargetManager.h"
+
+#define LOCTEXT_NAMESPACE "UUVSelectTool"
+
+using namespace UE::Geometry;
+
+/*
+ * ToolBuilder
+ */
+
+
+// TODO: We may need to rework how targets are passed to this tool. Right now GetTargetRequirements
+// is not used, and is not even accurate, since we don't even use UUVUnwrapDynamicMesh ourselves.
+// Instead, the UVMode is the one that uses UUVUnwrapDynamicMesh to create meshes that it can display,
+// and we just get those preview meshes passed to us. Perhaps we should get passed some sort of
+// Target-wrapped preview mesh?
+// The other option is to use UUVUnwrapDynamicMesh in the tool and cache it for use by others, but
+// this ends up being a lot less clean to code up, and it would be preferable to avoid that messiness.
+const FToolTargetTypeRequirements& UUVSelectToolBuilder::GetTargetRequirements() const
+{
+	static FToolTargetTypeRequirements TypeRequirements(UUVUnwrapDynamicMesh::StaticClass());
+	return TypeRequirements;
+}
+
+bool UUVSelectToolBuilder::CanBuildTool(const FToolBuilderState& SceneState) const
+{
+	return DisplayedMeshes->Num() > 0;
+}
+
+UInteractiveTool* UUVSelectToolBuilder::BuildTool(const FToolBuilderState& SceneState) const
+{
+	UUVSelectTool* NewTool = NewObject<UUVSelectTool>(SceneState.ToolManager);
+	NewTool->SetGizmoEnabled(bGizmoEnabled);
+	NewTool->SetStateObjectStore(StateObjectStore);
+	NewTool->SetDisplayedMeshes(*DisplayedMeshes);
+
+	return NewTool;
+}
+
+void UUVSelectTool::Setup()
+{
+	check(DisplayedMeshes.Num() > 0);
+
+	UInteractiveTool::Setup();
+
+	SetToolDisplayName(LOCTEXT("ToolName", "UV Select Tool"));
+
+	Settings = NewObject<UUVSelectToolProperties>(this);
+	Settings->RestoreProperties(this);
+	AddToolPropertySource(Settings);
+
+	SelectionMechanic = NewObject<UMeshSelectionMechanic>();
+	SelectionMechanic->Setup(this);
+	SelectionMechanic->OnSelectionChanged.AddUObject(this, &UUVSelectTool::OnSelectionChanged);
+
+	SelectionMechanic->SelectionMode =
+		Settings->SelectionMode == EUVSelectToolSelectionMode::Island ? EMeshSelectionMechanicMode::Component
+		: EMeshSelectionMechanicMode::Edge;
+
+
+	// TODO: Should StateObjectStore be mandatory? We could enforce it in builder.
+	// Restore any saved state
+	if (StateObjectStore)
+	{
+		UUVMeshAABBTrees* TreeStore = StateObjectStore->GetToolStateObject<UUVMeshAABBTrees>();
+		if (TreeStore)
+		{
+			AABBTrees = TreeStore->AABBTrees;
+		}
+			
+		UUVMeshSelection* SelectionStore = StateObjectStore->GetToolStateObject<UUVMeshSelection>();
+		if (SelectionStore)
+		{
+			SelectionMechanic->SetSelection(*SelectionStore->Selection);
+		}
+	}
+
+	// Initialize the spatial structures if they weren't there to load
+	if (AABBTrees.IsEmpty())
+	{
+		for (TObjectPtr<UPreviewMesh> Mesh : DisplayedMeshes)
+		{
+			TSharedPtr<FDynamicMeshAABBTree3> Tree = MakeShared<FDynamicMeshAABBTree3>();
+			Tree->SetMesh(Mesh->GetMesh());
+			AABBTrees.Add(Tree);
+		}
+
+		// Save them if appropriate
+		if (StateObjectStore)
+		{
+			UUVMeshAABBTrees* TreeStore = NewObject<UUVMeshAABBTrees>();
+			TreeStore->AABBTrees = AABBTrees;
+			StateObjectStore->SetToolStateObject(TreeStore);
+		}
+	}
+
+	// Add the spatial structures to the selection mechanic
+	for (TSharedPtr<FDynamicMeshAABBTree3> Tree : AABBTrees)
+	{
+		SelectionMechanic->AddSpatial(Tree);
+	}
+
+	// Gizmo setup
+	UInteractiveGizmoManager* GizmoManager = GetToolManager()->GetPairedGizmoManager();
+	UTransformProxy* TransformProxy = NewObject<UTransformProxy>(this);
+	TransformGizmo = GizmoManager->CreateCustomTransformGizmo(
+		ETransformGizmoSubElements::TranslateAxisX | ETransformGizmoSubElements::TranslateAxisY | ETransformGizmoSubElements::TranslatePlaneXY
+		//| ETransformGizmoSubElements::ScaleAxisX | ETransformGizmoSubElements::ScaleAxisY | ETransformGizmoSubElements::ScalePlaneXY
+		| ETransformGizmoSubElements::RotateAxisZ,
+		this);
+	TransformProxy->OnBeginTransformEdit.AddUObject(this, &UUVSelectTool::GizmoTransformStarted);
+	TransformProxy->OnTransformChanged.AddUObject(this, &UUVSelectTool::GizmoTransformChanged);
+	TransformProxy->OnEndTransformEdit.AddUObject(this, &UUVSelectTool::GizmoTransformEnded);
+
+	// Always align gizmo to x and y axes
+	TransformGizmo->bUseContextCoordinateSystem = false;
+	TransformGizmo->SetActiveTarget(TransformProxy, GetToolManager());
+
+	if (!SelectionMechanic->GetCurrentSelection().IsEmpty())
+	{
+		OnSelectionChanged();
+	}
+	UpdateGizmo();
+}
+
+void UUVSelectTool::Shutdown(EToolShutdownType ShutdownType)
+{
+	if (StateObjectStore)
+	{
+		UUVMeshSelection* SelectionStore = NewObject<UUVMeshSelection>();
+		SelectionStore->Selection = MakeShared<FDynamicMeshSelection>();
+		*SelectionStore->Selection = SelectionMechanic->GetCurrentSelection();
+		StateObjectStore->SetToolStateObject(SelectionStore);
+	}
+
+	Settings->SaveProperties(this);
+
+	SelectionMechanic->Shutdown();
+
+	// Calls shutdown on gizmo and destroys it.
+	GetToolManager()->GetPairedGizmoManager()->DestroyAllGizmosByOwner(this);
+}
+
+void UUVSelectTool::OnPropertyModified(UObject* PropertySet, FProperty* Property)
+{
+	SelectionMechanic->SelectionMode =
+		Settings->SelectionMode == EUVSelectToolSelectionMode::Island ? EMeshSelectionMechanicMode::Component
+		: EMeshSelectionMechanicMode::Edge;
+}
+
+void UUVSelectTool::UpdateGizmo()
+{
+	const FDynamicMeshSelection& Selection = SelectionMechanic->GetCurrentSelection();
+
+	if (!Selection.IsEmpty())
+	{
+		FVector3d Centroid = FVector3d::Zero();
+		for (int32 Vid : MovingVids)
+		{
+			Centroid += Selection.Mesh->GetVertex(Vid);
+		}
+		Centroid /= MovingVids.Num();
+
+		// TODO: We should either make it so that transparent materials don't affect the rendering
+		// of the gizmo component, or else we should pass in some depth offset into the component.
+		// For now we hack it by raising the gizmo a little off the plane manually.
+		Centroid += FVector3d::ZAxisVector * KINDA_SMALL_NUMBER;
+
+		TransformGizmo->ReinitializeGizmoTransform(FTransform((FVector)Centroid));
+	}
+
+	TransformGizmo->SetVisibility(bGizmoEnabled && !SelectionMechanic->GetCurrentSelection().IsEmpty());
+}
+
+void UUVSelectTool::OnSelectionChanged()
+{
+	const FDynamicMeshSelection& Selection = SelectionMechanic->GetCurrentSelection();
+
+	SelectionMeshIndex = -1;
+	MovingVids.Reset();
+	if (!Selection.IsEmpty())
+	{
+		// Note which mesh we're selecting in.
+		for (int32 i = 0; i < DisplayedMeshes.Num(); ++i)
+		{
+			if (DisplayedMeshes[i]->GetMesh() == Selection.Mesh)
+			{
+				SelectionMeshIndex = i;
+				break;
+			}
+		}
+		check(SelectionMeshIndex >= 0);
+
+		// Note the selected vids
+		TSet<int32> VidSet;
+		if (Selection.Type == FDynamicMeshSelection::EType::Triangle)
+		{
+			for (int32 Tid : Selection.SelectedIDs)
+			{
+				FIndex3i TriVids = Selection.Mesh->GetTriangle(Tid);
+				for (int i = 0; i < 3; ++i)
+				{
+					if (!VidSet.Contains(TriVids[i]))
+					{
+						VidSet.Add(TriVids[i]);
+						MovingVids.Add(TriVids[i]);
+					}
+				}
+			}
+		}
+		else if (Selection.Type == FDynamicMeshSelection::EType::Edge)
+		{
+			for (int32 Eid : Selection.SelectedIDs)
+			{
+				FIndex2i EdgeVids = Selection.Mesh->GetEdgeV(Eid);
+				for (int i = 0; i < 2; ++i)
+				{
+					if (!VidSet.Contains(EdgeVids[i]))
+					{
+						VidSet.Add(EdgeVids[i]);
+						MovingVids.Add(EdgeVids[i]);
+					}
+				}
+			}
+		}
+		else
+		{
+			check(false);
+		}
+	}
+	
+	UpdateGizmo();
+}
+
+void UUVSelectTool::SetGizmoEnabled(bool bEnabledIn)
+{
+	bGizmoEnabled = bEnabledIn;
+
+	// SetGizmoEnabled may be called before or after Setup, hence the check here to see if
+	// the gizmo is set up.
+	if (TransformGizmo)
+	{
+		UpdateGizmo();
+	}
+}
+
+void UUVSelectTool::GizmoTransformStarted(UTransformProxy* Proxy)
+{
+	InitialGizmoFrame = FFrame3d(TransformGizmo->ActiveTarget->GetTransform());
+	MovingVertOriginalPositions.SetNum(MovingVids.Num());
+	const FDynamicMesh3* Mesh = DisplayedMeshes[SelectionMeshIndex]->GetMesh();
+	for (int32 i = 0; i < MovingVids.Num(); ++i)
+	{
+		MovingVertOriginalPositions[i] = InitialGizmoFrame.ToFramePoint(Mesh->GetVertex(MovingVids[i]));
+	}
+}
+
+void UUVSelectTool::GizmoTransformChanged(UTransformProxy* Proxy, FTransform Transform)
+{
+	FFrame3d NewGizmoFrame(TransformGizmo->ActiveTarget->GetTransform());
+	FTransform DeltaTransform = NewGizmoFrame.ToFTransform().GetRelativeTransform(InitialGizmoFrame.ToFTransform());
+	const FDynamicMeshSelection& Selection = SelectionMechanic->GetCurrentSelection();
+
+	if (!DeltaTransform.GetTranslation().IsNearlyZero() || !DeltaTransform.GetRotation().IsIdentity())
+	{
+		DisplayedMeshes[SelectionMeshIndex]->EditMesh([&Selection, &NewGizmoFrame, this](FDynamicMesh3& MeshIn)
+		{
+			for (int32 i = 0; i < MovingVids.Num(); ++i)
+			{
+				MeshIn.SetVertex(MovingVids[i], NewGizmoFrame.FromFramePoint(MovingVertOriginalPositions[i]));
+			}
+		});
+
+		// TODO: This doesn't get broadcast when EditMesh is called. Should it? Or should there be some other way
+		// to inform the wireframe that the mesh changed?
+		DisplayedMeshes[SelectionMeshIndex]->GetOnMeshChanged().Broadcast();
+	}	
+}
+
+void UUVSelectTool::GizmoTransformEnded(UTransformProxy* Proxy)
+{
+	if (!AABBTrees[SelectionMeshIndex]->IsValid())
+	{
+		AABBTrees[SelectionMeshIndex]->Build();
+	}
+}
+
+void UUVSelectTool::Render(IToolsContextRenderAPI* RenderAPI)
+{
+	SelectionMechanic->Render(RenderAPI);
+}
+
+void UUVSelectTool::OnTick(float DeltaTime)
+{
+}
+
+
+#undef LOCTEXT_NAMESPACE
diff --git a/Engine/Plugins/Experimental/UVEditor/Source/UVEditorTools/Public/TargetInterfaces/UVUnwrapDynamicMesh.h b/Engine/Plugins/Experimental/UVEditor/Source/UVEditorTools/Public/TargetInterfaces/UVUnwrapDynamicMesh.h
new file mode 100644
--- /dev/null
+++ b/Engine/Plugins/Experimental/UVEditor/Source/UVEditorTools/Public/TargetInterfaces/UVUnwrapDynamicMesh.h
@@ -0,0 +1,33 @@
+// Copyright Epic Games, Inc. All Rights Reserved.
+
+#pragma once
+
+#include "CoreMinimal.h"
+
+#include "GeometryBase.h"
+#include "UObject/Interface.h"
+
+#include "UVUnwrapDynamicMesh.generated.h"
+
+PREDECLARE_USE_GEOMETRY_CLASS(FDynamicMesh3);
+
+UINTERFACE()
+class UVEDITORTOOLS_API UUVUnwrapDynamicMesh : public UInterface
+{
+	GENERATED_BODY()
+};
+
+/**
+ * Interface that is able to use the UV's of a target to generate a flat FDynamicMesh
+ * representing a given layer, and to bake a modified such mesh back to the UV's of
+ * the target provided that the triangle id's have all stayed the same.
+ */
+class UVEDITORTOOLS_API IUVUnwrapDynamicMesh
+{
+	GENERATED_BODY()
+
+public:
+	virtual int32 GetNumUVLayers() = 0;
+	virtual TSharedPtr<FDynamicMesh3> GetMesh(int32 LayerIndex) = 0;
+	virtual void SaveBackToUVs(const FDynamicMesh3* MeshToSave, int32 LayerIndex) = 0;
+};
\ No newline at end of file
diff --git a/Engine/Plugins/Experimental/UVEditor/Source/UVEditorTools/Public/UVEditorToolsModule.h b/Engine/Plugins/Experimental/UVEditor/Source/UVEditorTools/Public/UVEditorToolsModule.h
new file mode 100644
--- /dev/null
+++ b/Engine/Plugins/Experimental/UVEditor/Source/UVEditorTools/Public/UVEditorToolsModule.h
@@ -0,0 +1,15 @@
+// Copyright Epic Games, Inc. All Rights Reserved.
+
+#pragma once
+
+#include "CoreMinimal.h"
+#include "Modules/ModuleInterface.h"
+
+class FUVEditorToolsModule : public IModuleInterface
+{
+public:
+
+	/** IModuleInterface implementation */
+	virtual void StartupModule() override;
+	virtual void ShutdownModule() override;
+};
\ No newline at end of file
diff --git a/Engine/Plugins/Experimental/UVEditor/Source/UVEditorTools/Public/UVSelectTool.h b/Engine/Plugins/Experimental/UVEditor/Source/UVEditorTools/Public/UVSelectTool.h
new file mode 100644
--- /dev/null
+++ b/Engine/Plugins/Experimental/UVEditor/Source/UVEditorTools/Public/UVSelectTool.h
@@ -0,0 +1,135 @@
+// Copyright Epic Games, Inc. All Rights Reserved.
+
+#pragma once
+
+#include "CoreMinimal.h"
+
+#include "DynamicMeshAABBTree3.h"
+#include "FrameTypes.h"
+#include "InteractiveTool.h"
+#include "InteractiveToolBuilder.h"
+#include "TargetInterfaces/UVUnwrapDynamicMesh.h"
+
+#include "UVSelectTool.generated.h"
+
+class UMeshSelectionMechanic;
+class UToolTargetManager;
+class UTransformGizmo;
+class UTransformProxy;
+class UUVEditorMode;
+class UUVToolStateObjectStore;
+class UPreviewMesh;
+
+UCLASS()
+class UVEDITORTOOLS_API UUVSelectToolBuilder : public UInteractiveToolBuilder
+{
+	GENERATED_BODY()
+public:
+
+	virtual bool CanBuildTool(const FToolBuilderState& SceneState) const override;
+	virtual UInteractiveTool* BuildTool(const FToolBuilderState& SceneState) const override;
+
+	bool bGizmoEnabled = false;
+
+	// This is a pointer so that it can be updated under the builder without
+	// having to set it in the mode after initializing targets.
+	const TArray<TObjectPtr<UPreviewMesh>>* DisplayedMeshes;
+
+	// TODO: This is one way to pass around state between tools. Another could be to
+	// hang the state onto the targets via a tool target manager, or to have something
+	// like UInteractiveToolsSelectionStoreSubsystem. Which should we do?
+	TObjectPtr<UUVToolStateObjectStore> StateObjectStore = nullptr;
+
+protected:
+	virtual const FToolTargetTypeRequirements& GetTargetRequirements() const override;
+};
+
+UENUM()
+enum class EUVSelectToolSelectionMode : uint8
+{
+	Island,
+	Edge
+};
+
+UCLASS()
+class UVEDITORTOOLS_API UUVSelectToolProperties : public UInteractiveToolPropertySet
+{
+	GENERATED_BODY()
+
+public:
+
+	UPROPERTY(EditAnywhere, Category = Options)
+	EUVSelectToolSelectionMode SelectionMode = EUVSelectToolSelectionMode::Island;
+};
+
+
+/**
+ * A tool for selecting elements of a flat FDynamicMesh corresponding to a UV layer of some asset.
+ * If bGizmoEnabled is set to true, the selected elements can be moved around.
+ *
+ * TODO: Doesn't have undo/redo. Will get broken up into pieces later, probably.
+ */
+UCLASS()
+class UVEDITORTOOLS_API UUVSelectTool : public UInteractiveTool
+{
+	GENERATED_BODY()
+
+	using FDynamicMeshAABBTree3 = UE::Geometry::FDynamicMeshAABBTree3;
+
+public:
+	virtual void SetGizmoEnabled(bool bEnabledIn);
+	virtual void SetStateObjectStore(TObjectPtr<UUVToolStateObjectStore> StoreIn) { StateObjectStore = StoreIn; };
+	
+	/**
+	 * The tool will operate on the meshes given here.
+	 */
+	virtual void SetDisplayedMeshes(const TArray<TObjectPtr<UPreviewMesh>>& DisplayedMeshesIn)
+	{
+		DisplayedMeshes = DisplayedMeshesIn;
+	}
+
+	// UInteractiveTool
+	virtual void Setup() override;
+	virtual void Shutdown(EToolShutdownType ShutdownType) override;
+	virtual void OnPropertyModified(UObject* PropertySet, FProperty* Property) override;
+	virtual void Render(IToolsContextRenderAPI* RenderAPI) override;
+	virtual void OnTick(float DeltaTime) override;
+	virtual bool HasCancel() const override { return false; }
+	virtual bool HasAccept() const override { return false; }
+
+protected:
+	virtual void OnSelectionChanged();
+
+	// Callbacks we'll receive from the gizmo proxy
+	virtual void GizmoTransformChanged(UTransformProxy* Proxy, FTransform Transform);
+	virtual void GizmoTransformStarted(UTransformProxy* Proxy);
+	virtual void GizmoTransformEnded(UTransformProxy* Proxy);
+
+	virtual void UpdateGizmo();
+
+	UPROPERTY()
+	TObjectPtr<UUVToolStateObjectStore> StateObjectStore;
+
+	UPROPERTY()
+	TArray<TObjectPtr<UPreviewMesh>> DisplayedMeshes;
+
+	UPROPERTY()
+	TObjectPtr<UUVSelectToolProperties> Settings = nullptr;
+
+	UPROPERTY()
+	TObjectPtr<UMeshSelectionMechanic> SelectionMechanic = nullptr;
+
+	UPROPERTY()
+	UTransformGizmo* TransformGizmo = nullptr;
+
+	TArray<TSharedPtr<FDynamicMeshAABBTree3>> AABBTrees;
+
+	// Placeholders: used for transforming a triangle
+	UE::Geometry::FFrame3d InitialGizmoFrame;
+	TArray<int32> MovingVids;
+	TArray<FVector3d> MovingVertOriginalPositions;
+	int32 SelectionMeshIndex;
+
+	// We need this flag so that SetGizmoVisibility can be called before Setup() by the tool builder.
+	bool bGizmoEnabled = false;
+};
diff --git a/Engine/Plugins/Experimental/UVEditor/Source/UVEditorTools/Public/UVToolStateObjects.h b/Engine/Plugins/Experimental/UVEditor/Source/UVEditorTools/Public/UVToolStateObjects.h
new file mode 100644
--- /dev/null
+++ b/Engine/Plugins/Experimental/UVEditor/Source/UVEditorTools/Public/UVToolStateObjects.h
@@ -0,0 +1,85 @@
+// Copyright Epic Games, Inc. All Rights Reserved.
+
+#pragma once
+
+#include "CoreMinimal.h"
+
+#include "DynamicMeshAABBTree3.h"
+#include "Selection/DynamicMeshSelection.h"
+
+#include "UVToolStateObjects.generated.h"
+
+// This object store system is a simple way for tools to share intermediate
+// structures. The tool builder gets pointed to a state object store, which
+// is a simple map that the tool can query on whether it has a particular type
+// of state object. If it does, the tool can get it out and use it. If it doesn't,
+// it can create it itself and put it in the store for other tools to use later.
+
+// TODO: We should probably have a way to check that the store object is still relevant.
+// The AABBTrees already have change stamps built in, but the selection objects should
+// probably have topology change stamps.
+// We're also not currently bothering to clear the store. It's not really necessary for
+// the current use in the UV editor, but we could have the store track number of accesses
+// and throw away things that don't get accessed much. 
+
+/** Base class to inherit from, mostly to be held together with other state objects. */
+UCLASS()
+class UVEDITORTOOLS_API UUVToolStateObject : public UObject
+{
+	GENERATED_BODY()
+public:
+};
+
+/** Stores a UV mesh AABB tree */
+UCLASS()
+class UVEDITORTOOLS_API UUVMeshAABBTrees : public UUVToolStateObject
+{
+	GENERATED_BODY()
+public:
+
+	TArray<TSharedPtr<UE::Geometry::FDynamicMeshAABBTree3>> AABBTrees;
+};
+
+/** Stores a UV mesh selection */
+UCLASS()
+class UVEDITORTOOLS_API UUVMeshSelection : public UUVToolStateObject
+{
+	GENERATED_BODY()
+public:
+
+	TSharedPtr<UE::Geometry::FDynamicMeshSelection> Selection;
+};
+
+
+// TODO: If we do end up keeping this, it should probably go in its own file, and
+// it doesn't need to be UV specific. We'd also probably need to decide whether it
+// should continue to provide editable objects or give const ones.
+UCLASS()
+class UVEDITORTOOLS_API UUVToolStateObjectStore : public UObject
+{
+	GENERATED_BODY()
+public:
+
+	template<typename DesiredClassOfUObject>
+	TObjectPtr<DesiredClassOfUObject> GetToolStateObject()
+	{
+		TObjectPtr<UUVToolStateObject>* Result = ToolStateObjectMap.Find(DesiredClassOfUObject::StaticClass());
+		return Result ? Cast<DesiredClassOfUObject>(*Result) : nullptr;
+	}
+	void SetToolStateObject(TObjectPtr<UUVToolStateObject> StateObject)
+	{
+		ToolStateObjectMap.Add(StateObject->GetClass(), StateObject);
+	}
+	void RemoveToolStateObject(TObjectPtr<UClass> StateObjectClass)
+	{
+		ToolStateObjectMap.Remove(StateObjectClass);
+	}
+	void Clear()
+	{
+		ToolStateObjectMap.Empty();
+	}
+protected:
+
+	UPROPERTY()
+	TMap<TObjectPtr<UClass>, TObjectPtr<UUVToolStateObject>> ToolStateObjectMap;
+};
\ No newline at end of file
diff --git a/Engine/Plugins/Experimental/UVEditor/Source/UVEditorTools/UVEditorTools.Build.cs b/Engine/Plugins/Experimental/UVEditor/Source/UVEditorTools/UVEditorTools.Build.cs
new file mode 100644
--- /dev/null
+++ b/Engine/Plugins/Experimental/UVEditor/Source/UVEditorTools/UVEditorTools.Build.cs
@@ -0,0 +1,63 @@
+// Copyright Epic Games, Inc. All Rights Reserved.
+
+using UnrealBuildTool;
+
+public class UVEditorTools : ModuleRules
+{
+	public UVEditorTools(ReadOnlyTargetRules Target) : base(Target)
+	{
+		PCHUsage = ModuleRules.PCHUsageMode.UseExplicitOrSharedPCHs;
+		
+		PublicIncludePaths.AddRange(
+			new string[] {
+				// ... add public include paths required here ...
+			}
+			);
+				
+		
+		PrivateIncludePaths.AddRange(
+			new string[] {
+				// ... add other private include paths required here ...
+			}
+			);
+			
+		
+		PublicDependencyModuleNames.AddRange(
+			new string[]
+			{
+				"Core",
+				
+				"InteractiveToolsFramework",
+				
+				// ... add other public dependencies that you statically link with here ...
+			}
+			);
+			
+		
+		PrivateDependencyModuleNames.AddRange(
+			new string[]
+			{
+				// NOTE: UVEditorTools is a separate module so that it doesn't rely on the editor.
+				// So, do not add editor dependencies here.
+				
+				"CoreUObject",
+				"Engine",
+				"Slate",
+				"SlateCore",
+					
+				"MeshConversion",
+				"DynamicMesh",
+				"ModelingComponents",
+				// ... add private dependencies that you statically link with here ...	
+			}
+			);
+		
+		
+		DynamicallyLoadedModuleNames.AddRange(
+			new string[]
+			{
+				// ... add any modules that your module loads dynamically here ...
+			}
+			);
+	}
+}
diff --git a/Engine/Plugins/Experimental/UVEditor/UVEditor.uplugin b/Engine/Plugins/Experimental/UVEditor/UVEditor.uplugin
new file mode 100644
--- /dev/null
+++ b/Engine/Plugins/Experimental/UVEditor/UVEditor.uplugin
@@ -0,0 +1,38 @@
+{
+	"FileVersion": 3,
+	"Version": 1,
+	"VersionName": "0.1",
+	"FriendlyName": "UVEditor",
+	"Description": "Asset editor for modifying the UV's of a mesh, built on the Interactive Tools Framework",
+	"Category": "Other",
+	"CreatedBy": "Epic Games, Inc.",
+	"CreatedByURL": "http://epicgames.com",
+	"DocsURL": "",
+	"MarketplaceURL": "",
+	"SupportURL": "",
+	"CanContainContent": true,
+	"IsExperimentalVersion": true,
+	"Installed": false,
+	"Modules": [
+		{
+			"Name": "UVEditor",
+			"Type": "Editor",
+			"LoadingPhase": "Default"
+		},
+		{
+			"Name": "UVEditorTools",
+			"Type": "Runtime",
+			"LoadingPhase": "Default"
+		}
+	],
+	"Plugins": [
+		{
+			"Name": "GeometryProcessing",
+			"Enabled": true
+		},
+		{
+			"Name": "MeshModelingToolset",
+			"Enabled": true
+		}
+	]
+}
\ No newline at end of file
diff --git a/Engine/Source/Editor/Experimental/EditorInteractiveToolsFramework/Private/EdModeInteractiveToolsContext.cpp b/Engine/Source/Editor/Experimental/EditorInteractiveToolsFramework/Private/EdModeInteractiveToolsContext.cpp
--- a/Engine/Source/Editor/Experimental/EditorInteractiveToolsFramework/Private/EdModeInteractiveToolsContext.cpp
+++ b/Engine/Source/Editor/Experimental/EditorInteractiveToolsFramework/Private/EdModeInteractiveToolsContext.cpp
@@ -980,6 +980,26 @@ bool UEdModeInteractiveToolsContext::InputKey(FEditorViewportClient* ViewportCli
 					return true;
 				}
 			}
+			else if (Key == EKeys::MouseScrollUp || Key == EKeys::MouseScrollDown)
+			{
+				// Note that we get two events for each scroll- an IE_Pressed, and IE_Released.
+				// We pass both of these in, though only the first one will have WheelDelta set.
+				// If a behavior captures the mouse wheel interaction, the second event will give
+				// it the opportunity to immediately release capture. Not passing in the second
+				// event would probably be ok too- capture would get released on next hover event.
+
+				FInputDeviceState InputState = CurrentMouseState;
+				InputState.InputDevice = EInputDevices::Mouse;
+				InputState.SetModifierKeyStates(
+					ViewportClient->IsShiftPressed(), ViewportClient->IsAltPressed(),
+					ViewportClient->IsCtrlPressed(), ViewportClient->IsCmdPressed());
+
+				InputState.Mouse.WheelDelta = (Event != IE_Pressed) ? 0
+					: (Key == EKeys::MouseScrollUp) ? 1 
+					: -1;
+
+				InputRouter->PostInputEvent(InputState);
+			}
 		}
 		else if (Key.IsGamepadKey())
 		{
diff --git a/Engine/Source/Runtime/Experimental/InteractiveToolsFramework/Private/BaseBehaviors/MouseWheelBehavior.cpp b/Engine/Source/Runtime/Experimental/InteractiveToolsFramework/Private/BaseBehaviors/MouseWheelBehavior.cpp
new file mode 100644
--- /dev/null
+++ b/Engine/Source/Runtime/Experimental/InteractiveToolsFramework/Private/BaseBehaviors/MouseWheelBehavior.cpp
@@ -0,0 +1,57 @@
+// Copyright Epic Games, Inc. All Rights Reserved.
+
+#include "BaseBehaviors/MouseWheelBehavior.h"
+
+UMouseWheelInputBehavior::UMouseWheelInputBehavior()
+{
+}
+
+void UMouseWheelInputBehavior::Initialize(IMouseWheelBehaviorTarget* TargetIn)
+{
+	this->Target = TargetIn;
+}
+
+FInputCaptureRequest UMouseWheelInputBehavior::WantsCapture(const FInputDeviceState& Input)
+{
+	if (Input.Mouse.WheelDelta != 0 && (ModifierCheckFunc == nullptr || ModifierCheckFunc(Input)))
+	{
+		FInputRayHit HitResult = Target->ShouldRespondToMouseWheel(FInputDeviceRay(Input));
+		if (HitResult.bHit)
+		{
+			return FInputCaptureRequest::Begin(this, EInputCaptureSide::Any, HitResult.HitDepth);
+		}
+	}
+	return FInputCaptureRequest::Ignore();
+}
+
+
+FInputCaptureUpdate UMouseWheelInputBehavior::BeginCapture(const FInputDeviceState& Input, EInputCaptureSide eSide)
+{
+	check(Input.Mouse.WheelDelta != 0);
+
+	Modifiers.UpdateModifiers(Input, Target);
+
+	if (Input.Mouse.WheelDelta > 0)
+	{
+		Target->OnMouseWheelScrollUp(Input);
+	}
+	else
+	{
+		Target->OnMouseWheelScrollDown(Input);
+	}
+
+	return FInputCaptureUpdate::Begin(this, EInputCaptureSide::Any);
+}
+
+
+FInputCaptureUpdate UMouseWheelInputBehavior::UpdateCapture(const FInputDeviceState& Input, const FInputCaptureData& Data)
+{
+	return FInputCaptureUpdate::End();
+}
+
+
+void UMouseWheelInputBehavior::ForceEndCapture(const FInputCaptureData& data)
+{
+	// nothing to do
+}
+
diff --git a/Engine/Source/Runtime/Experimental/InteractiveToolsFramework/Private/ToolTargetManager.cpp b/Engine/Source/Runtime/Experimental/InteractiveToolsFramework/Private/ToolTargetManager.cpp
--- a/Engine/Source/Runtime/Experimental/InteractiveToolsFramework/Private/ToolTargetManager.cpp
+++ b/Engine/Source/Runtime/Experimental/InteractiveToolsFramework/Private/ToolTargetManager.cpp
@@ -18,7 +18,16 @@ void UToolTargetManager::Shutdown()
 
 void UToolTargetManager::AddTargetFactory(UToolTargetFactory* Factory)
 {
-	Factories.AddUnique(Factory);
+	// If this type of factory has already been added, skip it.
+	if (Factories.ContainsByPredicate(
+		[Factory](UToolTargetFactory* ExistingFactory) { 
+			return ExistingFactory->GetClass() == Factory->GetClass(); 
+		}))
+	{
+		return;
+	}
+
+	Factories.Add(Factory);
 }
 
 bool UToolTargetManager::CanBuildTarget(UObject* SourceObject, const FToolTargetTypeRequirements& TargetType) const
diff --git a/Engine/Source/Runtime/Experimental/InteractiveToolsFramework/Public/BaseBehaviors/BehaviorTargetInterfaces.h b/Engine/Source/Runtime/Experimental/InteractiveToolsFramework/Public/BaseBehaviors/BehaviorTargetInterfaces.h
--- a/Engine/Source/Runtime/Experimental/InteractiveToolsFramework/Public/BaseBehaviors/BehaviorTargetInterfaces.h
+++ b/Engine/Source/Runtime/Experimental/InteractiveToolsFramework/Public/BaseBehaviors/BehaviorTargetInterfaces.h
@@ -96,6 +96,39 @@ public:
 };
 
 
+/**
+ * Functions required to apply mouse wheel behavior
+ */
+class IMouseWheelBehaviorTarget : public IModifierToggleBehaviorTarget
+{
+public:
+	virtual ~IMouseWheelBehaviorTarget() {}
+
+	/**
+	 * The result's bHit property determines whether the mouse wheel action will be captured.
+	 * (Perhaps the mouse wheel only does something when mousing over some part of a mesh)
+	 * 
+	 * @param CurrentPos device position/ray at point where mouse wheel is engaged
+	 * @return hit information at this point
+	 */
+	virtual FInputRayHit ShouldRespondToMouseWheel(const FInputDeviceRay& CurrentPos) = 0;
+
+
+	/**
+	 *
+	 * @param CurrentPos device position/ray at point where mouse wheel is engaged
+	 */
+	virtual void OnMouseWheelScrollUp(const FInputDeviceRay& CurrentPos) = 0;
+
+	/**
+	 *
+	 * @param CurrentPos device position/ray at point where mouse wheel is engaged
+	 */
+	virtual void OnMouseWheelScrollDown(const FInputDeviceRay& CurrentPos) = 0;
+
+};
+
+
 
 
 
diff --git a/Engine/Source/Runtime/Experimental/InteractiveToolsFramework/Public/BaseBehaviors/MouseWheelBehavior.h b/Engine/Source/Runtime/Experimental/InteractiveToolsFramework/Public/BaseBehaviors/MouseWheelBehavior.h
new file mode 100644
--- /dev/null
+++ b/Engine/Source/Runtime/Experimental/InteractiveToolsFramework/Public/BaseBehaviors/MouseWheelBehavior.h
@@ -0,0 +1,52 @@
+// Copyright Epic Games, Inc. All Rights Reserved.
+
+#pragma once
+
+#include "CoreMinimal.h"
+
+#include "BaseBehaviors/AnyButtonInputBehavior.h"
+
+//#include "BehaviorTargetInterfaces.h"
+#include "InputBehaviorModifierStates.h"
+
+#include "MouseWheelBehavior.generated.h"
+
+UCLASS()
+class INTERACTIVETOOLSFRAMEWORK_API UMouseWheelInputBehavior : public UAnyButtonInputBehavior
+{
+	GENERATED_BODY()
+
+public:
+	UMouseWheelInputBehavior();
+
+	/**
+	 * Initialize this behavior with the given Target
+	 * @param Target implementor of hit-test and on-clicked functions
+	 */
+	virtual void Initialize(IMouseWheelBehaviorTarget* Target);
+
+
+	/**
+	 * WantsCapture() will only return capture request if this function returns true (or is null)
+	 */
+	TFunction<bool(const FInputDeviceState&)> ModifierCheckFunc = nullptr;
+
+
+	// UInputBehavior implementation
+	virtual FInputCaptureRequest WantsCapture(const FInputDeviceState& Input) override;
+	virtual FInputCaptureUpdate BeginCapture(const FInputDeviceState& Input, EInputCaptureSide eSide) override;
+	virtual FInputCaptureUpdate UpdateCapture(const FInputDeviceState& Input, const FInputCaptureData& Data) override;
+	virtual void ForceEndCapture(const FInputCaptureData& Data) override;
+
+	/**
+	 * The modifier set for this behavior
+	 */
+	FInputBehaviorModifierStates Modifiers;
+
+protected:
+	/**
+	 *
+	 */
+	IMouseWheelBehaviorTarget* Target;
+
+};
\ No newline at end of file
diff --git a/Engine/Source/Runtime/Experimental/InteractiveToolsFramework/Public/ToolTargetManager.h b/Engine/Source/Runtime/Experimental/InteractiveToolsFramework/Public/ToolTargetManager.h
--- a/Engine/Source/Runtime/Experimental/InteractiveToolsFramework/Public/ToolTargetManager.h
+++ b/Engine/Source/Runtime/Experimental/InteractiveToolsFramework/Public/ToolTargetManager.h
@@ -75,8 +75,13 @@ public:
 	virtual TArray<TObjectPtr<UToolTarget>> BuildAllSelectedTargetable(const FToolBuilderState& SceneState,
 		const FToolTargetTypeRequirements& TargetRequirements);
 
+	/** Initialize the ToolTargetManager. UInteractiveToolsContext calls this, you should not. */
+	virtual void Initialize();
+
+	/** Shutdown the ToolTargetManager. Called by UInteractiveToolsContext. */
+	virtual void Shutdown();
+
 protected:
-	friend class UInteractiveToolsContext;		// to call Initialize/Shutdown
 
 	// This should be removed once tools are transitioned to using tool targets, and the
 	// functions in ComponentSourceInterfaces.h do not exist. For now, this is here to call
@@ -85,12 +90,6 @@ protected:
 
 	UToolTargetManager(){};
 
-	/** Initialize the ToolTargetManager. UInteractiveToolsContext calls this, you should not. */
-	virtual void Initialize();
-
-	/** Shutdown the ToolTargetManager. Called by UInteractiveToolsContext. */
-	virtual void Shutdown();
-
 	/** This flag is set to true on Initialize() and false on Shutdown(). */
 	bool bIsActive = false;
 
diff --git a/Engine/Source/Runtime/MeshDescription/Public/MeshDescription.h b/Engine/Source/Runtime/MeshDescription/Public/MeshDescription.h
--- a/Engine/Source/Runtime/MeshDescription/Public/MeshDescription.h
+++ b/Engine/Source/Runtime/MeshDescription/Public/MeshDescription.h
@@ -831,7 +831,7 @@ public:
 	}
 
 	/** Return UV indices for this triangle for the given channel */
-	TArrayView<FUVID> GetTriangleUVIndices(const FTriangleID TriangleID, int32 UVChannel = 0)
+	TArrayView<FUVID> GetTriangleUVIndices(const FTriangleID TriangleID, int32 UVChannel = 0) const
 	{
 		return TriangleUVs.Get(TriangleID, UVChannel);
 	}
