GeometryProcessing: workaround for FAxisAlignedBox3d::Empty().Extents() result failing NaN check in LWC FVector3. If the mesh has no triangles, the box will never be used, so we don't need to use Empty() box. This does not resolve the larger issue, though.
#rnx

diff --git a/Engine/Plugins/Experimental/GeometryProcessing/Source/GeometricObjects/Public/Spatial/MeshAABBTree3.h b/Engine/Plugins/Experimental/GeometryProcessing/Source/GeometricObjects/Public/Spatial/MeshAABBTree3.h
--- a/Engine/Plugins/Experimental/GeometryProcessing/Source/GeometricObjects/Public/Spatial/MeshAABBTree3.h
+++ b/Engine/Plugins/Experimental/GeometryProcessing/Source/GeometricObjects/Public/Spatial/MeshAABBTree3.h
@@ -1141,7 +1141,8 @@ public:
 		int IStart, int ICount, int Depth, int MinTriCount,
 		FBoxesSet& Tris, FBoxesSet& Nodes, FAxisAlignedBox3d& Box)
 	{
-		Box = FAxisAlignedBox3d::Empty();
+		Box = (Triangles.Num() > 0) ?
+			FAxisAlignedBox3d::Empty() : FAxisAlignedBox3d(FVector3d::Zero(), 0.0);
 		int IBox = -1;
 
 		if (ICount < MinTriCount)
