MeshModelingTools: Fixing group edge inserter not saving some needed state in retriangulate mode after the undo refactor.

#rnx

diff --git a/Engine/Plugins/Experimental/GeometryProcessing/Source/DynamicMesh/Private/Operations/GroupEdgeInserter.cpp b/Engine/Plugins/Experimental/GeometryProcessing/Source/DynamicMesh/Private/Operations/GroupEdgeInserter.cpp
--- a/Engine/Plugins/Experimental/GeometryProcessing/Source/DynamicMesh/Private/Operations/GroupEdgeInserter.cpp
+++ b/Engine/Plugins/Experimental/GeometryProcessing/Source/DynamicMesh/Private/Operations/GroupEdgeInserter.cpp
@@ -1107,20 +1107,33 @@ bool InsertSingleWithRetriangulation(FDynamicMesh3& Mesh, FGroupTopology& Topolo
 		return false;
 	}
 
+	// Function we use to split the endpoints if needed
+	auto AddVertOnEdge = [](FDynamicMesh3& Mesh, int32 Eid, double EdgeTValue, TSet<int32>* ChangedTidsOut) {
+		FDynamicMesh3::FEdgeSplitInfo SplitInfo;
+		Mesh.SplitEdge(Eid, SplitInfo, EdgeTValue);
+
+		// Record the changed triangles
+		if (ChangedTidsOut)
+		{
+			ChangedTidsOut->Add(SplitInfo.OriginalTriangles.A);
+			if (SplitInfo.OriginalTriangles.B != FDynamicMesh3::InvalidID)
+			{
+				ChangedTidsOut->Add(SplitInfo.OriginalTriangles.B);
+			}
+		}
+		return SplitInfo.NewVertex;
+	};
+
 	int32 StartVid = StartPoint.ElementID;
 	if (!StartPoint.bIsVertex)
 	{
-		FDynamicMesh3::FEdgeSplitInfo SplitInfo;
-		Mesh.SplitEdge(StartPoint.ElementID, SplitInfo, StartPoint.EdgeTValue);
-		StartVid = SplitInfo.NewVertex;
+		StartVid = AddVertOnEdge(Mesh, StartPoint.ElementID, StartPoint.EdgeTValue, OptionalOut.ChangedTidsOut);
 	}
 
 	int32 EndVid = EndPoint.ElementID;
 	if (!EndPoint.bIsVertex)
 	{
-		FDynamicMesh3::FEdgeSplitInfo SplitInfo;
-		Mesh.SplitEdge(EndPoint.ElementID, SplitInfo, EndPoint.EdgeTValue);
-		EndVid = SplitInfo.NewVertex;
+		EndVid = AddVertOnEdge(Mesh, EndPoint.ElementID, EndPoint.EdgeTValue, OptionalOut.ChangedTidsOut);
 	}
 
 	const FGroupTopology::FGroup* Group = Topology.FindGroupByID(GroupID);
