GeometryProcessing updates for new MergeActors method:
- Add ECorrespondenceStrategy enum/setting to FMeshImageBakingCache to allow for different lowpoly/highpoly mapping polices, and implement those various policies. Add  FindSamplingHoles() to allow for image postprocessing to repair sampling failures.
- Add TImateBuilder::ContainsPixel()
- Add TMathUtil::Tan() function wrapper
- Add support for "Build Scale" in FMeshDescriptionTriangleMeshAdapter. StaticMesh MeshDescriptions can have an arbitrary scaling factor applied, this allows incorporating that scaling without having to make a mesh copy (at cost of a multiply for each vertex access). Add Vectortype-templated GetTriVertices() function.
- Add FWindingNumberBasedSolidify variant of TImplicitSolidify that supports an external Winding-Number-Like function, and does not have an explicit Mesh/AABBTree dependency. This allows us to mesh the winding number an arbitrary assembly of objects.
#rnx

diff --git a/Engine/Plugins/Experimental/GeometryProcessing/Source/DynamicMesh/Private/Sampling/MeshImageBakingCache.cpp b/Engine/Plugins/Experimental/GeometryProcessing/Source/DynamicMesh/Private/Sampling/MeshImageBakingCache.cpp
--- a/Engine/Plugins/Experimental/GeometryProcessing/Source/DynamicMesh/Private/Sampling/MeshImageBakingCache.cpp
+++ b/Engine/Plugins/Experimental/GeometryProcessing/Source/DynamicMesh/Private/Sampling/MeshImageBakingCache.cpp
@@ -26,28 +26,49 @@ void FMeshImageBakingCache::SetBakeTargetMesh(const FDynamicMesh3* Mesh)
 
 void FMeshImageBakingCache::SetDimensions(FImageDimensions DimensionsIn)
 {
-	Dimensions = DimensionsIn;
-	InvalidateSamples();
-	InvalidateOccupancy();
+	if (Dimensions != DimensionsIn)
+	{
+		Dimensions = DimensionsIn;
+		InvalidateSamples();
+		InvalidateOccupancy();
+	}
 }
 
 
 void FMeshImageBakingCache::SetUVLayer(int32 UVLayerIn)
 {
-	UVLayer = UVLayerIn;
-	InvalidateSamples();
-	InvalidateOccupancy();
+	if (UVLayer != UVLayerIn)
+	{
+		UVLayer = UVLayerIn;
+		InvalidateSamples();
+		InvalidateOccupancy();
+	}
 }
 
 
 void FMeshImageBakingCache::SetThickness(double ThicknessIn)
 {
-	Thickness = ThicknessIn;
-	InvalidateSamples();
-	InvalidateOccupancy();   // do we need to do this?
+	if (Thickness != ThicknessIn)
+	{
+		Thickness = ThicknessIn;
+		InvalidateSamples();
+		InvalidateOccupancy();   // do we need to do this?
+	}
 }
 
 
+void FMeshImageBakingCache::SetCorrespondenceStrategy(ECorrespondenceStrategy Strategy)
+{
+	if (CorrespondenceStrategy != Strategy)
+	{
+		CorrespondenceStrategy = Strategy;
+		InvalidateSamples();
+		InvalidateOccupancy();   // do we need to do this?
+	}
+}
+
+
+
 const FDynamicMeshNormalOverlay* FMeshImageBakingCache::GetDetailNormals() const
 {
 	check(DetailMesh && DetailMesh->HasAttributes());
@@ -97,18 +118,20 @@ void FMeshImageBakingCache::InvalidateOccupancy()
  *    3) cast a ray inwards along -Normal from BasePoint
  * We take (1) preferentially, and then (2), and then (3)
  * 
- * If all of those fail, we try to find a nearest-point within distance Thickess
+ * If all of those fail, if bFailToNearestPoint is true we fall back to nearest-point, 
  * 
  * If all the above fail, return false
  */
-static bool GetDetailTrianglePoint(
+static bool GetDetailTrianglePoint_Raycast(
 	const FDynamicMesh3& DetailMesh,
 	const FDynamicMeshAABBTree3& DetailSpatial,
 	const FVector3d& BasePoint,
 	const FVector3d& BaseNormal,
 	int32& DetailTriangleOut,
 	FVector3d& DetailTriBaryCoords,
-	double Thickness)
+	double Thickness,
+	bool bFailToNearestPoint
+)
 {
 	// TODO: should we check normals here? inverse normal should probably not be considered valid
 
@@ -173,15 +196,19 @@ static bool GetDetailTrianglePoint(
 		return true;
 	}
 
-	// if we get this far, both rays missed, so use absolute nearest point regardless of distance
-	//NearestTriID = DetailSpatial.FindNearestTriangle(BasePoint, NearDistSqr);
-	//if (DetailMesh.IsTriangle(NearestTriID))
-	//{
-	//	DetailTriangleOut = NearestTriID;
-	//	FDistPoint3Triangle3d DistQuery = TMeshQueries<FDynamicMesh3>::TriangleDistance(DetailMesh, NearestTriID, BasePoint);
-	//	DetailTriBaryCoords = DistQuery.TriangleBaryCoords;
-	//	return true;
-	//}
+	// if we get this far, all rays missed or were too far, so use absolute nearest point regardless of distance
+	if (bFailToNearestPoint)
+	{
+		double NearDistSqr = 0;
+		int32 NearestTriID = DetailSpatial.FindNearestTriangle(BasePoint, NearDistSqr);
+		if (DetailMesh.IsTriangle(NearestTriID))
+		{
+			DetailTriangleOut = NearestTriID;
+			FDistPoint3Triangle3d DistQuery = TMeshQueries<FDynamicMesh3>::TriangleDistance(DetailMesh, NearestTriID, BasePoint);
+			DetailTriBaryCoords = DistQuery.TriangleBaryCoords;
+			return true;
+		}
+	}
 
 	return false;
 }
@@ -190,6 +217,31 @@ static bool GetDetailTrianglePoint(
 
 
 
+/**
+ * Find point on Detail mesh that corresponds to point on Base mesh using minimum distance
+ */
+static bool GetDetailTrianglePoint_Nearest(
+	const FDynamicMesh3& DetailMesh,
+	const FDynamicMeshAABBTree3& DetailSpatial,
+	const FVector3d& BasePoint,
+	int32& DetailTriangleOut,
+	FVector3d& DetailTriBaryCoords)
+{
+	double NearDistSqr = 0;
+	int32 NearestTriID = DetailSpatial.FindNearestTriangle(BasePoint, NearDistSqr);
+	if (DetailMesh.IsTriangle(NearestTriID))
+	{
+		DetailTriangleOut = NearestTriID;
+		FDistPoint3Triangle3d DistQuery = TMeshQueries<FDynamicMesh3>::TriangleDistance(DetailMesh, NearestTriID, BasePoint);
+		DetailTriBaryCoords = DistQuery.TriangleBaryCoords;
+		return true;
+	}
+
+	return false;
+}
+
+
+
 
 bool FMeshImageBakingCache::ValidateCache()
 {
@@ -228,27 +280,46 @@ bool FMeshImageBakingCache::ValidateCache()
 
 	if (bSamplesValid == false)
 	{
+		ECorrespondenceStrategy UseStrategy = this->CorrespondenceStrategy;
+		if (UseStrategy == ECorrespondenceStrategy::Identity && ensure(DetailMesh == Mesh) == false)
+		{
+			// Identity strategy requires mesh to be the same. Could potentially have two copies, in which
+			// case this ensure is too conservative, but for now we will assume this
+			UseStrategy = ECorrespondenceStrategy::NearestPoint;
+		}
+
 		// this sampler finds the correspondence between base surface and detail surface
 		TMeshSurfaceUVSampler<FCorrespondenceSample> DetailMeshSampler;
 		DetailMeshSampler.Initialize(Mesh, UVOverlay, EMeshSurfaceSamplerQueryType::TriangleAndUV, FCorrespondenceSample(),
-			[Mesh, NormalOverlay, this](const FMeshUVSampleInfo& SampleInfo, FCorrespondenceSample& ValueOut)
+			[Mesh, NormalOverlay, UseStrategy, this](const FMeshUVSampleInfo& SampleInfo, FCorrespondenceSample& ValueOut)
 		{
-			//FVector3d BaseTriNormal = Mesh->GetTriNormal(SampleInfo.TriangleIndex);
 			NormalOverlay->GetTriBaryInterpolate<double>(SampleInfo.TriangleIndex, &SampleInfo.BaryCoords.X, &ValueOut.BaseNormal.X);
 			Normalize(ValueOut.BaseNormal);
 			FVector3d RayDir = ValueOut.BaseNormal;
 
 			ValueOut.BaseSample = SampleInfo;
+			ValueOut.DetailTriID = FDynamicMesh3::InvalidID;
 
-			double SampleThickness = this->GetThickness();		// could modulate w/ a map here...
-
-			// find detail mesh triangle point
-			bool bFoundTri = GetDetailTrianglePoint(*DetailMesh, *DetailSpatial, SampleInfo.SurfacePoint, RayDir,
-				ValueOut.DetailTriID, ValueOut.DetailBaryCoords, SampleThickness);
-			if (!bFoundTri)
+			if (UseStrategy == ECorrespondenceStrategy::Identity)
 			{
-				ValueOut.DetailTriID = FDynamicMesh3::InvalidID;
+				ValueOut.DetailTriID = SampleInfo.TriangleIndex;
+				ValueOut.DetailBaryCoords = SampleInfo.BaryCoords;
+			}
+			else if (UseStrategy == ECorrespondenceStrategy::NearestPoint)
+			{
+				bool bFoundTri = GetDetailTrianglePoint_Nearest(*DetailMesh, *DetailSpatial, SampleInfo.SurfacePoint,
+					ValueOut.DetailTriID, ValueOut.DetailBaryCoords);
+			}
+			else	// fall back to raycast strategy
+			{
+				double SampleThickness = this->GetThickness();		// could modulate w/ a map here...
+
+				// find detail mesh triangle point
+				bool bFoundTri = GetDetailTrianglePoint_Raycast(*DetailMesh, *DetailSpatial, SampleInfo.SurfacePoint, RayDir,
+					ValueOut.DetailTriID, ValueOut.DetailBaryCoords, SampleThickness, 
+					(UseStrategy == ECorrespondenceStrategy::RaycastStandardThenNearest) );
 			}
+
 		});
 
 
@@ -316,5 +387,44 @@ void FMeshImageBakingCache::EvaluateSamples(
 			SampleFunction(Coords, Sample);
 		}
 
+	}, !bParallel ? EParallelForFlags::ForceSingleThread : EParallelForFlags::None);
+}
+
+
+
+
+void FMeshImageBakingCache::FindSamplingHoles(
+	TFunctionRef<bool(const FVector2i&)> IsInvalidSampleFunction,
+	TArray<FVector2i>& HolePixelsOut,
+	bool bParallel) const
+{
+	check(IsCacheValid());
+
+	FCriticalSection HoleLock;
+
+	ParallelFor(Dimensions.GetHeight(), [this, &IsInvalidSampleFunction, &HoleLock, &HolePixelsOut](int32 ImgY)
+	{
+		if (CancelF())
+		{
+			return;
+		}
+		for (int32 ImgX = 0; ImgX < Dimensions.GetWidth(); ImgX++)
+		{
+			int64 LinearIdx = Dimensions.GetIndex(ImgX, ImgY);
+			if (OccupancyMap->IsInterior(LinearIdx) == false)
+			{
+				continue;
+			}
+			FVector2i Coords(ImgX, ImgY);
+
+			const FCorrespondenceSample& Sample = SampleMap[LinearIdx];
+			if (IsInvalidSampleFunction(Coords))
+			{
+				HoleLock.Lock();
+				HolePixelsOut.Add(Coords);
+				HoleLock.Unlock();
+			}
+		}
+
 	}, !bParallel ? EParallelForFlags::ForceSingleThread : EParallelForFlags::None);
 }
\ No newline at end of file
diff --git a/Engine/Plugins/Experimental/GeometryProcessing/Source/DynamicMesh/Public/Sampling/MeshImageBakingCache.h b/Engine/Plugins/Experimental/GeometryProcessing/Source/DynamicMesh/Public/Sampling/MeshImageBakingCache.h
--- a/Engine/Plugins/Experimental/GeometryProcessing/Source/DynamicMesh/Public/Sampling/MeshImageBakingCache.h
+++ b/Engine/Plugins/Experimental/GeometryProcessing/Source/DynamicMesh/Public/Sampling/MeshImageBakingCache.h
@@ -19,16 +19,35 @@ class DYNAMICMESH_API FMeshImageBakingCache
 {
 public:
 
+	/**
+	 * ECorrespondenceStrategy determines the basic approach that will be used to establish a 
+	 * mapping from points on the BakeTarget Mesh (usually low-poly) to points on the Detail Mesh (eg highpoly).
+	 * Geometrically this is not a 1-1 mapping so there are various options
+	 */
+	enum class ECorrespondenceStrategy
+	{
+		/** Raycast inwards from Point+Thickness*Normal, if that misses, try Outwards from Point, then Inwards from Point */
+		RaycastStandard,
+		/** Use geometrically nearest point. Thickness is ignored */
+		NearestPoint,
+		/** Use RaycastStandard but fall back to NearestPoint if none of the rays hit */
+		RaycastStandardThenNearest,
+		/** Assume that BakeTarget == DetailMesh and so no mapping is necessary */
+		Identity
+	};
+
 	void SetDetailMesh(const FDynamicMesh3* Mesh, const FDynamicMeshAABBTree3* Spatial);
 	void SetBakeTargetMesh(const FDynamicMesh3* Mesh);
 
 	void SetDimensions(FImageDimensions Dimensions);
 	void SetUVLayer(int32 UVLayer);
 	void SetThickness(double Thickness);
+	void SetCorrespondenceStrategy(ECorrespondenceStrategy Strategy);
 
 	FImageDimensions GetDimensions() const { return Dimensions; }
 	int32 GetUVLayer() const { return UVLayer; }
 	double GetThickness() const { return Thickness; }
+	ECorrespondenceStrategy GetCorrespondenceStrategy() const { return CorrespondenceStrategy; }
 
 	const FDynamicMesh3* GetBakeTargetMesh() const { return TargetMesh; }
 	const FDynamicMeshUVOverlay* GetBakeTargetUVs() const;
@@ -59,6 +78,20 @@ public:
 
 	const FImageOccupancyMap* GetOccupancyMap() const;
 
+
+	/**
+	 * Iterate over valid samples in the Occupancy map to find "holes", ie pixels
+	 * where no valid sample could be computed according to IsValidSampleFunction. 
+	 * This can be used after a bake to find locations where the baked map needs to
+	 * be repaired as a post-process (eg using image infilling/etc), which can produce 
+	 * better results than (eg) using a "worse" sampling strategy to guarantee some
+	 * non-garbage result (eg Nearest-point sampling)
+	 */
+	void FindSamplingHoles(TFunctionRef<bool(const FVector2i&)> IsValidSampleFunction,
+		TArray<FVector2i>& HolePixelsOut,
+		bool bParallel = true) const;
+
+
 	/** if this function returns true, we should abort calculation */
 	TFunction<bool(void)> CancelF = []() { return false; };
 
@@ -67,9 +100,10 @@ protected:
 	const FDynamicMeshAABBTree3* DetailSpatial = nullptr;
 	const FDynamicMesh3* TargetMesh = nullptr;
 
-	FImageDimensions Dimensions;
-	int32 UVLayer;
+	FImageDimensions Dimensions = FImageDimensions(128,128);
+	int32 UVLayer = 0;
 	double Thickness = 3.0;
+	ECorrespondenceStrategy CorrespondenceStrategy = ECorrespondenceStrategy::RaycastStandard;
 
 	TDenseGrid2<FCorrespondenceSample> SampleMap;
 	bool bSamplesValid = false;
diff --git a/Engine/Plugins/Experimental/GeometryProcessing/Source/GeometricObjects/Public/Image/ImageBuilder.h b/Engine/Plugins/Experimental/GeometryProcessing/Source/GeometricObjects/Public/Image/ImageBuilder.h
--- a/Engine/Plugins/Experimental/GeometryProcessing/Source/GeometricObjects/Public/Image/ImageBuilder.h
+++ b/Engine/Plugins/Experimental/GeometryProcessing/Source/GeometricObjects/Public/Image/ImageBuilder.h
@@ -43,6 +43,14 @@ public:
 		Image.AssignAll(ClearValue);
 	}
 
+	/**
+	 * @return true if the given X/Y coordinates are inside the image pixel bounds, ie can be used to index a pixel
+	 */
+	bool ContainsPixel(const FVector2i& ImageCoords) const
+	{
+		return ImageCoords.X >= 0 && ImageCoords.Y >= 0 &&
+			ImageCoords.X < Dimensions.GetWidth() && ImageCoords.Y < Dimensions.GetHeight();
+	}
 
 	/**
 	 * Get the Pixel at the given X/Y coordinates
diff --git a/Engine/Plugins/Experimental/GeometryProcessing/Source/GeometricObjects/Public/Image/ImageDimensions.h b/Engine/Plugins/Experimental/GeometryProcessing/Source/GeometricObjects/Public/Image/ImageDimensions.h
--- a/Engine/Plugins/Experimental/GeometryProcessing/Source/GeometricObjects/Public/Image/ImageDimensions.h
+++ b/Engine/Plugins/Experimental/GeometryProcessing/Source/GeometricObjects/Public/Image/ImageDimensions.h
@@ -132,6 +132,11 @@ public:
 	{
 		return Width == Other.Width && Height == Other.Height;
 	}
+
+	bool operator!=(const FImageDimensions& Other) const
+	{
+		return Width != Other.Width || Height != Other.Height;
+	}
 };
 
 
diff --git a/Engine/Plugins/Experimental/GeometryProcessing/Source/GeometricObjects/Public/Implicit/Solidify.h b/Engine/Plugins/Experimental/GeometryProcessing/Source/GeometricObjects/Public/Implicit/Solidify.h
--- a/Engine/Plugins/Experimental/GeometryProcessing/Source/GeometricObjects/Public/Implicit/Solidify.h
+++ b/Engine/Plugins/Experimental/GeometryProcessing/Source/GeometricObjects/Public/Implicit/Solidify.h
@@ -87,7 +87,7 @@ public:
 	}
 
 	/**
-	 * @param bReuseComputed If true, will attempt to reuse previously-computed AABB trees and SDFs where possible
+	 * 
 	 */
 	const FMeshShapeGenerator& Generate()
 	{
@@ -142,5 +142,150 @@ public:
 };
 
 
+
+
+
+
+
+/**
+ * Use marching cubes to remesh an arbitrary function that provides a winding-number like scalar value to a solid surface
+ */
+class FWindingNumberBasedSolidify
+{
+public:
+
+	FWindingNumberBasedSolidify(
+		TUniqueFunction<double(const FVector3d&)> WindingFunctionIn,
+		const FAxisAlignedBox3d& BoundsIn,
+		const TArray<FVector3d>& SeedPointsIn)
+	{
+		this->WindingFunction = MoveTemp(WindingFunctionIn);
+		this->FunctionBounds = BoundsIn;
+		this->SeedPoints = SeedPointsIn;
+	}
+
+	virtual ~FWindingNumberBasedSolidify()
+	{
+	}
+
+	///
+	/// Inputs
+	///
+
+	/** External Winding-Number Function */
+	TUniqueFunction<double(const FVector3d&)> WindingFunction = [](const FVector3d& Pos) { return 0.0; };
+
+	/** Bounds within which we will mesh things */
+	FAxisAlignedBox3d FunctionBounds;
+
+	/** Seed points for meshing */
+	TArray<FVector3d> SeedPoints;
+
+	/** Inside/outside winding number threshold */
+	double WindingThreshold = .5;
+
+	/** How much to extend bounds considered by marching cubes outside the original surface bounds */
+	double ExtendBounds = 1;
+
+	/** What to do if the surface extends outside the marching cubes bounds -- if true, puts a solid surface at the boundary */
+	bool bSolidAtBoundaries = true;
+
+	/** How many binary search steps to do when placing surface at boundary */
+	int SurfaceSearchSteps = 4;
+
+	/** size of the cells used when meshing the output (marching cubes' cube size) */
+	double MeshCellSize = 1.0;
+
+	/**
+	 * Set cell size to hit the target voxel count along the max dimension of the bounds
+	 */
+	void SetCellSizeAndExtendBounds(const FAxisAlignedBox3d& Bounds, double ExtendBoundsIn, int TargetOutputVoxelCount)
+	{
+		ExtendBounds = ExtendBoundsIn;
+		MeshCellSize = (Bounds.MaxDim() + ExtendBounds * 2.0) / double(TargetOutputVoxelCount);
+	}
+
+	/** if this function returns true, we should abort calculation */
+	TFunction<bool(void)> CancelF = []()
+	{
+		return false;
+	};
+	
+protected:
+
+	FMarchingCubes MarchingCubes;
+
+public:
+
+	/**
+	 * @return true if input parameters are valid
+	 */
+	bool Validate()
+	{
+		bool bValidParams = SurfaceSearchSteps >= 0 && MeshCellSize > 0 && FMath::IsFinite(MeshCellSize);
+		return bValidParams;
+	}
+
+	/**
+	 * 
+	 */
+	const FMeshShapeGenerator& Generate()
+	{
+		MarchingCubes.Reset();
+		if (!ensure(Validate()))
+		{
+			// give up and return and empty result on invalid parameters
+			return MarchingCubes;
+		}
+
+		FAxisAlignedBox3d InternalBounds = FunctionBounds;
+		InternalBounds.Expand(ExtendBounds);
+		
+		MarchingCubes.CubeSize = MeshCellSize;
+
+		MarchingCubes.Bounds = InternalBounds;
+
+		// expand marching cubes bounds beyond the 'internal' bounds to ensure we sample outside the bounds, if solid-at-boundaries is requested
+		if (bSolidAtBoundaries)
+		{
+			MarchingCubes.Bounds.Expand(MeshCellSize * .1);
+		}
+
+		MarchingCubes.RootMode = ERootfindingModes::Bisection;
+		MarchingCubes.RootModeSteps = SurfaceSearchSteps;
+		MarchingCubes.IsoValue = WindingThreshold;
+		MarchingCubes.CancelF = CancelF;
+
+		if (bSolidAtBoundaries)
+		{
+			MarchingCubes.Implicit = [this, InternalBounds](const FVector3d& Pos)
+			{
+				return InternalBounds.Contains(Pos) ? WindingFunction(Pos) : -(WindingThreshold + 1);
+			};
+		}
+		else
+		{
+			MarchingCubes.Implicit = [this](const FVector3d& Pos)
+			{
+				return WindingFunction(Pos);
+			};
+		}
+
+		TArray<FVector3<double>> MCSeeds;
+		for ( const FVector3d& SeedPoint : SeedPoints )
+		{
+			MCSeeds.Add(SeedPoint);
+		}
+		MarchingCubes.GenerateContinuation(MCSeeds);
+
+		return MarchingCubes;
+	}
+};
+
+
+
+
+
+
 } // end namespace UE::Geometry
 } // end namespace UE
\ No newline at end of file
diff --git a/Engine/Plugins/Experimental/GeometryProcessing/Source/GeometricObjects/Public/MathUtil.h b/Engine/Plugins/Experimental/GeometryProcessing/Source/GeometricObjects/Public/MathUtil.h
--- a/Engine/Plugins/Experimental/GeometryProcessing/Source/GeometricObjects/Public/MathUtil.h
+++ b/Engine/Plugins/Experimental/GeometryProcessing/Source/GeometricObjects/Public/MathUtil.h
@@ -154,6 +154,7 @@ public:
 	/** compute min and max of a,b,c with max 3 comparisons (sometimes 2) */
 	static inline void MinMax(RealType A, RealType B, RealType C, RealType& MinOut, RealType& MaxOut);
 	static inline RealType Sqrt(const RealType Value);
+	static inline RealType Tan(const RealType Value);
 	static inline RealType Atan2(const RealType ValueY, const RealType ValueX);
 	static inline RealType Sin(const RealType Value);
 	static inline RealType Cos(const RealType Value);
@@ -297,6 +298,12 @@ RealType TMathUtil<RealType>::Sqrt(const RealType Value)
 	return sqrt(Value);
 }
 
+template<typename RealType>
+RealType TMathUtil<RealType>::Tan(const RealType Value)
+{
+	return tan(Value);
+}
+
 template<typename RealType>
 RealType TMathUtil<RealType>::Atan2(const RealType ValueY, const RealType ValueX)
 {
diff --git a/Engine/Plugins/Experimental/GeometryProcessing/Source/MeshConversion/Public/MeshDescriptionAdapter.h b/Engine/Plugins/Experimental/GeometryProcessing/Source/MeshConversion/Public/MeshDescriptionAdapter.h
--- a/Engine/Plugins/Experimental/GeometryProcessing/Source/MeshConversion/Public/MeshDescriptionAdapter.h
+++ b/Engine/Plugins/Experimental/GeometryProcessing/Source/MeshConversion/Public/MeshDescriptionAdapter.h
@@ -25,6 +25,9 @@ protected:
 	TVertexAttributesConstRef<FVector> VertexPositions;
 	TVertexInstanceAttributesConstRef<FVector> VertexInstanceNormals;
 
+	FVector3d BuildScale = FVector3d::One();
+	bool bScaleNormals = false;
+
 public:
 	FMeshDescriptionTriangleMeshAdapter(const FMeshDescription* MeshIn) : Mesh(MeshIn)
 	{
@@ -34,6 +37,12 @@ public:
 		// @todo: can we hold TArrayViews of the attribute arrays here? Do we guarantee not to mutate the mesh description for the duration of this object?
 	}
 
+	void SetBuildScale(const FVector3d& BuildScaleIn, bool bScaleNormalsIn)
+	{
+		BuildScale = BuildScaleIn;
+		bScaleNormals = bScaleNormalsIn;
+	}
+
 	bool IsTriangle(int32 TID) const
 	{
 		return TID >= 0 && TID < Mesh->Triangles().Num();
@@ -72,17 +81,32 @@ public:
 	}
 	FVector3d GetVertex(int32 IDValue) const
 	{
-		return FVector3d(VertexPositions[FVertexID(IDValue)]);
+		const FVector& Position = VertexPositions[FVertexID(IDValue)];
+		return FVector3d(BuildScale.X * (double)Position.X, BuildScale.Y * (double)Position.Y, BuildScale.Z * (double)Position.Z);
 	}
 
 	inline void GetTriVertices(int32 IDValue, FVector3d& V0, FVector3d& V1, FVector3d& V2) const
 	{
 		TArrayView<const FVertexID> TriVertIDs = Mesh->GetTriangleVertices(FTriangleID(IDValue));
-		V0 = FVector3d(VertexPositions[TriVertIDs[0]]);
-		V1 = FVector3d(VertexPositions[TriVertIDs[1]]);
-		V2 = FVector3d(VertexPositions[TriVertIDs[2]]);
+		const FVector& A = VertexPositions[TriVertIDs[0]];
+		V0 = FVector3d(BuildScale.X * (double)A.X, BuildScale.Y * (double)A.Y, BuildScale.Z * (double)A.Z);
+		const FVector& B = VertexPositions[TriVertIDs[1]];
+		V1 = FVector3d(BuildScale.X * (double)B.X, BuildScale.Y * (double)B.Y, BuildScale.Z * (double)B.Z);
+		const FVector& C = VertexPositions[TriVertIDs[2]];
+		V2 = FVector3d(BuildScale.X * (double)C.X, BuildScale.Y * (double)C.Y, BuildScale.Z * (double)C.Z);
 	}
 
+	template<typename VectorType>
+	inline void GetTriVertices(int32 IDValue, VectorType& V0, VectorType& V1, VectorType& V2) const
+	{
+		TArrayView<const FVertexID> TriVertIDs = Mesh->GetTriangleVertices(FTriangleID(IDValue));
+		const FVector& A = VertexPositions[TriVertIDs[0]];
+		V0 = VectorType(BuildScale.X * (double)A.X, BuildScale.Y * (double)A.Y, BuildScale.Z * (double)A.Z);
+		const FVector& B = VertexPositions[TriVertIDs[1]];
+		V1 = VectorType(BuildScale.X * (double)B.X, BuildScale.Y * (double)B.Y, BuildScale.Z * (double)B.Z);
+		const FVector& C = VertexPositions[TriVertIDs[2]];
+		V2 = VectorType(BuildScale.X * (double)C.X, BuildScale.Y * (double)C.Y, BuildScale.Z * (double)C.Z);
+	}
 
 	inline bool HasNormals() const
 	{
@@ -102,7 +126,9 @@ public:
 	}
 	FVector3f GetNormal(int32 IDValue) const
 	{
-		return FVector3f(VertexInstanceNormals[FVertexInstanceID(IDValue)]);
+		const FVector& InstanceNormal = VertexInstanceNormals[FVertexInstanceID(IDValue)];
+		return (!bScaleNormals) ? FVector3f(InstanceNormal) :
+			UE::Geometry::Normalized(FVector3f(InstanceNormal.X/BuildScale.X, InstanceNormal.Y/BuildScale.Y, InstanceNormal.Z/BuildScale.Z));
 	}
 };
 
