GeometryProcessing: add basic ImageInfilling implementation and World-Position-Based Color/Normal Bakers
#rnx

diff --git a/Engine/Plugins/Experimental/GeometryProcessing/Source/DynamicMesh/Private/Sampling/MeshGenericWorldPositionBaker.cpp b/Engine/Plugins/Experimental/GeometryProcessing/Source/DynamicMesh/Private/Sampling/MeshGenericWorldPositionBaker.cpp
new file mode 100644
--- /dev/null
+++ b/Engine/Plugins/Experimental/GeometryProcessing/Source/DynamicMesh/Private/Sampling/MeshGenericWorldPositionBaker.cpp
@@ -0,0 +1,104 @@
+// Copyright Epic Games, Inc. All Rights Reserved.
+
+#include "Sampling/MeshGenericWorldPositionBaker.h"
+
+#include "ExplicitUseGeometryMathTypes.h"		// using UE::Geometry::(math types)
+using namespace UE::Geometry;
+
+void FMeshGenericWorldPositionColorBaker::Bake()
+{
+	const FMeshImageBakingCache* BakeCache = GetCache();
+	check(BakeCache);
+
+	auto SampleFunction = [&](const FMeshImageBakingCache::FCorrespondenceSample& SampleData)
+	{
+		FVector3d Position = SampleData.BaseSample.SurfacePoint;
+		FVector3d Normal = SampleData.BaseNormal;
+		int32 TriangleID = SampleData.BaseSample.TriangleIndex;
+		FVector3d BaryCoords = SampleData.BaseSample.BaryCoords;
+
+		FVector4f Color = ColorSampleFunction(Position, Normal);
+
+		return Color;
+	};
+
+	ResultBuilder = MakeUnique<TImageBuilder<FVector4f>>();
+	ResultBuilder->SetDimensions(BakeCache->GetDimensions());
+	ResultBuilder->Clear(DefaultColor);
+
+	BakeCache->EvaluateSamples([&](const FVector2i& Coords, const FMeshImageBakingCache::FCorrespondenceSample& Sample)
+	{
+		FVector4f Color = SampleFunction(Sample);
+		ResultBuilder->SetPixel(Coords, Color);
+	});
+
+
+
+	const FImageOccupancyMap& Occupancy = *BakeCache->GetOccupancyMap();
+	for (int64 k = 0; k < Occupancy.GutterTexels.Num(); k++)
+	{
+		TPair<int64, int64> GutterTexel = Occupancy.GutterTexels[k];
+		ResultBuilder->CopyPixel(GutterTexel.Value, GutterTexel.Key);
+		//ResultBuilder->SetPixel(GutterTexel.Key, FVector4f(0, 1, 0, 1) );
+	}
+}
+
+
+
+
+
+
+void FMeshGenericWorldPositionNormalBaker::Bake()
+{
+	const FMeshImageBakingCache* BakeCache = GetCache();
+	check(BakeCache);
+
+	// make sure we have per-triangle tangents computed so that GetInterpolatedTriangleTangent() below works
+	int32 MaxTriID = BakeCache->GetBakeTargetMesh()->MaxTriangleID();
+	check(BaseMeshTangents->GetTangents().Num() >= 3 * BakeCache->GetBakeTargetMesh()->MaxTriangleID());
+
+	auto SampleFunction = [&](const FMeshImageBakingCache::FCorrespondenceSample& SampleData)
+	{
+		FVector3d Position = SampleData.BaseSample.SurfacePoint;
+		FVector3d Normal = SampleData.BaseNormal;
+		int32 TriangleID = SampleData.BaseSample.TriangleIndex;
+		FVector3d BaryCoords = SampleData.BaseSample.BaryCoords;
+
+		// get tangents on base mesh
+		FVector3d BaseTangentX, BaseTangentY;
+		BaseMeshTangents->GetInterpolatedTriangleTangent(
+			SampleData.BaseSample.TriangleIndex,
+			SampleData.BaseSample.BaryCoords,
+			BaseTangentX, BaseTangentY);
+
+		// sample world normal at world position
+		FVector3f DetailNormal = NormalSampleFunction(Position, Normal);
+		Normalize(DetailNormal);
+
+		// compute normal in tangent space
+		float dx = DetailNormal.Dot((FVector3f)BaseTangentX);
+		float dy = DetailNormal.Dot((FVector3f)BaseTangentY);
+		float dz = DetailNormal.Dot((FVector3f)SampleData.BaseNormal);
+
+		return FVector3f(dx, dy, dz);
+	};
+
+	NormalsBuilder = MakeUnique<TImageBuilder<FVector3f>>();
+	NormalsBuilder->SetDimensions(BakeCache->GetDimensions());
+	FVector3f DefaultMapNormal = (DefaultNormal + FVector3f::One()) * 0.5f;
+	NormalsBuilder->Clear(DefaultMapNormal);
+
+	BakeCache->EvaluateSamples([&](const FVector2i& Coords, const FMeshImageBakingCache::FCorrespondenceSample& Sample)
+	{
+		FVector3f RelativeDetailNormal = SampleFunction(Sample);
+		FVector3f MapNormal = (RelativeDetailNormal + FVector3f::One()) * 0.5f;
+		NormalsBuilder->SetPixel(Coords, MapNormal);
+	});
+
+	const FImageOccupancyMap& Occupancy = *BakeCache->GetOccupancyMap();
+	for (int64 k = 0; k < Occupancy.GutterTexels.Num(); k++)
+	{
+		TPair<int64, int64> GutterTexel = Occupancy.GutterTexels[k];
+		NormalsBuilder->CopyPixel(GutterTexel.Value, GutterTexel.Key);
+	}
+}
\ No newline at end of file
diff --git a/Engine/Plugins/Experimental/GeometryProcessing/Source/DynamicMesh/Public/Sampling/MeshGenericWorldPositionBaker.h b/Engine/Plugins/Experimental/GeometryProcessing/Source/DynamicMesh/Public/Sampling/MeshGenericWorldPositionBaker.h
new file mode 100644
--- /dev/null
+++ b/Engine/Plugins/Experimental/GeometryProcessing/Source/DynamicMesh/Public/Sampling/MeshGenericWorldPositionBaker.h
@@ -0,0 +1,99 @@
+// Copyright Epic Games, Inc. All Rights Reserved.
+
+#pragma once
+
+
+#include "Sampling/MeshImageBaker.h"
+#include "Sampling/MeshImageBakingCache.h"
+#include "Image/ImageBuilder.h"
+#include "MeshTangents.h"
+
+namespace UE
+{
+namespace Geometry
+{
+
+/**
+ * Bake Colors based on arbitrary Position/Normal sampling. Useful for (eg) baking
+ * from an arbitrary function defined over space
+ */
+class DYNAMICMESH_API FMeshGenericWorldPositionColorBaker : public FMeshImageBaker
+{
+public:
+	virtual ~FMeshGenericWorldPositionColorBaker() {}
+
+	//
+	// Required input data
+	//
+
+	FVector4f DefaultColor = FVector4f(0, 0, 0, 1);
+
+	TFunction<FVector4f(FVector3d, FVector3d)> ColorSampleFunction 
+		= [this](FVector3d Position, FVector3d Normal) { return DefaultColor; };
+
+	//
+	// Compute functions
+	//
+
+	virtual void Bake() override;
+
+	//
+	// Output
+	//
+
+	const TUniquePtr<TImageBuilder<FVector4f>>& GetResult() const { return ResultBuilder; }
+
+	TUniquePtr<TImageBuilder<FVector4f>> TakeResult() { return MoveTemp(ResultBuilder); }
+
+protected:
+	TUniquePtr<TImageBuilder<FVector4f>> ResultBuilder;
+};
+
+
+
+
+/**
+ * Bake Tangent-Space Normals based on arbitrary Position/Normal sampling. Useful for (eg) 
+ * baking from an arbitrary function defined over space.
+ */
+class DYNAMICMESH_API FMeshGenericWorldPositionNormalBaker : public FMeshImageBaker
+{
+public:
+	virtual ~FMeshGenericWorldPositionNormalBaker() {}
+
+	//
+	// Required input data
+	//
+
+	const TMeshTangents<double>* BaseMeshTangents;
+
+	FVector3f DefaultNormal = FVector3f(0, 0, 1);
+
+	TFunction<FVector3f(FVector3d, FVector3d)> NormalSampleFunction
+		= [this](FVector3d Position, FVector3d Normal) { return DefaultNormal; };
+
+	//
+	// Compute functions
+	//
+
+	virtual void Bake() override;
+
+	//
+	// Output
+	//
+
+	const TUniquePtr<TImageBuilder<FVector3f>>& GetResult() const { return NormalsBuilder; }
+
+	TUniquePtr<TImageBuilder<FVector3f>> TakeResult() { return MoveTemp(NormalsBuilder); }
+
+protected:
+	TUniquePtr<TImageBuilder<FVector3f>> NormalsBuilder;
+};
+
+
+
+
+
+
+} // end namespace UE::Geometry
+} // end namespace UE
diff --git a/Engine/Plugins/Experimental/GeometryProcessing/Source/GeometricObjects/Public/Image/ImageInfilling.h b/Engine/Plugins/Experimental/GeometryProcessing/Source/GeometricObjects/Public/Image/ImageInfilling.h
new file mode 100644
--- /dev/null
+++ b/Engine/Plugins/Experimental/GeometryProcessing/Source/GeometricObjects/Public/Image/ImageInfilling.h
@@ -0,0 +1,186 @@
+// Copyright Epic Games, Inc. All Rights Reserved.
+
+#pragma once
+
+#include "Templates/Tuple.h"
+#include "VectorTypes.h"
+#include "Image/ImageDimensions.h"
+#include "Image/ImageBuilder.h"
+#include "Util/IndexPriorityQueue.h"
+
+namespace UE
+{
+namespace Geometry
+{
+
+
+/**
+ * TMarchingPixelInfill implements a very basic infilling strategy where the missing pixel with the most "known"
+ * neighbours is set to the average value of those neighbours, and this is iterated until all missing pixels are known.
+ * A Priority Queue is used to track the active front. 
+ * 
+ * The infill sequence is tracked, which allows it to be "replayed" on multiple images. This is useful 
+ * in cases where we have separate images with the same holes, which arises when baking textures from world sampling.
+ */
+template<typename PixelType>
+class TMarchingPixelInfill
+{
+public:
+	// Encoding of infill operations, so that they can be applied to multiple images
+	// Values are [ NumPixels, Pixel1Count, Pixel1X, Pixel1Y, Pixel1Nbr1X, Pixel1Nbr1Y, ...
+	TArray<int32> InfillSequence;
+
+
+	/**
+	 * Fill the values of MissingPixels in Image by propagating from known values.
+	 * @param MissingValue this value is used to indicate 'missing' in the image. Input pixels do not have to be set to this value, they are set internally to avoid map lookups 
+	 * @param NormalizeFunc Implements an "average" operation for the template PixelType, normally (SumOfNbrValues / Count)
+	 */
+	void ComputeInfill(
+		TImageBuilder<PixelType>& Image,
+		const TArray<FVector2i>& MissingPixels,
+		PixelType MissingValue,
+		TFunctionRef<PixelType(PixelType SumOfNbrValues, int NbrCount)> NormalizeFunc)
+	{
+		InfillSequence.Reset();
+		InfillSequence.Add(0);
+
+		// clear all missing values
+		for (int32 k = 0; k < MissingPixels.Num(); ++k)
+		{
+			Image.SetPixel(MissingPixels[k], MissingValue);
+		}
+
+		// march inward from hole boundary and compute fill values. This is probably not the
+		// right way to go and a smooth interpolant (eg mean-value) would be preferable
+		FIndexPriorityQueue ActiveSetQueue(MissingPixels.Num());
+		TMap<FVector2i, int32> PixelToIndexMap;
+		for (int32 k = 0; k < MissingPixels.Num(); ++k)
+		{
+			int32 Count = 0;
+			FVector2i CenterCoords = MissingPixels[k];
+			for (int32 j = 0; j < 8; ++j)
+			{
+				FVector2i NbrCoords = CenterCoords + IndexUtil::GridOffsets8[j];
+				if ( Image.ContainsPixel(NbrCoords) && Image.GetPixel(NbrCoords) != MissingValue)
+				{
+					Count++;
+				}
+			}
+			ActiveSetQueue.Insert(k, 8 - Count);
+			PixelToIndexMap.Add(CenterCoords, k);
+		}
+
+
+		while (ActiveSetQueue.GetCount() > 0)
+		{
+			int32 Index = ActiveSetQueue.Dequeue();
+			FVector2i CenterCoords = MissingPixels[Index];
+
+			int32 CountIdx = InfillSequence.Num();
+			InfillSequence.Add(1);
+			InfillSequence.Add(CenterCoords.X);
+			InfillSequence.Add(CenterCoords.Y);
+			InfillSequence[0]++;
+
+			// accumulate neighbour values in 8-neighbourhood
+			PixelType ValidNbrAvg = PixelType::Zero();
+			int32 ValidCount = 0;
+			for (int32 j = 0; j < 8; ++j)
+			{
+				FVector2i NbrCoords = CenterCoords + IndexUtil::GridOffsets8[j];
+				if (Image.ContainsPixel(NbrCoords))
+				{
+					FVector4f PixelColor = Image.GetPixel(NbrCoords);
+					if (PixelColor != MissingValue)
+					{
+						ValidNbrAvg += PixelColor;
+						++ValidCount;
+						InfillSequence.Add(NbrCoords.X);
+						InfillSequence.Add(NbrCoords.Y);
+						InfillSequence[CountIdx]++;
+					}
+				}
+			}
+
+			if (ValidCount > 0)
+			{
+				// update center pixel
+				PixelType InfillValue = NormalizeFunc(ValidNbrAvg, ValidCount);
+				Image.SetPixel(CenterCoords, InfillValue);
+
+				// Update queue for nbrs. We set one pixel so we subtract one from
+				// all those neighbours
+				for (int32 j = 0; j < 8; ++j)
+				{
+					FVector2i NbrCoords = CenterCoords + IndexUtil::GridOffsets8[j];
+					if (Image.ContainsPixel(NbrCoords) && Image.GetPixel(NbrCoords) == MissingValue)
+					{
+						const int32* FoundIndex = PixelToIndexMap.Find(NbrCoords);
+						if (FoundIndex != nullptr && ActiveSetQueue.Contains(*FoundIndex))
+						{
+							float CurPriority = ActiveSetQueue.GetPriority(*FoundIndex);
+							ActiveSetQueue.Update(*FoundIndex, CurPriority - 1.0);
+						}
+					}
+				}
+			}
+		}
+	}
+
+
+	/**
+	 * Fill the missing values in Image by replaying the infill sequence computed by ComputeInfill()
+	 * @param NormalizeFunc Implements an "average" operation for the template OtherPixelType, normally (SumOfNbrValues / NbrCount)
+	 */
+	template<typename OtherPixelType>
+	void ApplyInfill(
+		TImageBuilder<OtherPixelType>& Image,
+		TFunctionRef<OtherPixelType(OtherPixelType SumOfNbrValues, int NbrCount)> NormalizeFunc ) const
+	{
+		int32 InfillPixels = InfillSequence[0];
+		int32 i = 1;
+		for (int32 pi = 0; pi < InfillPixels; ++pi)
+		{
+			int32 PixelCount = InfillSequence[i++];
+			FVector2i CenterCoords;
+			CenterCoords.X = InfillSequence[i++];
+			CenterCoords.Y = InfillSequence[i++];
+			if (PixelCount > 1)
+			{
+				OtherPixelType NbrSum = OtherPixelType::Zero();
+				int32 NbrCount = 0;
+				for (int32 j = 1; j < PixelCount; ++j)
+				{
+					FVector2i NbrCoords;
+					NbrCoords.X = InfillSequence[i++];
+					NbrCoords.Y = InfillSequence[i++];
+					NbrSum += Image.GetPixel(NbrCoords);
+					NbrCount++;
+				}
+				PixelType InfillValue = NormalizeFunc(NbrSum, NbrCount);
+				Image.SetPixel(CenterCoords, InfillValue);
+			}
+		}
+	};
+
+
+	/**
+	 * Fill the missing values in Image by replaying the infill sequence computed by ComputeInfill()
+	 * @param NormalizeFunc Implements an "average" operation for the template PixelType, normally (SumOfNbrValues / NbrCount)
+	 */
+	void ApplyInfill(
+		TImageBuilder<PixelType>& Image,
+		TFunctionRef<PixelType(PixelType SumOfNbrValues, int NbrCount)> NormalizeFunc) const
+	{
+		ApplyInfill<PixelType>(Image, NormalizeFunc);
+	}
+
+};
+
+
+
+
+
+} // end namespace UE::Geometry
+} // end namespace UE
\ No newline at end of file
diff --git a/Engine/Plugins/Experimental/MeshModelingToolset/Source/ModelingOperators/Public/CleaningOps/RemoveOccludedTrianglesOp.h b/Engine/Plugins/Experimental/MeshModelingToolset/Source/ModelingOperators/Public/CleaningOps/RemoveOccludedTrianglesOp.h
--- a/Engine/Plugins/Experimental/MeshModelingToolset/Source/ModelingOperators/Public/CleaningOps/RemoveOccludedTrianglesOp.h
+++ b/Engine/Plugins/Experimental/MeshModelingToolset/Source/ModelingOperators/Public/CleaningOps/RemoveOccludedTrianglesOp.h
@@ -44,6 +44,7 @@ public:
 	// add triangle samples per triangle (in addition to TriangleSamplingMethod)
 	int AddTriangleSamples = 0;
 
+	// once triangles to remove are identified, do iterations of boundary erosion, ie contract selection by boundary vertex one-rings
 	int ShrinkRemoval = 0;
 
 	double MinAreaConnectedComponent = 0;
