GeometryProcessing: DynamicMesh serialization support

#rnx

diff --git a/Engine/Plugins/Experimental/GeometryProcessing/Source/DynamicMesh/Private/DynamicMesh3.cpp b/Engine/Plugins/Experimental/GeometryProcessing/Source/DynamicMesh/Private/DynamicMesh3.cpp
--- a/Engine/Plugins/Experimental/GeometryProcessing/Source/DynamicMesh/Private/DynamicMesh3.cpp
+++ b/Engine/Plugins/Experimental/GeometryProcessing/Source/DynamicMesh/Private/DynamicMesh3.cpp
@@ -4,6 +4,7 @@
 #include "DynamicMeshAttributeSet.h"
 #include "Generators/MeshShapeGenerator.h"
 #include "Templates/UniquePtr.h"
+#include "UObject/UE5MainStreamObjectVersion.h"
 
 #include "ExplicitUseGeometryMathTypes.h"		// using UE::Geometry::(math types)
 using namespace UE::Geometry;
@@ -369,6 +370,36 @@ void FDynamicMesh3::Clear()
 	CachedIsClosedTimestamp = -1;
 }
 
+void FDynamicMesh3::Serialize(FArchive& Ar)
+{
+	Ar.UsingCustomVersion(FUE5MainStreamObjectVersion::GUID);
+
+	Ar << Vertices;
+	Ar << VertexRefCounts;
+	Ar << VertexNormals;
+	Ar << VertexColors;
+	Ar << VertexUVs;
+	Ar << VertexEdgeLists;
+	Ar << Triangles;
+	Ar << TriangleRefCounts;
+	Ar << TriangleEdges;
+	Ar << TriangleGroups;
+	Ar << GroupIDCounter;
+	Ar << Edges;
+	Ar << EdgeRefCounts;
+
+	bool bHasAttributes = HasAttributes();
+	Ar << bHasAttributes;
+	if (bHasAttributes)
+	{
+		if (Ar.IsLoading())
+		{
+			EnableAttributes();
+		}
+		Ar << *AttributeSet;
+	}
+}
+
 int FDynamicMesh3::GetComponentsFlags() const
 {
 	int c = 0;
@@ -660,9 +691,7 @@ FString FDynamicMesh3::MeshInfoString()
 	return VtxString + "\n" + TriString + "\n" + EdgeString + "\n" + AttribString + "\n" + InfoString;
 }
 
-bool FDynamicMesh3::IsSameMesh(const FDynamicMesh3& m2, bool bCheckConnectivity, bool bCheckEdgeIDs,
-	bool bCheckNormals, bool bCheckColors, bool bCheckUVs, bool bCheckGroups,
-	float Epsilon)
+bool FDynamicMesh3::IsSameAs(const FDynamicMesh3& m2, const FSameAsOptions& Options)
 {
 	if (VertexCount() != m2.VertexCount())
 	{
@@ -674,7 +703,7 @@ bool FDynamicMesh3::IsSameMesh(const FDynamicMesh3& m2, bool bCheckConnectivity,
 	}
 	for (int vid : VertexIndicesItr())
 	{
-		if (m2.IsVertex(vid) == false || VectorUtil::EpsilonEqual(GetVertex(vid), m2.GetVertex(vid), (double)Epsilon) == false)
+		if (m2.IsVertex(vid) == false || VectorUtil::EpsilonEqual(GetVertex(vid), m2.GetVertex(vid), (double)Options.Epsilon) == false)
 		{
 			return false;
 		}
@@ -686,7 +715,7 @@ bool FDynamicMesh3::IsSameMesh(const FDynamicMesh3& m2, bool bCheckConnectivity,
 			return false;
 		}
 	}
-	if (bCheckConnectivity)
+	if (Options.bCheckConnectivity)
 	{
 		for (int eid : EdgeIndicesItr())
 		{
@@ -704,7 +733,7 @@ bool FDynamicMesh3::IsSameMesh(const FDynamicMesh3& m2, bool bCheckConnectivity,
 			}
 		}
 	}
-	if (bCheckEdgeIDs)
+	if (Options.bCheckEdgeIDs)
 	{
 		if (EdgeCount() != m2.EdgeCount())
 		{
@@ -718,7 +747,7 @@ bool FDynamicMesh3::IsSameMesh(const FDynamicMesh3& m2, bool bCheckConnectivity,
 			}
 		}
 	}
-	if (bCheckNormals)
+	if (Options.bCheckNormals)
 	{
 		if (HasVertexNormals() != m2.HasVertexNormals())
 		{
@@ -728,14 +757,14 @@ bool FDynamicMesh3::IsSameMesh(const FDynamicMesh3& m2, bool bCheckConnectivity,
 		{
 			for (int vid : VertexIndicesItr())
 			{
-				if (VectorUtil::EpsilonEqual(GetVertexNormal(vid), m2.GetVertexNormal(vid), Epsilon) == false)
+				if (VectorUtil::EpsilonEqual(GetVertexNormal(vid), m2.GetVertexNormal(vid), Options.Epsilon) == false)
 				{
 					return false;
 				}
 			}
 		}
 	}
-	if (bCheckColors)
+	if (Options.bCheckColors)
 	{
 		if (HasVertexColors() != m2.HasVertexColors())
 		{
@@ -745,14 +774,14 @@ bool FDynamicMesh3::IsSameMesh(const FDynamicMesh3& m2, bool bCheckConnectivity,
 		{
 			for (int vid : VertexIndicesItr())
 			{
-				if (VectorUtil::EpsilonEqual(GetVertexColor(vid), m2.GetVertexColor(vid), Epsilon) == false)
+				if (VectorUtil::EpsilonEqual(GetVertexColor(vid), m2.GetVertexColor(vid), Options.Epsilon) == false)
 				{
 					return false;
 				}
 			}
 		}
 	}
-	if (bCheckUVs)
+	if (Options.bCheckUVs)
 	{
 		if (HasVertexUVs() != m2.HasVertexUVs())
 		{
@@ -762,14 +791,14 @@ bool FDynamicMesh3::IsSameMesh(const FDynamicMesh3& m2, bool bCheckConnectivity,
 		{
 			for (int vid : VertexIndicesItr())
 			{
-				if (VectorUtil::EpsilonEqual(GetVertexUV(vid), m2.GetVertexUV(vid), Epsilon) == false)
+				if (VectorUtil::EpsilonEqual(GetVertexUV(vid), m2.GetVertexUV(vid), Options.Epsilon) == false)
 				{
 					return false;
 				}
 			}
 		}
 	}
-	if (bCheckGroups)
+	if (Options.bCheckGroups)
 	{
 		if (HasTriangleGroups() != m2.HasTriangleGroups())
 		{
@@ -786,6 +815,20 @@ bool FDynamicMesh3::IsSameMesh(const FDynamicMesh3& m2, bool bCheckConnectivity,
 			}
 		}
 	}
+	if (Options.bCheckAttributes)
+	{
+		if (HasAttributes() != m2.HasAttributes())
+		{
+			return false;
+		}
+		if (HasAttributes())
+		{
+			if (!AttributeSet->IsSameAs(*m2.AttributeSet))
+			{
+				return false;
+			}
+		}
+	}
 	return true;
 }
 
diff --git a/Engine/Plugins/Experimental/GeometryProcessing/Source/DynamicMesh/Private/DynamicMesh3_Edits.cpp b/Engine/Plugins/Experimental/GeometryProcessing/Source/DynamicMesh/Private/DynamicMesh3_Edits.cpp
--- a/Engine/Plugins/Experimental/GeometryProcessing/Source/DynamicMesh/Private/DynamicMesh3_Edits.cpp
+++ b/Engine/Plugins/Experimental/GeometryProcessing/Source/DynamicMesh/Private/DynamicMesh3_Edits.cpp
@@ -338,7 +338,7 @@ void FDynamicMesh3::CompactInPlace(FCompactMaps* CompactInfo)
 		iCurV++;
 	}
 
-	TDynamicVector<short> &vref = VertexRefCounts.GetRawRefCountsUnsafe();
+	TDynamicVector<unsigned short> &vref = VertexRefCounts.GetRawRefCountsUnsafe();
 
 	while (iCurV < iLastV)
 	{
@@ -431,7 +431,7 @@ void FDynamicMesh3::CompactInPlace(FCompactMaps* CompactInfo)
 		iCurT++;
 	}
 
-	TDynamicVector<short> &tref = TriangleRefCounts.GetRawRefCountsUnsafe();
+	TDynamicVector<unsigned short> &tref = TriangleRefCounts.GetRawRefCountsUnsafe();
 
 	while (iCurT < iLastT)
 	{
@@ -494,7 +494,7 @@ void FDynamicMesh3::CompactInPlace(FCompactMaps* CompactInfo)
 		iCurE++;
 	}
 
-	TDynamicVector<short> &eref = EdgeRefCounts.GetRawRefCountsUnsafe();
+	TDynamicVector<unsigned short> &eref = EdgeRefCounts.GetRawRefCountsUnsafe();
 
 	while (iCurE < iLastE)
 	{
@@ -877,7 +877,10 @@ EMeshResult FDynamicMesh3::SplitEdge(int eab, FEdgeSplitInfo& SplitInfo, double
 	}
 	FIndex3i T0tv = GetTriangle(t0);
 	int c = IndexUtil::OrientTriEdgeAndFindOtherVtx(a, b, T0tv);
-	if (VertexRefCounts.GetRawRefCount(c) > 32764)
+
+	// RefCount overflow check. Conservatively leave room for
+	// extra increments from other operations.
+	if (VertexRefCounts.GetRawRefCount(c) > FRefCountVector::INVALID_REF_COUNT - 3)
 	{
 		return EMeshResult::Failed_HitValenceLimit;
 	}
@@ -967,7 +970,10 @@ EMeshResult FDynamicMesh3::SplitEdge(int eab, FEdgeSplitInfo& SplitInfo, double
 		SplitInfo.OriginalTriangles.B = t1;
 		FIndex3i T1tv = GetTriangle(t1);
 		int d = IndexUtil::FindTriOtherVtx(a, b, T1tv);
-		if (VertexRefCounts.GetRawRefCount(d) > 32764)
+
+		// RefCount overflow check. Conservatively leave room for
+		// extra increments from other operations.
+		if (VertexRefCounts.GetRawRefCount(d) > FRefCountVector::INVALID_REF_COUNT - 3)
 		{
 			return EMeshResult::Failed_HitValenceLimit;
 		}
diff --git a/Engine/Plugins/Experimental/GeometryProcessing/Source/DynamicMesh/Private/DynamicMeshAttributeSet.cpp b/Engine/Plugins/Experimental/GeometryProcessing/Source/DynamicMesh/Private/DynamicMeshAttributeSet.cpp
--- a/Engine/Plugins/Experimental/GeometryProcessing/Source/DynamicMesh/Private/DynamicMeshAttributeSet.cpp
+++ b/Engine/Plugins/Experimental/GeometryProcessing/Source/DynamicMesh/Private/DynamicMeshAttributeSet.cpp
@@ -783,3 +783,115 @@ void FDynamicMeshAttributeSet::OnSplitVertex(const DynamicMeshInfo::FVertexSplit
 	}
 }
 
+bool FDynamicMeshAttributeSet::IsSameAs(const FDynamicMeshAttributeSet& Other) const
+{
+	if (UVLayers.Num() != Other.UVLayers.Num() ||
+		NormalLayers.Num() != Other.NormalLayers.Num() ||
+		PolygroupLayers.Num() != Other.PolygroupLayers.Num())
+	{
+		return false;
+	}
+
+	for (int Idx = 0; Idx < UVLayers.Num(); Idx++)
+	{
+		if (!UVLayers[Idx].IsSameAs(Other.UVLayers[Idx]))
+		{
+			return false;
+		}
+	}
+
+	for (int Idx = 0; Idx < NormalLayers.Num(); Idx++)
+	{
+		if (!NormalLayers[Idx].IsSameAs(Other.NormalLayers[Idx]))
+		{
+			return false;
+		}
+	}
+
+	for (int Idx = 0; Idx < PolygroupLayers.Num(); Idx++)
+	{
+		if (!PolygroupLayers[Idx].IsSameAs(Other.PolygroupLayers[Idx]))
+		{
+			return false;
+		}
+	}
+
+	if (HasPrimaryColors() != Other.HasPrimaryColors())
+	{
+		return false;
+	}
+	if (HasPrimaryColors())
+	{
+		if (!ColorLayer->IsSameAs(*Other.ColorLayer))
+		{
+			return false;
+		}
+	}
+
+	if (HasMaterialID() != Other.HasMaterialID())
+	{
+		return false;
+	}
+	if (HasMaterialID())
+	{
+		if (!MaterialIDAttrib->IsSameAs(*Other.MaterialIDAttrib))
+		{
+			return false;
+		}
+	}
+
+	// TODO: Test GenericAttributes
+
+	return true;
+}
+
+void FDynamicMeshAttributeSet::Serialize(FArchive& Ar)
+{
+	Ar << UVLayers;
+	Ar << NormalLayers;
+	Ar << PolygroupLayers;
+
+	if (Ar.IsLoading())
+	{
+		// Manually populate ParentMesh since deserialization of the individual
+		// layers cannot populate the pointer.
+		for (FDynamicMeshUVOverlay& Overlay : UVLayers)
+		{
+			Overlay.ParentMesh = ParentMesh;
+		}
+		for (FDynamicMeshNormalOverlay& Overlay : NormalLayers)
+		{
+			Overlay.ParentMesh = ParentMesh;
+		}
+		for (FDynamicMeshPolygroupAttribute& Attr : PolygroupLayers)
+		{
+			Attr.ParentMesh = ParentMesh;
+		}
+	}
+
+	// Use int32 here to futureproof for multiple color layers.
+	int32 bHasColorLayer = HasPrimaryColors() ? 1 : 0;
+	Ar << bHasColorLayer;
+	if (bHasColorLayer)
+	{
+		if (Ar.IsLoading())
+		{
+			EnablePrimaryColors();
+		}
+		Ar << *ColorLayer;
+	}
+
+	bool bHasMaterialID = HasMaterialID();
+	Ar << bHasMaterialID;
+	if (bHasMaterialID)
+	{
+		if (Ar.IsLoading())
+		{
+			EnableMaterialID();
+		}
+		Ar << *MaterialIDAttrib;
+	}
+
+	//Ar << GenericAttributes; // TODO
+}
+
diff --git a/Engine/Plugins/Experimental/GeometryProcessing/Source/DynamicMesh/Private/DynamicMeshOverlay.cpp b/Engine/Plugins/Experimental/GeometryProcessing/Source/DynamicMesh/Private/DynamicMeshOverlay.cpp
--- a/Engine/Plugins/Experimental/GeometryProcessing/Source/DynamicMesh/Private/DynamicMeshOverlay.cpp
+++ b/Engine/Plugins/Experimental/GeometryProcessing/Source/DynamicMesh/Private/DynamicMeshOverlay.cpp
@@ -1393,6 +1393,63 @@ bool TDynamicMeshOverlay<RealType, ElementSize>::CheckValidity(bool bAllowNonMan
 }
 
 
+template<typename RealType, int ElementSize>
+bool TDynamicMeshOverlay<RealType, ElementSize>::IsSameAs(const TDynamicMeshOverlay<RealType, ElementSize>& Other) const
+{
+	if (ElementsRefCounts.GetCount() != Other.ElementsRefCounts.GetCount() ||
+		ElementsRefCounts.GetMaxIndex() != Other.ElementsRefCounts.GetMaxIndex() ||
+		Elements.Num() != Other.Elements.Num() ||
+		ParentVertices.Num() != Other.ParentVertices.Num() ||
+		ElementTriangles.Num() != Other.ElementTriangles.Num())
+	{
+		return false;
+	}
+
+	for (int Idx = 0; Idx < ElementsRefCounts.GetMaxIndex(); Idx++)
+	{
+		if (ElementsRefCounts.GetRefCount(Idx) != Other.ElementsRefCounts.GetRefCount(Idx))
+		{
+			return false;
+		}
+	}
+
+	for (int Idx = 0; Idx < Elements.Num(); Idx++)
+	{
+		if (Elements[Idx] != Other.Elements[Idx])
+		{
+			return false;
+		}
+	}
+
+	for (int Idx = 0; Idx < ParentVertices.Num(); Idx++)
+	{
+		if (ParentVertices[Idx] != Other.ParentVertices[Idx])
+		{
+			return false;
+		}
+	}
+
+	for (int Idx = 0; Idx < ElementTriangles.Num(); Idx++)
+	{
+		if (ElementTriangles[Idx] != Other.ElementTriangles[Idx])
+		{
+			return false;
+		}
+	}
+
+	return true;
+}
+
+
+template<typename RealType, int ElementSize>
+void TDynamicMeshOverlay<RealType, ElementSize>::Serialize(FArchive& Ar)
+{
+	Ar << ElementsRefCounts;
+	Ar << Elements;
+	Ar << ParentVertices;
+	Ar << ElementTriangles;
+}
+
 
 namespace UE
 {
diff --git a/Engine/Plugins/Experimental/GeometryProcessing/Source/DynamicMesh/Public/DynamicMesh3.h b/Engine/Plugins/Experimental/GeometryProcessing/Source/DynamicMesh/Public/DynamicMesh3.h
--- a/Engine/Plugins/Experimental/GeometryProcessing/Source/DynamicMesh/Public/DynamicMesh3.h
+++ b/Engine/Plugins/Experimental/GeometryProcessing/Source/DynamicMesh/Public/DynamicMesh3.h
@@ -106,6 +106,26 @@ public:
 		{
 			return (e0.Vert != e1.Vert) || (e0.Tri != e1.Tri);
 		}
+
+		/**
+		 * Serialization operator for FEdge.
+		 *
+		 * @param Ar Archive to serialize with.
+		 * @param E Edge to serialize.
+		 * @returns Passing down serializing archive.
+		 */
+		friend FArchive& operator<<(FArchive& Ar, FEdge& E)
+		{
+			E.Serialize(Ar);
+			return Ar;
+		}
+
+		/** Serialize FEdge to an archive. */
+		void Serialize(FArchive& Ar)
+		{
+			Ar << Vert;
+			Ar << Tri;
+		}
 	};
 	/** InvalidID indicates that a vertex/edge/triangle ID is invalid */
 	constexpr static int InvalidID = IndexConstants::InvalidID;
@@ -213,6 +233,21 @@ public:
 	/** Discard all data */
 	void Clear();
 
+	/**
+	 * Serialization operator for FDynamicMesh3.
+	 *
+	 * @param Ar Archive to serialize with.
+	 * @param Mesh Mesh to serialize.
+	 * @returns Passing down serializing archive.
+	 */
+	friend DYNAMICMESH_API FArchive& operator<<(FArchive& Ar, FDynamicMesh3& Mesh)
+	{
+		Mesh.Serialize(Ar);
+		return Ar;
+	}
+
+	/** Serialize the mesh to an archive. */
+	void Serialize(FArchive& Ar);
 
 public:
 	/** @return number of vertices in the mesh */
@@ -1227,13 +1262,26 @@ public:
 	 */
 	virtual FString MeshInfoString();
 
+	/**
+	 * Options for the IsSameAs check
+	 */
+	struct FSameAsOptions
+	{
+		bool bCheckConnectivity = true;
+		bool bCheckEdgeIDs = false;
+		bool bCheckNormals = false;
+		bool bCheckColors = false;
+		bool bCheckUVs = false;
+		bool bCheckGroups = false;
+		bool bCheckAttributes = false;
+		float Epsilon = TMathUtil<float>::Epsilon;
+	};
+
 	/**
 	 * Check if another mesh is the same as this mesh. By default only checks
 	 * vertices and triangles, turn on other parameters w/ flags
 	 */
-	virtual bool IsSameMesh(const FDynamicMesh3& OtherMesh, bool bCheckConnectivity, bool bCheckEdgeIDs = false,
-	                        bool bCheckNormals = false, bool bCheckColors = false, bool bCheckUVs = false,
-	                        bool bCheckGroups = false, float Epsilon = TMathUtil<float>::Epsilon);
+	virtual bool IsSameAs(const FDynamicMesh3& OtherMesh, const FSameAsOptions& Options);
 
 	/**
 	 * Options for what the validity check will permit
diff --git a/Engine/Plugins/Experimental/GeometryProcessing/Source/DynamicMesh/Public/DynamicMeshAttributeSet.h b/Engine/Plugins/Experimental/GeometryProcessing/Source/DynamicMesh/Public/DynamicMeshAttributeSet.h
--- a/Engine/Plugins/Experimental/GeometryProcessing/Source/DynamicMesh/Public/DynamicMeshAttributeSet.h
+++ b/Engine/Plugins/Experimental/GeometryProcessing/Source/DynamicMesh/Public/DynamicMeshAttributeSet.h
@@ -343,6 +343,27 @@ public:
 		return GenericAttributes;
 	}
 
+	/**
+	 * Returns true if this AttributeSet is the same as Other.
+	 */
+	bool IsSameAs(const FDynamicMeshAttributeSet& Other) const;
+
+	/**
+	 * Serialization operator for FDynamicMeshAttributeSet.
+	 *
+	 * @param Ar Archive to serialize with.
+	 * @param Set Attribute set to serialize.
+	 * @returns Passing down serializing archive.
+	 */
+	friend DYNAMICMESH_API FArchive& operator<<(FArchive& Ar, FDynamicMeshAttributeSet& Set)
+	{
+		Set.Serialize(Ar);
+		return Ar;
+	}
+
+	/** Serialize the set to an archive. */
+	void Serialize(FArchive& Ar);
+
 protected:
 	/** Parent mesh of this attribute set */
 	FDynamicMesh3* ParentMesh;
diff --git a/Engine/Plugins/Experimental/GeometryProcessing/Source/DynamicMesh/Public/DynamicMeshOverlay.h b/Engine/Plugins/Experimental/GeometryProcessing/Source/DynamicMesh/Public/DynamicMeshOverlay.h
--- a/Engine/Plugins/Experimental/GeometryProcessing/Source/DynamicMesh/Public/DynamicMeshOverlay.h
+++ b/Engine/Plugins/Experimental/GeometryProcessing/Source/DynamicMesh/Public/DynamicMeshOverlay.h
@@ -153,7 +153,7 @@ public:
 			}
 		}
 
-		TDynamicVector<short>& ERef = ElementsRefCounts.GetRawRefCountsUnsafe();
+		TDynamicVector<unsigned short>& ERef = ElementsRefCounts.GetRawRefCountsUnsafe();
 		RealType Data[ElementSize];
 		while (iCurE < iLastE)
 		{
@@ -451,6 +451,26 @@ public:
 	 */
 	bool CheckValidity(bool bAllowNonManifoldVertices = true, EValidityCheckFailMode FailMode = EValidityCheckFailMode::Check) const;
 
+	/**
+	 * Returns true if this overlay is the same as Other.
+	 */
+	bool IsSameAs(const TDynamicMeshOverlay<RealType, ElementSize>& Other) const;
+
+	/**
+	 * Serialization operator for FDynamicMeshOverlay.
+	 *
+	 * @param Ar Archive to serialize with.
+	 * @param Overlay Mesh overlay to serialize.
+	 * @returns Passing down serializing archive.
+	 */
+	friend FArchive& operator<<(FArchive& Ar, TDynamicMeshOverlay<RealType,ElementSize>& Overlay)
+	{
+		Overlay.Serialize(Ar);
+		return Ar;
+	}
+
+	/** Serialize the overlay to an archive. */
+	void Serialize(FArchive& Ar);
 
 public:
 	/** Set a triangle's element indices to InvalidID */
diff --git a/Engine/Plugins/Experimental/GeometryProcessing/Source/DynamicMesh/Public/DynamicMeshTriangleAttribute.h b/Engine/Plugins/Experimental/GeometryProcessing/Source/DynamicMesh/Public/DynamicMeshTriangleAttribute.h
--- a/Engine/Plugins/Experimental/GeometryProcessing/Source/DynamicMesh/Public/DynamicMeshTriangleAttribute.h
+++ b/Engine/Plugins/Experimental/GeometryProcessing/Source/DynamicMesh/Public/DynamicMeshTriangleAttribute.h
@@ -370,7 +370,46 @@ public:
 		ResizeAttribStoreIfNeeded(TriangleID);
 	}
 
+public:
+	/**
+	 * Returns true if this AttributeSet is the same as Other.
+	 */
+	bool IsSameAs(const TDynamicMeshTriangleAttribute<AttribValueType, AttribDimension>& Other) const
+	{
+		if (AttribValues.Num() != Other.AttribValues.Num())
+		{
+			return false;
+		}
 
+		for (int Idx = 0, NumValues = AttribValues.Num(); Idx < NumValues; Idx++)
+		{
+			if (AttribValues[Idx] != Other.AttribValues[Idx])
+			{
+				return false;
+			}
+		}
+
+		return true;
+	}
+
+	/**
+	 * Serialization operator for TDynamicMeshTriangleAttribute.
+	 *
+	 * @param Ar Archive to serialize with.
+	 * @param Attr Mesh triangle attribute to serialize.
+	 * @returns Passing down serializing archive.
+	 */
+	friend FArchive& operator<<(FArchive& Ar, TDynamicMeshTriangleAttribute<AttribValueType, AttribDimension>& Attr)
+	{
+		Attr.Serialize(Ar);
+		return Ar;
+	}
+
+	/** Serialize the attribute to an archive. */
+	void Serialize(FArchive& Ar)
+	{
+		Ar << AttribValues;
+	}
 };
 
 
diff --git a/Engine/Plugins/Experimental/GeometryProcessing/Source/GeometricObjects/Public/IndexTypes.h b/Engine/Plugins/Experimental/GeometryProcessing/Source/GeometricObjects/Public/IndexTypes.h
--- a/Engine/Plugins/Experimental/GeometryProcessing/Source/GeometricObjects/Public/IndexTypes.h
+++ b/Engine/Plugins/Experimental/GeometryProcessing/Source/GeometricObjects/Public/IndexTypes.h
@@ -4,6 +4,7 @@
 
 #include "GeometryBase.h"
 #include "Math/IntVector.h"
+#include "Serialization/Archive.h"
 #include <limits>
 
 
@@ -80,6 +81,26 @@ struct FIndex2i
 		A = B;
 		B = tmp;
 	}
+
+	/**
+	 * Serialization operator for FIndex2i.
+	 *
+	 * @param Ar Archive to serialize with.
+	 * @param I Index to serialize.
+	 * @returns Passing down serializing archive.
+	 */
+	friend FArchive& operator<<(FArchive& Ar, FIndex2i& I)
+	{
+		I.Serialize(Ar);
+		return Ar;
+	}
+
+	/** Serialize FIndex2i to an archive. */
+	void Serialize(FArchive& Ar)
+	{
+		Ar << A;
+		Ar << B;
+	}
 };
 
 FORCEINLINE uint32 GetTypeHash(const FIndex2i& Index)
@@ -161,6 +182,27 @@ struct FIndex3i
 		B = Vec.Y;
 		C = Vec.Z;
 	}
+
+	/**
+	 * Serialization operator for FIndex3i.
+	 *
+	 * @param Ar Archive to serialize with.
+	 * @param I Index to serialize.
+	 * @returns Passing down serializing archive.
+	 */
+	friend FArchive& operator<<(FArchive& Ar, FIndex3i& I)
+	{
+		I.Serialize(Ar);
+		return Ar;
+	}
+
+	/** Serialize FIndex3i to an archive */
+	void Serialize(FArchive& Ar)
+	{
+		Ar << A;
+		Ar << B;
+		Ar << C;
+	}
 };
 
 FORCEINLINE uint32 GetTypeHash(const FIndex3i& Index)
@@ -234,6 +276,28 @@ struct FIndex4i
 	{
 		return A == Idx || B == Idx || C == Idx || D == Idx;
 	}
+
+	/**
+	 * Serialization operator for FIndex4i.
+	 *
+	 * @param Ar Archive to serialize with.
+	 * @param I Index to serialize.
+	 * @returns Passing down serializing archive.
+	 */
+	friend FArchive& operator<<(FArchive& Ar, FIndex4i& I)
+	{
+		I.Serialize(Ar);
+		return Ar;
+	}
+
+	/** Serialize FIndex3i to an archive */
+	void Serialize(FArchive& Ar)
+	{
+		Ar << A;
+		Ar << B;
+		Ar << C;
+		Ar << D;
+	}
 };
 
 FORCEINLINE uint32 GetTypeHash(const FIndex4i& Index)
@@ -246,3 +310,8 @@ FORCEINLINE uint32 GetTypeHash(const FIndex4i& Index)
 
 } // end namespace UE::Geometry
 } // end namespace UE
+
+template<> struct TCanBulkSerialize<UE::Geometry::FIndex2i> { enum { Value = true }; };
+template<> struct TCanBulkSerialize<UE::Geometry::FIndex3i> { enum { Value = true }; };
+template<> struct TCanBulkSerialize<UE::Geometry::FIndex4i> { enum { Value = true }; };
+
diff --git a/Engine/Plugins/Experimental/GeometryProcessing/Source/GeometricObjects/Public/Util/DynamicVector.h b/Engine/Plugins/Experimental/GeometryProcessing/Source/GeometricObjects/Public/Util/DynamicVector.h
--- a/Engine/Plugins/Experimental/GeometryProcessing/Source/GeometricObjects/Public/Util/DynamicVector.h
+++ b/Engine/Plugins/Experimental/GeometryProcessing/Source/GeometricObjects/Public/Util/DynamicVector.h
@@ -6,6 +6,7 @@
 
 #include <CoreMinimal.h>
 #include "Containers/StaticArray.h"
+#include "Serialization/Archive.h"
 #include "VectorTypes.h"
 #include "IndexTypes.h"
 
@@ -102,6 +103,27 @@ public:
 	template <typename Func>
 	void Apply(const Func& f);
 
+	/**
+	 * Serialization operator for TDynamicVector.
+	 *
+	 * @param Ar Archive to serialize with.
+	 * @param Vec Vector to serialize.
+	 * @returns Passing down serializing archive.
+	 */
+	friend FArchive& operator<<(FArchive& Ar, TDynamicVector& Vec)
+	{
+		Vec.Serialize(Ar);
+		return Ar;
+	}
+
+	/** Serialize vector to an archive */
+	void Serialize(FArchive& Ar)
+	{
+		Ar << CurBlock;
+		Ar << CurBlockUsed;
+		Ar << Blocks;
+	}
+
 public:
 	/*
 	 * FIterator class iterates over values of vector
@@ -301,6 +323,73 @@ private:
 			}
 			Elements.Empty(NewReservedSize);
 		}
+
+		/**
+		 * Count bytes needed to serialize this array.
+		 *
+		 * @param Ar Archive to count for.
+		 */
+		void CountBytes(FArchive& Ar) const
+		{
+			Elements.CountBytes(Ar);
+		}
+
+		/**
+		 * Serialization operator for TBlockVector.
+		 *
+		 * @param Ar Archive to serialize with.
+		 * @param Vec Vector to serialize.
+		 * @returns Passing down serializing archive.
+		 */
+		friend FArchive& operator<<(FArchive& Ar, TBlockVector& Vec)
+		{
+			Vec.Serialize(Ar);
+			return Ar;
+		}
+
+		/** Serialize TBlockVector to archive. */
+		void Serialize(FArchive& Ar)
+		{
+			CountBytes(Ar);
+			if (Ar.IsLoading())
+			{
+				// Load array.
+				int32 NewNum;
+				Ar << NewNum;
+				Empty(NewNum);
+				for (int32 Index = 0; Index < NewNum; Index++)
+				{
+					ArrayType* NewElement = new ArrayType;
+					// TODO: Consider moving bulk serialization code into TStaticArray.
+					if (TCanBulkSerialize<Type>::Value)
+					{
+						Ar.Serialize(NewElement->GetData(), NewElement->Num() * sizeof(Type));
+					}
+					else
+					{
+						Ar << *NewElement;
+					}
+					Add(NewElement);
+				}
+			}
+			else
+			{
+				// Save array.
+				int32 CurNum = Num();
+				Ar << CurNum;
+				for (int32 Index = 0; Index < CurNum; Index++)
+				{
+					if (TCanBulkSerialize<Type>::Value)
+					{
+						Ar.Serialize((*this)[Index].GetData(), (*this)[Index].Num() * sizeof(Type));
+					}
+					else
+					{
+						Ar << (*this)[Index];
+					}
+				}
+			}
+		}
 	};
 
 	// memory chunks for dynamic vector
diff --git a/Engine/Plugins/Experimental/GeometryProcessing/Source/GeometricObjects/Public/Util/RefCountVector.h b/Engine/Plugins/Experimental/GeometryProcessing/Source/GeometricObjects/Public/Util/RefCountVector.h
--- a/Engine/Plugins/Experimental/GeometryProcessing/Source/GeometricObjects/Public/Util/RefCountVector.h
+++ b/Engine/Plugins/Experimental/GeometryProcessing/Source/GeometricObjects/Public/Util/RefCountVector.h
@@ -24,7 +24,7 @@ namespace Geometry
 class FRefCountVector
 {
 public:
-	static constexpr short INVALID_REF_COUNT = -1;
+	static constexpr unsigned short INVALID_REF_COUNT = MAX_uint16;
 
 	FRefCountVector() = default;
 	FRefCountVector(const FRefCountVector&) = default;
@@ -103,20 +103,15 @@ public:
 		}
 	}
 
-	int Increment(int Index, short IncrementCount = 1)
+	int Increment(int Index, unsigned short IncrementCount = 1)
 	{
-		check(IsValid(Index));
-		// debug check for overflow...
-		check((short)(RefCounts[Index] + IncrementCount) > 0);
 		RefCounts[Index] += IncrementCount;
 		return RefCounts[Index];
 	}
 
-	void Decrement(int Index, short DecrementCount = 1)
+	void Decrement(int Index, unsigned short DecrementCount = 1)
 	{
-		check(IsValid(Index));
 		RefCounts[Index] -= DecrementCount;
-		check(RefCounts[Index] >= 0);
 		if (RefCounts[Index] == 0)
 		{
 			FreeIndices.Add(Index);
@@ -201,7 +196,7 @@ public:
 		}
 	}
 
-	const TDynamicVector<short>& GetRawRefCounts() const
+	const TDynamicVector<unsigned short>& GetRawRefCounts() const
 	{
 		return RefCounts;
 	}
@@ -209,7 +204,7 @@ public:
 	/**
 	 * @warning you should not use this!
 	 */
-	TDynamicVector<short>& GetRawRefCountsUnsafe()
+	TDynamicVector<unsigned short>& GetRawRefCountsUnsafe()
 	{
 		return RefCounts;
 	}
@@ -217,7 +212,7 @@ public:
 	/**
 	 * @warning you should not use this!
 	 */
-	void SetRefCountUnsafe(int Index, short ToCount)
+	void SetRefCountUnsafe(int Index, unsigned short ToCount)
 	{
 		RefCounts[Index] = ToCount;
 	}
@@ -453,8 +448,29 @@ public:
 			RefCounts.GetLength(), FreeIndices.GetLength(), (FreeIndices.GetByteCount() / 1024));
 	}
 
+	/**
+	 * Serialization operator for FRefCountVector.
+	 *
+	 * @param Ar Archive to serialize with.
+	 * @param Vec Vector to serialize.
+	 * @returns Passing down serializing archive.
+	 */
+	friend FArchive& operator<<(FArchive& Ar, FRefCountVector& Vec)
+	{
+		Vec.Serialize(Ar);
+		return Ar;
+	}
+
+	/** Serialize FRefCountVector to an archive. */
+	void Serialize(FArchive& Ar)
+	{
+		Ar << RefCounts;
+		Ar << FreeIndices;
+		Ar << UsedCount;
+	}
+
 private:
-	TDynamicVector<short> RefCounts{};
+	TDynamicVector<unsigned short> RefCounts{};
 	TDynamicVector<int> FreeIndices{};
 	int UsedCount{0};
 
diff --git a/Engine/Plugins/Experimental/GeometryProcessing/Source/GeometricObjects/Public/Util/SmallListSet.h b/Engine/Plugins/Experimental/GeometryProcessing/Source/GeometricObjects/Public/Util/SmallListSet.h
--- a/Engine/Plugins/Experimental/GeometryProcessing/Source/GeometricObjects/Public/Util/SmallListSet.h
+++ b/Engine/Plugins/Experimental/GeometryProcessing/Source/GeometricObjects/Public/Util/SmallListSet.h
@@ -181,6 +181,30 @@ public:
 	 */
 	void Enumerate(int32 ListIndex, TFunctionRef<void(int32)> ApplyFunc) const;
 
+	/**
+	 * Serialization operator for FSmallListSet.
+	 *
+	 * @param Ar Archive to serialize with.
+	 * @param Set Set to serialize.
+	 * @returns Passing down serializing archive.
+	 */
+	friend FArchive& operator<<(FArchive& Ar, FSmallListSet& Set)
+	{
+		Set.Serialize(Ar);
+		return Ar;
+	}
+
+	/** Serialize FSmallListSet to an archive. */
+	void Serialize(FArchive& Ar)
+	{
+		Ar << ListHeads;
+		Ar << ListBlocks;
+		Ar << FreeBlocks;
+		Ar << AllocatedCount;
+		Ar << LinkedListElements;
+		Ar << FreeHeadIndex;
+	}
+
 
 	//
 	// iterator support
diff --git a/Engine/Plugins/Experimental/GeometryProcessing/Source/GeometricObjects/Public/VectorTypes.h b/Engine/Plugins/Experimental/GeometryProcessing/Source/GeometricObjects/Public/VectorTypes.h
--- a/Engine/Plugins/Experimental/GeometryProcessing/Source/GeometricObjects/Public/VectorTypes.h
+++ b/Engine/Plugins/Experimental/GeometryProcessing/Source/GeometricObjects/Public/VectorTypes.h
@@ -5,6 +5,7 @@
 #include "Math/Vector.h"
 #include "Math/VectorLWC.h"
 #include "MathUtil.h"
+#include "Serialization/Archive.h"
 #include <sstream>
 
 
@@ -219,6 +220,25 @@ struct FVector2
 		return X != Other.X || Y != Other.Y;
 	}
 
+	/**
+	 * Serialization operator for FVector2.
+	 *
+	 * @param Ar Archive to serialize with.
+	 * @param Vec Vector to serialize.
+	 * @returns Passing down serializing archive.
+	 */
+	friend FArchive& operator<<(FArchive& Ar, FVector2& Vec)
+	{
+		Vec.Serialize(Ar);
+		return Ar;
+	}
+
+	/** Serialize FVector2 to an archive. */
+	void Serialize(FArchive& Ar)
+	{
+		Ar << X;
+		Ar << Y;
+	}
 };
 
 /** @return dot product of V1 with PerpCW(V2), ie V2 rotated 90 degrees clockwise */
@@ -1084,6 +1104,28 @@ struct TVector4
 	{
 		return X != Other.X || Y != Other.Y || Z != Other.Z || W != Other.W;
 	}
+
+	/**
+	 * Serialization operator for TVector4.
+	 *
+	 * @param Ar Archive to serialize with.
+	 * @param Vec Vector to serialize.
+	 * @returns Passing down serializing archive.
+	 */
+	friend FArchive& operator<<(FArchive& Ar, TVector4& Vec)
+	{
+		Vec.Serialize(Ar);
+		return Ar;
+	}
+
+	/** Serialize TVector4 to an archive. */
+	void Serialize(FArchive& Ar)
+	{
+		Ar << X;
+		Ar << Y;
+		Ar << Z;
+		Ar << W;
+	}
 };
 
 template <typename T>
@@ -1190,4 +1232,10 @@ typedef UE::Geometry::FVector2<double> FVector2d;
 //typedef UE::Geometry::FVector3<double> FVector3d;
 
 typedef UE::Geometry::TVector4<float> FVector4f;
-typedef UE::Geometry::TVector4<double> FVector4d;
\ No newline at end of file
+typedef UE::Geometry::TVector4<double> FVector4d;
+
+template<> struct TCanBulkSerialize<FVector2f> { enum { Value = true }; };
+template<> struct TCanBulkSerialize<FVector2d> { enum { Value = true }; };
+
+template<> struct TCanBulkSerialize<FVector4f> { enum { Value = true }; };
+template<> struct TCanBulkSerialize<FVector4d> { enum { Value = true }; };
