pass through tolerance and flag param when calling the MeshWeights CotanCentroidSafe

#rnx

diff --git a/Engine/Plugins/Experimental/GeometryProcessing/Source/DynamicMesh/Private/MeshWeights.cpp b/Engine/Plugins/Experimental/GeometryProcessing/Source/DynamicMesh/Private/MeshWeights.cpp
--- a/Engine/Plugins/Experimental/GeometryProcessing/Source/DynamicMesh/Private/MeshWeights.cpp
+++ b/Engine/Plugins/Experimental/GeometryProcessing/Source/DynamicMesh/Private/MeshWeights.cpp
@@ -234,12 +234,12 @@ FVector3d TCotanCentroidSafe(const FDynamicMesh3& mesh, int32 v_i, PositionFuncT
 
 FVector3d FMeshWeights::CotanCentroidSafe(const FDynamicMesh3& mesh, int32 v_i, double DegenerateTol, bool* bFailedToUniform)
 {
-	return TCotanCentroidSafe(mesh, v_i, [&](int32 vid) { return mesh.GetVertex(vid); });
+	return TCotanCentroidSafe(mesh, v_i, [&](int32 vid) { return mesh.GetVertex(vid); }, DegenerateTol, bFailedToUniform);
 }
 
 FVector3d FMeshWeights::CotanCentroidSafe(const FDynamicMesh3& mesh, int32 v_i, TFunctionRef<FVector3d(int32)> VertexPositionFunc, double DegenerateTol, bool* bFailedToUniform)
 {
-	return TCotanCentroid(mesh, v_i, VertexPositionFunc);
+	return TCotanCentroidSafe(mesh, v_i, VertexPositionFunc, DegenerateTol, bFailedToUniform);
 }
 
 
