GeometryProcessing: Fix RefCountVector unsigned short conversion errors.
Skip DynamicMesh Serialization Benchmark in regular tests.

#rnx

diff --git a/Engine/Plugins/Experimental/GeometryProcessing/Source/GeometricObjects/Public/Util/RefCountVector.h b/Engine/Plugins/Experimental/GeometryProcessing/Source/GeometricObjects/Public/Util/RefCountVector.h
--- a/Engine/Plugins/Experimental/GeometryProcessing/Source/GeometricObjects/Public/Util/RefCountVector.h
+++ b/Engine/Plugins/Experimental/GeometryProcessing/Source/GeometricObjects/Public/Util/RefCountVector.h
@@ -54,12 +54,12 @@ public:
 
 	bool IsValid(int Index) const
 	{
-		return (Index >= 0 && Index < (int)RefCounts.GetLength() && RefCounts[Index] > 0);
+		return (Index >= 0 && Index < (int)RefCounts.GetLength() && IsValidUnsafe(Index));
 	}
 
 	bool IsValidUnsafe(int Index) const
 	{
-		return RefCounts[Index] > 0;
+		return RefCounts[Index] > 0 && RefCounts[Index] < INVALID_REF_COUNT;
 	}
 
 	int GetRefCount(int Index) const
@@ -145,7 +145,7 @@ public:
 		}
 		else
 		{
-			if (RefCounts[Index] > 0)
+			if (IsValidUnsafe(Index))
 			{
 				return false;
 			}
@@ -186,7 +186,7 @@ public:
 		}
 		else
 		{
-			if (RefCounts[Index] > 0)
+			if (IsValidUnsafe(Index))
 			{
 				return false;
 			}
@@ -229,7 +229,7 @@ public:
 		int N = (int)RefCounts.GetLength();
 		for (int i = 0; i < N; ++i) 
 		{
-			if (RefCounts[i] > 0)
+			if (IsValidUnsafe(i))
 			{
 				UsedCount++;
 			}
