GeometryProcessing: Fix "Remesh Test Case Meshes" unit test.

#rnx

diff --git a/Engine/Plugins/Experimental/GeometryProcessing/Source/DynamicMesh/Public/DynamicMesh3.h b/Engine/Plugins/Experimental/GeometryProcessing/Source/DynamicMesh/Public/DynamicMesh3.h
--- a/Engine/Plugins/Experimental/GeometryProcessing/Source/DynamicMesh/Public/DynamicMesh3.h
+++ b/Engine/Plugins/Experimental/GeometryProcessing/Source/DynamicMesh/Public/DynamicMesh3.h
@@ -327,7 +327,7 @@ public:
 	inline bool IsReferencedVertex(int VertexID) const
 	{
 		return VertexID >= 0 && VertexID < (int)VertexRefCounts.GetMaxIndex() &&
-		    VertexRefCounts.GetRawRefCount(VertexID) > 1;
+		    VertexRefCounts.GetRefCount(VertexID) > 1;
 	}
 	/** @return true if TriangleID is a valid triangle in this mesh */
 	inline bool IsTriangle(int TriangleID) const
diff --git a/Engine/Plugins/Experimental/GeometryProcessing/Source/GeometricObjects/Public/Util/RefCountVector.h b/Engine/Plugins/Experimental/GeometryProcessing/Source/GeometricObjects/Public/Util/RefCountVector.h
--- a/Engine/Plugins/Experimental/GeometryProcessing/Source/GeometricObjects/Public/Util/RefCountVector.h
+++ b/Engine/Plugins/Experimental/GeometryProcessing/Source/GeometricObjects/Public/Util/RefCountVector.h
@@ -84,13 +84,13 @@ public:
 		}
 		else
 		{
-			int iFree = INVALID_REF_COUNT;
-			while (iFree == INVALID_REF_COUNT && FreeIndices.IsEmpty() == false)
+			int iFree = INDEX_NONE;
+			while (iFree == INDEX_NONE && FreeIndices.IsEmpty() == false)
 			{
 				iFree = FreeIndices.Back();
 				FreeIndices.PopBack();
 			}
-			if (iFree != INVALID_REF_COUNT)
+			if (iFree != INDEX_NONE)
 			{
 				RefCounts[iFree] = 1;
 				return iFree;
@@ -154,7 +154,7 @@ public:
 			{
 				if (FreeIndices[i] == Index)
 				{
-					FreeIndices[i] = INVALID_REF_COUNT;
+					FreeIndices[i] = INDEX_NONE;
 					RefCounts[Index] = 1;
 					UsedCount++;
 					return true;
