MeshModelingTools: Fix potential crash bug during mesh topology rebuild.

#rnx

diff --git a/Engine/Plugins/Experimental/GeometryProcessing/Source/DynamicMesh/Private/GroupTopology.cpp b/Engine/Plugins/Experimental/GeometryProcessing/Source/DynamicMesh/Private/GroupTopology.cpp
--- a/Engine/Plugins/Experimental/GeometryProcessing/Source/DynamicMesh/Private/GroupTopology.cpp
+++ b/Engine/Plugins/Experimental/GeometryProcessing/Source/DynamicMesh/Private/GroupTopology.cpp
@@ -817,7 +817,7 @@ bool FTriangleGroupTopology::RebuildTopology()
 	}
 
 	TArray<int32> MeshEdgeToGroupEdge;
-	MeshEdgeToGroupEdge.Init(-1, Mesh->MaxEdgeID());
+	MeshEdgeToGroupEdge.Init(INDEX_NONE, Mesh->MaxEdgeID());
 
 	// construct boundary loops
 	TArray<int> SpanVertices; SpanVertices.SetNum(2);
@@ -830,11 +830,8 @@ bool FTriangleGroupTopology::RebuildTopology()
 		FIndex3i TriEdges = Mesh->GetTriEdges(Group.GroupID);
 		for (int j = 0; j < 3; ++j)
 		{
-			if (MeshEdgeToGroupEdge[TriEdges[j]] != -1)
-			{
-				Boundary0.GroupEdges.Add(TriEdges[j]);
-			}
-			else
+			int& GroupEdgeIndex = MeshEdgeToGroupEdge[TriEdges[j]];
+			if (GroupEdgeIndex == INDEX_NONE)
 			{
 				FGroupEdge NewGroupEdge = { MakeEdgeID(TriEdges[j]) };
 				FIndex2i EdgeVerts = Mesh->GetEdgeV(TriEdges[j]);
@@ -843,10 +840,9 @@ bool FTriangleGroupTopology::RebuildTopology()
 				NewGroupEdge.Span.InitializeFromVertices(SpanVertices);
 				NewGroupEdge.EndpointCorners = FIndex2i(GetCornerIDFromVertexID(SpanVertices[0]), GetCornerIDFromVertexID(SpanVertices[1]));
 				check(NewGroupEdge.EndpointCorners.A != IndexConstants::InvalidID && NewGroupEdge.EndpointCorners.B != IndexConstants::InvalidID);
-				int32 EdgeIndex = Edges.Add(NewGroupEdge);
-				Boundary0.GroupEdges.Add(EdgeIndex);
-				MeshEdgeToGroupEdge[TriEdges[j]] = EdgeIndex;
+				GroupEdgeIndex = Edges.Add(NewGroupEdge);
 			}
+			Boundary0.GroupEdges.Add(GroupEdgeIndex);
 		}
 
 		// collect up .NeighbourGroupIDs and set .bIsOnBoundary
