Add support for converting mesh description with skin weights to/from dynamic mesh and topological editing of dynamic meshes, with skin weights that properly handle weight distribution under split/collapse.

#rnx

diff --git a/Engine/Plugins/Experimental/GeometryProcessing/Source/DynamicMesh/DynamicMesh.Build.cs b/Engine/Plugins/Experimental/GeometryProcessing/Source/DynamicMesh/DynamicMesh.Build.cs
--- a/Engine/Plugins/Experimental/GeometryProcessing/Source/DynamicMesh/DynamicMesh.Build.cs
+++ b/Engine/Plugins/Experimental/GeometryProcessing/Source/DynamicMesh/DynamicMesh.Build.cs
@@ -8,6 +8,15 @@ public class DynamicMesh : ModuleRules
 	{
         PCHUsage = ModuleRules.PCHUsageMode.UseExplicitOrSharedPCHs;
 
+        // For GPUSkinPublicDefs.h
+        PublicIncludePaths.Add("Runtime/Engine/Public");
+        
+        PublicIncludePathModuleNames.AddRange(
+			new string[] {
+				"AnimationCore",			// For the BoneWeights.h include
+			}
+		);
+
         PublicDependencyModuleNames.AddRange(
 			new string[] {
 				"Core",
diff --git a/Engine/Plugins/Experimental/GeometryProcessing/Source/DynamicMesh/Private/DynamicMeshAttributeSet.cpp b/Engine/Plugins/Experimental/GeometryProcessing/Source/DynamicMesh/Private/DynamicMeshAttributeSet.cpp
--- a/Engine/Plugins/Experimental/GeometryProcessing/Source/DynamicMesh/Private/DynamicMeshAttributeSet.cpp
+++ b/Engine/Plugins/Experimental/GeometryProcessing/Source/DynamicMesh/Private/DynamicMeshAttributeSet.cpp
@@ -1,13 +1,34 @@
 // Copyright Epic Games, Inc. All Rights Reserved.
 
 #include "DynamicMeshAttributeSet.h"
+
+#include "DynamicVertexSkinWeightsAttribute.h"
 #include "IndexTypes.h"
 #include "Async/Async.h"
 
 #include "ExplicitUseGeometryMathTypes.h"		// using UE::Geometry::(math types)
+
 using namespace UE::Geometry;
 
 
+FDynamicMeshAttributeSet::FDynamicMeshAttributeSet(FDynamicMesh3* Mesh)
+	: ParentMesh(Mesh)
+{
+	FDynamicMeshAttributeSet::SetNumUVLayers(1);
+	FDynamicMeshAttributeSet::SetNumNormalLayers(1);
+}
+
+FDynamicMeshAttributeSet::FDynamicMeshAttributeSet(FDynamicMesh3* Mesh, int32 NumUVLayers, int32 NumNormalLayers)
+	: ParentMesh(Mesh)
+{
+	FDynamicMeshAttributeSet::SetNumUVLayers(NumUVLayers);
+	FDynamicMeshAttributeSet::SetNumNormalLayers(NumNormalLayers);
+}
+
+FDynamicMeshAttributeSet::~FDynamicMeshAttributeSet()
+{
+}
+
 void FDynamicMeshAttributeSet::Copy(const FDynamicMeshAttributeSet& Copy)
 {
 	SetNumUVLayers(Copy.NumUVLayers());
@@ -45,8 +66,16 @@ void FDynamicMeshAttributeSet::Copy(const FDynamicMeshAttributeSet& Copy)
 		PolygroupLayers[GroupIdx].Copy(Copy.PolygroupLayers[GroupIdx]);
 	}
 
-	GenericAttributes.Reset();
 	ResetRegisteredAttributes();
+
+	SkinWeightAttributes.Reset();
+	for (const TPair<FName, TUniquePtr<FDynamicMeshVertexSkinWeightsAttribute>>& AttribPair : Copy.SkinWeightAttributes)
+	{
+		AttachSkinWeightsAttribute(AttribPair.Key,
+			static_cast<FDynamicMeshVertexSkinWeightsAttribute *>(AttribPair.Value->MakeCopy(ParentMesh)));
+	}
+	
+	GenericAttributes.Reset();
 	for (const TPair<FName, TUniquePtr<FDynamicMeshAttributeBase>>& AttribPair : Copy.GenericAttributes)
 	{
 		AttachAttribute(AttribPair.Key, AttribPair.Value->MakeCopy(ParentMesh));
@@ -124,8 +153,16 @@ void FDynamicMeshAttributeSet::CompactCopy(const FCompactMaps& CompactMaps, cons
 		PolygroupLayers[GroupIdx].CompactCopy(CompactMaps, Copy.PolygroupLayers[GroupIdx]);
 	}
 
-	GenericAttributes.Reset();
 	ResetRegisteredAttributes();
+
+	SkinWeightAttributes.Reset();
+	for (const TPair<FName, TUniquePtr<FDynamicMeshVertexSkinWeightsAttribute>>& AttribPair : Copy.SkinWeightAttributes)
+	{
+		AttachSkinWeightsAttribute(AttribPair.Key,
+			static_cast<FDynamicMeshVertexSkinWeightsAttribute *>(AttribPair.Value->MakeCompactCopy(CompactMaps, ParentMesh)));
+	}
+	
+	GenericAttributes.Reset();
 	for (const TPair<FName, TUniquePtr<FDynamicMeshAttributeBase>>& AttribPair : Copy.GenericAttributes)
 	{
 		AttachAttribute(AttribPair.Key, AttribPair.Value->MakeCompactCopy(CompactMaps, ParentMesh));
@@ -161,6 +198,11 @@ void FDynamicMeshAttributeSet::CompactInPlace(const FCompactMaps& CompactMaps)
 		PolygroupLayers[GroupIdx].CompactInPlace(CompactMaps);
 	}
 
+	for (const TPair<FName, TUniquePtr<FDynamicMeshVertexSkinWeightsAttribute>>& AttribPair : SkinWeightAttributes)
+	{
+		AttribPair.Value->CompactInPlace(CompactMaps);
+	}
+	
 	for (const TPair<FName, TUniquePtr<FDynamicMeshAttributeBase>>& AttribPair : GenericAttributes)
 	{
 		AttribPair.Value->CompactInPlace(CompactMaps);
@@ -252,8 +294,16 @@ void FDynamicMeshAttributeSet::EnableMatchingAttributes(const FDynamicMeshAttrib
 		PolygroupLayers[GroupIdx].Initialize((int32)0);
 	}
 
-	GenericAttributes.Reset();
 	ResetRegisteredAttributes();
+	
+	SkinWeightAttributes.Reset();
+	for (const TPair<FName, TUniquePtr<FDynamicMeshVertexSkinWeightsAttribute>>& AttribPair : ToMatch.SkinWeightAttributes)
+	{
+		AttachSkinWeightsAttribute(AttribPair.Key,
+			static_cast<FDynamicMeshVertexSkinWeightsAttribute *>(AttribPair.Value->MakeNew(ParentMesh)));
+	}
+
+	GenericAttributes.Reset();
 	for (const TPair<FName, TUniquePtr<FDynamicMeshAttributeBase>>& AttribPair : ToMatch.GenericAttributes)
 	{
 		AttachAttribute(AttribPair.Key, AttribPair.Value->MakeNew(ParentMesh));
@@ -289,6 +339,11 @@ void FDynamicMeshAttributeSet::Reparent(FDynamicMesh3* NewParent)
 		PolygroupLayers[GroupIdx].Reparent(NewParent);
 	}
 
+	for (const TPair<FName, TUniquePtr<FDynamicMeshVertexSkinWeightsAttribute>>& AttribPair : SkinWeightAttributes)
+	{
+		AttribPair.Value.Get()->Reparent(NewParent);
+	}
+	
 	for (const TPair<FName, TUniquePtr<FDynamicMeshAttributeBase>>& AttribPair : GenericAttributes)
 	{
 		AttribPair.Value->Reparent(NewParent);
@@ -415,6 +470,25 @@ void FDynamicMeshAttributeSet::DisableMaterialID()
 	MaterialIDAttrib.Reset();
 }
 
+void FDynamicMeshAttributeSet::AttachSkinWeightsAttribute(FName InProfileName, FDynamicMeshVertexSkinWeightsAttribute* InAttribute)
+{
+	RemoveSkinWeightsAttribute(InProfileName);
+
+	// Ensure proper ownership.
+	static_cast<FDynamicMeshAttributeBase *>(InAttribute)->Reparent(ParentMesh);
+	SkinWeightAttributes.Add(InProfileName, TUniquePtr<FDynamicMeshVertexSkinWeightsAttribute>(InAttribute));
+
+	RegisterExternalAttribute(InAttribute);
+}
+
+void FDynamicMeshAttributeSet::RemoveSkinWeightsAttribute(FName InProfileName)
+{
+	if (SkinWeightAttributes.Contains(InProfileName))
+	{
+		UnregisterExternalAttribute(SkinWeightAttributes[InProfileName].Get());
+		SkinWeightAttributes.Remove(InProfileName);
+	}
+}
 
 
 bool FDynamicMeshAttributeSet::IsSeamEdge(int eid) const
diff --git a/Engine/Plugins/Experimental/GeometryProcessing/Source/DynamicMesh/Private/DynamicMeshEditor.cpp b/Engine/Plugins/Experimental/GeometryProcessing/Source/DynamicMesh/Private/DynamicMeshEditor.cpp
--- a/Engine/Plugins/Experimental/GeometryProcessing/Source/DynamicMesh/Private/DynamicMeshEditor.cpp
+++ b/Engine/Plugins/Experimental/GeometryProcessing/Source/DynamicMesh/Private/DynamicMeshEditor.cpp
@@ -6,6 +6,7 @@
 #include "Util/BufferUtil.h"
 #include "MeshRegionBoundaryLoops.h"
 #include "DynamicSubmesh3.h"
+#include "DynamicVertexSkinWeightsAttribute.h"
 #include "MeshNormals.h"
 #include "MeshQueries.h"
 #include "Selections/MeshConnectedComponents.h"
@@ -1567,6 +1568,15 @@ void FDynamicMeshEditor::AppendMesh(const FDynamicMesh3* AppendMesh,
 			}
 		}
 
+		for (const TPair<FName, TUniquePtr<FDynamicMeshVertexSkinWeightsAttribute>>& AttribPair : AppendMesh->Attributes()->GetSkinWeightsAttributes())
+		{
+			FDynamicMeshVertexSkinWeightsAttribute* ToAttrib = Mesh->Attributes()->GetSkinWeightsAttribute(AttribPair.Key);
+			if (ToAttrib)
+			{
+				ToAttrib->CopyThroughMapping(AttribPair.Value.Get(), IndexMapsOut);
+			}
+		}
+		
 		for (const TPair<FName, TUniquePtr<FDynamicMeshAttributeBase>>& AttribPair : AppendMesh->Attributes()->GetAttachedAttributes())
 		{
 			if (Mesh->Attributes()->HasAttachedAttribute(AttribPair.Key))
@@ -1830,7 +1840,16 @@ static void AppendGenericAttributes(const FDynamicMesh3* FromMesh, FDynamicMesh3
 		return;
 	}
 
-	// copy generic attributes after full IndexMaps have been created
+	// Copy skin weight and generic attributes after full IndexMaps have been created. 	
+	for (const TPair<FName, TUniquePtr<FDynamicMeshVertexSkinWeightsAttribute>>& AttribPair : FromMesh->Attributes()->GetSkinWeightsAttributes())
+	{
+		FDynamicMeshVertexSkinWeightsAttribute* ToAttrib = ToMesh->Attributes()->GetSkinWeightsAttribute(AttribPair.Key);
+		if (ToAttrib)
+		{
+			ToAttrib->CopyThroughMapping(AttribPair.Value.Get(), IndexMaps);
+		}
+	}
+	
 	for (const TPair<FName, TUniquePtr<FDynamicMeshAttributeBase>>& AttribPair : FromMesh->Attributes()->GetAttachedAttributes())
 	{
 		if (ToMesh->Attributes()->HasAttachedAttribute(AttribPair.Key))
diff --git a/Engine/Plugins/Experimental/GeometryProcessing/Source/DynamicMesh/Public/DynamicMeshAttributeSet.h b/Engine/Plugins/Experimental/GeometryProcessing/Source/DynamicMesh/Public/DynamicMeshAttributeSet.h
--- a/Engine/Plugins/Experimental/GeometryProcessing/Source/DynamicMesh/Public/DynamicMeshAttributeSet.h
+++ b/Engine/Plugins/Experimental/GeometryProcessing/Source/DynamicMesh/Public/DynamicMeshAttributeSet.h
@@ -27,6 +27,11 @@ typedef TDynamicMeshScalarTriangleAttribute<int32> FDynamicMeshMaterialAttribute
 /** Per-triangle integer polygroup ID */
 typedef TDynamicMeshScalarTriangleAttribute<int32> FDynamicMeshPolygroupAttribute;
 
+/** Forward declarations */
+template<typename ParentType>
+class TDynamicVertexSkinWeightsAttribute;
+using FDynamicMeshVertexSkinWeightsAttribute = TDynamicVertexSkinWeightsAttribute<FDynamicMesh3>;
+	
 /**
  * FDynamicMeshAttributeSet manages a set of extended attributes for a FDynamicMesh3.
  * This includes UV and Normal overlays, etc.
@@ -39,24 +44,11 @@ typedef TDynamicMeshScalarTriangleAttribute<int32> FDynamicMeshPolygroupAttribut
 class DYNAMICMESH_API FDynamicMeshAttributeSet : public FDynamicMeshAttributeSetBase
 {
 public:
+	FDynamicMeshAttributeSet(FDynamicMesh3* Mesh);
 
-	FDynamicMeshAttributeSet(FDynamicMesh3* Mesh)
-		: ParentMesh(Mesh)
-	{
-		SetNumUVLayers(1);
-		SetNumNormalLayers(1);
-	}
+	FDynamicMeshAttributeSet(FDynamicMesh3* Mesh, int32 NumUVLayers, int32 NumNormalLayers);
 
-	FDynamicMeshAttributeSet(FDynamicMesh3* Mesh, int32 NumUVLayers, int32 NumNormalLayers)
-		: ParentMesh(Mesh)
-	{
-		SetNumUVLayers(NumUVLayers);
-		SetNumNormalLayers(NumNormalLayers);
-	}
-
-	virtual ~FDynamicMeshAttributeSet()
-	{
-	}
+	virtual ~FDynamicMeshAttributeSet() override;
 
 	void Copy(const FDynamicMeshAttributeSet& Copy);
 
@@ -303,6 +295,41 @@ public:
 		return MaterialIDAttrib.Get();
 	}
 
+	/** Skin weights */
+
+	/// Create a new skin weights attribute with a given skin weights profile name. If an attribute already exists with
+	/// that name, that existing attribute will be deleted.
+	void AttachSkinWeightsAttribute(FName InProfileName, FDynamicMeshVertexSkinWeightsAttribute* InAttribute);
+
+	/// Remove a skin weights attribute matching the given profile name.
+	void RemoveSkinWeightsAttribute(FName InProfileName);
+
+	/// Returns true if the list of skin weight attributes includes the given profile name.
+	bool HasSkinWeightsAttribute(FName InProfileName) const
+	{
+		return SkinWeightAttributes.Contains(InProfileName);
+	}
+
+	/// Returns a pointer to a skin weight attribute of the given profile name. If the attribute
+	/// does not exist, a nullptr is returned.
+	FDynamicMeshVertexSkinWeightsAttribute *GetSkinWeightsAttribute(FName InProfileName) const
+	{
+		if (SkinWeightAttributes.Contains(InProfileName))
+		{
+			return SkinWeightAttributes[InProfileName].Get();
+		}
+		else
+		{
+			return nullptr;
+		}
+	}
+	
+	/// Returns a map of all skin weight attributes.
+	const TMap<FName, TUniquePtr<FDynamicMeshVertexSkinWeightsAttribute>>& GetSkinWeightsAttributes() const
+	{
+		return SkinWeightAttributes;
+	}
+	
 	// Attach a new attribute (and transfer ownership of it to the attribute set)
 	void AttachAttribute(FName AttribName, FDynamicMeshAttributeBase* Attribute)
 	{
@@ -377,8 +404,10 @@ protected:
 
 	TIndirectArray<FDynamicMeshPolygroupAttribute> PolygroupLayers;
 
+	TMap<FName, TUniquePtr<FDynamicMeshVertexSkinWeightsAttribute>> SkinWeightAttributes;
+	
 	TMap<FName, TUniquePtr<FDynamicMeshAttributeBase>> GenericAttributes;
-
+	
 protected:
 	friend class FDynamicMesh3;
 
diff --git a/Engine/Plugins/Experimental/GeometryProcessing/Source/DynamicMesh/Public/DynamicVertexSkinWeightsAttribute.h b/Engine/Plugins/Experimental/GeometryProcessing/Source/DynamicMesh/Public/DynamicVertexSkinWeightsAttribute.h
new file mode 100644
--- /dev/null
+++ b/Engine/Plugins/Experimental/GeometryProcessing/Source/DynamicMesh/Public/DynamicVertexSkinWeightsAttribute.h
@@ -0,0 +1,411 @@
+// Copyright Epic Games, Inc. All Rights Reserved.
+
+#pragma once
+
+#include "DynamicMesh3.h"
+#include "DynamicAttribute.h"
+
+#include "BoneWeights.h"
+#include "HAL/UnrealMemory.h"
+
+
+namespace UE
+{
+namespace Geometry
+{
+
+template<typename ParentType>
+class TDynamicVertexSkinWeightsAttributeChange final : 
+	public TDynamicAttributeChangeBase<ParentType>
+{
+	struct FChangeVertexBoneWeights
+	{
+		UE::AnimationCore::FBoneWeights Weights;
+		int VertexID;
+	};
+
+	TArray<FChangeVertexBoneWeights> OldVertexBoneWeights, NewVertexBoneWeights;
+
+public:
+	TDynamicVertexSkinWeightsAttributeChange() = default;
+
+	virtual ~TDynamicVertexSkinWeightsAttributeChange() override = default;
+
+	void SaveInitialVertex(const TDynamicAttributeBase<ParentType>* Attribute, int VertexID) override;
+
+	void StoreAllFinalVertices(const TDynamicAttributeBase<ParentType>* Attribute, const TArray<int>& VertexIDs) override;
+
+	bool Apply(TDynamicAttributeBase<ParentType>* Attribute, bool bRevert) const override;
+};
+
+
+/**
+ * TDynamicVertexAttribute provides per-vertex storage of bone weights (skin weights)
+ */
+template<typename ParentType>
+class TDynamicVertexSkinWeightsAttribute final : 
+	public TDynamicAttributeBase<ParentType>
+{
+public:
+	using FBoneWeights = UE::AnimationCore::FBoneWeights;
+
+protected:
+	friend class FDynamicMeshAttributeSet;
+
+	/** The parent object (e.g. mesh, point set) this attribute belongs to */
+	ParentType* Parent = nullptr;
+
+	/** List of per-vertex bone weights values */
+	// FIXME: Replace with a local slab allocator storage to avoid per-item TInlineAllocator
+	TDynamicVector<FBoneWeights> VertexBoneWeights;
+
+public:
+
+	/** Create an empty overlay */
+	TDynamicVertexSkinWeightsAttribute() = default;
+
+	/** Create an attribute for the given parent */
+	TDynamicVertexSkinWeightsAttribute(ParentType* ParentIn, bool bAutoInit = true) : 
+		Parent(ParentIn)
+	{
+		if (bAutoInit)
+		{
+			Initialize();
+		}
+	}
+
+	virtual ~TDynamicVertexSkinWeightsAttribute() = default;
+
+	/** @return the parent for this attribute */
+	const ParentType* GetParent() const { return Parent; }
+	/** @return the parent for this attribute */
+	ParentType* GetParent() { return Parent; }
+private:
+	void Reparent(ParentType* NewParent) override { Parent = NewParent;  }
+public:
+	TDynamicAttributeBase<ParentType>* MakeNew(ParentType* ParentIn) const override
+	{
+		TDynamicVertexSkinWeightsAttribute<ParentType>* Matching = new TDynamicVertexSkinWeightsAttribute<ParentType>(ParentIn);
+		Matching->Initialize();
+		return Matching;
+	}
+	TDynamicAttributeBase<ParentType>* MakeCopy(ParentType* ParentIn) const override
+	{
+		TDynamicVertexSkinWeightsAttribute<ParentType>* ToFill = new TDynamicVertexSkinWeightsAttribute<ParentType>(ParentIn);
+		ToFill->Copy(*this);
+		return ToFill;
+	}
+
+	/** Set this overlay to contain the same arrays as the copy overlay */
+	void Copy(const TDynamicVertexSkinWeightsAttribute<ParentType>& Copy)
+	{
+		TDynamicAttributeBase<ParentType>::CopyParentClassData(Copy);
+		VertexBoneWeights = Copy.VertexBoneWeights;
+	}
+
+	TDynamicAttributeBase<ParentType>* MakeCompactCopy(const FCompactMaps& CompactMaps, ParentType* ParentTypeIn) const override
+	{
+		TDynamicVertexSkinWeightsAttribute<ParentType>* ToFill = new TDynamicVertexSkinWeightsAttribute<ParentType>(ParentTypeIn);
+		ToFill->Initialize();
+		ToFill->CompactCopy(CompactMaps, *this);
+		return ToFill;
+	}
+
+	void CompactInPlace(const FCompactMaps& CompactMaps) override
+	{
+		for (int32 VID = 0; VID < CompactMaps.MapV.Num(); VID++)
+		{
+			const int32 ToVID = CompactMaps.MapV[VID];
+			if (ToVID < 0)
+			{
+				continue;
+			}
+			if (ensure(ToVID <= VID))
+			{
+				CopyValue(VID, ToVID);
+			}
+		}
+		VertexBoneWeights.Resize(Parent->MaxVertexID());
+	}
+
+	void CompactCopy(const FCompactMaps& CompactMaps, const TDynamicVertexSkinWeightsAttribute<ParentType>& ToCopy)
+	{
+		TDynamicAttributeBase<ParentType>::CopyParentClassData(ToCopy);
+		check(CompactMaps.MapV.Num() <= VertexBoneWeights.Num());
+		FBoneWeights Data;
+		for (int32 VID = 0; VID < CompactMaps.MapV.Num(); VID++)
+		{
+			const int32 ToVID = CompactMaps.MapV[VID];
+			if (ToVID == -1)
+			{
+				continue;
+			}
+			ToCopy.GetValue(VID, Data);
+			SetValue(ToVID, Data);
+		}
+	}
+
+	/** Initialize the attribute values to the given max triangle ID */
+	void Initialize(const FBoneWeights InitialValue = {})
+	{
+		check(Parent != nullptr);
+		VertexBoneWeights.Resize(Parent->MaxVertexID());
+		VertexBoneWeights.Fill(InitialValue);
+	}
+
+	void SetNewValue(int32 InNewVertexID, const FBoneWeights InBoneWeights)
+	{
+		VertexBoneWeights.InsertAt(InBoneWeights, InNewVertexID);
+	}
+
+
+
+	//
+	// Accessors/Queries
+	//
+
+	bool CopyThroughMapping(const TDynamicAttributeBase<ParentType>* Source, const FMeshIndexMappings& Mapping) override
+	{
+		// Don't snarf the FBoneWeight as a concrete object, since it _may_ contain a pointer
+		// and we don't want it to destruct.
+		constexpr int32 BufferSize = sizeof(FBoneWeights);
+		int8 BufferData[BufferSize];
+		for (const TPair<int32, int32>& MapVID : Mapping.GetVertexMap().GetForwardMap())
+		{
+			if (!ensure(Source->CopyOut(MapVID.Key, &BufferData, BufferSize)))
+			{
+				return false;
+			}
+			SetValue(MapVID.Value, *reinterpret_cast<FBoneWeights *>(BufferData));
+		}
+		return true;
+	}
+	bool CopyOut(int RawID, void* Buffer, int BufferSize) const override
+	{
+		if (sizeof(FBoneWeights) != BufferSize)
+		{
+			return false;
+		}
+
+		// Avoid the copy constructor.
+		FMemory::Memcpy(Buffer, &VertexBoneWeights[RawID], sizeof(FBoneWeights));
+		return true;
+	}
+	bool CopyIn(int RawID, void* Buffer, int BufferSize) override
+	{
+		if (sizeof(FBoneWeights) != BufferSize)
+		{
+			return false;
+		}
+		
+		// Ensure the copy constructor is called.
+		VertexBoneWeights[RawID] = *static_cast<FBoneWeights*>(Buffer);
+		return true;
+	}
+
+	/** Get the element at a given index */
+	void GetValue(int VertexID, FBoneWeights& Data) const
+	{
+		Data = VertexBoneWeights[VertexID];
+	}
+
+	/** Get the element at a given index */
+	template<typename AsType>
+	void GetValue(int VertexID, AsType& Data) const
+	{
+		Data = VertexBoneWeights[VertexID];
+	}
+
+
+	/** Set the element at a given index */
+	void SetValue(int VertexID, const FBoneWeights& Data)
+	{
+		VertexBoneWeights[VertexID] = Data;
+	}
+
+	/** Set the element at a given index */
+	template<typename ContainerAdapter>
+	void SetValue(int VertexID, const UE::AnimationCore::TBoneWeights<ContainerAdapter>& Data)
+	{
+		VertexBoneWeights[VertexID] = FBoneWeights::Create(Data);
+	}
+
+	/**
+	 * Copy the attribute value at FromVertexID to ToVertexID
+	 */
+	void CopyValue(int FromVertexID, int ToVertexID)
+	{
+		VertexBoneWeights.InsertAt(VertexBoneWeights[FromVertexID], ToVertexID);
+	}
+
+
+public:
+
+	/** Update the overlay to reflect an edge split in the parent */
+	void OnSplitEdge(const FDynamicMesh3::FEdgeSplitInfo& SplitInfo) override
+	{
+		ResizeAttribStoreIfNeeded(SplitInfo.NewVertex);
+		SetBoneWeightsFromLerp(SplitInfo.NewVertex, SplitInfo.OriginalVertices.A, SplitInfo.OriginalVertices.B, SplitInfo.SplitT);
+	}
+
+	/** Update the overlay to reflect an edge flip in the parent */
+	void OnFlipEdge(const FDynamicMesh3::FEdgeFlipInfo& FlipInfo) override
+	{
+		// vertices unchanged
+	}
+
+	/** Update the overlay to reflect an edge collapse in the parent */
+	void OnCollapseEdge(const FDynamicMesh3::FEdgeCollapseInfo& CollapseInfo) override
+	{
+		SetBoneWeightsFromLerp(CollapseInfo.KeptVertex, CollapseInfo.KeptVertex, CollapseInfo.RemovedVertex, CollapseInfo.CollapseT);
+	}
+
+	void ResizeAttribStoreIfNeeded(int VertexID)
+	{
+		if (!ensure(VertexID >= 0))
+		{
+			return;
+		}
+		size_t NeededSize = (size_t(VertexID)+1);
+		if (NeededSize > VertexBoneWeights.Num())
+		{
+			VertexBoneWeights.Resize(NeededSize, FBoneWeights{});
+		}
+	}
+
+	void OnNewVertex(int VertexID, bool bInserted) override
+	{
+		ResizeAttribStoreIfNeeded(VertexID);
+	}
+
+	/** Update the overlay to reflect a face poke in the parent */
+	void OnPokeTriangle(const FDynamicMesh3::FPokeTriangleInfo& PokeInfo) override
+	{
+		FIndex3i Tri = PokeInfo.TriVertices;
+		ResizeAttribStoreIfNeeded(PokeInfo.NewVertex);
+		SetBoneWeightsFromBary(PokeInfo.NewVertex, Tri.A, Tri.B, Tri.C, PokeInfo.BaryCoords);
+	}
+
+	/** Update the overlay to reflect an edge merge in the parent */
+	void OnMergeEdges(const FDynamicMesh3::FMergeEdgesInfo& MergeInfo) override
+	{
+		// just blend the attributes?
+		if (MergeInfo.RemovedVerts.A != FDynamicMesh3::InvalidID)
+		{
+			SetBoneWeightsFromLerp(MergeInfo.KeptVerts.A, MergeInfo.KeptVerts.A, MergeInfo.RemovedVerts.A, .5);
+		}
+		if (MergeInfo.RemovedVerts.B != FDynamicMesh3::InvalidID)
+		{
+			SetBoneWeightsFromLerp(MergeInfo.KeptVerts.B, MergeInfo.KeptVerts.B, MergeInfo.RemovedVerts.B, .5);
+		}
+	}
+
+	/** Update the overlay to reflect a vertex split in the parent */
+	void OnSplitVertex(const FDynamicMesh3::FVertexSplitInfo& SplitInfo, const TArrayView<const int>& TrianglesToUpdate) override
+	{
+		CopyValue(SplitInfo.OriginalVertex, SplitInfo.NewVertex);
+	}
+
+	TUniquePtr<TDynamicAttributeChangeBase<ParentType>> NewBlankChange() override
+	{
+		return MakeUnique<TDynamicVertexSkinWeightsAttributeChange<ParentType>>();
+	}
+
+	/**
+	* Check validity of attribute
+	* 
+	* @param bAllowNonmanifold Accept non-manifold topology as valid. Note that this should almost always be true for attributes; non-manifold overlays are generally valid.
+	* @param FailMode Desired behavior if mesh is found invalid
+	*/
+	bool CheckValidity(bool bAllowNonmanifold, EValidityCheckFailMode FailMode) const override
+	{
+		// just check that the values buffer is big enough
+		if (Parent->MaxVertexID() > VertexBoneWeights.Num())
+		{
+			switch (FailMode)
+			{
+			case EValidityCheckFailMode::Check:
+				check(false);
+			case EValidityCheckFailMode::Ensure:
+				ensure(false);
+			default:
+				return false;
+			}
+		}
+
+		return true;
+	}
+
+
+protected:
+
+	// interpolation functions; default implementation assumes your attributes can be interpolated as reals
+
+	/** Set the value at an Attribute to be a linear interpolation of two other Attributes */
+	void SetBoneWeightsFromLerp(int SetAttribute, int AttributeA, int AttributeB, double Alpha)
+	{
+		Alpha = FMath::Clamp(Alpha, 0.0, 1.0);
+		VertexBoneWeights[SetAttribute] = FBoneWeights::Blend(VertexBoneWeights[AttributeA], VertexBoneWeights[AttributeB], Alpha);
+	}
+
+	/** Set the value at an Attribute to be a barycentric interpolation of three other Attributes */
+	void SetBoneWeightsFromBary(int SetAttribute, int AttributeA, int AttributeB, int AttributeC, const FVector3d& BaryCoords)
+	{
+		// Since FBoneWeights only defines Blend for two inputs, we need to split the barycentric coordinate interpolation 
+		// into two blends.
+		if (!FMath::IsNearlyZero(BaryCoords.Y + BaryCoords.Z))
+		{
+			const double BCW = BaryCoords.Y / (BaryCoords.Y + BaryCoords.Z);
+			const FBoneWeights BC = FBoneWeights::Blend(VertexBoneWeights[AttributeB], VertexBoneWeights[AttributeC], BCW);
+
+			VertexBoneWeights[SetAttribute] = FBoneWeights::Blend(VertexBoneWeights[AttributeA], BC, BaryCoords.X);
+		}
+		else if (SetAttribute != AttributeA)
+		{
+			VertexBoneWeights[SetAttribute] = VertexBoneWeights[AttributeA];
+		}
+	}
+};
+
+
+class FDynamicMesh3;
+using FDynamicMeshVertexSkinWeightsAttribute = TDynamicVertexSkinWeightsAttribute<FDynamicMesh3>;
+
+
+template<typename ParentType>
+void TDynamicVertexSkinWeightsAttributeChange<ParentType>::SaveInitialVertex(const TDynamicAttributeBase<ParentType>* Attribute, int VertexID)
+{
+	FChangeVertexBoneWeights& Change = OldVertexBoneWeights.Emplace_GetRef();
+	Change.VertexID = VertexID;
+	const TDynamicVertexSkinWeightsAttribute<ParentType>* AttribCast = static_cast<const TDynamicVertexSkinWeightsAttribute<ParentType>*>(Attribute);
+	AttribCast->GetValue(VertexID, Change.Weights);
+}
+
+template<typename ParentType>
+void TDynamicVertexSkinWeightsAttributeChange<ParentType>::StoreAllFinalVertices(const TDynamicAttributeBase<ParentType>* Attribute, const TArray<int>& VertexIDs)
+{
+	NewVertexBoneWeights.Reserve(NewVertexBoneWeights.Num() + VertexIDs.Num());
+	const TDynamicVertexSkinWeightsAttribute<ParentType>* AttribCast = static_cast<const TDynamicVertexSkinWeightsAttribute<ParentType>*>(Attribute);
+	for (int VertexID : VertexIDs)
+	{
+		FChangeVertexBoneWeights& Change = NewVertexBoneWeights.Emplace_GetRef();
+		Change.VertexID = VertexID;
+		AttribCast->GetValue(VertexID, Change.Weights);
+	}
+}
+
+template<typename ParentType>
+bool TDynamicVertexSkinWeightsAttributeChange<ParentType>::Apply(TDynamicAttributeBase<ParentType>* Attribute, bool bRevert) const
+{
+	const TArray<FChangeVertexBoneWeights> *Changes = bRevert ? &OldVertexBoneWeights : &NewVertexBoneWeights;
+	TDynamicVertexSkinWeightsAttribute<ParentType>* AttribCast = static_cast<TDynamicVertexSkinWeightsAttribute<ParentType>*>(Attribute);
+	for (const FChangeVertexBoneWeights& Change : *Changes)
+	{
+		check(AttribCast->GetParent()->IsVertex(Change.VertexID));
+		AttribCast->SetValue(Change.VertexID, Change.Weights);
+	}
+	return true;
+}
+
+}
+}
diff --git a/Engine/Plugins/Experimental/GeometryProcessing/Source/MeshConversion/MeshConversion.Build.cs b/Engine/Plugins/Experimental/GeometryProcessing/Source/MeshConversion/MeshConversion.Build.cs
--- a/Engine/Plugins/Experimental/GeometryProcessing/Source/MeshConversion/MeshConversion.Build.cs
+++ b/Engine/Plugins/Experimental/GeometryProcessing/Source/MeshConversion/MeshConversion.Build.cs
@@ -14,6 +14,7 @@ public class MeshConversion : ModuleRules
 				"Engine",
                 "MeshDescription",
 				"StaticMeshDescription",
+				"SkeletalMeshDescription",
 				"GeometricObjects",
 				"DynamicMesh"
             }
diff --git a/Engine/Plugins/Experimental/GeometryProcessing/Source/MeshConversion/Private/DynamicMeshToMeshDescription.cpp b/Engine/Plugins/Experimental/GeometryProcessing/Source/MeshConversion/Private/DynamicMeshToMeshDescription.cpp
--- a/Engine/Plugins/Experimental/GeometryProcessing/Source/MeshConversion/Private/DynamicMeshToMeshDescription.cpp
+++ b/Engine/Plugins/Experimental/GeometryProcessing/Source/MeshConversion/Private/DynamicMeshToMeshDescription.cpp
@@ -1,9 +1,11 @@
 // Copyright Epic Games, Inc. All Rights Reserved. 
 
 #include "DynamicMeshToMeshDescription.h"
-#include "StaticMeshAttributes.h"
+#include "SkeletalMeshAttributes.h"
 #include "DynamicMeshAttributeSet.h"
+
 #include "DynamicMeshOverlay.h"
+#include "DynamicVertexSkinWeightsAttribute.h"
 #include "MeshDescriptionBuilder.h"
 #include "MeshTangents.h"
 
@@ -845,6 +847,30 @@ void FDynamicMeshToMeshDescription::Convert_NoSharedInstances(const FDynamicMesh
 	// convert polygroup layers
 	ConvertPolygroupLayers(MeshIn, MeshOut, IndexToTriangleIDMap);
 
+	// Convert all attached skin weights, if we're converting a mesh description that originated from
+	// a USkeletalMesh.
+	FSkeletalMeshAttributes MeshOutAttributes(MeshOut);
+	for (const TTuple<FName, TUniquePtr<FDynamicMeshVertexSkinWeightsAttribute>>& AttributeInfo: MeshIn->Attributes()->GetSkinWeightsAttributes())
+	{
+		FName ProfileName = AttributeInfo.Key;
+
+		MeshOutAttributes.RegisterSkinWeightAttribute(ProfileName);
+
+		const FDynamicMeshVertexSkinWeightsAttribute *MeshSkinWeights = AttributeInfo.Value.Get();
+		
+		FSkinWeightsVertexAttributesRef VertexBoneWeights = MeshOutAttributes.GetVertexSkinWeights(ProfileName);
+
+		for (int32 VertexIndex = 0; VertexIndex < MapV.Num(); VertexIndex++)
+		{
+			if (const FVertexID VertexID = MapV[VertexIndex]; VertexID != INDEX_NONE)
+			{
+				UE::AnimationCore::FBoneWeights BW;
+				MeshSkinWeights->GetValue(VertexIndex, BW);
+				VertexBoneWeights.Set(VertexID, BW);
+			}
+		}
+	}
+	
 	Builder.ResumeMeshDescriptionIndexing();
 }
 
diff --git a/Engine/Plugins/Experimental/GeometryProcessing/Source/MeshConversion/Private/MeshDescriptionToDynamicMesh.cpp b/Engine/Plugins/Experimental/GeometryProcessing/Source/MeshConversion/Private/MeshDescriptionToDynamicMesh.cpp
--- a/Engine/Plugins/Experimental/GeometryProcessing/Source/MeshConversion/Private/MeshDescriptionToDynamicMesh.cpp
+++ b/Engine/Plugins/Experimental/GeometryProcessing/Source/MeshConversion/Private/MeshDescriptionToDynamicMesh.cpp
@@ -3,9 +3,12 @@
 #include "MeshDescriptionToDynamicMesh.h"
 #include "DynamicMeshAttributeSet.h"
 #include "DynamicMeshOverlay.h"
+#include "DynamicVertexAttribute.h"
 #include "MeshTangents.h"
 #include "MeshDescriptionBuilder.h"
 #include "StaticMeshAttributes.h"
+#include "SkeletalMeshAttributes.h"
+#include "DynamicVertexSkinWeightsAttribute.h"
 #include "Async/Async.h"
 
 #include "ExplicitUseGeometryMathTypes.h"		// using UE::Geometry::(math types)
@@ -159,17 +162,15 @@ public:
 };
 
 
-
-
-
-
-
-
-
+struct FSkinWeightsAttribCopyInfo
+{
+	FSkinWeightsVertexAttributesConstRef MeshDescAttribRef;
+	FDynamicMeshVertexSkinWeightsAttribute* DynaMeshAttribRef = nullptr;
+};
 
 
 void FMeshDescriptionToDynamicMesh::Convert(const FMeshDescription* MeshIn, FDynamicMesh3& MeshOut, bool bCopyTangents )
-{ 
+{
 	TriIDMap.Reset();
 	VertIDMap.Reset();
 
@@ -189,7 +190,7 @@ void FMeshDescriptionToDynamicMesh::Convert(const FMeshDescription* MeshIn, FDyn
 
 	// look up vertex-instance UVs and normals
 	// @todo: does the MeshDescription always have UVs and Normals?
-	FStaticMeshConstAttributes Attributes(*MeshIn);
+	FSkeletalMeshConstAttributes Attributes(*MeshIn);
 	TVertexInstanceAttributesConstRef<FVector2D> InstanceUVs = Attributes.GetVertexInstanceUVs();
 	TVertexInstanceAttributesConstRef<FVector> InstanceNormals = Attributes.GetVertexInstanceNormals();
 	TVertexInstanceAttributesConstRef<FVector> InstanceTangents = Attributes.GetVertexInstanceTangents();
@@ -345,7 +346,7 @@ void FMeshDescriptionToDynamicMesh::Convert(const FMeshDescription* MeshIn, FDyn
 		
 		TriIDMap[NewTriangleID] = TriangleID;
 	
-	}
+		}
 
 	int32 MaxTriID = MeshOut.MaxTriangleID(); // really MaxTriID+1
 	// if the source mesh had duplicates then TriIDMap will be too long, this will trim excess
@@ -386,8 +387,8 @@ void FMeshDescriptionToDynamicMesh::Convert(const FMeshDescription* MeshIn, FDyn
 	FDynamicMeshNormalOverlay* TangentOverlay = nullptr;
 	FDynamicMeshNormalOverlay* BiTangentOverlay = nullptr;
 	FDynamicMeshMaterialAttribute* MaterialIDAttrib = nullptr;
-	
 
+	TArray<FSkinWeightsAttribCopyInfo> SkinWeightAttribs;
 	if (!bDisableAttributes)
 	{
 		MeshOut.EnableAttributes(); // by default 1-UV layer and 1-normal layer
@@ -419,7 +420,7 @@ void FMeshDescriptionToDynamicMesh::Convert(const FMeshDescription* MeshIn, FDyn
 		{
 			MeshOut.Attributes()->GetUVLayer(i)->InitializeTriangles(MeshOut.MaxTriangleID());
 		}
-
+	
 		if (InstanceColors.IsValid())
 		{
 			MeshOut.Attributes()->EnablePrimaryColors();
@@ -430,6 +431,23 @@ void FMeshDescriptionToDynamicMesh::Convert(const FMeshDescription* MeshIn, FDyn
 		// always enable Material ID if there are any attributes
 		MeshOut.Attributes()->EnableMaterialID();
 		MaterialIDAttrib = MeshOut.Attributes()->GetMaterialID();
+
+		// Copy all skin weight attributes, if they exist.
+		MeshIn->VertexAttributes().ForEach([&](const FName InAttributeName, auto InAttributesRef)
+		{
+			if (FSkeletalMeshAttributes::IsSkinWeightAttribute(InAttributeName))
+			{
+				const FName ProfileName = FSkeletalMeshAttributes::GetProfileNameFromAttribute(InAttributeName);
+				
+                FDynamicMeshVertexSkinWeightsAttribute *VertexSkinWeightsAttrib = new FDynamicMeshVertexSkinWeightsAttribute(&MeshOut);
+
+                SkinWeightAttribs.Add({
+                    Attributes.GetVertexSkinWeightsFromAttributeName(InAttributeName),
+                	VertexSkinWeightsAttrib                    
+                });
+				MeshOut.Attributes()->AttachSkinWeightsAttribute(ProfileName, VertexSkinWeightsAttrib);
+			}
+		});
 	}
 
 
@@ -439,7 +457,7 @@ void FMeshDescriptionToDynamicMesh::Convert(const FMeshDescription* MeshIn, FDyn
 	TArray<FName> PolygroupAttribNames;
 	TriAttribsSet.ForEach([&](const FName AttributeName, auto AttributesRef)
 	{
-		if (IsReservedAttributeName(AttributeName)) return;
+		if (FSkeletalMeshAttributes::IsReservedAttributeName(AttributeName)) return;
 
 		if (TriAttribsSet.template HasAttributeOfType<int32>(AttributeName))
 		{
@@ -448,8 +466,6 @@ void FMeshDescriptionToDynamicMesh::Convert(const FMeshDescription* MeshIn, FDyn
 		}
 	});
 
-	
-
 	if (!bDisableAttributes)
 	{
 		bool bFoundNonDefaultVertexInstanceColor = false;
@@ -747,6 +763,19 @@ void FMeshDescriptionToDynamicMesh::Convert(const FMeshDescription* MeshIn, FDyn
 			Pending.Add(MoveTemp(MaterialFuture));
 		}
 
+		for (const FSkinWeightsAttribCopyInfo& SkinWeightAttribInfo: SkinWeightAttribs)
+		{
+			auto SkinWeightsFuture = Async(EAsyncExecution::ThreadPool, [&, SkinWeightAttribInfo]() -> void
+            {
+                for (int32 VertexID: MeshOut.VertexIndicesItr())
+                {
+                	SkinWeightAttribInfo.DynaMeshAttribRef->SetValue(VertexID, SkinWeightAttribInfo.MeshDescAttribRef.Get(VertIDMap[VertexID]));
+                }
+				
+            });
+			Pending.Add(MoveTemp(SkinWeightsFuture));
+		}
+		
 		// initialize polygroup layers
 		for (int32 GroupLayerIdx = 0; GroupLayerIdx < PolygroupAttribs.Num(); GroupLayerIdx++)
 		{
@@ -856,35 +885,3 @@ void FMeshDescriptionToDynamicMesh::CopyTangents(const FMeshDescription* SourceM
 	if (!ensureMsgf(TriIDMap.Num() == TargetMesh->TriangleCount(), TEXT("Tried to CopyTangents to mesh with different triangle count"))) return;
 	CopyTangents_Internal<double>(SourceMesh, TargetMesh, TangentsOut, TriIDMap);
 }
-
-
-
-
-bool FMeshDescriptionToDynamicMesh::IsReservedAttributeName(FName AttributeName)
-{
-	// @todo: use FAttributesSetBase::DoesAttributeHaveAnyFlags(EMeshAttributeFlags::Mandatory) to determine this
-	// will need a way to look up MeshDescription attribute sets from their category.
-	return (AttributeName == MeshAttribute::Vertex::Position)
-		|| (AttributeName == MeshAttribute::VertexInstance::VertexIndex)
-		|| (AttributeName == MeshAttribute::VertexInstance::TextureCoordinate)
-		|| (AttributeName == MeshAttribute::VertexInstance::Normal)
-		|| (AttributeName == MeshAttribute::VertexInstance::Tangent)
-		|| (AttributeName == MeshAttribute::VertexInstance::BinormalSign)
-		|| (AttributeName == MeshAttribute::VertexInstance::Color)
-		|| (AttributeName == MeshAttribute::Edge::IsHard)
-		|| (AttributeName == MeshAttribute::Edge::VertexIndex)
-		|| (AttributeName == MeshAttribute::Triangle::VertexInstanceIndex)
-		|| (AttributeName == MeshAttribute::Triangle::PolygonIndex)
-		|| (AttributeName == MeshAttribute::Triangle::EdgeIndex)
-		|| (AttributeName == MeshAttribute::Triangle::VertexIndex)
-		|| (AttributeName == MeshAttribute::Triangle::UVIndex)
-		|| (AttributeName == MeshAttribute::Triangle::PolygonGroupIndex)
-		|| (AttributeName == MeshAttribute::Triangle::Normal)
-		|| (AttributeName == MeshAttribute::Triangle::Tangent)
-		|| (AttributeName == MeshAttribute::Triangle::Binormal)
-		|| (AttributeName == MeshAttribute::UV::UVCoordinate)
-		|| (AttributeName == MeshAttribute::Polygon::PolygonGroupIndex)
-		|| (AttributeName == MeshAttribute::PolygonGroup::ImportedMaterialSlotName)
-		|| (AttributeName == ExtendedMeshAttribute::PolyTriGroups);
-}
-
diff --git a/Engine/Plugins/Experimental/GeometryProcessing/Source/MeshConversion/Public/MeshDescriptionToDynamicMesh.h b/Engine/Plugins/Experimental/GeometryProcessing/Source/MeshConversion/Public/MeshDescriptionToDynamicMesh.h
--- a/Engine/Plugins/Experimental/GeometryProcessing/Source/MeshConversion/Public/MeshDescriptionToDynamicMesh.h
+++ b/Engine/Plugins/Experimental/GeometryProcessing/Source/MeshConversion/Public/MeshDescriptionToDynamicMesh.h
@@ -79,12 +79,4 @@ public:
 	 * @warning Convert() must have been used to create the TargetMesh before calling this function
 	 */
 	void CopyTangents(const FMeshDescription* SourceMesh, const FDynamicMesh3* TargetMesh, UE::Geometry::TMeshTangents<double>* TangentsOut);
-
-
-public:
-
-	/**
-	 * @return true if the Name is the name of a reserved FMeshDescription attribute
-	 */
-	static bool IsReservedAttributeName(FName Name);
 };
\ No newline at end of file
diff --git a/Engine/Plugins/Experimental/MeshModelingToolset/Source/MeshModelingToolsEditorOnly/MeshModelingToolsEditorOnly.Build.cs b/Engine/Plugins/Experimental/MeshModelingToolset/Source/MeshModelingToolsEditorOnly/MeshModelingToolsEditorOnly.Build.cs
--- a/Engine/Plugins/Experimental/MeshModelingToolset/Source/MeshModelingToolsEditorOnly/MeshModelingToolsEditorOnly.Build.cs
+++ b/Engine/Plugins/Experimental/MeshModelingToolset/Source/MeshModelingToolsEditorOnly/MeshModelingToolsEditorOnly.Build.cs
@@ -29,6 +29,7 @@ public class MeshModelingToolsEditorOnly : ModuleRules
                 "InteractiveToolsFramework",
                 "MeshDescription",
 				"StaticMeshDescription",
+				"SkeletalMeshDescription",
 				"GeometricObjects",
 				"DynamicMesh",
 				"MeshConversion",
diff --git a/Engine/Plugins/Experimental/MeshModelingToolset/Source/MeshModelingToolsEditorOnly/Private/AttributeEditorTool.cpp b/Engine/Plugins/Experimental/MeshModelingToolset/Source/MeshModelingToolsEditorOnly/Private/AttributeEditorTool.cpp
--- a/Engine/Plugins/Experimental/MeshModelingToolset/Source/MeshModelingToolsEditorOnly/Private/AttributeEditorTool.cpp
+++ b/Engine/Plugins/Experimental/MeshModelingToolset/Source/MeshModelingToolsEditorOnly/Private/AttributeEditorTool.cpp
@@ -10,7 +10,7 @@
 #include "AssetUtils/MeshDescriptionUtil.h"
 
 #include "MeshDescription.h"
-#include "StaticMeshAttributes.h"
+#include "SkeletalMeshAttributes.h"
 #include "StaticMeshOperations.h"
 
 // for lightmap access
@@ -718,11 +718,11 @@ void UAttributeEditorTool::DeleteAttribute()
 	FMeshDescription* CurMesh = TargetMeshProviderInterface(0)->GetMeshDescription();
 	FName SelectedName(ModifyAttributeProps->Attribute);
 
-	// @todo: use FAttributesSetBase::DoesAttributeHaveAnyFlags(EMeshAttributeFlags::Mandatory) to determine this ?
-	// @todo: introduce 'reserved' attribute flag which we can check here
-	if (FMeshDescriptionToDynamicMesh::IsReservedAttributeName(SelectedName))
+	// We check on the skeletal mesh attributes because it is a superset of the static mesh
+	// attributes.
+	if (FSkeletalMeshAttributes::IsReservedAttributeName(SelectedName))
 	{
-		GetToolManager()->DisplayMessage(LOCTEXT("CannotDeleteReservedNameError", "Cannot delete required mesh Attributes"), EToolMessageLevel::UserWarning);
+		GetToolManager()->DisplayMessage(LOCTEXT("CannotDeleteReservedNameError", "Cannot delete reserved mesh Attributes"), EToolMessageLevel::UserWarning);
 		return;
 	}
 
diff --git a/Engine/Source/Runtime/SkeletalMeshDescription/Private/SkeletalMeshAttributes.cpp b/Engine/Source/Runtime/SkeletalMeshDescription/Private/SkeletalMeshAttributes.cpp
--- a/Engine/Source/Runtime/SkeletalMeshDescription/Private/SkeletalMeshAttributes.cpp
+++ b/Engine/Source/Runtime/SkeletalMeshDescription/Private/SkeletalMeshAttributes.cpp
@@ -55,7 +55,6 @@ TArray<FName> FSkeletalMeshAttributesShared::GetSkinWeightProfileNames() const
 	MeshDescriptionShared.VertexAttributes().GetAttributeNames(AllAttributeNames);
 
 	TArray<FName> SkinWeightProfileNames;
-	const FString Prefix = SkinWeightAttributeNamePrefix();
 	bool bHasDefault = false; 
 
 	for (const FName AttributeName: AllAttributeNames)
@@ -64,11 +63,9 @@ TArray<FName> FSkeletalMeshAttributesShared::GetSkinWeightProfileNames() const
 		{
 			bHasDefault = true;
 		}
-		else if (AttributeName.ToString().StartsWith(Prefix))
+		else if (IsSkinWeightAttribute(AttributeName))
 		{
-			FName ProfileName = FName(AttributeName.ToString().Mid(Prefix.Len()));
-
-			SkinWeightProfileNames.Add(ProfileName);
+			SkinWeightProfileNames.Add(GetProfileNameFromAttribute(AttributeName));
 		}
 	}
 
@@ -92,6 +89,24 @@ bool FSkeletalMeshAttributesShared::IsSkinWeightAttribute(const FName InAttribut
 		   InAttributeName.ToString().StartsWith(SkinWeightAttributeNamePrefix());
 }
 
+FName FSkeletalMeshAttributesShared::GetProfileNameFromAttribute(const FName InAttributeName)
+{
+	if (InAttributeName == MeshAttribute::Vertex::SkinWeights)
+	{
+		return DefaultSkinWeightProfileName;
+	}
+	
+	const FString Prefix = SkinWeightAttributeNamePrefix();
+	if (InAttributeName.ToString().StartsWith(Prefix))
+	{
+		return FName(InAttributeName.ToString().Mid(Prefix.Len()));
+	}
+	else
+	{
+		return NAME_None;
+	}
+}
+
 FName FSkeletalMeshAttributesShared::CreateSkinWeightAttributeName(const FName InProfileName)
 {
 	// If it's the default profile, then return the base skin weights attribute name.
diff --git a/Engine/Source/Runtime/SkeletalMeshDescription/Public/SkeletalMeshAttributes.h b/Engine/Source/Runtime/SkeletalMeshDescription/Public/SkeletalMeshAttributes.h
--- a/Engine/Source/Runtime/SkeletalMeshDescription/Public/SkeletalMeshAttributes.h
+++ b/Engine/Source/Runtime/SkeletalMeshDescription/Public/SkeletalMeshAttributes.h
@@ -37,6 +37,10 @@ public:
 	/// Helper function that indicates whether an attribute name represents a skin weight attribute.
 	static bool IsSkinWeightAttribute(const FName InAttributeName);
 
+	/// Returns a skin profile name from the attribute name, if the attribute name is a valid skin weights
+	/// attribute.
+	static FName GetProfileNameFromAttribute(const FName InAttributeName);
+
 	FSkinWeightsVertexAttributesConstRef GetVertexSkinWeights(const FName InProfileName = NAME_None) const
 	{
 		return MeshDescriptionShared.VertexAttributes().GetAttributesRef<TArrayAttribute<int32>>(CreateSkinWeightAttributeName(InProfileName));
