Partially parallelize two heavy functions in FNormalFlowRemesher.

#rnx

diff --git a/Engine/Plugins/Experimental/GeometryProcessing/Source/DynamicMesh/Private/NormalFlowRemesher.cpp b/Engine/Plugins/Experimental/GeometryProcessing/Source/DynamicMesh/Private/NormalFlowRemesher.cpp
--- a/Engine/Plugins/Experimental/GeometryProcessing/Source/DynamicMesh/Private/NormalFlowRemesher.cpp
+++ b/Engine/Plugins/Experimental/GeometryProcessing/Source/DynamicMesh/Private/NormalFlowRemesher.cpp
@@ -62,7 +62,15 @@ void FNormalFlowRemesher::RemeshWithFaceProjection()
 		}
 
 		double MaxProjectionDistance = 0.0;
-		TrackedFaceProjectionPass(MaxProjectionDistance);
+		if (bEnableParallelProjection)
+		{
+			TrackedFaceProjectionPass(MaxProjectionDistance);
+		}
+		else
+		{
+			TrackedFaceProjectionPass_Serial(MaxProjectionDistance);
+		}
+		
 
 		// Stop if we've hit max iterations, or both:
 		// - queue is empty and
@@ -83,7 +91,14 @@ void FNormalFlowRemesher::RemeshWithFaceProjection()
 			}
 
 			double MaxProjectionDistance = 0.0;
-			TrackedFaceProjectionPass(MaxProjectionDistance);
+			if (bEnableParallelProjection)
+			{
+				TrackedFaceProjectionPass(MaxProjectionDistance);
+			}
+			else
+			{
+				TrackedFaceProjectionPass_Serial(MaxProjectionDistance);
+			}
 
 			if (MaxProjectionDistance == 0.0)
 			{
@@ -107,6 +122,157 @@ void FNormalFlowRemesher::TrackedFaceProjectionPass(double& MaxDistanceMoved)
 
 	InitializeVertexBufferForFacePass();
 
+	// For each triangle, rotate it such that it aligns with closest normal on the target surface.
+	// Each vertex is then assigned a weighted combination of its corresponding triangle corner positions. The weighting
+	// is chosen to favor triangles that don't rotate much (i.e. are already aligned with the target surface.)
+
+	// First compute all rotated triangles and store their corner positions and weights
+	ParallelFor(Mesh->MaxTriangleID(), [this, NormalProjTarget](int32 TriangleIndex)
+	{
+		FVector3d TriangleNormal, Centroid;
+		double Area;
+		Mesh->GetTriInfo(TriangleIndex, TriangleNormal, Area, Centroid);
+
+		FVector3d ProjectedNormal{ 1e30, 1e30, 1e30 };
+		FVector3d ProjectedPosition = NormalProjTarget->Project(Centroid, ProjectedNormal);
+
+		if (!ensure(ProjectedNormal[0] != 1e30))
+		{
+			return;
+		}
+		if (!ensure(ProjectedNormal.Length() > 1e-6))
+		{
+			return;
+		}
+
+		FVector3d V0, V1, V2;
+		Mesh->GetTriVertices(TriangleIndex, V0, V1, V2);
+
+		FFrame3d TriF(Centroid, TriangleNormal);
+		V0 = TriF.ToFramePoint(V0);
+		V1 = TriF.ToFramePoint(V1);
+		V2 = TriF.ToFramePoint(V2);
+
+		TriF.AlignAxis(2, ProjectedNormal);
+		TriF.Origin = ProjectedPosition;
+		V0 = TriF.FromFramePoint(V0);
+		V1 = TriF.FromFramePoint(V1);
+		V2 = TriF.FromFramePoint(V2);
+
+		double Dot = TriangleNormal.Dot(ProjectedNormal);
+		Dot = FMath::Clamp(Dot, 0.0, 1.0);
+		const double Weight = Area * (Dot * Dot * Dot);
+
+		ProjectionVertexBuffer[3 * TriangleIndex] = Weight * V0;
+		ProjectionWeightBuffer[3 * TriangleIndex] = Weight;
+		ProjectionVertexBuffer[3 * TriangleIndex + 1] = Weight * V1;
+		ProjectionWeightBuffer[3 * TriangleIndex + 1] = Weight;
+		ProjectionVertexBuffer[3 * TriangleIndex + 2] = Weight * V2;
+		ProjectionWeightBuffer[3 * TriangleIndex + 2] = Weight;
+	});
+
+	// Next compute all the weighted average of triangle corners for each vertex
+	ParallelFor(Mesh->MaxVertexID(), [this, &MaxDistanceMoved](int32 VertexIndex)
+	{
+		TempFlagBuffer[VertexIndex] = false;
+
+		if (!Mesh->IsVertex(VertexIndex))
+		{
+			return;
+		}
+
+		TempPosBuffer[VertexIndex] = Mesh->GetVertex(VertexIndex);		// Don't modify the position of constrained vertex
+
+		if (IsVertexPositionConstrained(VertexIndex))
+		{
+			return;
+		}
+
+		if (VertexControlF != nullptr && ((int)VertexControlF(VertexIndex) & (int)EVertexControl::NoProject) != 0)
+		{
+			return;
+		}
+
+		TempWeightBuffer[VertexIndex] = 0.0;
+		TempPosBuffer[VertexIndex] = { 0.0, 0.0, 0.0 };
+
+		TArray<int> IncidentTriangles;
+		Mesh->GetVertexOneRingTriangles(VertexIndex, IncidentTriangles);
+
+		for (int TriID : IncidentTriangles)
+		{
+			const FIndex3i Tri = Mesh->GetTriangle(TriID);
+
+			// Find VertexIndex in Tri:
+			for (int IndexInTri = 0; IndexInTri < 3; ++IndexInTri)
+			{
+				if (Tri[IndexInTri] == VertexIndex)
+				{
+					int BufferIndex = 3 * TriID + IndexInTri;
+					TempWeightBuffer[VertexIndex] += ProjectionWeightBuffer[BufferIndex];
+					TempPosBuffer[VertexIndex] += ProjectionVertexBuffer[BufferIndex];
+				}
+			}
+		}
+
+		if (FMath::IsNearlyZero(TempWeightBuffer[VertexIndex]))
+		{
+			return;
+		}
+
+		FVector3d CurrentPosition = Mesh->GetVertex(VertexIndex);
+		FVector3d ProjectedPosition = TempPosBuffer[VertexIndex] / TempWeightBuffer[VertexIndex];
+
+		if (VectorUtil::EpsilonEqual(CurrentPosition, ProjectedPosition, FMathd::ZeroTolerance))
+		{
+			return;
+		}
+
+		TempFlagBuffer[VertexIndex] = true;
+		TempPosBuffer[VertexIndex] = ProjectedPosition;
+	});
+
+
+	MaxDistanceMoved = 0.0;
+
+	// We queue any edges that moved far enough to fall under min/max edge length thresholds. Also find the max distance
+	// moved by a vertex.
+	for (int EdgeID : Mesh->EdgeIndicesItr())
+	{
+		FIndex2i EdgeVertices = Mesh->GetEdgeV(EdgeID);
+		double NewEdgeLength = Distance(TempPosBuffer[EdgeVertices[0]], TempPosBuffer[EdgeVertices[1]]);
+		if (NewEdgeLength < MinEdgeLength || NewEdgeLength > MaxEdgeLength)
+		{
+			QueueEdge(EdgeID);
+		}
+
+		for (int I = 0; I < 2; ++I)
+		{
+			if (TempFlagBuffer[EdgeVertices[I]])
+			{
+				const FVector3d& CurrentPosition = Mesh->GetVertex(EdgeVertices[I]);
+				const FVector3d& ProjectedPosition = TempPosBuffer[EdgeVertices[I]];
+				MaxDistanceMoved = FMath::Max(MaxDistanceMoved, Distance(CurrentPosition, ProjectedPosition));
+			}
+		}
+	}
+
+	// Update vertices
+	constexpr bool bUpdateParallel = true;
+	ApplyVertexBuffer(bUpdateParallel);
+}
+
+
+
+void FNormalFlowRemesher::TrackedFaceProjectionPass_Serial(double& MaxDistanceMoved)
+{
+	ensure(ProjTarget != nullptr);
+
+	IOrientedProjectionTarget* NormalProjTarget = static_cast<IOrientedProjectionTarget*>(ProjTarget);
+	ensure(NormalProjTarget != nullptr);
+
+	InitializeVertexBufferForFacePass();
+
 	// this function computes rotated position of triangle, such that it
 	// aligns with face normal on target surface. We accumulate weighted-average
 	// of vertex positions, which we will then use further down where possible.
@@ -203,7 +369,8 @@ void FNormalFlowRemesher::TrackedFaceProjectionPass(double& MaxDistanceMoved)
 	}
 
 	// update vertices
-	ApplyVertexBuffer(true);
+	constexpr bool bUpdateParallel = true;
+	ApplyVertexBuffer(bUpdateParallel);
 }
 
 
@@ -285,20 +452,54 @@ void FNormalFlowRemesher::TrackedEdgeFlipPass()
 	IOrientedProjectionTarget* NormalProjTarget = static_cast<IOrientedProjectionTarget*>(ProjTarget);
 	check(NormalProjTarget != nullptr);
 
-	for (auto EdgeID : Mesh->EdgeIndicesItr())
+	if (bEnableParallelEdgeFlipPass)
 	{
-		check(Mesh->IsEdge(EdgeID));
+		if (EdgeShouldBeQueuedBuffer.Num() < Mesh->MaxEdgeID())
+		{
+			EdgeShouldBeQueuedBuffer.SetNum(2 * Mesh->MaxEdgeID());
+		}
+		EdgeShouldBeQueuedBuffer.Init(false, EdgeShouldBeQueuedBuffer.Num());
 
-		FEdgeConstraint Constraint =
-			(!Constraints) ? FEdgeConstraint::Unconstrained() : Constraints->GetEdgeConstraint(EdgeID);
+		int NumEdges = Mesh->MaxEdgeID();
 
-		if (!Constraint.CanFlip())
+		ParallelFor(NumEdges, [this](int EdgeID)
 		{
-			continue;
-		}
+			if (!Mesh->IsEdge(EdgeID))
+			{
+				return;
+			}
+
+			FEdgeConstraint Constraint =
+				(!Constraints) ? FEdgeConstraint::Unconstrained() : Constraints->GetEdgeConstraint(EdgeID);
+
+			if (!Constraint.CanFlip())
+			{
+				return;
+			}
 
-		if (EdgeFlipWouldReduceNormalError(EdgeID))
+			if (EdgeFlipWouldReduceNormalError(EdgeID))
+			{
+				EdgeShouldBeQueuedBuffer[EdgeID] = true;
+			}
+		});
+
+		for (int32 EdgeID = 0; EdgeID < NumEdges; ++EdgeID)
 		{
+			if (!EdgeShouldBeQueuedBuffer[EdgeID])
+			{
+				continue;
+			}
+
+			if (!Mesh->IsEdge(EdgeID))
+			{
+				continue;
+			}
+
+			if (!EdgeFlipWouldReduceNormalError(EdgeID))
+			{
+				continue;
+			}
+
 			DynamicMeshInfo::FEdgeFlipInfo FlipInfo;
 			auto Result = Mesh->FlipEdge(EdgeID, FlipInfo);
 
@@ -315,5 +516,39 @@ void FNormalFlowRemesher::TrackedEdgeFlipPass()
 			}
 		}
 	}
+	else
+	{
+		for (auto EdgeID : Mesh->EdgeIndicesItr())
+		{
+			check(Mesh->IsEdge(EdgeID));
+
+			FEdgeConstraint Constraint =
+				(!Constraints) ? FEdgeConstraint::Unconstrained() : Constraints->GetEdgeConstraint(EdgeID);
+
+			if (!Constraint.CanFlip())
+			{
+				continue;
+			}
+
+			if (EdgeFlipWouldReduceNormalError(EdgeID))
+			{
+				DynamicMeshInfo::FEdgeFlipInfo FlipInfo;
+				auto Result = Mesh->FlipEdge(EdgeID, FlipInfo);
+
+				if (Result == EMeshResult::Ok)
+				{
+					FIndex2i EdgeVertices = Mesh->GetEdgeV(EdgeID);
+					FIndex2i OpposingEdgeVertices = Mesh->GetEdgeOpposingV(EdgeID);
+
+					QueueOneRing(EdgeVertices.A);
+					QueueOneRing(EdgeVertices.B);
+					QueueOneRing(OpposingEdgeVertices.A);
+					QueueOneRing(OpposingEdgeVertices.B);
+					OnEdgeFlip(EdgeID, FlipInfo);
+				}
+			}
+		}
+	}
 
 }
+
diff --git a/Engine/Plugins/Experimental/GeometryProcessing/Source/DynamicMesh/Public/NormalFlowRemesher.h b/Engine/Plugins/Experimental/GeometryProcessing/Source/DynamicMesh/Public/NormalFlowRemesher.h
--- a/Engine/Plugins/Experimental/GeometryProcessing/Source/DynamicMesh/Public/NormalFlowRemesher.h
+++ b/Engine/Plugins/Experimental/GeometryProcessing/Source/DynamicMesh/Public/NormalFlowRemesher.h
@@ -33,6 +33,9 @@ public:
 		RemeshWithFaceProjection();
 	}
 
+	/// Gather candidates for edge flipping in parallel, then do the actual flips sequentially
+	bool bEnableParallelEdgeFlipPass = true;
+
 protected:
 
 	// Do remeshing with face-aligned projection. During the Projection step of RemeshIteration, perform N passes of 
@@ -42,6 +45,7 @@ protected:
 
 	// Perform face-aligned projection onto the target mesh. Queue edges whose lengths change because of it.
 	void TrackedFaceProjectionPass(double& MaxDistanceMoved);
+	void TrackedFaceProjectionPass_Serial(double& MaxDistanceMoved);
 
 	// This is called during RemeshIteration 
 	void TrackedFullProjectionPass(bool bParallel) override
@@ -49,14 +53,27 @@ protected:
 		for (int i = 0; i < FaceProjectionPassesPerRemeshIteration; ++i)
 		{
 			double ProjectionDistance;
-			TrackedFaceProjectionPass(ProjectionDistance);
+			if (bParallel)
+			{
+				TrackedFaceProjectionPass(ProjectionDistance);
+			}
+			else
+			{
+				TrackedFaceProjectionPass_Serial(ProjectionDistance);
+			}
 		}
 	}
 
 
-	// Per-vertex weights for face projection pass
 	TArray<double> TempWeightBuffer;
 
+	// Per triangle-vertex weights for face projection pass. (Size >= 3 * MaxTriangleID)
+	TArray<double> ProjectionWeightBuffer;
+
+	// Per triangle-vertex projection positions for face projection pass. (Size >= 3 * MaxTriangleID)
+	TArray<FVector3d> ProjectionVertexBuffer;
+
+
 	// Similar to InitializeVertexBufferForPass, but also initialize the additional per-vertex weight buffer
 	void InitializeVertexBufferForFacePass()
 	{
@@ -67,6 +84,21 @@ protected:
 			TempWeightBuffer.SetNum(2 * Mesh->MaxVertexID());
 		}
 		TempWeightBuffer.Init(0.0, TempWeightBuffer.Num());
+
+		int NumTriangleVertices = 3 * Mesh->MaxTriangleID();
+
+		if (ProjectionWeightBuffer.Num() < NumTriangleVertices)
+		{
+			ProjectionWeightBuffer.SetNum(2 * NumTriangleVertices);
+		}
+		ProjectionWeightBuffer.Init(0.0, ProjectionWeightBuffer.Num());
+
+		if (ProjectionVertexBuffer.Num() < NumTriangleVertices)
+		{
+			ProjectionVertexBuffer.SetNum(2 * NumTriangleVertices);
+		}
+		ProjectionVertexBuffer.Init(FVector3d{ 0,0,0 }, ProjectionWeightBuffer.Num());
+
 	}
 
 
