MeshModelingTools: Undo the 1-V conversion when displaying things in the UV Editor.

#rnx

diff --git a/Engine/Plugins/Experimental/GeometryProcessing/Source/MeshConversion/Private/MeshDescriptionUVsToDynamicMesh.cpp b/Engine/Plugins/Experimental/GeometryProcessing/Source/MeshConversion/Private/MeshDescriptionUVsToDynamicMesh.cpp
--- a/Engine/Plugins/Experimental/GeometryProcessing/Source/MeshConversion/Private/MeshDescriptionUVsToDynamicMesh.cpp
+++ b/Engine/Plugins/Experimental/GeometryProcessing/Source/MeshConversion/Private/MeshDescriptionUVsToDynamicMesh.cpp
@@ -22,6 +22,25 @@ namespace MeshDescriptionUVsToDynamicMeshLocals
 		return NumSharedUVLayers > LayerIndex && MeshDescription->UVs(LayerIndex).GetArraySize() != 0;
 	}
 
+	// These determine the mapping between UV values and the resulting mesh vertex positions. If we're
+	// looking down on the unwrapped mesh, with the Z axis towards us, we want U's to be right, and
+	// V's to be up. In Unreal's left-handed coordinate system, this means that we map U's to world Y
+	// and V's to world X.
+	// Also, Unreal changes the V coordinates of imported meshes to 1-V internally, and we undo this
+	// while displaying the UV's because the users likely expect to see the original UV's (it would
+	// be particularly confusing for users working with UDIM assets, where internally stored V's 
+	// frequently end up negative).
+	// The ScaleFactor just scales the mesh up. Scaling the mesh up makes it easier to zoom in
+	// further into the display before getting issues with the camera near plane distance.
+	FORCEINLINE FVector3d UVToVertPosition(const FVector2D& UV, double ScaleFactor)
+	{
+		return FVector3d((1 - UV.Y) * ScaleFactor, UV.X * ScaleFactor, 0);
+	}
+	FORCEINLINE FVector2D VertPositionToUV(const FVector3d& VertPosition, double ScaleFactor)
+	{
+		return FVector2D(VertPosition.Y / ScaleFactor, 1 - (VertPosition.X / ScaleFactor));
+	}
+
 	// This class is copied from MeshDescriptionToDynamicMesh.cpp. We could consider putting it
 	// somewhere common.
 	struct FVertexUV
@@ -64,7 +83,7 @@ namespace MeshDescriptionUVsToDynamicMeshLocals
 				return *FoundIndex;
 			}
 
-			int32 NewIndex = Mesh->AppendVertex(FVector3d(UV.Y * ScaleFactor, UV.X * ScaleFactor, 0));
+			int32 NewIndex = Mesh->AppendVertex(UVToVertPosition(UV, ScaleFactor));
 			UniqueVertexUVs.Add(VertUV, NewIndex);
 			return NewIndex;
 		}
@@ -106,7 +125,7 @@ TSharedPtr<FDynamicMesh3> FMeshDescriptionUVsToDynamicMesh::GetUVMesh(const FMes
 		for (FUVID UVID : SharedUVs.GetElementIDs())
 		{
 			const FVector2D UV = SharedUVCoordinates[UVID];
-			MeshOut->InsertVertex(UVID.GetValue(), FVector3d(UV.Y * ScaleFactor, UV.X * ScaleFactor, 0), true);
+			MeshOut->InsertVertex(UVID.GetValue(), UVToVertPosition(UV, ScaleFactor), true);
 		}
 		MeshOut->EndUnsafeTrianglesInsert();
 	}
@@ -255,7 +274,7 @@ void FMeshDescriptionUVsToDynamicMesh::BakeBackUVsFromUVMesh(const FDynamicMesh3
 		for (int i = 0; i < 3; ++i)
 		{
 			FVector3d VertPosition = DynamicMesh->GetVertex(DynamicMeshTriVids[i]);
-			InstanceUVs.Set(MeshDescriptionTriInstanceIds[i], UVLayerIndex, FVector2D(VertPosition.Y / ScaleFactor, VertPosition.X / ScaleFactor));
+			InstanceUVs.Set(MeshDescriptionTriInstanceIds[i], UVLayerIndex, VertPositionToUV(VertPosition, ScaleFactor));
 		}
 	}
 
@@ -310,7 +329,7 @@ void FMeshDescriptionUVsToDynamicMesh::BakeBackUVsFromUVMesh(const FDynamicMesh3
 				// If we haven't dealt with this vert, we will need to update its entry in SharedUVCoordinates.
 
 				FVector3d VertPosition = DynamicMesh->GetVertex(Vid);
-				FVector2D NewUVValue = FVector2D(VertPosition.Y / ScaleFactor, VertPosition.X / ScaleFactor);
+				FVector2D NewUVValue = VertPositionToUV(VertPosition, ScaleFactor);
 				
 				// See if the vert is the same one that we were expecting (including one that we may have split off just for
 				// the conversion, in which case the location will match).
