Modeling Tools: MeshConversions.  When converting to SkeletalMeshLODModel, initialize the triangles in the color overlay.
#rnx

diff --git a/Engine/Plugins/Experimental/GeometryProcessing/Source/MeshConversion/Private/ToDynamicMesh.h b/Engine/Plugins/Experimental/GeometryProcessing/Source/MeshConversion/Private/ToDynamicMesh.h
--- a/Engine/Plugins/Experimental/GeometryProcessing/Source/MeshConversion/Private/ToDynamicMesh.h
+++ b/Engine/Plugins/Experimental/GeometryProcessing/Source/MeshConversion/Private/ToDynamicMesh.h
@@ -394,6 +394,7 @@ protected:
 		// create color overlay
 		MeshOut.Attributes()->EnablePrimaryColors();
 		ColorOverlay = MeshOut.Attributes()->PrimaryColors();
+		ColorOverlay->InitializeTriangles(MeshOut.MaxTriangleID());
 
 		// always enable Material ID if there are any attributes
 		MeshOut.Attributes()->EnableMaterialID();
