ProjectToTarget tool: enable world space projection option. Currently projection happens in local space for both the "mesh" and the "projection target". The new WorldSpace option transforms the mesh vertices into world coordinates, then into the coordinates of the projection target before projecting. This is similar to transforming both the mesh and target (an associated acceleration structure) into world coordinates and then doing the projection query.

diff --git a/Engine/Plugins/Experimental/GeometryProcessing/Source/DynamicMesh/Public/ProjectionTargets.h b/Engine/Plugins/Experimental/GeometryProcessing/Source/DynamicMesh/Public/ProjectionTargets.h
--- a/Engine/Plugins/Experimental/GeometryProcessing/Source/DynamicMesh/Public/ProjectionTargets.h
+++ b/Engine/Plugins/Experimental/GeometryProcessing/Source/DynamicMesh/Public/ProjectionTargets.h
@@ -25,11 +25,8 @@ public:
 	/** An AABBTree for Mesh */
 	FDynamicMeshAABBTree3* Spatial = nullptr;
 
-	~FMeshProjectionTarget() {}
-
-	FMeshProjectionTarget()
-	{
-	}
+	virtual ~FMeshProjectionTarget() = default;
+	FMeshProjectionTarget() = default;
 
 	FMeshProjectionTarget(const FDynamicMesh3* MeshIn, FDynamicMeshAABBTree3* SpatialIn)
 	{
@@ -52,5 +49,70 @@ public:
 };
 
 
+class DYNAMICMESH_API FWorldSpaceProjectionTarget : public FMeshProjectionTarget
+{
+public:
+
+	virtual ~FWorldSpaceProjectionTarget() = default;
+	FWorldSpaceProjectionTarget() = default;
+
+	FWorldSpaceProjectionTarget(const FDynamicMesh3* MeshIn, 
+								FDynamicMeshAABBTree3* SpatialIn,
+								const FTransform3d& TargetMeshLocalToWorldIn,
+								const FTransform3d& ToolMeshLocalToWorldIn) :
+		FMeshProjectionTarget(MeshIn, SpatialIn),
+		TargetMeshLocalToWorld(TargetMeshLocalToWorldIn),
+		ToolMeshLocalToWorld(ToolMeshLocalToWorldIn)
+	{}
+
+	/**
+	 * @return Projection of Point onto this target
+	*/
+	virtual FVector3d Project(const FVector3d& Point, int Identifier = -1) override
+	{
+		// Put query point in Target space
+		FVector3d TransformedPoint = Point;
+		TransformedPoint = ToolMeshLocalToWorld.TransformPosition(TransformedPoint);
+		TransformedPoint = TargetMeshLocalToWorld.InverseTransformPosition(TransformedPoint);
+
+		FVector3d ProjectedPosition = FMeshProjectionTarget::Project(TransformedPoint, Identifier);
+
+		// ProjectedPosition is in Target space, put it in Tool space
+		ProjectedPosition = TargetMeshLocalToWorld.TransformPosition(ProjectedPosition);
+		ProjectedPosition = ToolMeshLocalToWorld.InverseTransformPosition(ProjectedPosition);
+
+		return ProjectedPosition;
+	}
+
+	/**
+	 * @return Projection of Point onto this target, and set ProjectNormalOut to the triangle normal at the returned point (*not* interpolated vertex normal)
+	 */
+	virtual FVector3d Project(const FVector3d& Point, FVector3d& ProjectNormalOut, int Identifier = -1) override
+	{
+		// Put query point in Target space
+		FVector3d TransformedPoint = Point;
+		TransformedPoint = ToolMeshLocalToWorld.TransformPosition(TransformedPoint);
+		TransformedPoint = TargetMeshLocalToWorld.InverseTransformPosition(TransformedPoint);
+
+		FVector3d ProjectedPosition = FMeshProjectionTarget::Project(TransformedPoint, ProjectNormalOut, Identifier);
+
+		// ProjectedPosition and ProjectedNormal are in Target space, put them in Tool space
+		ProjectedPosition = TargetMeshLocalToWorld.TransformPosition(ProjectedPosition);
+		ProjectedPosition = ToolMeshLocalToWorld.InverseTransformPosition(ProjectedPosition);
+
+		ProjectNormalOut = TargetMeshLocalToWorld.TransformNormal(ProjectNormalOut);
+		ProjectNormalOut = ToolMeshLocalToWorld.InverseTransformNormal(ProjectNormalOut);
+
+		return ProjectedPosition;
+	}
+
+protected:
+	FTransform3d TargetMeshLocalToWorld;
+	FTransform3d ToolMeshLocalToWorld;
+
+};
+
+
+
 } // end namespace UE::Geometry
-} // end namespace UE
\ No newline at end of file
+} // end namespace UE
diff --git a/Engine/Plugins/Experimental/MeshModelingToolset/Source/MeshModelingTools/Private/ProjectToTargetTool.cpp b/Engine/Plugins/Experimental/MeshModelingToolset/Source/MeshModelingTools/Private/ProjectToTargetTool.cpp
--- a/Engine/Plugins/Experimental/MeshModelingToolset/Source/MeshModelingTools/Private/ProjectToTargetTool.cpp
+++ b/Engine/Plugins/Experimental/MeshModelingToolset/Source/MeshModelingTools/Private/ProjectToTargetTool.cpp
@@ -64,6 +64,12 @@ void UProjectToTargetTool::Setup()
 
 TUniquePtr<FDynamicMeshOperator> UProjectToTargetTool::MakeNewOperator()
 {
+	UProjectToTargetToolProperties* ProjectProperties = Cast<UProjectToTargetToolProperties>(BasicProperties);
+	if (!ensure(ProjectProperties))
+	{
+		return nullptr;
+	}
+
 	TUniquePtr<FDynamicMeshOperator> Op = URemeshMeshTool::MakeNewOperator();
 
 	FDynamicMeshOperator* RawOp = Op.Get();
@@ -73,6 +79,10 @@ TUniquePtr<FDynamicMeshOperator> UProjectToTargetTool::MakeNewOperator()
 	RemeshOp->ProjectionTarget = ProjectionTarget.Get();
 	RemeshOp->ProjectionTargetSpatial = ProjectionTargetSpatial.Get();
 
+	RemeshOp->ToolMeshLocalToWorld = FTransform3d(Cast<IPrimitiveComponentBackedTarget>(Targets[0])->GetWorldTransform());
+	RemeshOp->TargetMeshLocalToWorld = FTransform3d(Cast<IPrimitiveComponentBackedTarget>(Targets[1])->GetWorldTransform());
+	RemeshOp->bUseWorldSpace = ProjectProperties->bWorldSpace;
+
 	return Op;
 }
 
diff --git a/Engine/Plugins/Experimental/MeshModelingToolset/Source/MeshModelingTools/Public/ProjectToTargetTool.h b/Engine/Plugins/Experimental/MeshModelingToolset/Source/MeshModelingTools/Public/ProjectToTargetTool.h
--- a/Engine/Plugins/Experimental/MeshModelingToolset/Source/MeshModelingTools/Public/ProjectToTargetTool.h
+++ b/Engine/Plugins/Experimental/MeshModelingToolset/Source/MeshModelingTools/Public/ProjectToTargetTool.h
@@ -37,8 +37,12 @@ class MESHMODELINGTOOLS_API UProjectToTargetToolProperties : public URemeshMeshT
 
 public:
 
+	UPROPERTY(EditAnywhere, Category = ProjectionSpace)
+	bool bWorldSpace = true;
+
 	UProjectToTargetToolProperties() :
-		URemeshMeshToolProperties()
+		URemeshMeshToolProperties(),
+		bWorldSpace(true)
 	{
 		bPreserveSharpEdges = false;
 		RemeshType = ERemeshType::NormalFlow;
diff --git a/Engine/Plugins/Experimental/MeshModelingToolset/Source/ModelingOperators/Private/CleaningOps/RemeshMeshOp.cpp b/Engine/Plugins/Experimental/MeshModelingToolset/Source/ModelingOperators/Private/CleaningOps/RemeshMeshOp.cpp
--- a/Engine/Plugins/Experimental/MeshModelingToolset/Source/ModelingOperators/Private/CleaningOps/RemeshMeshOp.cpp
+++ b/Engine/Plugins/Experimental/MeshModelingToolset/Source/ModelingOperators/Private/CleaningOps/RemeshMeshOp.cpp
@@ -124,8 +124,17 @@ void FRemeshMeshOp::CalculateResult(FProgressCancel* Progress)
 		ProjectionTargetSpatial = OriginalMeshSpatial.Get();
 	}
 
-	FMeshProjectionTarget ProjTarget(ProjectionTarget, ProjectionTargetSpatial);
-	Remesher->SetProjectionTarget(&ProjTarget);
+	TUniquePtr<IProjectionTarget> ProjTarget;
+	if (bUseWorldSpace)
+	{
+		ProjTarget = MakeUnique<FWorldSpaceProjectionTarget>(ProjectionTarget, ProjectionTargetSpatial, TargetMeshLocalToWorld, ToolMeshLocalToWorld);
+	}
+	else
+	{
+		ProjTarget = MakeUnique<FMeshProjectionTarget>(ProjectionTarget, ProjectionTargetSpatial);
+	}
+	Remesher->SetProjectionTarget(ProjTarget.Get());
+
 
 	Remesher->Progress = Progress;
 
diff --git a/Engine/Plugins/Experimental/MeshModelingToolset/Source/ModelingOperators/Public/CleaningOps/RemeshMeshOp.h b/Engine/Plugins/Experimental/MeshModelingToolset/Source/ModelingOperators/Public/CleaningOps/RemeshMeshOp.h
--- a/Engine/Plugins/Experimental/MeshModelingToolset/Source/ModelingOperators/Public/CleaningOps/RemeshMeshOp.h
+++ b/Engine/Plugins/Experimental/MeshModelingToolset/Source/ModelingOperators/Public/CleaningOps/RemeshMeshOp.h
@@ -69,6 +69,10 @@ public:
 	FDynamicMesh3* ProjectionTarget = nullptr;
 	FDynamicMeshAABBTree3* ProjectionTargetSpatial = nullptr;
 
+	FTransform3d TargetMeshLocalToWorld;
+	FTransform3d ToolMeshLocalToWorld;
+	bool bUseWorldSpace = false;
+
 	void SetTransform(const FTransform& Transform);
 
 	//
