GeometryProcessing: clean up check() that should be checkSlow() or ensure, to avoid hard failures when used at runtime
#rnx

diff --git a/Engine/Plugins/Experimental/GeometryProcessing/Source/DynamicMesh/Private/DynamicMesh3.cpp b/Engine/Plugins/Experimental/GeometryProcessing/Source/DynamicMesh/Private/DynamicMesh3.cpp
--- a/Engine/Plugins/Experimental/GeometryProcessing/Source/DynamicMesh/Private/DynamicMesh3.cpp
+++ b/Engine/Plugins/Experimental/GeometryProcessing/Source/DynamicMesh/Private/DynamicMesh3.cpp
@@ -144,7 +144,7 @@ FDynamicMesh3::FDynamicMesh3(const FMeshShapeGenerator* Generator)
 	Copy(Generator);
 }
 
-void FDynamicMesh3::Copy(const FMeshShapeGenerator* Generator)
+bool FDynamicMesh3::Copy(const FMeshShapeGenerator* Generator)
 {
 	Clear();
 
@@ -161,6 +161,7 @@ void FDynamicMesh3::Copy(const FMeshShapeGenerator* Generator)
 		AppendVertex(Generator->Vertices[i]);
 	}
 
+	int NumTris = Generator->Triangles.Num();
 	if (Generator->HasAttributes())
 	{
 		FDynamicMeshUVOverlay* UVOverlay = Attributes()->PrimaryUV();
@@ -176,36 +177,35 @@ void FDynamicMesh3::Copy(const FMeshShapeGenerator* Generator)
 			NormalOverlay->AppendElement(Generator->Normals[i]);
 		}
 
-		int NumTris = Generator->Triangles.Num();
 		for (int i = 0; i < NumTris; ++i)
 		{
 			int PolyID = Generator->TrianglePolygonIDs.Num() > 0 ? 1 + Generator->TrianglePolygonIDs[i] : 0;
 			int tid = AppendTriangle(Generator->Triangles[i], PolyID);
-			check(tid == i);
-			UVOverlay->SetTriangle(tid, Generator->TriangleUVs[i]);
-			NormalOverlay->SetTriangle(tid, Generator->TriangleNormals[i]);
+			if (ensure(tid == i))
+			{
+				UVOverlay->SetTriangle(tid, Generator->TriangleUVs[i]);
+				NormalOverlay->SetTriangle(tid, Generator->TriangleNormals[i]);
+			}
 		}
 	}
 	else if (Generator->TrianglePolygonIDs.Num()) // no attributes, yes polygon ids
 	{
-		int NumTris = Generator->Triangles.Num();
 		for (int i = 0; i < NumTris; ++i)
 		{
 			int tid = AppendTriangle(Generator->Triangles[i], 1 + Generator->TrianglePolygonIDs[i]);
-			check(tid == i);
+			ensure(tid == i);
 		}
 	}
 	else // no attribute and no polygon ids
 	{
-		int NumTris = Generator->Triangles.Num();
 		for (int i = 0; i < NumTris; ++i)
 		{
 			int tid = AppendTriangle(Generator->Triangles[i], 0);
-			check(tid == i);
+			ensure(tid == i);
 		}
 	}
 
-
+	return (TriangleCount() == NumTris);
 }
 
 void FDynamicMesh3::Copy(const FDynamicMesh3& copy, bool bNormals, bool bColors, bool bUVs, bool bAttributes)
@@ -641,37 +641,43 @@ FIndex3i FDynamicMesh3::GetTriNeighbourTris(int tID) const
 
 void FDynamicMesh3::GetVertexOneRingTriangles(int VertexID, TArray<int>& TrianglesOut) const
 {
-	check(VertexRefCounts.IsValid(VertexID));
-	VertexEdgeLists.Enumerate(VertexID, [&](int32 EdgeID)
+	checkSlow(VertexRefCounts.IsValid(VertexID));
+	if (VertexRefCounts.IsValid(VertexID))
 	{
-		FIndex2i EdgePair = GetOrderedOneRingEdgeTris(VertexID, EdgeID);
-		if (EdgePair.A != InvalidID)
+		VertexEdgeLists.Enumerate(VertexID, [&](int32 EdgeID)
 		{
-			TrianglesOut.Add(EdgePair.A);
-		}
-		if (EdgePair.B != InvalidID)
-		{
-			TrianglesOut.Add(EdgePair.B);
-		}
-	});
+			FIndex2i EdgePair = GetOrderedOneRingEdgeTris(VertexID, EdgeID);
+			if (EdgePair.A != InvalidID)
+			{
+				TrianglesOut.Add(EdgePair.A);
+			}
+			if (EdgePair.B != InvalidID)
+			{
+				TrianglesOut.Add(EdgePair.B);
+			}
+		});
+	}
 }
 
 
 void FDynamicMesh3::EnumerateVertexTriangles(int32 VertexID, TFunctionRef<void(int32)> ApplyFunc) const
 {
-	check(VertexRefCounts.IsValid(VertexID));
-	VertexEdgeLists.Enumerate(VertexID, [&](int32 EdgeID)
+	checkSlow(VertexRefCounts.IsValid(VertexID));
+	if (VertexRefCounts.IsValid(VertexID))
 	{
-		FIndex2i EdgePair = GetOrderedOneRingEdgeTris(VertexID, EdgeID);
-		if (EdgePair.A != InvalidID)
+		VertexEdgeLists.Enumerate(VertexID, [&](int32 EdgeID)
 		{
-			ApplyFunc(EdgePair.A);
-		}
-		if (EdgePair.B != InvalidID)
-		{
-			ApplyFunc(EdgePair.B);
-		}
-	});
+			FIndex2i EdgePair = GetOrderedOneRingEdgeTris(VertexID, EdgeID);
+			if (EdgePair.A != InvalidID)
+			{
+				ApplyFunc(EdgePair.A);
+			}
+			if (EdgePair.B != InvalidID)
+			{
+				ApplyFunc(EdgePair.B);
+			}
+		});
+	}
 }
 
 
@@ -1166,8 +1172,8 @@ int32 FDynamicMesh3::FindEdgeInternal(int32 vA, int32 vB, bool& bIsBoundary) con
 
 int FDynamicMesh3::FindEdge(int vA, int vB) const
 {
-	check(IsVertex(vA));
-	check(IsVertex(vB));
+	checkSlow(IsVertex(vA));
+	checkSlow(IsVertex(vB));
 	if (vA == vB)
 	{
 		// self-edges are not allowed, and if we fall through to the search below on a self edge we will incorrectly
@@ -1182,10 +1188,17 @@ int FDynamicMesh3::FindEdge(int vA, int vB) const
 	{
 		vMax = vB; vMin = vA;
 	}
-	return VertexEdgeLists.Find(vMin, [&](int32 eid)
+	if (IsVertex(vMin))
 	{
-		return (Edges[eid].Vert[1] == vMax);
-	}, InvalidID);
+		return VertexEdgeLists.Find(vMin, [&](int32 eid)
+		{
+			return (Edges[eid].Vert[1] == vMax);
+		}, InvalidID);
+	}
+	else
+	{
+		return InvalidID;
+	}
 
 	// this is slower, likely because it creates func<> every time. can we do w/o that?
 	//return VertexEdgeLists.Find(vI, (eid) => { return Edges[4 * eid + 1] == vO; }, InvalidID);
diff --git a/Engine/Plugins/Experimental/GeometryProcessing/Source/DynamicMesh/Private/DynamicMesh3_Edits.cpp b/Engine/Plugins/Experimental/GeometryProcessing/Source/DynamicMesh/Private/DynamicMesh3_Edits.cpp
--- a/Engine/Plugins/Experimental/GeometryProcessing/Source/DynamicMesh/Private/DynamicMesh3_Edits.cpp
+++ b/Engine/Plugins/Experimental/GeometryProcessing/Source/DynamicMesh/Private/DynamicMesh3_Edits.cpp
@@ -152,7 +152,7 @@ int FDynamicMesh3::AppendTriangle(const FIndex3i& tv, int gid)
 {
 	if (IsVertex(tv[0]) == false || IsVertex(tv[1]) == false || IsVertex(tv[2]) == false)
 	{
-		check(false);
+		checkSlow(false);
 		return InvalidID;
 	}
 	if (!ensure(tv[0] != tv[1] && tv[0] != tv[2] && tv[1] != tv[2]))
@@ -227,12 +227,12 @@ EMeshResult FDynamicMesh3::InsertTriangle(int tid, const FIndex3i& tv, int gid,
 
 	if (IsVertex(tv[0]) == false || IsVertex(tv[1]) == false || IsVertex(tv[2]) == false)
 	{
-		check(false);
+		checkSlow(false);
 		return EMeshResult::Failed_NotAVertex;
 	}
 	if (tv[0] == tv[1] || tv[0] == tv[2] || tv[1] == tv[2])
 	{
-		check(false);
+		checkSlow(false);
 		return EMeshResult::Failed_InvalidNeighbourhood;
 	}
 
@@ -534,7 +534,7 @@ void FDynamicMesh3::CompactInPlace(FCompactMaps* CompactInfo)
 
 	if (HasAttributes())
 	{
-		check(CompactInfo);
+		checkSlow(CompactInfo);		// can this ever fail?
 		AttributeSet->CompactInPlace(*CompactInfo);
 	}
 
@@ -638,7 +638,7 @@ EMeshResult FDynamicMesh3::RemoveVertex(int vID, bool bRemoveAllTriangles, bool
 	}
 
 	VertexRefCounts.Decrement(vID);
-	check(VertexRefCounts.IsValid(vID) == false);
+	ensure(VertexRefCounts.IsValid(vID) == false);
 	VertexEdgeLists.Clear(vID);
 
 	UpdateTimeStamp(true, true);
@@ -656,7 +656,7 @@ EMeshResult FDynamicMesh3::RemoveTriangle(int tID, bool bRemoveIsolatedVertices,
 {
 	if (!TriangleRefCounts.IsValid(tID))
 	{
-		check(false);
+		ensure(false);
 		return EMeshResult::Failed_NotATriangle;
 	}
 
@@ -701,7 +701,7 @@ EMeshResult FDynamicMesh3::RemoveTriangle(int tID, bool bRemoveIsolatedVertices,
 
 	// free this triangle
 	TriangleRefCounts.Decrement(tID);
-	check(TriangleRefCounts.IsValid(tID) == false);
+	checkSlow(TriangleRefCounts.IsValid(tID) == false);
 
 	// Decrement vertex refcounts. If any hit 1 and we got remove-isolated flag,
 	// we need to remove that vertex
@@ -712,7 +712,7 @@ EMeshResult FDynamicMesh3::RemoveTriangle(int tID, bool bRemoveIsolatedVertices,
 		if (bRemoveIsolatedVertices && VertexRefCounts.GetRefCount(vid) == 1)
 		{
 			VertexRefCounts.Decrement(vid);
-			check(VertexRefCounts.IsValid(vid) == false);
+			checkSlow(VertexRefCounts.IsValid(vid) == false);
 			VertexEdgeLists.Clear(vid);
 		}
 	}
@@ -735,7 +735,10 @@ EMeshResult FDynamicMesh3::RemoveTriangle(int tID, bool bRemoveIsolatedVertices,
 EMeshResult FDynamicMesh3::SetTriangle(int tID, const FIndex3i& newv, bool bRemoveIsolatedVertices)
 {
 	// @todo support this.
-	check(HasAttributes() == false);
+	if (ensure(HasAttributes()) == false)
+	{
+		return EMeshResult::Failed_Unsupported;
+	}
 
 	FIndex3i tv = GetTriangle(tID);
 	FIndex3i te = GetTriEdges(tID);
@@ -754,17 +757,17 @@ EMeshResult FDynamicMesh3::SetTriangle(int tID, const FIndex3i& newv, bool bRemo
 
 	if (!TriangleRefCounts.IsValid(tID))
 	{
-		check(false);
+		checkSlow(false);
 		return EMeshResult::Failed_NotATriangle;
 	}
 	if (IsVertex(newv[0]) == false || IsVertex(newv[1]) == false || IsVertex(newv[2]) == false)
 	{
-		check(false);
+		checkSlow(false);
 		return EMeshResult::Failed_NotAVertex;
 	}
 	if (newv[0] == newv[1] || newv[0] == newv[2] || newv[1] == newv[2])
 	{
-		check(false);
+		checkSlow(false);
 		return EMeshResult::Failed_BrokenTopology;
 	}
 	// look up edges. if any already have two triangles, this would
@@ -817,7 +820,7 @@ EMeshResult FDynamicMesh3::SetTriangle(int tID, const FIndex3i& newv, bool bRemo
 		if (bRemoveIsolatedVertices && VertexRefCounts.GetRefCount(vid) == 1)
 		{
 			VertexRefCounts.Decrement(vid);
-			check(VertexRefCounts.IsValid(vid) == false);
+			checkSlow(VertexRefCounts.IsValid(vid) == false);
 			VertexEdgeLists.Clear(vid);
 		}
 	}
@@ -1129,12 +1132,12 @@ EMeshResult FDynamicMesh3::FlipEdge(int eab, FEdgeFlipInfo& FlipInfo)
 	// update the two other edges whose triangle nbrs have changed
 	if (ReplaceEdgeTriangle(eca, t0, t1) == -1)
 	{
-		checkf(false, TEXT("FDynamicMesh3.FlipEdge: first ReplaceEdgeTriangle failed"));
+		checkfSlow(false, TEXT("FDynamicMesh3.FlipEdge: first ReplaceEdgeTriangle failed"));
 		return EMeshResult::Failed_UnrecoverableError;
 	}
 	if (ReplaceEdgeTriangle(edb, t1, t0) == -1)
 	{
-		checkf(false, TEXT("FDynamicMesh3.FlipEdge: second ReplaceEdgeTriangle failed"));
+		checkfSlow(false, TEXT("FDynamicMesh3.FlipEdge: second ReplaceEdgeTriangle failed"));
 		return EMeshResult::Failed_UnrecoverableError;
 	}
 
@@ -1145,19 +1148,19 @@ EMeshResult FDynamicMesh3::FlipEdge(int eab, FEdgeFlipInfo& FlipInfo)
 	// remove old eab from verts a and b, and Decrement ref counts
 	if (VertexEdgeLists.Remove(a, eab) == false)
 	{
-		checkf(false, TEXT("FDynamicMesh3.FlipEdge: first edge list remove failed"));
+		checkfSlow(false, TEXT("FDynamicMesh3.FlipEdge: first edge list remove failed"));
 		return EMeshResult::Failed_UnrecoverableError;
 	}
 	if (VertexEdgeLists.Remove(b, eab) == false)
 	{
-		checkf(false, TEXT("FDynamicMesh3.FlipEdge: second edge list remove failed"));
+		checkfSlow(false, TEXT("FDynamicMesh3.FlipEdge: second edge list remove failed"));
 		return EMeshResult::Failed_UnrecoverableError;
 	}
 	VertexRefCounts.Decrement(a);
 	VertexRefCounts.Decrement(b);
 	if (IsVertex(a) == false || IsVertex(b) == false)
 	{
-		checkf(false, TEXT("FDynamicMesh3.FlipEdge: either a or b is not a vertex?"));
+		checkfSlow(false, TEXT("FDynamicMesh3.FlipEdge: either a or b is not a vertex?"));
 		return EMeshResult::Failed_UnrecoverableError;
 	}
 
@@ -1430,7 +1433,7 @@ EMeshResult FDynamicMesh3::CollapseEdge(int vKeep, int vRemove, double collapse_
 		{
 			if (VertexEdgeLists.Remove(b, eid) != true)
 			{
-				checkf(false, TEXT("FDynamicMesh3::CollapseEdge: failed at remove case o == b"));
+				checkfSlow(false, TEXT("FDynamicMesh3::CollapseEdge: failed at remove case o == b"));
 				return EMeshResult::Failed_UnrecoverableError;
 			}
 		}
@@ -1438,7 +1441,7 @@ EMeshResult FDynamicMesh3::CollapseEdge(int vKeep, int vRemove, double collapse_
 		{
 			if (VertexEdgeLists.Remove(c, eid) != true)
 			{
-				checkf(false, TEXT("FDynamicMesh3::CollapseEdge: failed at remove case o == c"));
+				checkfSlow(false, TEXT("FDynamicMesh3::CollapseEdge: failed at remove case o == c"));
 				return EMeshResult::Failed_UnrecoverableError;
 			}
 			tac = GetOtherEdgeTriangle(eid, t0);
@@ -1447,7 +1450,7 @@ EMeshResult FDynamicMesh3::CollapseEdge(int vKeep, int vRemove, double collapse_
 		{
 			if (VertexEdgeLists.Remove(d, eid) != true)
 			{
-				checkf(false, TEXT("FDynamicMesh3::CollapseEdge: failed at remove case o == c, step 1"));
+				checkfSlow(false, TEXT("FDynamicMesh3::CollapseEdge: failed at remove case o == c, step 1"));
 				return EMeshResult::Failed_UnrecoverableError;
 			}
 			tad = GetOtherEdgeTriangle(eid, t1);
@@ -1456,7 +1459,7 @@ EMeshResult FDynamicMesh3::CollapseEdge(int vKeep, int vRemove, double collapse_
 		{
 			if (ReplaceEdgeVertex(eid, a, b) == -1)
 			{
-				checkf(false, TEXT("FDynamicMesh3::CollapseEdge: failed at remove case else"));
+				checkfSlow(false, TEXT("FDynamicMesh3::CollapseEdge: failed at remove case else"));
 				return EMeshResult::Failed_UnrecoverableError;
 			}
 			VertexEdgeLists.Insert(b, eid);
@@ -1473,7 +1476,7 @@ EMeshResult FDynamicMesh3::CollapseEdge(int vKeep, int vRemove, double collapse_
 				{
 					if (ReplaceTriangleVertex(t_j, a, b) == -1)
 					{
-						checkf(false, TEXT("FDynamicMesh3::CollapseEdge: failed at remove last check"));
+						checkfSlow(false, TEXT("FDynamicMesh3::CollapseEdge: failed at remove last check"));
 						return EMeshResult::Failed_UnrecoverableError;
 					}
 					VertexRefCounts.Increment(b);
@@ -1487,9 +1490,9 @@ EMeshResult FDynamicMesh3::CollapseEdge(int vKeep, int vRemove, double collapse_
 	{
 		// remove all edges from vtx a, then remove vtx a
 		VertexEdgeLists.Clear(a);
-		check(VertexRefCounts.GetRefCount(a) == 3);		// in t0,t1, and initial ref
+		checkSlow(VertexRefCounts.GetRefCount(a) == 3);		// in t0,t1, and initial ref
 		VertexRefCounts.Decrement(a, 3);
-		check(VertexRefCounts.IsValid(a) == false);
+		checkSlow(VertexRefCounts.IsValid(a) == false);
 
 		// remove triangles T0 and T1, and update b/c/d refcounts
 		TriangleRefCounts.Decrement(t0);
@@ -1497,16 +1500,16 @@ EMeshResult FDynamicMesh3::CollapseEdge(int vKeep, int vRemove, double collapse_
 		VertexRefCounts.Decrement(c);
 		VertexRefCounts.Decrement(d);
 		VertexRefCounts.Decrement(b, 2);
-		check(TriangleRefCounts.IsValid(t0) == false);
-		check(TriangleRefCounts.IsValid(t1) == false);
+		checkSlow(TriangleRefCounts.IsValid(t0) == false);
+		checkSlow(TriangleRefCounts.IsValid(t1) == false);
 
 		// remove edges ead, eab, eac
 		EdgeRefCounts.Decrement(ead);
 		EdgeRefCounts.Decrement(eab);
 		EdgeRefCounts.Decrement(eac);
-		check(EdgeRefCounts.IsValid(ead) == false);
-		check(EdgeRefCounts.IsValid(eab) == false);
-		check(EdgeRefCounts.IsValid(eac) == false);
+		checkSlow(EdgeRefCounts.IsValid(ead) == false);
+		checkSlow(EdgeRefCounts.IsValid(eab) == false);
+		checkSlow(EdgeRefCounts.IsValid(eac) == false);
 
 		// replace t0 and t1 in edges ebd and ebc that we kept
 		ebd = FindEdgeFromTri(b, d, t1);
@@ -1517,13 +1520,13 @@ EMeshResult FDynamicMesh3::CollapseEdge(int vKeep, int vRemove, double collapse_
 
 		if (ReplaceEdgeTriangle(ebd, t1, tad) == -1)
 		{
-			checkf(false, TEXT("FDynamicMesh3::CollapseEdge: failed at isboundary=false branch, ebd replace triangle"));
+			checkfSlow(false, TEXT("FDynamicMesh3::CollapseEdge: failed at isboundary=false branch, ebd replace triangle"));
 			return EMeshResult::Failed_UnrecoverableError;
 		}
 
 		if (ReplaceEdgeTriangle(ebc, t0, tac) == -1)
 		{
-			checkf(false, TEXT("FDynamicMesh3::CollapseEdge: failed at isboundary=false branch, ebc replace triangle"));
+			checkfSlow(false, TEXT("FDynamicMesh3::CollapseEdge: failed at isboundary=false branch, ebc replace triangle"));
 			return EMeshResult::Failed_UnrecoverableError;
 		}
 
@@ -1532,7 +1535,7 @@ EMeshResult FDynamicMesh3::CollapseEdge(int vKeep, int vRemove, double collapse_
 		{
 			if (ReplaceTriangleEdge(tad, ead, ebd) == -1)
 			{
-				checkf(false, TEXT("FDynamicMesh3::CollapseEdge: failed at isboundary=false branch, ebd replace triangle"));
+				checkfSlow(false, TEXT("FDynamicMesh3::CollapseEdge: failed at isboundary=false branch, ebd replace triangle"));
 				return EMeshResult::Failed_UnrecoverableError;
 			}
 		}
@@ -1540,7 +1543,7 @@ EMeshResult FDynamicMesh3::CollapseEdge(int vKeep, int vRemove, double collapse_
 		{
 			if (ReplaceTriangleEdge(tac, eac, ebc) == -1)
 			{
-				checkf(false, TEXT("FDynamicMesh3::CollapseEdge: failed at isboundary=false branch, ebd replace triangle"));
+				checkfSlow(false, TEXT("FDynamicMesh3::CollapseEdge: failed at isboundary=false branch, ebd replace triangle"));
 				return EMeshResult::Failed_UnrecoverableError;
 			}
 		}
@@ -1552,27 +1555,27 @@ EMeshResult FDynamicMesh3::CollapseEdge(int vKeep, int vRemove, double collapse_
 
 		// remove all edges from vtx a, then remove vtx a
 		VertexEdgeLists.Clear(a);
-		check(VertexRefCounts.GetRefCount(a) == 2);		// in t0 and initial ref
+		checkSlow(VertexRefCounts.GetRefCount(a) == 2);		// in t0 and initial ref
 		VertexRefCounts.Decrement(a, 2);
-		check(VertexRefCounts.IsValid(a) == false);
+		checkSlow(VertexRefCounts.IsValid(a) == false);
 
 		// remove triangle T0 and update b/c refcounts
 		TriangleRefCounts.Decrement(t0);
 		VertexRefCounts.Decrement(c);
 		VertexRefCounts.Decrement(b);
-		check(TriangleRefCounts.IsValid(t0) == false);
+		checkSlow(TriangleRefCounts.IsValid(t0) == false);
 
 		// remove edges eab and eac
 		EdgeRefCounts.Decrement(eab);
 		EdgeRefCounts.Decrement(eac);
-		check(EdgeRefCounts.IsValid(eab) == false);
-		check(EdgeRefCounts.IsValid(eac) == false);
+		checkSlow(EdgeRefCounts.IsValid(eab) == false);
+		checkSlow(EdgeRefCounts.IsValid(eac) == false);
 
 		// replace t0 in edge ebc that we kept
 		ebc = FindEdgeFromTri(b, c, t0);
 		if (ReplaceEdgeTriangle(ebc, t0, tac) == -1)
 		{
-			checkf(false, TEXT("FDynamicMesh3::CollapseEdge: failed at isboundary=false branch, ebc replace triangle"));
+			checkfSlow(false, TEXT("FDynamicMesh3::CollapseEdge: failed at isboundary=false branch, ebc replace triangle"));
 			return EMeshResult::Failed_UnrecoverableError;
 		}
 
@@ -1581,7 +1584,7 @@ EMeshResult FDynamicMesh3::CollapseEdge(int vKeep, int vRemove, double collapse_
 		{
 			if (ReplaceTriangleEdge(tac, eac, ebc) == -1)
 			{
-				checkf(false, TEXT("FDynamicMesh3::CollapseEdge: failed at isboundary=true branch, ebd replace triangle"));
+				checkfSlow(false, TEXT("FDynamicMesh3::CollapseEdge: failed at isboundary=true branch, ebd replace triangle"));
 				return EMeshResult::Failed_UnrecoverableError;
 			}
 		}
diff --git a/Engine/Plugins/Experimental/GeometryProcessing/Source/DynamicMesh/Private/DynamicMesh3_Queries.cpp b/Engine/Plugins/Experimental/GeometryProcessing/Source/DynamicMesh/Private/DynamicMesh3_Queries.cpp
--- a/Engine/Plugins/Experimental/GeometryProcessing/Source/DynamicMesh/Private/DynamicMesh3_Queries.cpp
+++ b/Engine/Plugins/Experimental/GeometryProcessing/Source/DynamicMesh/Private/DynamicMesh3_Queries.cpp
@@ -48,8 +48,7 @@ int FDynamicMesh3::GetVtxBoundaryEdges(int vID, int& e0, int& e1) const
 		}
 		return count;
 	}
-	check(false);
-	return -1;
+	return 0;
 }
 
 
@@ -68,8 +67,7 @@ int FDynamicMesh3::GetAllVtxBoundaryEdges(int vID, TArray<int>& EdgeListOut) con
 		}
 		return count;
 	}
-	check(false);
-	return -1;
+	return 0;
 }
 
 
@@ -278,8 +276,11 @@ EMeshResult FDynamicMesh3::GetVtxContiguousTriangles(int VertexID, TArray<int>&
 			if (NextTriID == InvalidID)
 			{
 				// remove the corresponding boundary
-				check(StartEdgeIDs.Num() > 0);
-				StartEdgeIDs.RemoveSingleSwap(NextEID);
+				checkSlow(StartEdgeIDs.Num() > 0);
+				if (StartEdgeIDs.Num() > 0)
+				{
+					StartEdgeIDs.RemoveSingleSwap(NextEID);
+				}
 				break;
 			}
 			WalkTri = NextTriID;
@@ -288,19 +289,20 @@ EMeshResult FDynamicMesh3::GetVtxContiguousTriangles(int VertexID, TArray<int>&
 		SpanLengths.Add(TrianglesOut.Num() - SpanStart);
 	}
 
-	check(SpanLengths.Num() == IsLoop.Num());
-
-	return EMeshResult::Ok;
+	return ensure(SpanLengths.Num() == IsLoop.Num()) ? EMeshResult::Ok : EMeshResult::Failed_InvalidNeighbourhood;
 }
 
 bool FDynamicMesh3::IsBoundaryVertex(int vID) const
 {
-	check(IsVertex(vID));
-	for (int eid : VertexEdgeLists.Values(vID))
+	checkSlow(IsVertex(vID));
+	if ( IsVertex(vID) )
 	{
-		if (Edges[eid].Tri[1] == InvalidID)
+		for (int eid : VertexEdgeLists.Values(vID))
 		{
-			return true;
+			if (Edges[eid].Tri[1] == InvalidID)
+			{
+				return true;
+			}
 		}
 	}
 	return false;
@@ -308,9 +310,16 @@ bool FDynamicMesh3::IsBoundaryVertex(int vID) const
 
 bool FDynamicMesh3::IsBoundaryTriangle(int tID) const
 {
-	check(IsTriangle(tID));
-	const FIndex3i& TriEdgeIDs = TriangleEdges[tID];
-	return IsBoundaryEdge(TriEdgeIDs[0]) || IsBoundaryEdge(TriEdgeIDs[1]) || IsBoundaryEdge(TriEdgeIDs[2]);
+	checkSlow(IsTriangle(tID));
+	if (IsTriangle(tID))
+	{
+		const FIndex3i& TriEdgeIDs = TriangleEdges[tID];
+		return IsBoundaryEdge(TriEdgeIDs[0]) || IsBoundaryEdge(TriEdgeIDs[1]) || IsBoundaryEdge(TriEdgeIDs[2]);
+	}
+	else
+	{
+		return false;
+	}
 }
 
 FIndex2i FDynamicMesh3::GetOrientedBoundaryEdgeV(int eID) const
@@ -327,7 +336,7 @@ FIndex2i FDynamicMesh3::GetOrientedBoundaryEdgeV(int eID) const
 			return FIndex2i(tri[ai], tri[(ai + 1) % 3]);
 		}
 	}
-	check(false);
+	checkSlow(false);
 	return InvalidEdge;
 }
 
@@ -489,7 +498,11 @@ bool FDynamicMesh3::GetAllVertexGroups(int vID, TArray<int>& GroupsOut) const
  */
 bool FDynamicMesh3::IsBowtieVertex(int vID) const
 {
-	check(VertexRefCounts.IsValid(vID));
+	checkSlow(VertexRefCounts.IsValid(vID));
+	if (VertexRefCounts.IsValid(vID) == false)
+	{
+		return false;
+	}
 
 	int nEdges = VertexEdgeLists.GetCount(vID);
 	if (nEdges == 0)
@@ -690,7 +703,7 @@ FVector3d FDynamicMesh3::GetEdgeNormal(int eID) const
 		}
 		return n;
 	}
-	check(false);
+	checkSlow(false);
 	return FVector3d::Zero();
 }
 
@@ -705,7 +718,7 @@ FVector3d FDynamicMesh3::GetEdgePoint(int eID, double t) const
 		double mt = 1.0 - t;
 		return mt*Vertices[iv0] + t * Vertices[iv1];
 	}
-	check(false);
+	checkSlow(false);
 	return FVector3d::Zero();
 }
 
@@ -732,7 +745,10 @@ void FDynamicMesh3::GetVtxOneRingCentroid(int vID, FVector3d& centroid) const
 
 FFrame3d FDynamicMesh3::GetVertexFrame(int vID, bool bFrameNormalY) const
 {
-	check(HasVertexNormals());
+	if (ensure(HasVertexNormals()) == false)
+	{
+		return FFrame3d();
+	}
 
 	FVector3d v = Vertices[vID];
 	const TDynamicVector<FVector3f>& Normals = VertexNormals.GetValue();
@@ -793,12 +809,16 @@ FVector3d FDynamicMesh3::GetTriBaryPoint(int tID, double bary0, double bary1, do
 
 FVector3d FDynamicMesh3::GetTriBaryNormal(int tID, double bary0, double bary1, double bary2) const
 {
-	check(HasVertexNormals());
-	const FIndex3i& tIDs = Triangles[tID];
-	const TDynamicVector<FVector3f>& normalsR = VertexNormals.GetValue();
-	FVector3d n = FVector3d(bary0 * normalsR[tIDs[0]] + bary1 * normalsR[tIDs[1]] + bary2 * normalsR[tIDs[2]]);
-	Normalize(n);
-	return n;
+	checkSlow(HasVertexNormals());
+	if (HasVertexNormals())
+	{
+		const FIndex3i& tIDs = Triangles[tID];
+		const TDynamicVector<FVector3f>& normalsR = VertexNormals.GetValue();
+		FVector3d n = FVector3d(bary0 * normalsR[tIDs[0]] + bary1 * normalsR[tIDs[1]] + bary2 * normalsR[tIDs[2]]);
+		Normalize(n);
+		return n;
+	}
+	return FVector3d::Zero();
 }
 
 FVector3d FDynamicMesh3::GetTriCentroid(int tID) const
diff --git a/Engine/Plugins/Experimental/GeometryProcessing/Source/DynamicMesh/Private/DynamicMeshAttributeSet.cpp b/Engine/Plugins/Experimental/GeometryProcessing/Source/DynamicMesh/Private/DynamicMeshAttributeSet.cpp
--- a/Engine/Plugins/Experimental/GeometryProcessing/Source/DynamicMesh/Private/DynamicMeshAttributeSet.cpp
+++ b/Engine/Plugins/Experimental/GeometryProcessing/Source/DynamicMesh/Private/DynamicMeshAttributeSet.cpp
@@ -369,7 +369,7 @@ void FDynamicMeshAttributeSet::SetNumUVLayers(int Num)
 	{
 		UVLayers.RemoveAt(Num, UVLayers.Num() - Num);
 	}
-	check(UVLayers.Num() == Num);
+	ensure(UVLayers.Num() == Num);
 }
 
 
@@ -402,7 +402,7 @@ void FDynamicMeshAttributeSet::SetNumNormalLayers(int Num)
 	{
 		NormalLayers.RemoveAt(Num, UVLayers.Num() - Num);
 	}
-	check(NormalLayers.Num() == Num);
+	ensure(NormalLayers.Num() == Num);
 }
 
 void FDynamicMeshAttributeSet::EnablePrimaryColors()
@@ -441,7 +441,7 @@ void FDynamicMeshAttributeSet::SetNumPolygroupLayers(int32 Num)
 	{
 		PolygroupLayers.RemoveAt(Num, PolygroupLayers.Num() - Num);
 	}
-	check(PolygroupLayers.Num() == Num);
+	ensure(PolygroupLayers.Num() == Num);
 }
 
 FDynamicMeshPolygroupAttribute* FDynamicMeshAttributeSet::GetPolygroupLayer(int Index)
@@ -600,17 +600,21 @@ bool FDynamicMeshAttributeSet::IsMaterialBoundaryEdge(int EdgeID) const
 	{
 		return false;
 	}
-	check(ParentMesh->IsEdge(EdgeID));
-	const FDynamicMesh3::FEdge Edge = ParentMesh->GetEdge(EdgeID);
-	const int Tri0 = Edge.Tri[0];
-	const int Tri1 = Edge.Tri[1];
-	if (( Tri0 == IndexConstants::InvalidID ) || (Tri1 == IndexConstants::InvalidID))
+	checkSlow(ParentMesh->IsEdge(EdgeID));
+	if (ParentMesh->IsEdge(EdgeID))
 	{
-		return false;
+		const FDynamicMesh3::FEdge Edge = ParentMesh->GetEdge(EdgeID);
+		const int Tri0 = Edge.Tri[0];
+		const int Tri1 = Edge.Tri[1];
+		if ((Tri0 == IndexConstants::InvalidID) || (Tri1 == IndexConstants::InvalidID))
+		{
+			return false;
+		}
+		const int MatID0 = MaterialIDAttrib->GetValue(Tri0);
+		const int MatID1 = MaterialIDAttrib->GetValue(Tri1);
+		return MatID0 != MatID1;
 	}
-	const int MatID0 = MaterialIDAttrib->GetValue(Tri0);
-	const int MatID1 = MaterialIDAttrib->GetValue(Tri1);
-	return MatID0 != MatID1;
+	return false;
 }
 
 void FDynamicMeshAttributeSet::OnNewVertex(int VertexID, bool bInserted)
diff --git a/Engine/Plugins/Experimental/GeometryProcessing/Source/DynamicMesh/Private/DynamicMeshChangeTracker.cpp b/Engine/Plugins/Experimental/GeometryProcessing/Source/DynamicMesh/Private/DynamicMeshChangeTracker.cpp
--- a/Engine/Plugins/Experimental/GeometryProcessing/Source/DynamicMesh/Private/DynamicMeshChangeTracker.cpp
+++ b/Engine/Plugins/Experimental/GeometryProcessing/Source/DynamicMesh/Private/DynamicMeshChangeTracker.cpp
@@ -114,7 +114,7 @@ void TDynamicMeshAttributeChange<RealType,ElementSize>::ApplyReplaceChange(TDyna
 		{
 			const RealType* Data = &InsertElementData[ElemInfo.DataIndex];
 			EMeshResult Result = Overlay->InsertElement(ElemInfo.ElementID, Data, bUseUnsafe);
-			check(Result == EMeshResult::Ok);
+			ensure(Result == EMeshResult::Ok);
 		}
 		else
 		{
@@ -130,10 +130,10 @@ void TDynamicMeshAttributeChange<RealType,ElementSize>::ApplyReplaceChange(TDyna
 	// set new element triangles
 	for (const FChangeTriangle& TriInfo : InsertTris)
 	{
-		check(Overlay->GetParentMesh()->IsTriangle(TriInfo.TriangleID));
+		checkSlow(Overlay->GetParentMesh()->IsTriangle(TriInfo.TriangleID));
 		// existing overlay triangle should be empty...
 		EMeshResult Result = Overlay->SetTriangle(TriInfo.TriangleID, TriInfo.Elements);
-		check(Result == EMeshResult::Ok);
+		ensure(Result == EMeshResult::Ok);
 	}
 }
 
@@ -249,7 +249,7 @@ void FDynamicMeshChange::ApplyReplaceChange(FDynamicMesh3* Mesh,
 	for (const FChangeTriangle& TriInfo : RemoveTris)
 	{
 		EMeshResult Result = Mesh->RemoveTriangle(TriInfo.TriangleID);
-		check(Result == EMeshResult::Ok);
+		ensure(Result == EMeshResult::Ok);
 	}
 
 	// if this is true, we use unsafe insert, which means we take shortcuts
@@ -269,7 +269,7 @@ void FDynamicMeshChange::ApplyReplaceChange(FDynamicMesh3* Mesh,
 		if (Mesh->IsVertex(VertInfo.VertexID) == false)
 		{
 			EMeshResult Result = Mesh->InsertVertex(VertInfo.VertexID, VertInfo.Info, bUseUnsafe);
-			check(Result == EMeshResult::Ok);
+			ensure(Result == EMeshResult::Ok);
 		}
 		else
 		{
@@ -301,7 +301,7 @@ void FDynamicMeshChange::ApplyReplaceChange(FDynamicMesh3* Mesh,
 	for (const FChangeTriangle& TriInfo : InsertTris)
 	{
 		EMeshResult Result = Mesh->InsertTriangle(TriInfo.TriangleID, TriInfo.Vertices, TriInfo.GroupID, bUseUnsafe);
-		check(Result == EMeshResult::Ok);
+		ensure(Result == EMeshResult::Ok);
 	}
 	if (bUseUnsafe)
 	{
@@ -331,9 +331,9 @@ void FDynamicMeshChange::VerifySaveState() const
 
 	for (const FChangeTriangle& TriInfo : OldTriangles)
 	{
-		check(SavedVertexIDs.Contains(TriInfo.Vertices.A));
-		check(SavedVertexIDs.Contains(TriInfo.Vertices.B));
-		check(SavedVertexIDs.Contains(TriInfo.Vertices.C));
+		ensure(SavedVertexIDs.Contains(TriInfo.Vertices.A));
+		ensure(SavedVertexIDs.Contains(TriInfo.Vertices.B));
+		ensure(SavedVertexIDs.Contains(TriInfo.Vertices.C));
 	}
 }
 
@@ -358,7 +358,7 @@ FDynamicMeshAttributeSetChangeTracker::FDynamicMeshAttributeSetChangeTracker(con
 void FDynamicMeshAttributeSetChangeTracker::BeginChange()
 {
 	// initialize new attribute set change
-	check(Change == nullptr);
+	ensure(Change == nullptr);
 	Change = new FDynamicMeshAttributeChangeSet();
 	int NumUVLayers = Attribs->NumUVLayers();
 	Change->UVChanges.SetNum(NumUVLayers);
@@ -625,7 +625,8 @@ void FDynamicMeshAttributeSetChangeTracker::StoreAllFinalVertices(const TArray<i
 bool FDynamicMeshAttributeChangeSet::Apply(FDynamicMeshAttributeSet* Attributes, bool bRevert) const
 {
 	int NumUVLayers = Attributes->NumUVLayers();
-	check(NumUVLayers == UVChanges.Num());
+	ensure(NumUVLayers == UVChanges.Num());
+	NumUVLayers = FMath::Min(NumUVLayers, UVChanges.Num());
 	for (int k = 0; k < NumUVLayers; ++k)
 	{
 		FDynamicMeshUVOverlay* UVLayer = Attributes->GetUVLayer(k);
@@ -633,7 +634,8 @@ bool FDynamicMeshAttributeChangeSet::Apply(FDynamicMeshAttributeSet* Attributes,
 	}
 
 	int NumNormalLayers = Attributes->NumNormalLayers();
-	check(NumNormalLayers == NormalChanges.Num());
+	ensure(NumNormalLayers == NormalChanges.Num());
+	NumNormalLayers = FMath::Min(NumNormalLayers, NormalChanges.Num());
 	for (int k = 0; k < NumNormalLayers; ++k)
 	{
 		FDynamicMeshNormalOverlay* NormalLayer = Attributes->GetNormalLayer(k);
@@ -700,7 +702,7 @@ FDynamicMeshChangeTracker::~FDynamicMeshChangeTracker()
 
 void FDynamicMeshChangeTracker::BeginChange()
 {
-	check(Change == nullptr);
+	ensure(Change == nullptr);
 	Change = new FDynamicMeshChange();
 
 	// @todo should we do this on EndChange() so that delay is at end of stroke instead of beginning?
diff --git a/Engine/Plugins/Experimental/GeometryProcessing/Source/DynamicMesh/Private/DynamicMeshOverlay.cpp b/Engine/Plugins/Experimental/GeometryProcessing/Source/DynamicMesh/Private/DynamicMeshOverlay.cpp
--- a/Engine/Plugins/Experimental/GeometryProcessing/Source/DynamicMesh/Private/DynamicMeshOverlay.cpp
+++ b/Engine/Plugins/Experimental/GeometryProcessing/Source/DynamicMesh/Private/DynamicMeshOverlay.cpp
@@ -271,7 +271,10 @@ void TDynamicMeshOverlay<RealType, ElementSize>::SplitBowties()
 		{
 			bool bIsLoop = GroupIsLoop[GroupIdx];
 			int TriInGroupNum = ContiguousGroupLengths[GroupIdx];
-			check(TriInGroupNum > 0);
+			if (ensure(TriInGroupNum > 0) == false)
+			{
+				continue;
+			}
 			int TriSubEnd = TriSubStart + TriInGroupNum;
 			
 			GroupElementIDs.Reset();
@@ -364,12 +367,12 @@ EMeshResult TDynamicMeshOverlay<RealType, ElementSize>::SetTriangle(int tid, con
 {
 	if (IsElement(tv[0]) == false || IsElement(tv[1]) == false || IsElement(tv[2]) == false)
 	{
-		check(false);
+		checkSlow(false);
 		return EMeshResult::Failed_NotAVertex;
 	}
 	if (tv[0] == tv[1] || tv[0] == tv[2] || tv[1] == tv[2]) 
 	{
-		check(false);
+		checkSlow(false);
 		return EMeshResult::Failed_InvalidNeighbourhood;
 	}
 
@@ -404,7 +407,10 @@ void TDynamicMeshOverlay<RealType, ElementSize>::UnsetTriangle(int TriangleID)
 template<typename RealType, int ElementSize>
 void TDynamicMeshOverlay<RealType, ElementSize>::InternalSetTriangle(int tid, const FIndex3i& tv, bool bIncrementRefCounts)
 {
-	check(ParentMesh);
+	if (ensure(ParentMesh) == false)
+	{
+		return;
+	}
 
 	int i = 3 * tid;
 	ElementTriangles.InsertAt(tv[2], i + 2);
@@ -746,15 +752,18 @@ int TDynamicMeshOverlay<RealType, ElementSize>::CountVertexElements(int vid, boo
 template<typename RealType, int ElementSize>
 void TDynamicMeshOverlay<RealType, ElementSize>::GetElementTriangles(int ElementID, TArray<int>& OutTriangles) const
 {
-	check(ElementsRefCounts.IsValid(ElementID));
-	int VertexID = ParentVertices[ElementID];
-
-	for (int TriangleID : ParentMesh->VtxTrianglesItr(VertexID))
+	checkSlow(ElementsRefCounts.IsValid(ElementID));
+	if (ElementsRefCounts.IsValid(ElementID))
 	{
-		int i = 3 * TriangleID;
-		if (ElementTriangles[i] == ElementID || ElementTriangles[i+1] == ElementID || ElementTriangles[i+2] == ElementID)
+		int VertexID = ParentVertices[ElementID];
+
+		for (int TriangleID : ParentMesh->VtxTrianglesItr(VertexID))
 		{
-			OutTriangles.Add(TriangleID);
+			int i = 3 * TriangleID;
+			if (ElementTriangles[i] == ElementID || ElementTriangles[i+1] == ElementID || ElementTriangles[i+2] == ElementID)
+			{
+				OutTriangles.Add(TriangleID);
+			}
 		}
 	}
 }
@@ -781,7 +790,7 @@ void TDynamicMeshOverlay<RealType, ElementSize>::OnRemoveTriangle(int TriangleID
 		{
 			ElementsRefCounts.Decrement(elemid);
 			ParentVertices[elemid] = FDynamicMesh3::InvalidID;
-			check(ElementsRefCounts.IsValid(elemid) == false);
+			ensure(ElementsRefCounts.IsValid(elemid) == false);
 		}
 	}
 }
@@ -904,7 +913,7 @@ void TDynamicMeshOverlay<RealType, ElementSize>::OnFlipEdge(const FDynamicMesh3:
 	bool bT0Set = IsSetTriangle(orig_t0), bT1Set = IsSetTriangle(orig_t1);
 	if (!bT0Set)
 	{
-		check(!bT1Set); // flipping across a set/unset boundary is not allowed?
+		ensure(!bT1Set); // flipping across a set/unset boundary is not allowed?
 		return; // nothing to do on the overlay if both triangles are unset
 	}
 
@@ -928,12 +937,12 @@ void TDynamicMeshOverlay<RealType, ElementSize>::OnFlipEdge(const FDynamicMesh3:
 	int idx_base_d = IndexUtil::GetOtherTriIndex(idx_base_a2, idx_base_b2);
 
 	// sanity checks
-	check(idx_base_c == BaseTriangle0.IndexOf(base_c));
-	check(idx_base_d == BaseTriangle1.IndexOf(base_d));
+	checkSlow(idx_base_c == BaseTriangle0.IndexOf(base_c));
+	checkSlow(idx_base_d == BaseTriangle1.IndexOf(base_d));
 
 	// we should not have been called on a non-shared edge!!
 	bool bHasSharedUVEdge = IndexUtil::SamePairUnordered(Triangle0[idx_base_a1], Triangle0[idx_base_b1], Triangle1[idx_base_a2], Triangle1[idx_base_b2]);
-	check(bHasSharedUVEdge);
+	checkSlow(bHasSharedUVEdge);
 
 	int A = Triangle0[idx_base_a1];
 	int B = Triangle0[idx_base_b1];
@@ -1170,7 +1179,7 @@ void TDynamicMeshOverlay<RealType, ElementSize>::OnCollapseEdge(const FDynamicMe
 		if (removed_elemid[k] != FDynamicMesh3::InvalidID)
 	{
 			int rc = ElementsRefCounts.GetRefCount(removed_elemid[k]);
-			check(rc == 1);
+			ensure(rc == 1);
 			ElementsRefCounts.Decrement(removed_elemid[k]);
 			ParentVertices[removed_elemid[k]] = FDynamicMesh3::InvalidID;
 		}
diff --git a/Engine/Plugins/Experimental/GeometryProcessing/Source/DynamicMesh/Public/DynamicMesh3.h b/Engine/Plugins/Experimental/GeometryProcessing/Source/DynamicMesh/Public/DynamicMesh3.h
--- a/Engine/Plugins/Experimental/GeometryProcessing/Source/DynamicMesh/Public/DynamicMesh3.h
+++ b/Engine/Plugins/Experimental/GeometryProcessing/Source/DynamicMesh/Public/DynamicMesh3.h
@@ -219,8 +219,11 @@ public:
 	void Copy(const FDynamicMesh3& CopyMesh, bool bNormals = true, bool bColors = true, bool bUVs = true,
 	          bool bAttributes = true);
 
-	/** Initialize mesh from the output of a MeshShapeGenerator (assumes Generate() was already called) */
-	void Copy(const FMeshShapeGenerator* Generator);
+	/** 
+	  * Initialize mesh from the output of a MeshShapeGenerator (assumes Generate() was already called) 
+	  * @return false if any triangles defined in Generator were skipped due to invalid topology
+	  */
+	bool Copy(const FMeshShapeGenerator* Generator);
 
 	/** 
 	 * Copy input mesh while compacting, i.e. removing unused vertices/triangles/edges. 
diff --git a/Engine/Plugins/Experimental/GeometryProcessing/Source/DynamicMesh/Public/DynamicMeshTriangleAttribute.h b/Engine/Plugins/Experimental/GeometryProcessing/Source/DynamicMesh/Public/DynamicMeshTriangleAttribute.h
--- a/Engine/Plugins/Experimental/GeometryProcessing/Source/DynamicMesh/Public/DynamicMeshTriangleAttribute.h
+++ b/Engine/Plugins/Experimental/GeometryProcessing/Source/DynamicMesh/Public/DynamicMeshTriangleAttribute.h
@@ -131,8 +131,10 @@ public:
 			{
 				continue;
 			}
-			check(ToTID <= TID);
-			CopyValue(TID, ToTID);
+			if (ensure(ToTID <= TID))
+			{
+				CopyValue(TID, ToTID);
+			}
 		}
 		AttribValues.Resize(ParentMesh->MaxTriangleID() * AttribDimension);
 	}
@@ -443,8 +445,12 @@ bool FDynamicMeshTriangleAttributeChange<AttribValueType, AttribDimension>::Appl
 	const TArray<FChangeTriangleAttribute> *Changes = bRevert ? &OldTriangleAttributes : &NewTriangleAttributes;
 	for (const FChangeTriangleAttribute& Change : *Changes)
 	{
-		check(AttribCast->GetParentMesh()->IsTriangle(Change.TriangleID));
-		AttribCast->SetValue(Change.TriangleID, Change.Data);
+		bool bIsTriangle = AttribCast->GetParentMesh()->IsTriangle(Change.TriangleID);
+		checkSlow(bIsTriangle);
+		if (bIsTriangle)
+		{
+			AttribCast->SetValue(Change.TriangleID, Change.Data);
+		}
 	}
 	return true;
 }
diff --git a/Engine/Plugins/Experimental/GeometryProcessing/Source/DynamicMesh/Public/DynamicVertexAttribute.h b/Engine/Plugins/Experimental/GeometryProcessing/Source/DynamicMesh/Public/DynamicVertexAttribute.h
--- a/Engine/Plugins/Experimental/GeometryProcessing/Source/DynamicMesh/Public/DynamicVertexAttribute.h
+++ b/Engine/Plugins/Experimental/GeometryProcessing/Source/DynamicMesh/Public/DynamicVertexAttribute.h
@@ -118,8 +118,10 @@ public:
 			{
 				continue;
 			}
-			check(ToVID <= VID);
-			CopyValue(VID, ToVID);
+			if (ensure(ToVID <= VID))
+			{
+				CopyValue(VID, ToVID);
+			}
 		}
 		AttribValues.Resize(Parent->MaxVertexID() * AttribDimension);
 	}
@@ -435,8 +437,12 @@ bool TDynamicVertexAttributeChange<AttribValueType, AttribDimension, ParentType>
 	TDynamicVertexAttribute<AttribValueType, AttribDimension, ParentType>* AttribCast = static_cast<TDynamicVertexAttribute<AttribValueType, AttribDimension, ParentType>*>(Attribute);
 	for (const FChangeVertexAttribute& Change : *Changes)
 	{
-		check(AttribCast->GetParent()->IsVertex(Change.VertexID));
-		AttribCast->SetValue(Change.VertexID, Change.Data);
+		bool bIsVertex = AttribCast->GetParent()->IsVertex(Change.VertexID);
+		check(bIsVertex);
+		if (bIsVertex)
+		{
+			AttribCast->SetValue(Change.VertexID, Change.Data);
+		}
 	}
 	return true;
 }
diff --git a/Engine/Plugins/Experimental/GeometryProcessing/Source/GeometricObjects/Public/GeometryTypes.h b/Engine/Plugins/Experimental/GeometryProcessing/Source/GeometricObjects/Public/GeometryTypes.h
--- a/Engine/Plugins/Experimental/GeometryProcessing/Source/GeometricObjects/Public/GeometryTypes.h
+++ b/Engine/Plugins/Experimental/GeometryProcessing/Source/GeometricObjects/Public/GeometryTypes.h
@@ -40,7 +40,8 @@ enum class EMeshResult
 	Failed_TriangleAlreadyExists = 51,
 	Failed_CannotAllocateTriangle = 52,
 
-	Failed_UnrecoverableError = 1000
+	Failed_UnrecoverableError = 1000,
+	Failed_Unsupported = 1001
 };
 
 
