GeometryProcessing: add GroupTopology::GetSelectionBounds() function, use in new UPolygonSelectionMechanic::GetSelectionBounds() function
#rnx

diff --git a/Engine/Plugins/Experimental/GeometryProcessing/Source/DynamicMesh/Private/GroupTopology.cpp b/Engine/Plugins/Experimental/GeometryProcessing/Source/DynamicMesh/Private/GroupTopology.cpp
--- a/Engine/Plugins/Experimental/GeometryProcessing/Source/DynamicMesh/Private/GroupTopology.cpp
+++ b/Engine/Plugins/Experimental/GeometryProcessing/Source/DynamicMesh/Private/GroupTopology.cpp
@@ -739,7 +739,44 @@ FFrame3d FGroupTopology::GetSelectionFrame(const FGroupTopologySelection& Select
 
 
 
+FAxisAlignedBox3d FGroupTopology::GetSelectionBounds(
+	const FGroupTopologySelection& Selection,
+	TFunctionRef<FVector3d(const FVector3d&)> TransformFunc) const
+{
+	if (ensure(!Selection.IsEmpty()) == false)
+	{
+		return Mesh->GetBounds();
+	}
+
+	FAxisAlignedBox3d Bounds = FAxisAlignedBox3d::Empty();
+
+	for (int32 CornerID : Selection.SelectedCornerIDs)
+	{
+		Bounds.Contain(TransformFunc(Mesh->GetVertex(GetCornerVertexID(CornerID))));
+	}
 
+	for (int32 EdgeID : Selection.SelectedEdgeIDs)
+	{
+		const FGroupEdge& Edge = Edges[EdgeID];
+		for (int32 vid : Edge.Span.Vertices)
+		{
+			Bounds.Contain(TransformFunc(Mesh->GetVertex(vid)));
+		}
+	}
+
+	for (int32 GroupID : Selection.SelectedGroupIDs)
+	{
+		for (int32 TriangleID : GetGroupTriangles(GroupID))
+		{
+			FIndex3i TriVerts = Mesh->GetTriangle(TriangleID);
+			Bounds.Contain(TransformFunc(Mesh->GetVertex(TriVerts.A)));
+			Bounds.Contain(TransformFunc(Mesh->GetVertex(TriVerts.B)));
+			Bounds.Contain(TransformFunc(Mesh->GetVertex(TriVerts.C)));
+		}
+	}
+
+	return Bounds;
+}
 
 void FGroupTopology::GetSelectedTriangles(const FGroupTopologySelection& Selection, TArray<int32>& Triangles) const
 {
diff --git a/Engine/Plugins/Experimental/GeometryProcessing/Source/DynamicMesh/Public/GroupTopology.h b/Engine/Plugins/Experimental/GeometryProcessing/Source/DynamicMesh/Public/GroupTopology.h
--- a/Engine/Plugins/Experimental/GeometryProcessing/Source/DynamicMesh/Public/GroupTopology.h
+++ b/Engine/Plugins/Experimental/GeometryProcessing/Source/DynamicMesh/Public/GroupTopology.h
@@ -379,6 +379,12 @@ public:
 	 */
 	FFrame3d GetSelectionFrame(const FGroupTopologySelection& Selection, FFrame3d* InitialLocalFrame = nullptr) const;
 
+	/**
+	 * @return A Bounding Box for the vertices of the current selection, with TransformFunc applied
+	 */
+	FAxisAlignedBox3d GetSelectionBounds(const FGroupTopologySelection& Selection, TFunctionRef<FVector3d(const FVector3d&)> TransformFunc) const;
+
+
 	/**
 	 * Get the set of selected triangles for a given GroupTopologySelection
 	 */
diff --git a/Engine/Plugins/Experimental/MeshModelingToolset/Source/ModelingComponents/Private/Selection/PolygonSelectionMechanic.cpp b/Engine/Plugins/Experimental/MeshModelingToolset/Source/ModelingComponents/Private/Selection/PolygonSelectionMechanic.cpp
--- a/Engine/Plugins/Experimental/MeshModelingToolset/Source/ModelingComponents/Private/Selection/PolygonSelectionMechanic.cpp
+++ b/Engine/Plugins/Experimental/MeshModelingToolset/Source/ModelingComponents/Private/Selection/PolygonSelectionMechanic.cpp
@@ -695,6 +695,25 @@ FFrame3d UPolygonSelectionMechanic::GetSelectionFrame(bool bWorld, FFrame3d* Ini
 	return UseFrame;
 }
 
+FAxisAlignedBox3d UPolygonSelectionMechanic::GetSelectionBounds(bool bWorld) const
+{
+	if ( ! PersistentSelection.IsEmpty() )
+	{
+		if (bWorld)
+		{
+			return Topology->GetSelectionBounds(PersistentSelection, [this](const FVector3d& Pos) { return TargetTransform.TransformPosition(Pos); });
+		}
+		else
+		{
+			return Topology->GetSelectionBounds(PersistentSelection, [this](const FVector3d& Pos) { return Pos; });
+		}
+	}
+	else 
+	{
+		return Mesh->GetBounds();
+	}
+}
+
 
 
 void FPolygonSelectionMechanicSelectionChange::Apply(UObject* Object)
diff --git a/Engine/Plugins/Experimental/MeshModelingToolset/Source/ModelingComponents/Public/Selection/PolygonSelectionMechanic.h b/Engine/Plugins/Experimental/MeshModelingToolset/Source/ModelingComponents/Public/Selection/PolygonSelectionMechanic.h
--- a/Engine/Plugins/Experimental/MeshModelingToolset/Source/ModelingComponents/Public/Selection/PolygonSelectionMechanic.h
+++ b/Engine/Plugins/Experimental/MeshModelingToolset/Source/ModelingComponents/Public/Selection/PolygonSelectionMechanic.h
@@ -83,6 +83,7 @@ class MODELINGCOMPONENTS_API UPolygonSelectionMechanic : public UInteractionMech
 	GENERATED_BODY()
 	using FFrame3d = UE::Geometry::FFrame3d;
 	using FTransform3d = UE::Geometry::FTransform3d;	
+	using FAxisAlignedBox3d = UE::Geometry::FAxisAlignedBox3d;
 public:
 
 	virtual ~UPolygonSelectionMechanic();
@@ -263,6 +264,12 @@ public:
 	 */
 	FFrame3d GetSelectionFrame(bool bWorld, FFrame3d* InitialLocalFrame = nullptr) const;
 
+	/**
+	 * @return Bounding box for the current selection
+	 * @param bWorld if true, the box is in world space, otherwise it is in local space of the target MeshComponent
+	 */
+	FAxisAlignedBox3d GetSelectionBounds(bool bWorld) const;
+
 	//
 	// Change Tracking
 	//
