GeometryProcessing: Add option to do explicit geometric tolerance/deviation checking in TMeshSimplification. Checks edge midpoints and face centroids against a fixed tolerance.
MeshModelingToolset: expose new option in SimplifyMeshOp and SimplifyMeshTool
#rnx

diff --git a/Engine/Plugins/Experimental/GeometryProcessing/Source/DynamicMesh/Private/MeshSimplification.cpp b/Engine/Plugins/Experimental/GeometryProcessing/Source/DynamicMesh/Private/MeshSimplification.cpp
--- a/Engine/Plugins/Experimental/GeometryProcessing/Source/DynamicMesh/Private/MeshSimplification.cpp
+++ b/Engine/Plugins/Experimental/GeometryProcessing/Source/DynamicMesh/Private/MeshSimplification.cpp
@@ -1328,8 +1328,8 @@ ESimplificationResult TMeshSimplification<QuadricErrorType>::CollapseEdge(int ed
 		
 		// [TODO] maybe skip Projection call when !bModeAllowsVertMovment.
 		vNewPos = GetProjectedCollapsePosition(iKeep, vNewPos);
-		double div = Distance(vA, vB);
-		collapse_t = (div < FMathd::ZeroTolerance) ? 0.5 : (Distance(vNewPos, Mesh->GetVertex(iKeep))) / div;
+		double EdgeLength = Distance(vA, vB);
+		collapse_t = (EdgeLength < FMathd::ZeroTolerance) ? 0.5 : (Distance(vNewPos, Mesh->GetVertex(iKeep))) / EdgeLength;
 		collapse_t = VectorUtil::Clamp(collapse_t, 0.0, 1.0);
 
 		if (bMinimalVertexMode)
@@ -1338,6 +1338,24 @@ ESimplificationResult TMeshSimplification<QuadricErrorType>::CollapseEdge(int ed
 			// round here to make sure later attribute lerps don't change values
 			collapse_t = FMath::RoundToDouble(collapse_t);
 		}
+
+		// If vertex is not explicitly constrained, and geometric error constraint is requested, we check it here.
+		// If the check with the predicted vNewPos fails, try a second time with the linear-interpolated point.
+		// (could optionally do a line search here, and/or be smarter about avoiding duplicate work, although
+		//  it is small in the context of the larger algorithm)
+		if (GeometricErrorConstraint != EGeometricErrorCriteria::None )
+		{
+			if (CheckIfCollapseWithinGeometricTolerance(iKeep, iCollapse, vNewPos, t0, t1) == false)
+			{
+				// project new position back onto the edge
+				vNewPos = (1.0-collapse_t)*Mesh->GetVertex(iKeep) + (collapse_t)*Mesh->GetVertex(iCollapse);
+				if (CheckIfCollapseWithinGeometricTolerance(iKeep, iCollapse, vNewPos, t0, t1) == false)
+				{
+					ProfileEndCollapse();
+					return ESimplificationResult::Failed_GeometricDeviation;
+				}
+			}
+		}
 	}
 	else
 	{
@@ -1441,6 +1459,80 @@ ESimplificationResult TMeshSimplification<QuadricErrorType>::CollapseEdge(int ed
 	return retVal;
 }
 
+template <typename QuadricErrorType>
+bool TMeshSimplification<QuadricErrorType>::CheckIfCollapseWithinGeometricTolerance(int vKeep, int vRemove, const FVector3d& NewPosition, int tc, int td)
+{
+	if (GeometricErrorConstraint == EGeometricErrorCriteria::PredictedPointToProjectionTarget)
+	{
+		// currently assuming projection target is what we want to measure geometric error against
+		if (ProjectionTarget() != nullptr)
+		{
+			double ToleranceSqr = GeometricErrorTolerance * GeometricErrorTolerance;
+
+			// test new position to see if it is within geometric tolerance of projection surface
+			FVector3d TargetPos = ProjectionTarget()->Project(NewPosition);
+			double DistSqr = DistanceSquared(TargetPos, NewPosition);
+			if (DistSqr > ToleranceSqr)
+			{
+				return false;
+			}
+
+			// test edge midpoints, except the edge being collapsed
+			int32 CollapseEdgeID = Mesh->FindEdge(vKeep, vRemove);
+			auto EdgeMidpointsWithinTolerance = [this, CollapseEdgeID, NewPosition, ToleranceSqr](int32 vid) {
+				for ( int32 eid : Mesh->VtxEdgesItr(vid) )
+				{
+					if (eid != CollapseEdgeID)
+					{
+						FIndex2i EdgeV = Mesh->GetEdgeV(eid);
+						FVector3d OtherVertexPos = (EdgeV.A == vid) ? Mesh->GetVertex(EdgeV.B) : Mesh->GetVertex(EdgeV.A);
+						FVector3d NewMidpoint = (OtherVertexPos + NewPosition) * 0.5;
+						FVector3d MidpointTargetPos = ProjectionTarget()->Project(NewMidpoint);
+						if (DistanceSquared(NewMidpoint, MidpointTargetPos) > ToleranceSqr)
+						{
+							return false;
+						}
+					}
+				}
+				return true;
+			};
+			if (EdgeMidpointsWithinTolerance(vKeep) == false || EdgeMidpointsWithinTolerance(vRemove) == false)
+			{
+				return false;
+			}
+
+
+			// check tri centers, except the triangles being collapsed
+			auto CentroidsWithinToleranceFunc = [this, vKeep, vRemove, tc, td, NewPosition, ToleranceSqr](int32 vid) {
+				bool bInTolerance = true;
+				Mesh->EnumerateVertexTriangles(vid, [&](int32 tid)
+					{
+						if (bInTolerance && tid != tc && tid != td)
+						{
+							FIndex3i Tri = Mesh->GetTriangle(tid);
+							FVector3d NewCentroid = FVector3d::Zero();
+							for (int32 j = 0; j < 3; ++j)
+							{
+								NewCentroid += (Tri[j] == vRemove || Tri[j] == vKeep) ? NewPosition : Mesh->GetVertex(Tri[j]);
+							}
+							NewCentroid *= (1.0 / 3.0);
+							FVector3d CentroidTargetPos = ProjectionTarget()->Project(NewCentroid);
+							if (DistanceSquared(NewCentroid, CentroidTargetPos) > ToleranceSqr)
+							{
+								bInTolerance = false;
+							}
+						}
+					});
+				return bInTolerance;
+			};
+			if (CentroidsWithinToleranceFunc(vKeep) == false || CentroidsWithinToleranceFunc(vRemove) == false)
+			{
+				return false;
+			}
+		}
+	}
+	return true;
+}
 
 template <typename QuadricErrorType>
 bool TMeshSimplification<QuadricErrorType>::RemoveIsolatedTriangle(int tID)
diff --git a/Engine/Plugins/Experimental/GeometryProcessing/Source/DynamicMesh/Public/MeshSimplification.h b/Engine/Plugins/Experimental/GeometryProcessing/Source/DynamicMesh/Public/MeshSimplification.h
--- a/Engine/Plugins/Experimental/GeometryProcessing/Source/DynamicMesh/Public/MeshSimplification.h
+++ b/Engine/Plugins/Experimental/GeometryProcessing/Source/DynamicMesh/Public/MeshSimplification.h
@@ -26,7 +26,8 @@ enum class ESimplificationResult
 	Ignored_CreatesFlip = 5,
 	Failed_OpNotSuccessful = 6,
 	Failed_NotAnEdge = 7,
-	Failed_IsolatedTriangle = 8
+	Failed_IsolatedTriangle = 8,
+	Failed_GeometricDeviation = 9
 };
 
 /**
@@ -78,6 +79,23 @@ public:
 	/** When using the constraint system, these options will apply to the appropriate boundaries. */
 	EEdgeRefineFlags MeshBoundaryConstraint, GroupBoundaryConstraint, MaterialBoundaryConstraint;
 
+
+	/** Ways to measure geometric error */
+	enum class EGeometricErrorCriteria
+	{
+		/** No geometric error checking */
+		None = 0,
+		/** Fixed envelope around projection target */
+		PredictedPointToProjectionTarget = 1
+	};
+
+	/** Geometric error measurement/check to perform before allowing an edge collapse */
+	EGeometricErrorCriteria GeometricErrorConstraint = EGeometricErrorCriteria::None;
+
+	/** Tolerance to use in geometric error checking */
+	double GeometricErrorTolerance = 0.0;
+
+
 	TMeshSimplification(FDynamicMesh3* m) : FMeshRefinerBase(m)
 	{
 		NormalOverlay = nullptr;
@@ -324,6 +342,16 @@ protected:
 	virtual void ApplyToProjectVertices(const TFunction<void(int)>& apply_f);
 
 
+	/**
+	 * Check if edge collapse would violate geometric error criteria
+	 * @param vid first vertex of edge
+	 * @param vother other vertex of edge
+	 * @param newv new vertex position after collapse
+	 * @param tc triangle on one side of edge
+	 * @param td triangle on other side of edge
+	 */
+	bool CheckIfCollapseWithinGeometricTolerance(int vid, int vother, const FVector3d& newv, int tc, int td);
+
 
 	/*
 	 * testing/debug/profiling stuff
diff --git a/Engine/Plugins/Experimental/MeshModelingToolset/Source/MeshModelingToolsEditorOnly/Private/SimplifyMeshTool.cpp b/Engine/Plugins/Experimental/MeshModelingToolset/Source/MeshModelingToolsEditorOnly/Private/SimplifyMeshTool.cpp
--- a/Engine/Plugins/Experimental/MeshModelingToolset/Source/MeshModelingToolsEditorOnly/Private/SimplifyMeshTool.cpp
+++ b/Engine/Plugins/Experimental/MeshModelingToolset/Source/MeshModelingToolsEditorOnly/Private/SimplifyMeshTool.cpp
@@ -63,6 +63,7 @@ USimplifyMeshToolProperties::USimplifyMeshToolProperties()
 	bReproject = false;
 	bPreventNormalFlips = true;
 	bDiscardAttributes = false;
+	bGeometricConstraint = false;
 	bShowWireframe = true;
 	bShowGroupColors = false;
 	GroupBoundaryConstraint = EGroupBoundaryConstraint::Ignore;
@@ -196,6 +197,8 @@ TUniquePtr<FDynamicMeshOperator> USimplifyMeshTool::MakeNewOperator()
 	Op->MeshBoundaryConstraint = (EEdgeRefineFlags)SimplifyProperties->MeshBoundaryConstraint;
 	Op->GroupBoundaryConstraint = (EEdgeRefineFlags)SimplifyProperties->GroupBoundaryConstraint;
 	Op->MaterialBoundaryConstraint = (EEdgeRefineFlags)SimplifyProperties->MaterialBoundaryConstraint;
+	Op->bGeometricDeviationConstraint = SimplifyProperties->bGeometricConstraint;
+	Op->GeometricTolerance = SimplifyProperties->GeometricTolerance;
 	FTransform LocalToWorld = Cast<IPrimitiveComponentBackedTarget>(Target)->GetWorldTransform();
 	Op->SetTransform(LocalToWorld);
 
diff --git a/Engine/Plugins/Experimental/MeshModelingToolset/Source/MeshModelingToolsEditorOnly/Public/SimplifyMeshTool.h b/Engine/Plugins/Experimental/MeshModelingToolset/Source/MeshModelingToolsEditorOnly/Public/SimplifyMeshTool.h
--- a/Engine/Plugins/Experimental/MeshModelingToolset/Source/MeshModelingToolsEditorOnly/Public/SimplifyMeshTool.h
+++ b/Engine/Plugins/Experimental/MeshModelingToolset/Source/MeshModelingToolsEditorOnly/Public/SimplifyMeshTool.h
@@ -67,6 +67,14 @@ public:
 	UPROPERTY(EditAnywhere, Category = Options)
 	bool bDiscardAttributes;
 
+	/** If true, then simplification will consider geometric deviation with the input mesh  */
+	UPROPERTY(EditAnywhere, Category = Options)
+	bool bGeometricConstraint;
+
+	/** Geometric deviation tolerance used when bGeometricConstraint is enabled, to limit the geometric deviation between the simplified and original meshes */
+	UPROPERTY(EditAnywhere, Category = Options, meta = (UIMin = "0.0", UIMax = "10.0", ClampMin = "0.0", ClampMax = "10000000.0", EditCondition = "bGeometricConstraint && SimplifierType != ESimplifyType::UEStandard"))
+	float GeometricTolerance;
+
 	/** If true, display wireframe */
 	UPROPERTY(EditAnywhere, Category = Display)
 	bool bShowWireframe;
diff --git a/Engine/Plugins/Experimental/MeshModelingToolset/Source/ModelingOperatorsEditorOnly/Private/CleaningOps/SimplifyMeshOp.cpp b/Engine/Plugins/Experimental/MeshModelingToolset/Source/ModelingOperatorsEditorOnly/Private/CleaningOps/SimplifyMeshOp.cpp
--- a/Engine/Plugins/Experimental/MeshModelingToolset/Source/ModelingOperatorsEditorOnly/Private/CleaningOps/SimplifyMeshOp.cpp
+++ b/Engine/Plugins/Experimental/MeshModelingToolset/Source/ModelingOperatorsEditorOnly/Private/CleaningOps/SimplifyMeshOp.cpp
@@ -34,7 +34,8 @@ void ComputeSimplify(FDynamicMesh3* TargetMesh, const bool bReproject,
 					 const float TargetPercentage, const int TargetCount, const float TargetEdgeLength,
 					 const float AngleThreshold,
 	                 typename SimplificationType::ESimplificationCollapseModes CollapseMode,
-					 bool bUseQuadricMemory)
+					 bool bUseQuadricMemory,
+					 float GeometricTolerance )
 {
 	SimplificationType Reducer(TargetMesh);
 
@@ -71,17 +72,25 @@ void ComputeSimplify(FDynamicMesh3* TargetMesh, const bool bReproject,
 	Reducer.GroupBoundaryConstraint = GroupBoundaryConstraint;
 	Reducer.MaterialBoundaryConstraint = MaterialBoundaryConstraint;
 	
+	// MinimalPlanar is a special path that ignores the collapse mode (and quadric error ), so don't bother setting it.
+	if (TargetMode != ESimplifyTargetType::MinimalPlanar)
+	{
+		Reducer.CollapseMode = CollapseMode;
+		GeometricTolerance = 0;		// MinimalPlanar does not allow vertices to move off the input surface
+	}
 
-	FMeshProjectionTarget  ProjTarget(&OriginalMesh, &OriginalMeshSpatial);
-	if (bReproject)
+	// use projection target if we are reprojecting or doing geometric error checking
+	FMeshProjectionTarget ProjTarget(&OriginalMesh, &OriginalMeshSpatial);
+	if (bReproject || GeometricTolerance > 0)
 	{
 		Reducer.SetProjectionTarget(&ProjTarget);
 	}
 
-	// MinimalPlanar is a special path that ignores the collapse mode (and quadric error ), so don't bother setting it.
-	if (TargetMode != ESimplifyTargetType::MinimalPlanar)
+	// configure geometric error settings
+	if (GeometricTolerance > 0)
 	{
-		Reducer.CollapseMode = CollapseMode;
+		Reducer.GeometricErrorConstraint = SimplificationType::EGeometricErrorCriteria::PredictedPointToProjectionTarget;
+		Reducer.GeometricErrorTolerance = GeometricTolerance;
 	}
 
 	if (TargetMode == ESimplifyTargetType::Percentage)
@@ -127,6 +136,7 @@ void FSimplifyMeshOp::CalculateResult(FProgressCancel* Progress)
 	}
 
 	int OriginalTriCount = OriginalMesh->TriangleCount();
+	double UseGeometricTolerance = (bGeometricDeviationConstraint) ? GeometricTolerance : 0.0;
 	if (SimplifierType == ESimplifyType::QEM)
 	{
 		bool bUseQuadricMemory = true;
@@ -137,7 +147,8 @@ void FSimplifyMeshOp::CalculateResult(FProgressCancel* Progress)
 											MaterialBoundaryConstraint,
 											bPreserveSharpEdges, bAllowSeamCollapse,
 											TargetMode, TargetPercentage, TargetCount, TargetEdgeLength, MinimalPlanarAngleThresh,
-											FQEMSimplification::ESimplificationCollapseModes::MinimalQuadricPositionError, bUseQuadricMemory);
+											FQEMSimplification::ESimplificationCollapseModes::MinimalQuadricPositionError, bUseQuadricMemory, 
+											UseGeometricTolerance);
 	}
 	else if (SimplifierType == ESimplifyType::Attribute)
 	{
@@ -153,7 +164,8 @@ void FSimplifyMeshOp::CalculateResult(FProgressCancel* Progress)
 													MaterialBoundaryConstraint,
 													bPreserveSharpEdges, bAllowSeamCollapse,
 													TargetMode, TargetPercentage, TargetCount, TargetEdgeLength, MinimalPlanarAngleThresh,
-													FAttrMeshSimplification::ESimplificationCollapseModes::MinimalQuadricPositionError, bUseQuadricMemory);
+													FAttrMeshSimplification::ESimplificationCollapseModes::MinimalQuadricPositionError, bUseQuadricMemory,
+													UseGeometricTolerance);
 	}
 	else if (SimplifierType == ESimplifyType::MinimalPlanar)
 	{
@@ -169,7 +181,8 @@ void FSimplifyMeshOp::CalculateResult(FProgressCancel* Progress)
 			MaterialBoundaryConstraint,
 			bPreserveSharpEdges, bAllowSeamCollapse,
 			ESimplifyTargetType::MinimalPlanar, TargetPercentage, TargetCount, TargetEdgeLength, MinimalPlanarAngleThresh,
-			FAttrMeshSimplification::ESimplificationCollapseModes::MinimalQuadricPositionError, bUseQuadricMemory);
+			FAttrMeshSimplification::ESimplificationCollapseModes::MinimalQuadricPositionError, bUseQuadricMemory,
+			UseGeometricTolerance);
 	}
 	else if (SimplifierType == ESimplifyType::MinimalExistingVertex)
 	{
@@ -181,7 +194,8 @@ void FSimplifyMeshOp::CalculateResult(FProgressCancel* Progress)
 			MaterialBoundaryConstraint,
 			bPreserveSharpEdges, bAllowSeamCollapse,
 			TargetMode, TargetPercentage, TargetCount, TargetEdgeLength, MinimalPlanarAngleThresh, 
-			FQEMSimplification::ESimplificationCollapseModes::MinimalExistingVertexError, bUseQuadricMemory);
+			FQEMSimplification::ESimplificationCollapseModes::MinimalExistingVertexError, bUseQuadricMemory,
+			UseGeometricTolerance);
 	}
 	else // SimplifierType == ESimplifyType::UEStandard
 	{
diff --git a/Engine/Plugins/Experimental/MeshModelingToolset/Source/ModelingOperatorsEditorOnly/Public/CleaningOps/SimplifyMeshOp.h b/Engine/Plugins/Experimental/MeshModelingToolset/Source/ModelingOperatorsEditorOnly/Public/CleaningOps/SimplifyMeshOp.h
--- a/Engine/Plugins/Experimental/MeshModelingToolset/Source/ModelingOperatorsEditorOnly/Public/CleaningOps/SimplifyMeshOp.h
+++ b/Engine/Plugins/Experimental/MeshModelingToolset/Source/ModelingOperatorsEditorOnly/Public/CleaningOps/SimplifyMeshOp.h
@@ -76,6 +76,9 @@ public:
 	/** Angle threshold in degrees used for testing if two triangles should be considered coplanar, or two lines collinear */
 	float MinimalPlanarAngleThresh = 0.01f;
 
+	bool bGeometricDeviationConstraint = false;
+	float GeometricTolerance = 0.0f;
+
 	// stored for the UEStandard path
 	TSharedPtr<FMeshDescription, ESPMode::ThreadSafe> OriginalMeshDescription;
 	// stored for the GeometryProcessing custom simplifier paths (currently precomputed once in tool setup)
