GeometryProcessing:
Add FDynamicMeshEditor::AppendMesh() variant that takes a TTriangleMeshAdapter<double>
Add FMeshDescriptionMeshAdapterd that wraps FMeshDescriptionTriangleMeshAdapter in TTriangleMeshAdapter<double> API
Add TTransform3::HasNonUniformScale()
#rnx

diff --git a/Engine/Plugins/Experimental/GeometryProcessing/Source/DynamicMesh/Private/DynamicMeshEditor.cpp b/Engine/Plugins/Experimental/GeometryProcessing/Source/DynamicMesh/Private/DynamicMeshEditor.cpp
--- a/Engine/Plugins/Experimental/GeometryProcessing/Source/DynamicMesh/Private/DynamicMeshEditor.cpp
+++ b/Engine/Plugins/Experimental/GeometryProcessing/Source/DynamicMesh/Private/DynamicMeshEditor.cpp
@@ -1590,6 +1590,64 @@ void FDynamicMeshEditor::AppendMesh(const FDynamicMesh3* AppendMesh,
 
 
 
+
+void FDynamicMeshEditor::AppendMesh(const TTriangleMeshAdapter<double>* AppendMesh,
+	FMeshIndexMappings& IndexMapsOut, 
+	TFunction<FVector3d(int, const FVector3d&)> PositionTransform)
+{
+	IndexMapsOut.Reset();
+	//IndexMapsOut.Initialize(Mesh);		// not supported
+
+	FIndexMapi& VertexMap = IndexMapsOut.GetVertexMap();
+	VertexMap.Reserve(AppendMesh->VertexCount());
+	int32 MaxVertexID = AppendMesh->MaxVertexID();
+	for ( int32 VertID = 0; VertID < MaxVertexID; ++VertID )
+	{
+		if (AppendMesh->IsVertex(VertID) == false) continue;
+
+		FVector3d Position = AppendMesh->GetVertex(VertID);
+		if (PositionTransform != nullptr)
+		{
+			Position = PositionTransform(VertID, Position);
+		}
+		int NewVertID = Mesh->AppendVertex(Position);
+		VertexMap.Add(VertID, NewVertID);
+	}
+
+	// set face normals if they exist
+	FDynamicMeshNormalOverlay* SetNormals = Mesh->HasAttributes() ? Mesh->Attributes()->PrimaryNormals() : nullptr;
+
+	FIndexMapi& TriangleMap = IndexMapsOut.GetTriangleMap();
+	int32 MaxTriangleID = AppendMesh->MaxTriangleID();
+	for (int TriID = 0; TriID < MaxTriangleID; ++TriID )
+	{
+		if (AppendMesh->IsTriangle(TriID) == false) continue;
+
+		int GroupID = FDynamicMesh3::InvalidID;
+		FIndex3i Tri = AppendMesh->GetTriangle(TriID);
+		FIndex3i NewTri = FIndex3i(VertexMap.GetTo(Tri.A), VertexMap.GetTo(Tri.B), VertexMap.GetTo(Tri.C));
+		int NewTriID = Mesh->AppendTriangle(NewTri.A, NewTri.B, NewTri.C, GroupID);
+		TriangleMap.Add(TriID, NewTriID);
+
+		if (SetNormals)
+		{
+			FVector3f TriNormal = Mesh->GetTriNormal(NewTriID);
+			FIndex3i NormalTri;
+			for (int32 j = 0; j < 3; ++j)
+			{
+				NormalTri[j] = SetNormals->AppendElement(TriNormal);
+				SetNormals->SetParentVertex(NormalTri[j], NewTri[j]);
+			}
+			SetNormals->SetTriangle(NewTriID, NormalTri);
+		}
+	}
+}
+
+
+
+
+
+
 void FDynamicMeshEditor::AppendNormals(const FDynamicMesh3* AppendMesh, 
 	const FDynamicMeshNormalOverlay* FromNormals, FDynamicMeshNormalOverlay* ToNormals,
 	const FIndexMapi& VertexMap, const FIndexMapi& TriangleMap,
diff --git a/Engine/Plugins/Experimental/GeometryProcessing/Source/DynamicMesh/Public/DynamicMeshEditor.h b/Engine/Plugins/Experimental/GeometryProcessing/Source/DynamicMesh/Public/DynamicMeshEditor.h
--- a/Engine/Plugins/Experimental/GeometryProcessing/Source/DynamicMesh/Public/DynamicMeshEditor.h
+++ b/Engine/Plugins/Experimental/GeometryProcessing/Source/DynamicMesh/Public/DynamicMeshEditor.h
@@ -9,6 +9,7 @@
 #include "EdgeLoop.h"
 #include "Util/SparseIndexCollectionTypes.h"
 #include "MeshIndexMappings.h"
+#include "MeshAdapter.h"
 
 namespace UE
 {
@@ -444,6 +445,16 @@ public:
 		TFunction<FVector3d(int, const FVector3d&)> PositionTransform = nullptr,
 		TFunction<FVector3d(int, const FVector3d&)> NormalTransform = nullptr);
 
+	/**
+	 * Append input mesh to our internal Mesh. If the internal Mesh has attributes enabled,
+	 * per-triangle normals will be computed and set. No other attributes are initialized.
+	 * @param AppendMesh mesh to append
+	 * @param IndexMapsOut mesh element index mappings generated in this append operation
+	 * @param PositionTransform optional transformation function applied to mesh vertex positions
+	 */
+	void AppendMesh(const TTriangleMeshAdapter<double>* AppendMesh, FMeshIndexMappings& IndexMapsOut,
+		TFunction<FVector3d(int, const FVector3d&)> PositionTransform = nullptr);
+
 
 	/**
 	 * Append normals from one attribute overlay to another.
diff --git a/Engine/Plugins/Experimental/GeometryProcessing/Source/DynamicMesh/Public/Operations/OffsetMeshRegion.h b/Engine/Plugins/Experimental/GeometryProcessing/Source/DynamicMesh/Public/Operations/OffsetMeshRegion.h
--- a/Engine/Plugins/Experimental/GeometryProcessing/Source/DynamicMesh/Public/Operations/OffsetMeshRegion.h
+++ b/Engine/Plugins/Experimental/GeometryProcessing/Source/DynamicMesh/Public/Operations/OffsetMeshRegion.h
@@ -6,6 +6,7 @@
 #include "VectorTypes.h"
 #include "GeometryTypes.h"
 #include "MeshRegionBoundaryLoops.h"
+#include "DynamicMeshChangeTracker.h"
 
 
 namespace UE
diff --git a/Engine/Plugins/Experimental/GeometryProcessing/Source/GeometricObjects/Public/TransformTypes.h b/Engine/Plugins/Experimental/GeometryProcessing/Source/GeometricObjects/Public/TransformTypes.h
--- a/Engine/Plugins/Experimental/GeometryProcessing/Source/GeometricObjects/Public/TransformTypes.h
+++ b/Engine/Plugins/Experimental/GeometryProcessing/Source/GeometricObjects/Public/TransformTypes.h
@@ -124,11 +124,24 @@ public:
 		Scale3D = ScaleIn;
 	}
 
-	RealType GetDeterminant()
+	/**
+	 * @return determinant of scale matrix
+	 */
+	RealType GetDeterminant() const
 	{
 		return Scale3D.X * Scale3D.Y * Scale3D.Z;
 	}
 
+	/**
+	 * @return true if scale is nonuniform, within tolerance
+	 */
+	bool HasNonUniformScale(RealType Tolerance = TMathUtil<RealType>::ZeroTolerance) const
+	{
+		return (TMathUtil<RealType>::Abs(Scale3D.X - Scale3D.Y) > Tolerance)
+			|| (TMathUtil<RealType>::Abs(Scale3D.X - Scale3D.Z) > Tolerance)
+			|| (TMathUtil<RealType>::Abs(Scale3D.Y - Scale3D.Z) > Tolerance);
+	}
+
 	/**
 	 * @return inverse of this transform
 	 */
diff --git a/Engine/Plugins/Experimental/GeometryProcessing/Source/MeshConversion/Public/MeshDescriptionAdapter.h b/Engine/Plugins/Experimental/GeometryProcessing/Source/MeshConversion/Public/MeshDescriptionAdapter.h
--- a/Engine/Plugins/Experimental/GeometryProcessing/Source/MeshConversion/Public/MeshDescriptionAdapter.h
+++ b/Engine/Plugins/Experimental/GeometryProcessing/Source/MeshConversion/Public/MeshDescriptionAdapter.h
@@ -6,6 +6,7 @@
 #include "MeshDescription.h"
 #include "StaticMeshAttributes.h"
 #include "Spatial/MeshAABBTree3.h"
+#include "MeshAdapter.h"
 
 
 /**
@@ -233,3 +234,39 @@ public:
 	}
 };
 
+
+
+/**
+ * TTriangleMeshAdapter version of FMeshDescriptionTriangleMeshAdapter
+ */
+struct FMeshDescriptionMeshAdapterd : public UE::Geometry::TTriangleMeshAdapter<double>
+{
+	FMeshDescriptionTriangleMeshAdapter ParentAdapter;
+
+	FMeshDescriptionMeshAdapterd(const FMeshDescription* MeshIn) : ParentAdapter(MeshIn)
+	{
+		IsTriangle = [&](int index) { return ParentAdapter.IsTriangle(index);};
+		IsVertex = [&](int index) { return ParentAdapter.IsVertex(index); };
+		MaxTriangleID = [&]() { return ParentAdapter.MaxTriangleID();};
+		MaxVertexID = [&]() { return ParentAdapter.MaxVertexID();};
+		TriangleCount = [&]() { return ParentAdapter.TriangleCount();};
+		VertexCount = [&]() { return ParentAdapter.VertexCount();};
+		GetShapeTimestamp = [&]() { return ParentAdapter.GetShapeTimestamp();};
+		GetTriangle = [&](int32 TriangleID) { return ParentAdapter.GetTriangle(TriangleID); };
+		GetVertex = [&](int32 VertexID) { return ParentAdapter.GetVertex(VertexID); };
+	}
+
+	FMeshDescriptionMeshAdapterd(FMeshDescriptionTriangleMeshAdapter ParentAdapterIn) : ParentAdapter(ParentAdapterIn)
+	{
+		IsTriangle = [&](int index) { return ParentAdapter.IsTriangle(index);};
+		IsVertex = [&](int index) { return ParentAdapter.IsVertex(index); };
+		MaxTriangleID = [&]() { return ParentAdapter.MaxTriangleID();};
+		MaxVertexID = [&]() { return ParentAdapter.MaxVertexID();};
+		TriangleCount = [&]() { return ParentAdapter.TriangleCount();};
+		VertexCount = [&]() { return ParentAdapter.VertexCount();};
+		GetShapeTimestamp = [&]() { return ParentAdapter.GetShapeTimestamp();};
+		GetTriangle = [&](int32 TriangleID) { return ParentAdapter.GetTriangle(TriangleID); };
+		GetVertex = [&](int32 VertexID) { return ParentAdapter.GetVertex(VertexID); };
+	}
+
+};
