GeometryProcessing: optimize SweepingMeshSDF make_level_set3_parallel_floodfill() function. Replaced single full-grid lock with per-chunk locks to reduce lock contention. Also made per-chunk queues that are flattened into a single queue for the next iteration after each pass. Reduced cost of a large-distance morphology operation from 49s to 12s on a 64-core machine (where 7s of the remaining time is mesh generation). Instrument Morphology and SweepingMeshSDF with profiling macros.
#rnx

diff --git a/Engine/Plugins/Experimental/GeometryProcessing/Source/GeometricObjects/Public/Implicit/Morphology.h b/Engine/Plugins/Experimental/GeometryProcessing/Source/GeometricObjects/Public/Implicit/Morphology.h
--- a/Engine/Plugins/Experimental/GeometryProcessing/Source/GeometricObjects/Public/Implicit/Morphology.h
+++ b/Engine/Plugins/Experimental/GeometryProcessing/Source/GeometricObjects/Public/Implicit/Morphology.h
@@ -121,6 +121,8 @@ public:
 protected:
 	void ComputeFirstPass(double UnsignedOffset, double SignedOffset)
 	{
+		TRACE_CPUPROFILER_EVENT_SCOPE(Geometry_Morphology_FirstPass);
+
 		ComputedSDF.Mesh = Source;
 
 		ComputedSDF.Spatial = SourceSpatial;
@@ -139,7 +141,10 @@ protected:
 			ComputedSDF.ComputeMode = TSweepingMeshSDF<TriangleMeshType>::EComputeModes::NarrowBandOnly;
 		}
 
-		ComputedSDF.Compute(SourceSpatial->GetBoundingBox());
+		{
+			TRACE_CPUPROFILER_EVENT_SCOPE(Geometry_Morphology_FirstPass_ComputeSDF);
+			ComputedSDF.Compute(SourceSpatial->GetBoundingBox());
+		}
 
 		TTriLinearGridInterpolant<TSweepingMeshSDF<TriangleMeshType>> Interpolant = ComputedSDF.MakeInterpolant();
 
@@ -164,7 +169,10 @@ protected:
 		{
 			return -Interpolant.Value(Pt);
 		};
-		MarchingCubes.Generate();
+		{
+			TRACE_CPUPROFILER_EVENT_SCOPE(Geometry_Morphology_FirstPass_GenerateMesh);
+			MarchingCubes.Generate();
+		}
 
 		// TODO: refactor FMarchingCubes to not retain the implicit function, or refactor this function so the implicit function isn't invalid after returning,
 		/// ..... then remove this line
@@ -173,6 +181,8 @@ protected:
 
 	void ComputeSecondPass(double UnsignedOffset, double SignedOffset)
 	{
+		TRACE_CPUPROFILER_EVENT_SCOPE(Geometry_Morphology_SecondPass);
+
 		if (MarchingCubes.Triangles.Num() == 0)
 		{
 			MarchingCubes.Reset();
@@ -211,7 +221,10 @@ protected:
 			return;
 		}
 
-		SecondSDF.Compute(Bounds);
+		{
+			TRACE_CPUPROFILER_EVENT_SCOPE(Geometry_Morphology_SecondPass_ComputeSDF);
+			SecondSDF.Compute(Bounds);
+		}
 		TTriLinearGridInterpolant<TSweepingMeshSDF<TIndexVectorMeshArrayAdapter<FIndex3i, double, FVector3d>>> Interpolant = SecondSDF.MakeInterpolant();
 
 		MarchingCubes.Reset();
@@ -234,7 +247,10 @@ protected:
 			return -Interpolant.Value(Pt);
 		};
 
-		MarchingCubes.Generate();
+		{
+			TRACE_CPUPROFILER_EVENT_SCOPE(Geometry_Morphology_SecondPass_GenerateMesh);
+			MarchingCubes.Generate();
+		}
 
 		// TODO: refactor FMarchingCubes to not retain the implicit function, or refactor this function so the implicit function isn't invalid after returning,
 		/// ..... then remove this line
diff --git a/Engine/Plugins/Experimental/GeometryProcessing/Source/GeometricObjects/Public/Implicit/SweepingMeshSDF.h b/Engine/Plugins/Experimental/GeometryProcessing/Source/GeometricObjects/Public/Implicit/SweepingMeshSDF.h
--- a/Engine/Plugins/Experimental/GeometryProcessing/Source/GeometricObjects/Public/Implicit/SweepingMeshSDF.h
+++ b/Engine/Plugins/Experimental/GeometryProcessing/Source/GeometricObjects/Public/Implicit/SweepingMeshSDF.h
@@ -735,114 +735,169 @@ private:
 		double ox = (double)Origin[0], oy = (double)Origin[1], oz = (double)Origin[2];
 		double invdx = 1.0 / DX;
 
-		// compute values at vertices
+		// the steps below that populate the grid will potentially touch the same cells at the
+		// same time, so we need to lock them. However locking the entire grid for each cell 
+		// basically means the threads are constantly fighting. So we split the grid up into
+		// chunks for locking
+		int32 NumSections = 256;		// a bit arbitrary, stopped seeing improvement on a 64-core machine at this point
+		TArray<FCriticalSection> GridSections;
+		GridSections.SetNum(NumSections + 1);
+		int64 TotalGridCellCount = NI * NJ * NK;
+		int64 SectionSize = TotalGridCellCount / (int)NumSections;
+		// this returns the FCriticalSection to use for the given span of values
+		auto GetGridSectionLock = [this, SectionSize, &Distances, &GridSections](FVector3i CellGridIndex) -> FCriticalSection*
+		{
+			int64 CellLinearIndex = Distances.ToLinear(CellGridIndex);
+			int64 SectionIndex = CellLinearIndex / SectionSize;
+			return &GridSections[SectionIndex];
+		};
+
+		// per-grid-chunk queue, each one of these cooresponds to a GridSections lock
+		TArray<TArray<int>> SectionQueues;
+		SectionQueues.SetNum(NumSections + 1);
+		auto AddToSectionQueue = [SectionSize, &SectionQueues](int64 CellLinearIndex)
+		{
+			int64 SectionIndex = CellLinearIndex / SectionSize;
+			SectionQueues[SectionIndex].Add(CellLinearIndex);
+		};
 
-		FCriticalSection GridSection;
-		TArray<int> Q[2]; // store two queues, an active one to process and a second to accumulate indices into for the next processing pass
-		int ActiveQ = 0;
-		TArray<bool> done; done.SetNumZeroed(Distances.Size());
 
-		bool bAbort = false;
-		ParallelFor(Mesh->MaxVertexID(), [this, DX, NI, NJ, NK, &Distances, &upper_bound, &closest_tri, &intersection_count, &ox, &oy, &oz, &invdx, &GridSection, &Q, &ActiveQ, &done, &bAbort](int vid)
+		TArray<int> PendingCellQueue;
+		TArray<bool> done; 
+		done.SetNumZeroed(Distances.Size());
+
+		// this lambda combines the per-chunk queues into a single queue which we can pass to a ParallelFor
+		auto FlattenSectionQueues = [&PendingCellQueue, &SectionQueues]()
 		{
-			if (!Mesh->IsVertex(vid))
+			PendingCellQueue.Reset();
+			for (TArray<int>& SectionQueue : SectionQueues)
 			{
-				return;
-			}
-			if (vid % 100 == 0)
-			{
-				bAbort = CancelF();
-			}
-			if (bAbort)
-			{
-				return;
+				for (int Value : SectionQueue)
+				{
+					PendingCellQueue.Add(Value);
+				}
+				SectionQueue.Reset();
 			}
+		};
 
-			FVector3d v = Mesh->GetVertex(vid);
-			// real IJK coordinates of v
-			double fi = (v.X - ox) * invdx, fj = (v.Y - oy) * invdx, fk = (v.Z - oz) * invdx;
-			FVector3i Idx(
-				FMath::Clamp((int)fi, 0, NI - 1),
-				FMath::Clamp((int)fj, 0, NJ - 1),
-				FMath::Clamp((int)fk, 0, NK - 1));
-
-			if (Distances[Idx] < upper_bound)
+		// compute values at vertices
+		bool bAbort = false;
+		{
+			TRACE_CPUPROFILER_EVENT_SCOPE(Geometry_SweepingMeshSDF_Distances);
+			ParallelFor(Mesh->MaxVertexID(), [&](int vid)
 			{
-				return;
-			}
+				if (!Mesh->IsVertex(vid))
+				{
+					return;
+				}
+				if (vid % 100 == 0)
+				{
+					bAbort = CancelF();
+				}
+				if (bAbort)
+				{
+					return;
+				}
 
-			{
-				FScopeLock GridLock(&GridSection);
+				FVector3d v = Mesh->GetVertex(vid);
+				// real IJK coordinates of v
+				double fi = (v.X - ox) * invdx, fj = (v.Y - oy) * invdx, fk = (v.Z - oz) * invdx;
+				FVector3i Idx(
+					FMath::Clamp((int)fi, 0, NI - 1),
+					FMath::Clamp((int)fj, 0, NJ - 1),
+					FMath::Clamp((int)fk, 0, NK - 1));
 
-				FVector3d p = (FVector3d)cell_center(Idx);
-				double dsqr;
-				int near_tid = Spatial->FindNearestTriangle(p, dsqr);
-				Distances[Idx] = (float)FMath::Sqrt(dsqr);
-				closest_tri[Idx] = near_tid;
-				int idx_linear = Distances.ToLinear(Idx);
-				Q[ActiveQ].Add(idx_linear);
-				done[idx_linear] = true;
-			}
-		}, !bUseParallel);
+				{
+					FScopeLock GridLock(GetGridSectionLock(Idx));
+					if (Distances[Idx] < upper_bound)
+					{
+						return;
+					}
+					FVector3d p = (FVector3d)cell_center(Idx);
+					double dsqr;
+					int near_tid = Spatial->FindNearestTriangle(p, dsqr);
+					Distances[Idx] = (float)FMathd::Sqrt(dsqr);
+					closest_tri[Idx] = near_tid;
+					int idx_linear = Distances.ToLinear(Idx);
+					if (done[idx_linear] == false)
+					{
+						done[idx_linear] = true;
+						AddToSectionQueue(idx_linear);
+					}
+				}
+			}, !bUseParallel);
+		}
 		if (CancelF())
 		{
 			return;
 		}
 
+		FlattenSectionQueues();
+
 		// we could do this parallel w/ some kind of producer-consumer...
 		FAxisAlignedBox3i Bounds = Distances.BoundsInclusive();
 		double max_dist = NarrowBandMaxDistance;
 		double max_query_dist = max_dist + (2.0 * DX * FMathd::Sqrt2);
-		int next_pass_count = Q[ActiveQ].Num();
-		while (next_pass_count > 0)
 		{
-			Q[1-ActiveQ].Reset();
-			ParallelFor(Q[ActiveQ].Num(), [this, DX, NI, NJ, NK, &Distances, &upper_bound, &closest_tri, &intersection_count, &ox, &oy, &oz, &invdx, &GridSection, &Q, &ActiveQ, &done, &Bounds, max_dist, max_query_dist](int QIdx)
+			TRACE_CPUPROFILER_EVENT_SCOPE(Geometry_SweepingMeshSDF_FloodFill);
+			while (PendingCellQueue.Num() > 0)
 			{
-				int cur_linear_index = Q[ActiveQ][QIdx];
-				FVector3i cur_idx = Distances.ToIndex(cur_linear_index);
-				for (FVector3i idx_offset : IndexUtil::GridOffsets26) {
-					FVector3i nbr_idx = cur_idx + idx_offset;
-					if (Bounds.Contains(nbr_idx) == false)
-					{
-						continue;
-					}
-					int nbr_linear_idx = Distances.ToLinear(nbr_idx);
-					if (done[nbr_linear_idx])
-					{
-						continue;
-					}
-
-					FVector3d p = (FVector3d)cell_center(nbr_idx);
-					double dsqr;
-					int near_tid = Spatial->FindNearestTriangle(p, dsqr, max_query_dist);
-					if (near_tid == -1)
+				ParallelFor(PendingCellQueue.Num(), [&](int QIdx)
+				{
+					int cur_linear_index = PendingCellQueue[QIdx];
+					FVector3i cur_idx = Distances.ToIndex(cur_linear_index);
+					for (FVector3i idx_offset : IndexUtil::GridOffsets26) 
 					{
-						done[nbr_linear_idx] = true;
-						continue;
-					}
-					double dist = FMath::Sqrt(dsqr);
+						FVector3i nbr_idx = cur_idx + idx_offset;
+						if (Bounds.Contains(nbr_idx) == false)
+						{
+							continue;
+						}
+						int nbr_linear_idx = Distances.ToLinear(nbr_idx);
+
+						// Note: this is technically unsafe to do because other threads might be writing to this memory location
+						// which could in theory mean that the value being read is garbage or stale. However the only possible values are true and false. 
+						// If we get an incorrect false, we will just do an extra unnecessary spatial query, but because we lock before
+						// any writes, that value will just be discarded. If we get an incorrect true, we potentially skip a grid cell
+						// we need. However in this context we are doing a flood-fill and so except at the very edges of the populated cells,
+						// each cell will be considered multiple times as we touch it's neighbours. In practice we have not observed
+						// this to be a problem, however locking here dramatically slows the algorithm down (eg by 3-5x on a 64-core machine)
+						if (done[nbr_linear_idx])
+						{
+							continue;
+						}
 
-					{	// locked section -- if index not already done, update the grid
-						FScopeLock GridLock(&GridSection);
-						if (done[nbr_linear_idx] == false)
+						FVector3d p = (FVector3d)cell_center(nbr_idx);
+						double dsqr;
+						int near_tid = Spatial->FindNearestTriangle(p, dsqr, max_query_dist);
+						if (near_tid == -1)
 						{
-							Distances[nbr_linear_idx] = (float)dist;
-							closest_tri[nbr_linear_idx] = near_tid;
+							FScopeLock GridLock(GetGridSectionLock(nbr_idx));
 							done[nbr_linear_idx] = true;
-							if (dist < max_dist)
+							continue;
+						}
+						double dist = FMathd::Sqrt(dsqr);
+
+						{	
+							// locked section -- if index not already done, update the grid
+							FScopeLock GridLock(GetGridSectionLock(nbr_idx));
+							if (done[nbr_linear_idx] == false)
 							{
-								Q[1-ActiveQ].Add(nbr_linear_idx);
+								Distances[nbr_linear_idx] = (float)dist;
+								closest_tri[nbr_linear_idx] = near_tid;
+								done[nbr_linear_idx] = true;
+								if (dist < max_dist)
+								{
+									AddToSectionQueue(nbr_linear_idx);
+								}
 							}
 						}
 					}
-				}
-			}, !bUseParallel);
-			// swap lists
-			ActiveQ = 1 - ActiveQ;
-			next_pass_count = Q[ActiveQ].Num();
-		}
+				}, !bUseParallel);
 
+				FlattenSectionQueues();
+			}
+		}
 		if (CancelF())
 		{
 			return;
@@ -851,6 +906,7 @@ private:
 
 		if (bComputeSigns == true)
 		{
+			TRACE_CPUPROFILER_EVENT_SCOPE(Geometry_SweepingMeshSDF_Signs);
 			compute_intersections(Origin, DX, NI, NJ, NK, intersection_count);
 			if (CancelF())
 			{
