Normal Flow Remesh improvements:
- add scalar damping to vertex projection and normal fitting
- detect areas that cover holes in the target mesh and don't project in those regions
- expose smoothing parameters for these areas

diff --git a/Engine/Plugins/Experimental/GeometryProcessing/Source/DynamicMesh/Private/NormalFlowRemesher.cpp b/Engine/Plugins/Experimental/GeometryProcessing/Source/DynamicMesh/Private/NormalFlowRemesher.cpp
--- a/Engine/Plugins/Experimental/GeometryProcessing/Source/DynamicMesh/Private/NormalFlowRemesher.cpp
+++ b/Engine/Plugins/Experimental/GeometryProcessing/Source/DynamicMesh/Private/NormalFlowRemesher.cpp
@@ -62,20 +62,24 @@ void FNormalFlowRemesher::RemeshWithFaceProjection()
 		}
 
 		double MaxProjectionDistance = 0.0;
+		constexpr bool bIsTuningIteration = false;
 		if (bEnableParallelProjection)
 		{
-			TrackedFaceProjectionPass(MaxProjectionDistance);
+			TrackedFaceProjectionPass(MaxProjectionDistance, bIsTuningIteration);
 		}
 		else
 		{
-			TrackedFaceProjectionPass_Serial(MaxProjectionDistance);
+			TrackedFaceProjectionPass_Serial(MaxProjectionDistance, bIsTuningIteration);
 		}
 		
 
-		// Stop if we've hit max iterations, or both:
+		// Stop if we've hit max iterations, hit triangle limit, or both:
 		// - queue is empty and
 		// - projection isn't moving anything
-		bContinue = (Iterations++ < MaxRemeshIterations) && ((ModifiedEdges->Num() > 0) || (MaxProjectionDistance > ProjectionDistanceThreshold));
+		bContinue =
+			(Iterations++ < MaxRemeshIterations)
+			&& ((ModifiedEdges->Num() > 0) || (MaxProjectionDistance > ProjectionDistanceThreshold))
+			&& ((MaxTriangleCount == 0) || (Mesh->TriangleCount() < MaxTriangleCount));
 	}
 
 	SmoothSpeedT = OriginalSmoothSpeed;
@@ -91,13 +95,14 @@ void FNormalFlowRemesher::RemeshWithFaceProjection()
 			}
 
 			double MaxProjectionDistance = 0.0;
+			constexpr bool bIsTuningIteration = true;
 			if (bEnableParallelProjection)
 			{
-				TrackedFaceProjectionPass(MaxProjectionDistance);
+				TrackedFaceProjectionPass(MaxProjectionDistance, bIsTuningIteration);
 			}
 			else
 			{
-				TrackedFaceProjectionPass_Serial(MaxProjectionDistance);
+				TrackedFaceProjectionPass_Serial(MaxProjectionDistance, bIsTuningIteration);
 			}
 
 			if (MaxProjectionDistance == 0.0)
@@ -113,7 +118,8 @@ void FNormalFlowRemesher::RemeshWithFaceProjection()
 }
 
 
-void FNormalFlowRemesher::TrackedFaceProjectionPass(double& MaxDistanceMoved)
+
+void FNormalFlowRemesher::TrackedFaceProjectionPass(double& MaxDistanceMoved, bool bIsTuningIteration)
 {
 	ensure(ProjTarget != nullptr);
 
@@ -122,12 +128,17 @@ void FNormalFlowRemesher::TrackedFaceProjectionPass(double& MaxDistanceMoved)
 
 	InitializeVertexBufferForFacePass();
 
+	double MaxStepDistance = (bIsTuningIteration) ? (0.25 * MaxEdgeLength) : (0.5 * MaxEdgeLength);		// kind of arbitrary...
+	double SmoothAreaDistance = FillAreaDistanceMultiplier * MaxEdgeLength;
+
+	TFunction<FVector3d(const FDynamicMesh3&, int, double)> UseSmoothFunc = GetSmoothFunction();
+
 	// For each triangle, rotate it such that it aligns with closest normal on the target surface.
 	// Each vertex is then assigned a weighted combination of its corresponding triangle corner positions. The weighting
 	// is chosen to favor triangles that don't rotate much (i.e. are already aligned with the target surface.)
 
 	// First compute all rotated triangles and store their corner positions and weights
-	ParallelFor(Mesh->MaxTriangleID(), [this, NormalProjTarget](int32 TriangleIndex)
+	ParallelFor(Mesh->MaxTriangleID(), [this, NormalProjTarget, bIsTuningIteration, SmoothAreaDistance, UseSmoothFunc, MaxStepDistance](int32 TriangleIndex)
 	{
 		FVector3d TriangleNormal, Centroid;
 		double Area;
@@ -136,6 +147,11 @@ void FNormalFlowRemesher::TrackedFaceProjectionPass(double& MaxDistanceMoved)
 		FVector3d ProjectedNormal{ 1e30, 1e30, 1e30 };
 		FVector3d ProjectedPosition = NormalProjTarget->Project(Centroid, ProjectedNormal);
 
+		if (TriangleNormal.Length() < 0.9 || ProjectedNormal.Length() < 0.9)
+		{
+			return;		// skip this triangle
+		}
+
 		if (!ensure(ProjectedNormal[0] != 1e30))
 		{
 			return;
@@ -145,6 +161,52 @@ void FNormalFlowRemesher::TrackedFaceProjectionPass(double& MaxDistanceMoved)
 			return;
 		}
 
+		FIndex3i TriangleVertices = Mesh->GetTriangle(TriangleIndex);
+
+		// if we are tuning and we are not within distance band from the target mesh, then this 
+		// is likely an area that cannot project, ie a "fill" area. In those areas we just want to
+		// smooth because the projection+remesh will have created very ugly geometry
+		// Note: it seems like this ought to be done per-vertex, not triangle based. However this
+		// so far has not produced as good of results - not entirely clear why, but it seems that
+		// around the boundaries, the triangles that sill have a good projection extert some pull
+		// back on the boundary vertices, which keeps them more stable, otherwise they shrink too
+		// much and this causes artifacts around the border (perhaps because the normals are never converging).
+		// It *might* work better to do this as a fully separate step, after everything else has 
+		// converged, and do some kind of weighted-identification of these hole areas
+
+		if (bIsTuningIteration && bSmoothInFillAreas && Distance(ProjectedPosition, Centroid) > SmoothAreaDistance && FillAreaSmoothMultiplier > 0)
+		{
+			for (int32 j = 0; j < 3; ++j)
+			{
+				bool bModified = false;
+				FVector3d SmoothPos = ComputeSmoothedVertexPos(TriangleVertices[j], UseSmoothFunc, bModified);
+				if (bModified)
+				{
+					double Weight = Area;
+					SmoothPos = Lerp(Centroid, SmoothPos, FillAreaSmoothMultiplier * SmoothSpeedT);
+
+					ProjectionVertexBuffer[3*TriangleIndex + j] = Weight * SmoothPos;
+					ProjectionWeightBuffer[3*TriangleIndex + j] = Weight;
+				}
+			}
+			return;
+		}
+
+		// apply damping to new position/normal
+		ProjectedPosition = Lerp(Centroid, ProjectedPosition, SurfaceProjectionSpeed);
+		ProjectedNormal = Normalized(Lerp(TriangleNormal, ProjectedNormal, NormalAlignmentSpeed));
+		if (ProjectedNormal.Length() < 0.1)
+		{
+			ProjectedNormal = TriangleNormal;
+		}
+
+		// clamp movement of target position, to prevent moving too far in a single step
+		FVector3d MoveDelta = (ProjectedPosition - Centroid);
+		double MoveLength = Normalize(MoveDelta);
+
+		MoveLength = FMathd::Min(MoveLength, MaxStepDistance);
+		ProjectedPosition = Centroid + MoveLength * MoveDelta;
+
 		FVector3d V0, V1, V2;
 		Mesh->GetTriVertices(TriangleIndex, V0, V1, V2);
 
@@ -172,7 +234,7 @@ void FNormalFlowRemesher::TrackedFaceProjectionPass(double& MaxDistanceMoved)
 	});
 
 	// Next compute all the weighted average of triangle corners for each vertex
-	ParallelFor(Mesh->MaxVertexID(), [this, &MaxDistanceMoved](int32 VertexIndex)
+	ParallelFor(Mesh->MaxVertexID(), [this, &MaxDistanceMoved, bIsTuningIteration, MaxStepDistance](int32 VertexIndex)
 	{
 		TempFlagBuffer[VertexIndex] = false;
 
@@ -223,6 +285,12 @@ void FNormalFlowRemesher::TrackedFaceProjectionPass(double& MaxDistanceMoved)
 		FVector3d CurrentPosition = Mesh->GetVertex(VertexIndex);
 		FVector3d ProjectedPosition = TempPosBuffer[VertexIndex] / TempWeightBuffer[VertexIndex];
 
+		// clamp movement of target position, to prevent moving too far in a single step
+		FVector3d MoveDelta = (ProjectedPosition - CurrentPosition);
+		double MoveLength = Normalize(MoveDelta);
+		MoveLength = FMathd::Min(MoveLength, MaxStepDistance);
+		ProjectedPosition = CurrentPosition + MoveLength * MoveDelta;
+
 		if (VectorUtil::EpsilonEqual(CurrentPosition, ProjectedPosition, FMathd::ZeroTolerance))
 		{
 			return;
@@ -264,7 +332,7 @@ void FNormalFlowRemesher::TrackedFaceProjectionPass(double& MaxDistanceMoved)
 
 
 
-void FNormalFlowRemesher::TrackedFaceProjectionPass_Serial(double& MaxDistanceMoved)
+void FNormalFlowRemesher::TrackedFaceProjectionPass_Serial(double& MaxDistanceMoved, bool bIsTuningIteration)
 {
 	ensure(ProjTarget != nullptr);
 
@@ -273,6 +341,11 @@ void FNormalFlowRemesher::TrackedFaceProjectionPass_Serial(double& MaxDistanceMo
 
 	InitializeVertexBufferForFacePass();
 
+	double MaxStepDistance = (bIsTuningIteration) ? (0.25*MaxEdgeLength) : (0.5*MaxEdgeLength);		// kind of arbitrary...
+	double SmoothAreaDistance = FillAreaDistanceMultiplier * MaxEdgeLength;
+
+	TFunction<FVector3d(const FDynamicMesh3&, int, double)> UseSmoothFunc = GetSmoothFunction();
+
 	// this function computes rotated position of triangle, such that it
 	// aligns with face normal on target surface. We accumulate weighted-average
 	// of vertex positions, which we will then use further down where possible.
@@ -282,12 +355,56 @@ void FNormalFlowRemesher::TrackedFaceProjectionPass_Serial(double& MaxDistanceMo
 		FVector3d TriangleNormal, Centroid;
 		double Area;
 		Mesh->GetTriInfo(TriangleIndex, TriangleNormal, Area, Centroid);
-
+		FIndex3i TriangleVertices = Mesh->GetTriangle(TriangleIndex);
+		
 		FVector3d ProjectedNormal{ 1e30, 1e30, 1e30 };
 		FVector3d ProjectedPosition = NormalProjTarget->Project(Centroid, ProjectedNormal);
+		if (TriangleNormal.Length() < 0.9 || ProjectedNormal.Length() < 0.9 )
+		{
+			continue;		// skip this triangle
+		}
+
+		// if we are tuning and we are not within distance band from the target mesh, then this 
+		// is likely an area that cannot project, ie a "fill" area. In those areas we just want to
+		// smooth because the projection+remesh will have created very ugly geometry
+		// Note: it seems like this ought to be done per-vertex, not triangle based. However this
+		// so far has not produced as good of results - not entirely clear why, but it seems that
+		// around the boundaries, the triangles that sill have a good projection extert some pull
+		// back on the boundary vertices, which keeps them more stable, otherwise they shrink too
+		// much and this causes artifacts around the border (perhaps because the normals are never converging).
+		// It *might* work better to do this as a fully separate step, after everything else has 
+		// converged, and do some kind of weighted-identification of these hole areas
+		if (bIsTuningIteration && bSmoothInFillAreas && Distance(ProjectedPosition, Centroid) > SmoothAreaDistance && FillAreaSmoothMultiplier > 0)
+		{
+			for (int32 j = 0; j < 3; ++j)
+			{
+				bool bModified = false;
+				FVector3d SmoothPos = ComputeSmoothedVertexPos(TriangleVertices[j], UseSmoothFunc, bModified);
+				if (bModified)
+				{
+					double Weight = Area;
+					SmoothPos = Lerp(Centroid, SmoothPos, FillAreaSmoothMultiplier*SmoothSpeedT);
+					TempPosBuffer[TriangleVertices[j]] += Weight * SmoothPos;
+					TempWeightBuffer[TriangleVertices[j]] += Weight;
+				}
+			}
+			continue;
+		}
 
-		check(ProjectedNormal[0] != 1e30);
-		check(ProjectedNormal.Length() > 1e-6);
+		// apply damping to new position/normal
+		ProjectedPosition = Lerp(Centroid, ProjectedPosition, SurfaceProjectionSpeed);
+		ProjectedNormal = Normalized(Lerp(TriangleNormal, ProjectedNormal, NormalAlignmentSpeed));
+		if (ProjectedNormal.Length() < 0.1)
+		{
+			ProjectedNormal = TriangleNormal;
+		}
+
+		// clamp movement of target position, to prevent moving too far in a single step
+		FVector3d MoveDelta = (ProjectedPosition - Centroid);
+		double MoveLength = Normalize(MoveDelta);
+
+		MoveLength = FMathd::Min(MoveLength, MaxStepDistance);
+		ProjectedPosition = Centroid + MoveLength * MoveDelta;
 
 		FVector3d V0, V1, V2;
 		Mesh->GetTriVertices(TriangleIndex, V0, V1, V2);
@@ -307,7 +424,6 @@ void FNormalFlowRemesher::TrackedFaceProjectionPass_Serial(double& MaxDistanceMo
 		Dot = FMath::Clamp(Dot, 0.0, 1.0);
 		double Weight = Area * (Dot * Dot * Dot);
 
-		FIndex3i TriangleVertices = Mesh->GetTriangle(TriangleIndex);
 		TempPosBuffer[TriangleVertices.A] += Weight * V0;
 		TempWeightBuffer[TriangleVertices.A] += Weight;
 		TempPosBuffer[TriangleVertices.B] += Weight * V1;
@@ -349,6 +465,12 @@ void FNormalFlowRemesher::TrackedFaceProjectionPass_Serial(double& MaxDistanceMo
 			continue;
 		}
 
+		// clamp movement of target position, to prevent moving too far in a single step
+		FVector3d MoveDelta = (ProjectedPosition - CurrentPosition);
+		double MoveLength = Normalize(MoveDelta);
+		MoveLength = FMathd::Min(MoveLength, MaxStepDistance);
+		ProjectedPosition = CurrentPosition + MoveLength * MoveDelta;
+
 		MaxDistanceMoved = FMath::Max(MaxDistanceMoved, Distance(CurrentPosition, ProjectedPosition));
 
 		TempFlagBuffer[VertexID] = true;
diff --git a/Engine/Plugins/Experimental/GeometryProcessing/Source/DynamicMesh/Public/NormalFlowRemesher.h b/Engine/Plugins/Experimental/GeometryProcessing/Source/DynamicMesh/Public/NormalFlowRemesher.h
--- a/Engine/Plugins/Experimental/GeometryProcessing/Source/DynamicMesh/Public/NormalFlowRemesher.h
+++ b/Engine/Plugins/Experimental/GeometryProcessing/Source/DynamicMesh/Public/NormalFlowRemesher.h
@@ -27,6 +27,25 @@ public:
 	/// Additional projection iterations after the usual remesh step
 	int NumExtraProjectionIterations = 5;
 
+	/// drag on surface projection
+	double SurfaceProjectionSpeed = 0.2;
+
+	/// drag on normal alignment
+	double NormalAlignmentSpeed = 0.2;
+
+	/// Control whether or not we want to apply mesh smoothing in "free" areas that have not projected to target surface.
+	/// This smoothing is on applied in the ExtraProjections Iterations
+	bool bSmoothInFillAreas = true;
+
+	/// This is used as a multiplier on MaxEdgeLength to determine when we identify points as being in "free" areas
+	float FillAreaDistanceMultiplier = 0.25;
+
+	/// This is used as a multiplier on the Remesher smoothing rate, applied to points identified as being in "free" areas
+	float FillAreaSmoothMultiplier = 0.25;
+
+	/// cap on triangle count, to prevent runaway mesh messes. Ignored if zero.
+	int MaxTriangleCount = 0;
+
 	/// "Outer loop" for all remeshing operations
 	void BasicRemeshPass() override
 	{
@@ -44,8 +63,8 @@ protected:
 	void RemeshWithFaceProjection();
 
 	// Perform face-aligned projection onto the target mesh. Queue edges whose lengths change because of it.
-	void TrackedFaceProjectionPass(double& MaxDistanceMoved);
-	void TrackedFaceProjectionPass_Serial(double& MaxDistanceMoved);
+	void TrackedFaceProjectionPass(double& MaxDistanceMoved, bool bIsTuningIteration);
+	void TrackedFaceProjectionPass_Serial(double& MaxDistanceMoved, bool bIsTuningIteration);
 
 	// This is called during RemeshIteration 
 	void TrackedFullProjectionPass(bool bParallel) override
@@ -53,13 +72,14 @@ protected:
 		for (int i = 0; i < FaceProjectionPassesPerRemeshIteration; ++i)
 		{
 			double ProjectionDistance;
+			constexpr bool bIsTuningIteration = false;
 			if (bParallel)
 			{
-				TrackedFaceProjectionPass(ProjectionDistance);
+				TrackedFaceProjectionPass(ProjectionDistance, bIsTuningIteration);
 			}
 			else
 			{
-				TrackedFaceProjectionPass_Serial(ProjectionDistance);
+				TrackedFaceProjectionPass_Serial(ProjectionDistance, bIsTuningIteration);
 			}
 		}
 	}
diff --git a/Engine/Plugins/Experimental/MeshModelingToolset/Source/MeshModelingTools/Private/ProjectToTargetTool.cpp b/Engine/Plugins/Experimental/MeshModelingToolset/Source/MeshModelingTools/Private/ProjectToTargetTool.cpp
--- a/Engine/Plugins/Experimental/MeshModelingToolset/Source/MeshModelingTools/Private/ProjectToTargetTool.cpp
+++ b/Engine/Plugins/Experimental/MeshModelingToolset/Source/MeshModelingTools/Private/ProjectToTargetTool.cpp
@@ -82,6 +82,17 @@ TUniquePtr<FDynamicMeshOperator> UProjectToTargetTool::MakeNewOperator()
 	RemeshOp->ToolMeshLocalToWorld = FTransform3d(Cast<IPrimitiveComponentBackedTarget>(Targets[0])->GetWorldTransform());
 	RemeshOp->TargetMeshLocalToWorld = FTransform3d(Cast<IPrimitiveComponentBackedTarget>(Targets[1])->GetWorldTransform());
 	RemeshOp->bUseWorldSpace = ProjectProperties->bWorldSpace;
+	RemeshOp->bParallel = ProjectProperties->bParallel;
+
+	if (ProjectProperties->RemeshType == ERemeshType::NormalFlow)
+	{
+		RemeshOp->FaceProjectionPassesPerRemeshIteration = ProjectProperties->FaceProjectionPassesPerRemeshIteration;
+		RemeshOp->SurfaceProjectionSpeed = ProjectProperties->SurfaceProjectionSpeed;
+		RemeshOp->NormalAlignmentSpeed = ProjectProperties->NormalAlignmentSpeed;
+		RemeshOp->bSmoothInFillAreas = ProjectProperties->bSmoothInFillAreas;
+		RemeshOp->FillAreaDistanceMultiplier = ProjectProperties->FillAreaDistanceMultiplier;
+		RemeshOp->FillAreaSmoothMultiplier = ProjectProperties->FillAreaSmoothMultiplier;
+	}
 
 	return Op;
 }
diff --git a/Engine/Plugins/Experimental/MeshModelingToolset/Source/MeshModelingTools/Private/RemeshMeshTool.cpp b/Engine/Plugins/Experimental/MeshModelingToolset/Source/MeshModelingTools/Private/RemeshMeshTool.cpp
--- a/Engine/Plugins/Experimental/MeshModelingToolset/Source/MeshModelingTools/Private/RemeshMeshTool.cpp
+++ b/Engine/Plugins/Experimental/MeshModelingToolset/Source/MeshModelingTools/Private/RemeshMeshTool.cpp
@@ -200,6 +200,7 @@ TUniquePtr<FDynamicMeshOperator> URemeshMeshTool::MakeNewOperator()
 	if (!BasicProperties->bUseTargetEdgeLength)
 	{
 		Op->TargetEdgeLength = CalculateTargetEdgeLength(BasicProperties->TargetTriangleCount);
+		Op->TriangleCountHint = 2.0 * BasicProperties->TargetTriangleCount;
 	}
 	else
 	{
diff --git a/Engine/Plugins/Experimental/MeshModelingToolset/Source/MeshModelingTools/Public/ProjectToTargetTool.h b/Engine/Plugins/Experimental/MeshModelingToolset/Source/MeshModelingTools/Public/ProjectToTargetTool.h
--- a/Engine/Plugins/Experimental/MeshModelingToolset/Source/MeshModelingTools/Public/ProjectToTargetTool.h
+++ b/Engine/Plugins/Experimental/MeshModelingToolset/Source/MeshModelingTools/Public/ProjectToTargetTool.h
@@ -40,9 +40,30 @@ public:
 	UPROPERTY(EditAnywhere, Category = ProjectionSpace)
 	bool bWorldSpace = true;
 
-	UProjectToTargetToolProperties() :
-		URemeshMeshToolProperties(),
-		bWorldSpace(true)
+	UPROPERTY(EditAnywhere, Category = Remeshing)
+	bool bParallel = true;
+
+	UPROPERTY(EditAnywhere, Category = NormalFlow, meta = (EditCondition = "RemeshType == ERemeshType::NormalFlow", UIMin = "0", UIMax = "100", ClampMin = "0", ClampMax = "100"))
+	int FaceProjectionPassesPerRemeshIteration = 1;
+
+	UPROPERTY(EditAnywhere, Category = NormalFlow, meta = (EditCondition = "RemeshType == ERemeshType::NormalFlow", UIMin = "0", UIMax = "1.0", ClampMin = "0", ClampMax = "10.0"))
+	float SurfaceProjectionSpeed = 0.2f;
+
+	UPROPERTY(EditAnywhere, Category = NormalFlow, meta = (EditCondition = "RemeshType == ERemeshType::NormalFlow", UIMin = "0", UIMax = "1.0", ClampMin = "0", ClampMax = "10.0"))
+	float NormalAlignmentSpeed = 0.2f;
+
+	UPROPERTY(EditAnywhere, Category = NormalFlow, meta = (EditCondition = "RemeshType == ERemeshType::NormalFlow"))
+	bool bSmoothInFillAreas = true;
+
+	UPROPERTY(EditAnywhere, Category = NormalFlow, meta = (EditCondition = "RemeshType == ERemeshType::NormalFlow", UIMin = "0", UIMax = "1.0", ClampMin = "0", ClampMax = "10.0"))
+	float FillAreaDistanceMultiplier = 0.25f;
+
+	UPROPERTY(EditAnywhere, Category = NormalFlow, meta = (EditCondition = "RemeshType == ERemeshType::NormalFlow", UIMin = "0", UIMax = "1.0", ClampMin = "0", ClampMax = "10.0"))
+	float FillAreaSmoothMultiplier = 0.25f;
+
+
+	UProjectToTargetToolProperties() 
+		: URemeshMeshToolProperties()
 	{
 		bPreserveSharpEdges = false;
 		RemeshType = ERemeshType::NormalFlow;
diff --git a/Engine/Plugins/Experimental/MeshModelingToolset/Source/ModelingOperators/Private/CleaningOps/RemeshMeshOp.cpp b/Engine/Plugins/Experimental/MeshModelingToolset/Source/ModelingOperators/Private/CleaningOps/RemeshMeshOp.cpp
--- a/Engine/Plugins/Experimental/MeshModelingToolset/Source/ModelingOperators/Private/CleaningOps/RemeshMeshOp.cpp
+++ b/Engine/Plugins/Experimental/MeshModelingToolset/Source/ModelingOperators/Private/CleaningOps/RemeshMeshOp.cpp
@@ -34,6 +34,13 @@ TUniquePtr<FRemesher> FRemeshMeshOp::CreateRemesher(ERemeshType Type, FDynamicMe
 		TUniquePtr<FNormalFlowRemesher> NormalFlowRemesher = MakeUnique<FNormalFlowRemesher>(TargetMesh);
 		NormalFlowRemesher->MaxRemeshIterations = MaxRemeshIterations;
 		NormalFlowRemesher->NumExtraProjectionIterations = ExtraProjectionIterations;
+		NormalFlowRemesher->MaxTriangleCount = FMath::Max(0, TriangleCountHint);
+		NormalFlowRemesher->FaceProjectionPassesPerRemeshIteration = FaceProjectionPassesPerRemeshIteration;
+		NormalFlowRemesher->SurfaceProjectionSpeed = SurfaceProjectionSpeed;
+		NormalFlowRemesher->NormalAlignmentSpeed = NormalAlignmentSpeed;
+		NormalFlowRemesher->bSmoothInFillAreas = bSmoothInFillAreas;
+		NormalFlowRemesher->FillAreaDistanceMultiplier = FillAreaDistanceMultiplier;
+		NormalFlowRemesher->FillAreaSmoothMultiplier = FillAreaSmoothMultiplier;
 		return NormalFlowRemesher;
 	}
 	default:
@@ -69,6 +76,8 @@ void FRemeshMeshOp::CalculateResult(FProgressCancel* Progress)
 
 	TUniquePtr<FRemesher> Remesher = CreateRemesher(RemeshType, TargetMesh);
 
+	Remesher->bEnableParallelProjection = this->bParallel;
+
 	Remesher->bEnableSplits = bSplits;
 	Remesher->bEnableFlips = bFlips;
 	Remesher->bEnableCollapses = bCollapses;
diff --git a/Engine/Plugins/Experimental/MeshModelingToolset/Source/ModelingOperators/Public/CleaningOps/RemeshMeshOp.h b/Engine/Plugins/Experimental/MeshModelingToolset/Source/ModelingOperators/Public/CleaningOps/RemeshMeshOp.h
--- a/Engine/Plugins/Experimental/MeshModelingToolset/Source/ModelingOperators/Public/CleaningOps/RemeshMeshOp.h
+++ b/Engine/Plugins/Experimental/MeshModelingToolset/Source/ModelingOperators/Public/CleaningOps/RemeshMeshOp.h
@@ -61,6 +61,7 @@ public:
 	int RemeshIterations;
 	int MaxRemeshIterations;
 	int ExtraProjectionIterations = 5;
+	int TriangleCountHint = 0;
 	float SmoothingStrength, TargetEdgeLength;
 	ERemeshSmoothingType SmoothingType;
 	bool bDiscardAttributes, bPreserveSharpEdges, bFlips, bSplits, bCollapses, bReproject, bPreventNormalFlips;
@@ -72,6 +73,31 @@ public:
 	FTransform3d TargetMeshLocalToWorld;
 	FTransform3d ToolMeshLocalToWorld;
 	bool bUseWorldSpace = false;
+	bool bParallel = true;
+
+	// Normal flow only:
+
+	/// During each call to RemeshIteration, do this many passes of face-aligned projection
+	int FaceProjectionPassesPerRemeshIteration = 1;
+
+	/// drag on surface projection
+	double SurfaceProjectionSpeed = 0.2;
+
+	/// drag on normal alignment
+	double NormalAlignmentSpeed = 0.2;
+
+	/// Control whether or not we want to apply mesh smoothing in "free" areas that have not projected to target surface.
+	/// This smoothing is on applied in the ExtraProjections Iterations
+	bool bSmoothInFillAreas = true;
+
+	/// This is used as a multiplier on MaxEdgeLength to determine when we identify points as being in "free" areas
+	float FillAreaDistanceMultiplier = 0.25;
+
+	/// This is used as a multiplier on the Remesher smoothing rate, applied to points identified as being in "free" areas
+	float FillAreaSmoothMultiplier = 0.25;
+
+	// End normal flow only
+
 
 	void SetTransform(const FTransform& Transform);
 
