Fix static analysis warnings in MeshResampleImageBaker.

#rnx

diff --git a/Engine/Plugins/Experimental/GeometryProcessing/Source/DynamicMesh/Private/Sampling/MeshResampleImageBaker.cpp b/Engine/Plugins/Experimental/GeometryProcessing/Source/DynamicMesh/Private/Sampling/MeshResampleImageBaker.cpp
--- a/Engine/Plugins/Experimental/GeometryProcessing/Source/DynamicMesh/Private/Sampling/MeshResampleImageBaker.cpp
+++ b/Engine/Plugins/Experimental/GeometryProcessing/Source/DynamicMesh/Private/Sampling/MeshResampleImageBaker.cpp
@@ -143,9 +143,11 @@ void FMeshMultiResampleImageBaker::Bake()
 void FMeshMultiResampleImageBaker::PreEvaluate(const FMeshMapBaker& Baker)
 {
 	DetailMesh = Baker.GetDetailMesh();
-	DetailMaterialIDAttrib = DetailMesh->Attributes()->GetMaterialID();
-
-	bValidDetailMesh = ensure(DetailMesh) && ensure(DetailMaterialIDAttrib) && ensure(DetailUVOverlay);
+	if (ensure(DetailMesh))
+	{
+		DetailMaterialIDAttrib = DetailMesh->Attributes()->GetMaterialID();
+		bValidDetailMesh = ensure(DetailMaterialIDAttrib) && ensure(DetailUVOverlay);
+	}
 }
 
 FVector4f FMeshMultiResampleImageBaker::EvaluateSample(const FMeshMapBaker& Baker, const FCorrespondenceSample& Sample)
