Fix MeshResampleImageBaker preevaluate initialization.

#rnx

diff --git a/Engine/Plugins/Experimental/GeometryProcessing/Source/DynamicMesh/Private/Sampling/MeshResampleImageBaker.cpp b/Engine/Plugins/Experimental/GeometryProcessing/Source/DynamicMesh/Private/Sampling/MeshResampleImageBaker.cpp
--- a/Engine/Plugins/Experimental/GeometryProcessing/Source/DynamicMesh/Private/Sampling/MeshResampleImageBaker.cpp
+++ b/Engine/Plugins/Experimental/GeometryProcessing/Source/DynamicMesh/Private/Sampling/MeshResampleImageBaker.cpp
@@ -143,11 +143,8 @@ void FMeshMultiResampleImageBaker::Bake()
 void FMeshMultiResampleImageBaker::PreEvaluate(const FMeshMapBaker& Baker)
 {
 	DetailMesh = Baker.GetDetailMesh();
-	if (ensure(DetailMesh))
-	{
-		DetailMaterialIDAttrib = DetailMesh->Attributes()->GetMaterialID();
-		bValidDetailMesh = ensure(DetailMaterialIDAttrib) && ensure(DetailUVOverlay);
-	}
+	DetailMaterialIDAttrib = ensure(DetailMesh) ? DetailMesh->Attributes()->GetMaterialID() : nullptr;
+	bValidDetailMesh = ensure(DetailMaterialIDAttrib) && ensure(DetailUVOverlay);
 }
 
 FVector4f FMeshMultiResampleImageBaker::EvaluateSample(const FMeshMapBaker& Baker, const FCorrespondenceSample& Sample)
