GeometryProcessing:
- Skip full-mesh normals computation in OffsetMeshRegion if using vertex normals.
- Add operator[] to MeshConnectedComponents.
- Add FSparseDynamicOctree3::ContainmentQuery() to allow for single-point range query.
- Add FMeshFaceSelection::IsSelectionBoundaryEdge() and optimize ExpandToOneRingNeighbours() and ContractBorderByOneRingNeighbours(), now significantly faster for very large selections. Convert VtxTrianglesItr calls to EnumerateVtxTriangles
- Use TestAnyHitTriangle() instead of FindNearestHitTriangle() in TRemoveOccludedTriangles. Reduces cost from 6.4s to 1.2s on a 64-core threadripper for a large-scale (1.4m triangles) visibility-based removal problem
- Optimize FMarchingCubes implementation for large thread counts. Replace single FCriticalSection with indexed sets of FCriticalSection to reduce thread contention. Instead of directly writing to mesh, accumulate mesh vertices and triangles in separate blocks (hashed on grid coordinates) and then build output mesh after generation completes. Add bEnableValueCaching option to skip caching of corner values, which causes lots of thread contention if Implicit function is very cheap (eg a grid-based SDF evaluation). Enable this option in Morphology meshing. Reduces marching cubes cost from ~3.5s to <500ms on a 64-core threadripper for a large-scale (1.4m triangles) Morphology problem.

#rnx

diff --git a/Engine/Plugins/Experimental/GeometryProcessing/Source/DynamicMesh/Private/Operations/OffsetMeshRegion.cpp b/Engine/Plugins/Experimental/GeometryProcessing/Source/DynamicMesh/Private/Operations/OffsetMeshRegion.cpp
--- a/Engine/Plugins/Experimental/GeometryProcessing/Source/DynamicMesh/Private/Operations/OffsetMeshRegion.cpp
+++ b/Engine/Plugins/Experimental/GeometryProcessing/Source/DynamicMesh/Private/Operations/OffsetMeshRegion.cpp
@@ -25,7 +25,7 @@ bool FOffsetMeshRegion::Apply()
 {
 	FMeshNormals Normals;
 	bool bHaveVertexNormals = Mesh->HasVertexNormals();
-	if (!bHaveVertexNormals)
+	if (bHaveVertexNormals == false && bUseFaceNormals == false)
 	{
 		Normals = FMeshNormals(Mesh);
 		Normals.ComputeVertexNormals();
diff --git a/Engine/Plugins/Experimental/GeometryProcessing/Source/DynamicMesh/Private/Selections/MeshFaceSelection.cpp b/Engine/Plugins/Experimental/GeometryProcessing/Source/DynamicMesh/Private/Selections/MeshFaceSelection.cpp
--- a/Engine/Plugins/Experimental/GeometryProcessing/Source/DynamicMesh/Private/Selections/MeshFaceSelection.cpp
+++ b/Engine/Plugins/Experimental/GeometryProcessing/Source/DynamicMesh/Private/Selections/MeshFaceSelection.cpp
@@ -44,4 +44,151 @@ FMeshFaceSelection::FMeshFaceSelection(const FDynamicMesh3* mesh, const FMeshVer
 			}
 		}
 	}
+}
+
+
+
+
+
+void FMeshFaceSelection::ExpandToOneRingNeighbours(const TUniqueFunction<bool(int)>& FilterF)
+{
+	TSet<int> TrisToAdd;
+
+	for (int tid : Selected) 
+	{ 
+		FIndex3i TriEdges = Mesh->GetTriEdges(tid);
+		TArray<int32, TFixedAllocator<3>> VertsToProcess;
+		for (int j = 0; j < 3; ++j)
+		{
+			TPair<bool, bool> BoundaryInfo = IsSelectionBoundaryEdge(TriEdges[j]);
+			if ( BoundaryInfo.Key != BoundaryInfo.Value)		// selection boundary but not mesh boundary
+			{
+				FIndex2i EdgeV = Mesh->GetEdgeV(TriEdges[j]);
+				VertsToProcess.AddUnique(EdgeV.A);
+				VertsToProcess.AddUnique(EdgeV.B);
+			}
+		}
+		for (int vid : VertsToProcess)
+		{
+			Mesh->EnumerateVertexTriangles(vid, [&](int nbr_t)
+			{
+				if (FilterF && FilterF(nbr_t) == false)
+				{
+					return;
+				}
+				if (!IsSelected(nbr_t))
+				{
+					TrisToAdd.Add(nbr_t);
+				}
+			});
+		}
+	}
+
+	for (int tid : TrisToAdd)
+	{
+		add(tid);
+	}
+}
+
+
+void FMeshFaceSelection::ExpandToOneRingNeighbours(int nRings, const TUniqueFunction<bool(int)>& FilterF)
+{
+	if (nRings == 1)
+	{
+		ExpandToOneRingNeighbours(FilterF);
+		return;
+	}
+
+	// todo: rewrite to use TSet and mesh edges as in ExpandToOneRingNeighbours above
+
+	TArray<int> triArrays[2];
+	int addIdx = 0, checkIdx = 1;
+	triArrays[checkIdx] = Selected.Array();
+
+	TBitArray<FDefaultBitArrayAllocator> Bitmap = AsBitArray();
+
+	for (int ri = 0; ri < nRings; ++ri)
+	{
+		TArray<int>& addTris = triArrays[addIdx];
+		TArray<int>& checkTris = triArrays[checkIdx];
+		addTris.Empty();
+
+		for (int tid : checkTris)
+		{
+			FIndex3i tri_v = Mesh->GetTriangle(tid);
+			for (int j = 0; j < 3; ++j)
+			{
+				int vid = tri_v[j];
+				Mesh->EnumerateVertexTriangles(vid, [&](int nbr_t)
+				{
+					if (FilterF && FilterF(nbr_t) == false)
+					{
+						return;
+					}
+					if (Bitmap[nbr_t] == false)
+					{
+						addTris.Add(nbr_t);
+						Bitmap[nbr_t] = true;
+					}
+				});
+			}
+		}
+
+		for (int TID : addTris)
+		{
+			add(TID);
+		}
+
+		Swap(addIdx, checkIdx); // check in the next iter what we added in this iter
+	}
+}
+
+
+
+
+void FMeshFaceSelection::ContractBorderByOneRingNeighbours(int NumRings, bool bContractFromMeshBoundary, const TUniqueFunction<bool(int)>& FilterF)
+{
+	TSet<int> BorderVIDs;   // border vertices
+
+	// TODO: can we track the boundary across iterations? Slightly tricky because the modified tris
+	// are no longer in the selection set. However it seems like that is still going to work because
+	// all the selection boundary edges will still be adjacent to that set of deselected triangles
+
+	for (int RingIdx = 0; RingIdx < NumRings; ++RingIdx)
+	{
+		BorderVIDs.Reset();
+
+		// find set of vertices on border
+		for (int tid : Selected)
+		{
+			FIndex3i TriEdges = Mesh->GetTriEdges(tid);
+			TArray<int32, TInlineAllocator<3>> BoundaryVertsToProcess;
+			for (int j = 0; j < 3; ++j)
+			{
+				TPair<bool, bool> BoundaryInfo = IsSelectionBoundaryEdge(TriEdges[j]);
+				if (BoundaryInfo.Key != false)		// we are some kind of boundary edge
+				{
+					if (bContractFromMeshBoundary || BoundaryInfo.Value == false)
+					{
+						FIndex2i EdgeV = Mesh->GetEdgeV(TriEdges[j]);
+						BorderVIDs.Add(EdgeV.A);
+						BorderVIDs.Add(EdgeV.B);
+					}
+				}
+			}
+		}
+
+		for (int VID : BorderVIDs)
+		{
+			Mesh->EnumerateVertexTriangles(VID, [&](int32 NbrT)
+			{
+				if (FilterF && !FilterF(NbrT))
+				{
+					return;
+				}
+				Deselect(NbrT);
+			});
+		}
+	}
+
 }
\ No newline at end of file
diff --git a/Engine/Plugins/Experimental/GeometryProcessing/Source/DynamicMesh/Public/Operations/RemoveOccludedTriangles.h b/Engine/Plugins/Experimental/GeometryProcessing/Source/DynamicMesh/Public/Operations/RemoveOccludedTriangles.h
--- a/Engine/Plugins/Experimental/GeometryProcessing/Source/DynamicMesh/Public/Operations/RemoveOccludedTriangles.h
+++ b/Engine/Plugins/Experimental/GeometryProcessing/Source/DynamicMesh/Public/Operations/RemoveOccludedTriangles.h
@@ -259,8 +259,8 @@ public:
 					FVector3d XFPt = SpatialTransforms[SpatialIdx].InverseTransformUnlessMatch(Pt, OriginalPt, OriginalXF);
 					Ray.Direction = SpatialTransforms[SpatialIdx].InverseTransformVector(RayDirs[RayIdx]);
 					Ray.Origin = XFPt;
-					int HitTID = Spatials[SpatialIdx]->FindNearestHitTriangle(Ray);
-					if (HitTID != IndexConstants::InvalidID)
+					bool bFoundHit = Spatials[SpatialIdx]->TestAnyHitTriangle(Ray);
+					if (bFoundHit)
 					{
 						bAnyHit = true;
 						break;
@@ -288,8 +288,8 @@ public:
 				for (int RayIdx = 0; RayIdx < NR; ++RayIdx)
 				{
 					Ray.Direction = SpatialTransforms[SpatialIdx].InverseTransformVector(RayDirs[RayIdx]);
-					int HitTID = Spatials[SpatialIdx]->FindNearestHitTriangle(Ray);
-					if (HitTID == IndexConstants::InvalidID)
+					bool bFoundHit = Spatials[SpatialIdx]->TestAnyHitTriangle(Ray);
+					if (bFoundHit == false)
 					{
 						bIsOccluded = false;
 						break;
diff --git a/Engine/Plugins/Experimental/GeometryProcessing/Source/DynamicMesh/Public/Selections/MeshConnectedComponents.h b/Engine/Plugins/Experimental/GeometryProcessing/Source/DynamicMesh/Public/Selections/MeshConnectedComponents.h
--- a/Engine/Plugins/Experimental/GeometryProcessing/Source/DynamicMesh/Public/Selections/MeshConnectedComponents.h
+++ b/Engine/Plugins/Experimental/GeometryProcessing/Source/DynamicMesh/Public/Selections/MeshConnectedComponents.h
@@ -98,6 +98,11 @@ public:
 	/** @return element of Components array at given Index */
 	FComponent& GetComponent(int32 Index) { return Components[Index]; }
 
+	/** @return element of Components array at given Index */
+	const FComponent& operator[](int32 Index) const { return Components[Index]; }
+
+	/** @return element of Components array at given Index */
+	FComponent& operator[](int32 Index) { return Components[Index]; }
 
 	/** 
 	 * @return index of largest component by element count 
diff --git a/Engine/Plugins/Experimental/GeometryProcessing/Source/DynamicMesh/Public/Selections/MeshFaceSelection.h b/Engine/Plugins/Experimental/GeometryProcessing/Source/DynamicMesh/Public/Selections/MeshFaceSelection.h
--- a/Engine/Plugins/Experimental/GeometryProcessing/Source/DynamicMesh/Public/Selections/MeshFaceSelection.h
+++ b/Engine/Plugins/Experimental/GeometryProcessing/Source/DynamicMesh/Public/Selections/MeshFaceSelection.h
@@ -121,20 +121,18 @@ public:
 	}
 
 
-	void SelectVertexOneRing(int vid) {
-		for (int tid : Mesh->VtxTrianglesItr(vid))
+	void SelectVertexOneRing(int vid) 
+	{
+		Mesh->EnumerateVertexTriangles(vid, [&](int tid)
 		{
 			add(tid);
-		}
+		});
 	}
 	void SelectVertexOneRings(TArrayView<const int> Vertices)
 	{
 		for (int vid : Vertices)
 		{
-			for (int tid : Mesh->VtxTrianglesItr(vid))
-			{
-				add(tid);
-			}
+			SelectVertexOneRing(vid);
 		}
 	}
 
@@ -259,6 +257,26 @@ public:
 		return Bitmap;
 	}
 
+	/**
+	 * @return pair (bIsSelectionBoundaryEdge, bIsMeshBoundaryEdge) for a given edge
+	 */
+	TPair<bool, bool> IsSelectionBoundaryEdge(int32 eid) const
+	{
+		if (Mesh->IsEdge(eid))
+		{
+			FIndex2i EdgeT = Mesh->GetEdgeT(eid);
+			bool bA = Selected.Contains(EdgeT.A);
+			if (EdgeT.B == IndexConstants::InvalidID)
+			{
+				return TPair<bool, bool>(true, true);
+			}
+			else if (Selected.Contains(EdgeT.B) != bA)
+			{
+				return TPair<bool, bool>(true, false);
+			}
+		}
+		return TPair<bool, bool>(false, false);
+	}
 
 
 	/**
@@ -340,35 +358,7 @@ public:
 	 *  
 	 *  Return false from FilterF to prevent triangles from being included.
 	 */
-	void ExpandToOneRingNeighbours(const TUniqueFunction<bool(int)>& FilterF = nullptr)
-	{
-		TArray<int> ToAdd;
-
-		for (int tid : Selected) { 
-			FIndex3i tri_v = Mesh->GetTriangle(tid);
-			for (int j = 0; j < 3; ++j)
-			{
-				int vid = tri_v[j];
-				for (int nbr_t : Mesh->VtxTrianglesItr(vid))
-				{
-					if (FilterF && FilterF(nbr_t) == false)
-					{
-						continue;
-					}
-					if (!IsSelected(nbr_t))
-					{
-						ToAdd.Add(nbr_t);
-					}
-				}
-			}
-		}
-
-		for (int ID : ToAdd)
-		{
-			add(ID);
-		}
-	}
-
+	void ExpandToOneRingNeighbours(const TUniqueFunction<bool(int)>& FilterF = nullptr);
 
 	/**
 	 *  Expand selection by N vertex one-rings. This is *significantly* faster
@@ -377,59 +367,7 @@ public:
 	 *  
 	 *  Return false from FilterF to prevent triangles from being included.
 	 */
-	void ExpandToOneRingNeighbours(int nRings, const TUniqueFunction<bool(int)>& FilterF = nullptr)
-	{
-		if (nRings == 1)
-		{
-			ExpandToOneRingNeighbours(FilterF);
-			return;
-		}
-
-		TArray<int> triArrays[2];
-		int addIdx = 0, checkIdx = 1;
-		triArrays[checkIdx] = Selected.Array();
-
-		TBitArray<FDefaultBitArrayAllocator> Bitmap = AsBitArray();
-
-		for (int ri = 0; ri < nRings; ++ri)
-		{
-			TArray<int>& addTris = triArrays[addIdx];
-			TArray<int>& checkTris = triArrays[checkIdx];
-			addTris.Empty();
-
-			for (int tid : checkTris)
-			{
-				FIndex3i tri_v = Mesh->GetTriangle(tid);
-				for (int j = 0; j < 3; ++j)
-				{
-					int vid = tri_v[j];
-					for (int nbr_t : Mesh->VtxTrianglesItr(vid))
-					{
-						if (FilterF && FilterF(nbr_t) == false)
-						{
-							continue;
-						}
-						if (Bitmap[nbr_t] == false)
-						{
-							addTris.Add(nbr_t);
-							Bitmap[nbr_t] = true;
-						}
-					}
-				}
-			}
-
-			for (int TID : addTris)
-			{
-				add(TID);
-			}
-
-			Swap(addIdx, checkIdx); // check in the next iter what we added in this iter
-		}
-	}
-
-
-
-
+	void ExpandToOneRingNeighbours(int nRings, const TUniqueFunction<bool(int)>& FilterF = nullptr);
 
 	/**
 	 *  remove all triangles in vertex one-rings of current selection to set.
@@ -438,62 +376,7 @@ public:
 	 *  
 	 *  Return false from FilterF to prevent triangles from being deselected.
 	 */
-	void ContractBorderByOneRingNeighbours(int NumRings = 1, bool bContractFromMeshBoundary = false, const TUniqueFunction<bool(int)>& FilterF = nullptr)
-	{
-		TArray<int> BorderVIDs;   // border vertices
-		// TODO: track the border across removals if NumRings > 1, so we don't have to iterate over the whole selection each time
-		for (int RingIdx = 0; RingIdx < NumRings; ++RingIdx)
-		{
-			BorderVIDs.Reset();
-
-			// [TODO] border vertices are pushed onto the temp list multiple times.
-			// minor inefficiency, but maybe we could improve it?
-
-			// find set of vertices on border
-			for (int TID : Selected)
-			{
-				FIndex3i TriV = Mesh->GetTriangle(TID);
-				for (int Idx = 0; Idx < 3; ++Idx)
-				{
-					int VID = TriV[Idx];
-					int UnusedE1, UnusedE2; // Unused edge boundary indices, required to call GetVtxBoundaryEdges to test if vertex is on boundary
-					bool bIsBorder = bContractFromMeshBoundary && Mesh->GetVtxBoundaryEdges(VID, UnusedE1, UnusedE2) > 0;
-					if (!bIsBorder)
-					{
-						for (int NbrT : Mesh->VtxTrianglesItr(VID))
-						{
-							if (!IsSelected(NbrT))
-							{
-								bIsBorder = true;
-								break;
-							}
-						}
-					}
-					if (bIsBorder)
-					{
-						BorderVIDs.Add(VID);
-					}
-				}
-			}
-
-			for (int VID : BorderVIDs)
-			{
-				for (int NbrT : Mesh->VtxTrianglesItr(VID))
-				{
-					if (FilterF && !FilterF(NbrT))
-					{
-						continue;
-					}
-					Deselect(NbrT);
-				}
-			}
-		}
-
-	}
-
-
-
-
+	void ContractBorderByOneRingNeighbours(int NumRings = 1, bool bContractFromMeshBoundary = false, const TUniqueFunction<bool(int)>& FilterF = nullptr);
 
 	/**
 	 *  Grow selection outwards from seed triangle, until it hits boundaries defined by triangle and edge filters.
@@ -646,12 +529,14 @@ public:
 				vertices.Add(tv.A); vertices.Add(tv.B); vertices.Add(tv.C);
 			}
 
-			for (int vid : vertices) {
-				if (is_bowtie_vtx(vid)) {
-					for (int TID : Mesh->VtxTrianglesItr(vid))
+			for (int vid : vertices) 
+			{
+				if (is_bowtie_vtx(vid)) 
+				{
+					Mesh->EnumerateVertexTriangles(vid, [&](int TID)
 					{
 						Deselect(TID);
-					}
+					});
 					done = false;
 				}
 			}
@@ -667,7 +552,8 @@ private:
 	bool is_bowtie_vtx(int vid) const
 	{
 		int border_edges = 0;
-		for (int eid : Mesh->VtxEdgesItr(vid)) {
+		for (int eid : Mesh->VtxEdgesItr(vid)) 
+		{
 			FIndex2i et = Mesh->GetEdgeT(eid);
 			if (et.B != FDynamicMesh3::InvalidID)
 			{
diff --git a/Engine/Plugins/Experimental/GeometryProcessing/Source/GeometricObjects/Private/Spatial/SparseDynamicOctree3.cpp b/Engine/Plugins/Experimental/GeometryProcessing/Source/GeometricObjects/Private/Spatial/SparseDynamicOctree3.cpp
--- a/Engine/Plugins/Experimental/GeometryProcessing/Source/GeometricObjects/Private/Spatial/SparseDynamicOctree3.cpp
+++ b/Engine/Plugins/Experimental/GeometryProcessing/Source/GeometricObjects/Private/Spatial/SparseDynamicOctree3.cpp
@@ -313,6 +313,55 @@ int32 FSparseDynamicOctree3::FindNearestHitObject(const FRay3d& Ray,
 
 
 
+void FSparseDynamicOctree3::ContainmentQuery(
+	const FVector3d& Point,
+	TFunctionRef<void(int)> ObjectIDFunc) const
+{
+	// todo: this should take advantage of raster!
+
+	// always process spill objects
+	for (int ObjectID : SpillObjectSet)
+	{
+		ObjectIDFunc(ObjectID);
+	}
+
+	TArray<const FSparseOctreeCell*, TInlineAllocator<32>> Queue;
+
+	// start at root cells
+	RootCells.AllocatedIteration([&](const uint32* RootCellID)
+	{
+		const FSparseOctreeCell* RootCell = &Cells[*RootCellID];
+		if ( GetCellBox(*RootCell, MaxExpandFactor).Contains(Point) )
+		{
+			Queue.Add(&Cells[*RootCellID]);
+		}
+	});
+
+
+	while (Queue.Num() > 0)
+	{
+		const FSparseOctreeCell* CurCell = Queue.Pop(false);
+
+		// process elements
+		CellObjectLists.Enumerate(CurCell->CellID, [&](int32 ObjectID)
+		{
+			ObjectIDFunc(ObjectID);
+		});
+
+		for (int k = 0; k < 8; ++k)
+		{
+			if (CurCell->HasChild(k))
+			{
+				const FSparseOctreeCell* ChildCell = &Cells[CurCell->GetChildCellID(k)];
+				if (GetCellBox(*ChildCell, MaxExpandFactor).Contains(Point))
+				{
+					Queue.Add(ChildCell);
+				}
+			}
+		}
+	}
+}
+
 
 void FSparseDynamicOctree3::RangeQuery(
 	const FAxisAlignedBox3d& Bounds,
diff --git a/Engine/Plugins/Experimental/GeometryProcessing/Source/GeometricObjects/Public/Generators/MarchingCubes.h b/Engine/Plugins/Experimental/GeometryProcessing/Source/GeometricObjects/Public/Generators/MarchingCubes.h
--- a/Engine/Plugins/Experimental/GeometryProcessing/Source/GeometricObjects/Public/Generators/MarchingCubes.h
+++ b/Engine/Plugins/Experimental/GeometryProcessing/Source/GeometricObjects/Public/Generators/MarchingCubes.h
@@ -61,6 +61,13 @@ public:
 	 */
 	bool bParallelCompute = true;
 
+	/**
+	 * If true, code will assume that Implicit() is expensive enough that it is worth it to cache
+	 * evaluations when possible. For something simple like evaluation of an SDF defined by a discrete
+	 * grid, this is generally not worth the overhead.
+	 */
+	bool bEnableValueCaching = true;
+
 	/**
 	 * Max number of cells on any dimension; if exceeded, CubeSize will be automatically increased to fix
 	 */
@@ -108,6 +115,8 @@ public:
 	*/
 	FMeshShapeGenerator& Generate() override
 	{
+		TRACE_CPUPROFILER_EVENT_SCOPE(Geometry_MCMesh_Generate);
+
 		if (!ensure(Validate()))
 		{
 			return *this;
@@ -116,22 +125,34 @@ public:
 		SetDimensions();
 		GridBounds = FAxisAlignedBox3i(FVector3i::Zero(), CellDimensions - FVector3i(1,1,1)); // grid bounds are inclusive
 
-		corner_values_grid = FDenseGrid3f(CellDimensions.X+1, CellDimensions.Y+1, CellDimensions.Z+1, FMathf::MaxReal);
-		edge_vertices.Reset();
+		if (bEnableValueCaching)
+		{
+			corner_values_grid = FDenseGrid3f(CellDimensions.X + 1, CellDimensions.Y + 1, CellDimensions.Z + 1, FMathf::MaxReal);
+		}
 		corner_values.Reset();
+		InitHashTables();
+		ResetMesh();
 
-		if (bParallelCompute) {
+		if (bParallelCompute) 
+		{
 			generate_parallel();
-		} else {
+		} 
+		else 
+		{
 			generate_basic();
 		}
 
+		// finalize mesh
+		BuildMesh();
+
 		return *this;
 	}
 
 
 	FMeshShapeGenerator& GenerateContinuation(TArrayView<const FVector3<double>> Seeds)
 	{
+		TRACE_CPUPROFILER_EVENT_SCOPE(Geometry_MCMesh_GenerateContinuation);
+
 		if (!ensure(Validate()))
 		{
 			return *this;
@@ -140,11 +161,16 @@ public:
 		SetDimensions();
 		GridBounds = FAxisAlignedBox3i(FVector3i::Zero(), CellDimensions - FVector3i(1,1,1)); // grid bounds are inclusive
 
+		InitHashTables();
+		ResetMesh();
+		corner_values.Reset();
+
 		if (LastGridBounds != GridBounds)
 		{
-			corner_values_grid = FDenseGrid3f(CellDimensions.X + 1, CellDimensions.Y + 1, CellDimensions.Z + 1, FMathf::MaxReal);
-			edge_vertices.Reset();
-			corner_values.Reset();
+			if (bEnableValueCaching)
+			{
+				corner_values_grid = FDenseGrid3f(CellDimensions.X + 1, CellDimensions.Y + 1, CellDimensions.Z + 1, FMathf::MaxReal);
+			}
 			if (bParallelCompute)
 			{
 				done_cells = FDenseGrid3i(CellDimensions.X, CellDimensions.Y, CellDimensions.Z, 0);
@@ -152,21 +178,28 @@ public:
 		}
 		else
 		{
-			edge_vertices.Reset();
-			corner_values.Reset();
-			corner_values_grid.Assign(FMathf::MaxReal);
+			if (bEnableValueCaching)
+			{
+				corner_values_grid.Assign(FMathf::MaxReal);
+			}
 			if (bParallelCompute)
 			{
 				done_cells.Assign(0);
 			}
 		}
 
-		if (bParallelCompute) {
+		if (bParallelCompute) 
+		{
 			generate_continuation_parallel(Seeds);
-		} else {
+		} 
+		else 
+		{
 			generate_continuation(Seeds);
 		}
 
+		// finalize mesh
+		BuildMesh();
+
 		LastGridBounds = GridBounds;
 
 		return *this;
@@ -271,28 +304,41 @@ protected:
 	// Hash table for edge vertices
 	//
 
-	TMap<int64, int> edge_vertices;
-	FCriticalSection edge_vertices_lock;
+	int64 NumEdgeVertexSections = 64;
+	TArray<TMap<int64, int>> EdgeVertexSections;
+	TArray<FCriticalSection> EdgeVertexSectionLocks;
+	
+	int FindVertexID(int64 hash)
+	{
+		int32 SectionIndex = (int32)(hash % (NumEdgeVertexSections - 1));
+		FScopeLock Lock(&EdgeVertexSectionLocks[SectionIndex]);
+		int* Found = EdgeVertexSections[SectionIndex].Find(hash);
+		return (Found != nullptr) ? *Found : IndexConstants::InvalidID;
+	}
+
+	int AppendOrFindVertexID(int64 hash, FVector3<double> Pos)
+	{
+		int32 SectionIndex = (int32)(hash % (NumEdgeVertexSections - 1));
+		FScopeLock Lock(&EdgeVertexSectionLocks[SectionIndex]);
+		int* FoundVID = EdgeVertexSections[SectionIndex].Find(hash);
+		if (FoundVID != nullptr)
+		{
+			return *FoundVID;
+		}
+		int NewVID = append_vertex(Pos, hash);
+		EdgeVertexSections[SectionIndex].Add(hash, NewVID);
+		return NewVID;
+	}
+
 
 	int edge_vertex_id(const FVector3i& Idx1, const FVector3i& Idx2, double F1, double F2)
 	{
 		int64 hash = edge_hash(Idx1, Idx2);
 
-		int vid = IndexConstants::InvalidID;
-		bool found = false;
+		int foundvid = FindVertexID(hash);
+		if (foundvid != IndexConstants::InvalidID)
 		{
-			FScopeLock Lock(&edge_vertices_lock);
-			int* foundvid = edge_vertices.Find(hash);
-			if (foundvid)
-			{
-				vid = *foundvid;
-				found = true;
-			}
-		}
-		
-		if (found)
-		{
-			return vid;
+			return foundvid;
 		}
 
 		// ok this is a bit messy. We do not want to lock the entire hash table 
@@ -307,21 +353,7 @@ protected:
 		FVector3<double> Pos = FVector3<double>::Zero();
 		find_iso(pa, pb, F1, F2, Pos);
 
-		{
-			FScopeLock Lock(&edge_vertices_lock);
-			int* foundvid = edge_vertices.Find(hash);
-			if (!foundvid)
-			{
-				vid = append_vertex(Pos);
-				edge_vertices.Add(hash, vid);
-			}
-			else
-			{
-				vid = *foundvid;
-			}
-		}
-
-		return vid;
+		return AppendOrFindVertexID(hash, Pos);
 	}
 
 
@@ -384,8 +416,43 @@ protected:
 
 	FDenseGrid3f corner_values_grid;
 
+	int64 NumCornerValueSections = 64;
+	TArray<FCriticalSection> CornerValueSectionLocks;
+	int32 GetCornerValueSectionIndex(int64 hash)
+	{
+		return (int32)(hash % (NumCornerValueSections - 1));
+	}
+
+	double corner_value_grid_parallel(const FVector3i& Idx)
+	{
+		int32 SectionIndex = GetCornerValueSectionIndex(corner_hash(Idx));
+		float val = 0.0;
+		{
+			FScopeLock Lock(&CornerValueSectionLocks[SectionIndex]);
+			val = corner_values_grid[Idx];
+		}
+		if (val != FMathf::MaxReal)
+		{
+			return (double)val;
+		}
+
+		FVector3<double> V = corner_pos(Idx);
+		val = (float)Implicit(V);
+
+		{
+			FScopeLock Lock(&CornerValueSectionLocks[SectionIndex]);
+			corner_values_grid[Idx] = val;
+		}
+
+		return (double)val;
+	}
 	double corner_value_grid(const FVector3i& Idx)
 	{
+		if (bParallelCompute)
+		{
+			return corner_value_grid_parallel(Idx);
+		}
+
 		float val = corner_values_grid[Idx];
 		if (val != FMathf::MaxReal)
 		{
@@ -397,6 +464,7 @@ protected:
 		corner_values_grid[Idx] = val;
 		return (double)val;
 	}
+
 	void initialize_cell_values_grid(FGridCell& Cell, bool Shift)
 	{
 		if (Shift)
@@ -422,7 +490,8 @@ protected:
 	//
 	//
 
-	double corner_value_nohash(const FVector3i& Idx) {
+	double corner_value_nohash(const FVector3i& Idx) 
+	{
 		FVector3<double> V = corner_pos(Idx);
 		return Implicit(V);
 	}
@@ -446,8 +515,6 @@ protected:
 
 
 
-
-
 	/**
 	*  compute 3D corner-positions and field values for cell at index
 	*/
@@ -463,8 +530,14 @@ protected:
 		Cell.i[7] = FVector3i(Idx.X + 0, Idx.Y + 1, Idx.Z + 1);
 
 		//initialize_cell_values(Cell, false);
-		initialize_cell_values_grid(Cell, false);
-		//initialize_cell_values_nohash(Cell, false);
+		if (bEnableValueCaching)
+		{
+			initialize_cell_values_grid(Cell, false);
+		}
+		else
+		{
+			initialize_cell_values_nohash(Cell, false);
+		}
 	}
 
 
@@ -482,13 +555,34 @@ protected:
 		Cell.i[4].X = XIdx; Cell.i[5].X = XIdx+1; Cell.i[6].X = XIdx+1; Cell.i[7].X = XIdx;
 
 		//initialize_cell_values(Cell, true);
-		initialize_cell_values_grid(Cell, true);
-		//initialize_cell_values_nohash(Cell, true);
+		if (bEnableValueCaching)
+		{
+			initialize_cell_values_grid(Cell, true);
+		}
+		else
+		{
+			initialize_cell_values_nohash(Cell, true);
+		}
+	}
+
+
+	void InitHashTables()
+	{
+		EdgeVertexSections.Reset();
+		EdgeVertexSections.SetNum(NumEdgeVertexSections);
+		EdgeVertexSectionLocks.Reset();
+		EdgeVertexSectionLocks.SetNum(NumEdgeVertexSections);
+
+		CornerValueSectionLocks.Reset();
+		CornerValueSectionLocks.SetNum(NumCornerValueSections);
+
+		DoneCellSectionLocks.Reset();
+		DoneCellSectionLocks.SetNum(NumDoneCellSections);
 	}
 
 
 	bool parallel_mesh_access = false;
-	FCriticalSection mesh_lock;
+
 
 	/**
 	*  processing z-slabs of cells in parallel
@@ -666,16 +760,25 @@ protected:
 
 
 	FDenseGrid3i done_cells;
-	FCriticalSection done_cells_lock;
+
+	int64 NumDoneCellSections = 64;
+	TArray<FCriticalSection> DoneCellSectionLocks;
+	int GetDoneCellSectionIndex(int64 hash)
+	{
+		return (int32)(hash % (NumDoneCellSections - 1));
+	}
 
 	bool set_cell_if_not_done(const FVector3i& Idx)
 	{
+		int32 SectionIndex = GetDoneCellSectionIndex(corner_hash(Idx));
 		bool was_set = false;
-		FScopeLock Lock(&done_cells_lock);
-		if (done_cells[Idx] == 0)
 		{
-			done_cells[Idx] = 1;
-			was_set = true;
+			FScopeLock Lock(&DoneCellSectionLocks[SectionIndex]);
+			if (done_cells[Idx] == 0)
+			{
+				done_cells[Idx] = 1;
+				was_set = true;
+			}
 		}
 		return was_set;
 	}
@@ -729,6 +832,8 @@ protected:
 			Shift <<= 1;
 		}
 
+		int64 CellHash = corner_hash(Cell.i[0]);
+
 		// now iterate through the set of triangles in TriTable for this cube,
 		// and emit triangles using the vertices we found.
 		int tri_count = 0;
@@ -747,7 +852,7 @@ protected:
 				continue;
 			}
 
-			append_triangle(a, b, c);
+			append_triangle(a, b, c, CellHash);
 			tri_count++;
 		}
 
@@ -755,39 +860,113 @@ protected:
 	}
 
 
+	struct FIndexedVertex
+	{
+		int32 Index;
+		FVector3d Position;
+	};
+	std::atomic<int32> VertexCounter;
 
+	int64 NumVertexSections = 64;
+	TArray<FCriticalSection> VertexSectionLocks;
+	TArray<TArray<FIndexedVertex>> VertexSectionLists;
+	int GetVertexSectionIndex(int64 hash)
+	{
+		return (int32)(hash % (NumVertexSections - 1));
+	}
 
 	/**
 	*  add vertex to mesh, with locking if we are computing in parallel
 	*/
-	int append_vertex(FVector3<double> V)
+	int append_vertex(FVector3<double> V, int64 CellHash)
 	{
+		int SectionIndex = GetVertexSectionIndex(CellHash);
+		int32 NewIndex = VertexCounter++;
+
 		if (parallel_mesh_access)
 		{
-			FScopeLock Lock(&mesh_lock);
-			return AppendVertex((FVector3d)V);
+			FScopeLock Lock(&VertexSectionLocks[SectionIndex]);
+			VertexSectionLists[SectionIndex].Add(FIndexedVertex{ NewIndex, V });
 		}
 		else
 		{
-			return AppendVertex((FVector3d)V);
+			VertexSectionLists[SectionIndex].Add(FIndexedVertex{ NewIndex, V });
 		}
+
+		return NewIndex;
 	}
 
 
+	int64 NumTriangleSections = 64;
+	TArray<FCriticalSection> TriangleSectionLocks;
+	TArray<TArray<FIndex3i>> TriangleSectionLists;
+	int GetTriangleSectionIndex(int64 hash)
+	{
+		return (int32)(hash % (NumTriangleSections - 1));
+	}
 
 	/**
 	*  add triangle to mesh, with locking if we are computing in parallel
 	*/
-	int append_triangle(int A, int B, int C)
+	void append_triangle(int A, int B, int C, int64 CellHash)
 	{
+		int SectionIndex = GetTriangleSectionIndex(CellHash);
 		if (parallel_mesh_access)
 		{
-			FScopeLock Lock(&mesh_lock);
-			return AppendTriangle(A, B, C);
+			FScopeLock Lock(&TriangleSectionLocks[SectionIndex]);
+			TriangleSectionLists[SectionIndex].Add(FIndex3i(A, B, C));
 		}
 		else
 		{
-			return AppendTriangle(A, B, C);
+			TriangleSectionLists[SectionIndex].Add(FIndex3i(A, B, C));
+		}
+	}
+
+
+	/**
+	 * Reset internal mesh-assembly data structures
+	 */
+	void ResetMesh()
+	{
+		VertexSectionLocks.SetNum(NumVertexSections);
+		VertexSectionLists.Reset();
+		VertexSectionLists.SetNum(NumVertexSections);
+		VertexCounter = 0;
+
+		TriangleSectionLocks.SetNum(NumTriangleSections);
+		TriangleSectionLists.Reset();
+		TriangleSectionLists.SetNum(NumTriangleSections);
+	}
+
+	/**
+	 * Populate FMeshShapeGenerator data structures from accumulated
+	 * vertex/triangle sets
+	 */
+	void BuildMesh()
+	{
+		TRACE_CPUPROFILER_EVENT_SCOPE(Geometry_MCMesh_BuildMesh);
+
+		int32 NumVertices = VertexCounter;
+		TArray<FVector3d> VertexBuffer;
+		VertexBuffer.SetNum(NumVertices);
+		for (const TArray<FIndexedVertex>& VertexList : VertexSectionLists)
+		{
+			for (FIndexedVertex Vtx : VertexList)
+			{
+				VertexBuffer[Vtx.Index] = Vtx.Position;
+			}
+		}
+		for (int32 k = 0; k < NumVertices; ++k)
+		{
+			int32 vid = AppendVertex(VertexBuffer[k]);
+		}
+
+		for (const TArray<FIndex3i>& TriangleList : TriangleSectionLists)
+		{
+			for (FIndex3i Tri : TriangleList)
+			{
+				AppendTriangle(Tri.A, Tri.B, Tri.C);
+			}
 		}
 	}
 
diff --git a/Engine/Plugins/Experimental/GeometryProcessing/Source/GeometricObjects/Public/Implicit/Morphology.h b/Engine/Plugins/Experimental/GeometryProcessing/Source/GeometricObjects/Public/Implicit/Morphology.h
--- a/Engine/Plugins/Experimental/GeometryProcessing/Source/GeometricObjects/Public/Implicit/Morphology.h
+++ b/Engine/Plugins/Experimental/GeometryProcessing/Source/GeometricObjects/Public/Implicit/Morphology.h
@@ -169,6 +169,8 @@ protected:
 		{
 			return -Interpolant.Value(Pt);
 		};
+		MarchingCubes.bEnableValueCaching = false;
+
 		{
 			TRACE_CPUPROFILER_EVENT_SCOPE(Geometry_Morphology_FirstPass_GenerateMesh);
 			MarchingCubes.Generate();
@@ -246,6 +248,7 @@ protected:
 		{
 			return -Interpolant.Value(Pt);
 		};
+		MarchingCubes.bEnableValueCaching = false;
 
 		{
 			TRACE_CPUPROFILER_EVENT_SCOPE(Geometry_Morphology_SecondPass_GenerateMesh);
diff --git a/Engine/Plugins/Experimental/GeometryProcessing/Source/GeometricObjects/Public/Spatial/SparseDynamicOctree3.h b/Engine/Plugins/Experimental/GeometryProcessing/Source/GeometricObjects/Public/Spatial/SparseDynamicOctree3.h
--- a/Engine/Plugins/Experimental/GeometryProcessing/Source/GeometricObjects/Public/Spatial/SparseDynamicOctree3.h
+++ b/Engine/Plugins/Experimental/GeometryProcessing/Source/GeometricObjects/Public/Spatial/SparseDynamicOctree3.h
@@ -252,6 +252,14 @@ public:
 		TFunctionRef<double(int, const FRay3d&)> HitObjectDistFunc,
 		double MaxDistance = TNumericLimits<double>::Max()) const;
 
+	/**
+	 * Process ObjectIDs from all the cells with bounding boxes that contain query point
+	 * @param Point query point
+	 * @param ObjectIDFunc this function is called
+	 */
+	void ContainmentQuery(const FVector3d& Point,
+		TFunctionRef<void(int)> ObjectIDFunc) const;
+
 	/**
 	 * Process ObjectIDs from all the cells with bounding boxes that intersect Bounds
 	 * @param Bounds query box
