More optimization of FNormalFlowRemesher:
- for temp buffers, don't pre-initialize all their values, just make sure their size is big enough
- parallelize the decision of choosing which edges should be queued for later processing

diff --git a/Engine/Plugins/Experimental/GeometryProcessing/Source/DynamicMesh/Private/NormalFlowRemesher.cpp b/Engine/Plugins/Experimental/GeometryProcessing/Source/DynamicMesh/Private/NormalFlowRemesher.cpp
--- a/Engine/Plugins/Experimental/GeometryProcessing/Source/DynamicMesh/Private/NormalFlowRemesher.cpp
+++ b/Engine/Plugins/Experimental/GeometryProcessing/Source/DynamicMesh/Private/NormalFlowRemesher.cpp
@@ -126,12 +126,44 @@ void FNormalFlowRemesher::TrackedFaceProjectionPass(double& MaxDistanceMoved, bo
 	IOrientedProjectionTarget* NormalProjTarget = static_cast<IOrientedProjectionTarget*>(ProjTarget);
 	ensure(NormalProjTarget != nullptr);
 
-	InitializeVertexBufferForFacePass();
+	{
+		// Make sure the temp buffers are an adequate size
 
-	double MaxStepDistance = (bIsTuningIteration) ? (0.25 * MaxEdgeLength) : (0.5 * MaxEdgeLength);		// kind of arbitrary...
-	double SmoothAreaDistance = FillAreaDistanceMultiplier * MaxEdgeLength;
+		int NumTriangleVertices = 3 * Mesh->MaxTriangleID();
+		if (ProjectionVertexBuffer.Num() < NumTriangleVertices)
+		{
+			ProjectionVertexBuffer.SetNumUninitialized(2 * NumTriangleVertices);
+		}
+		if (ProjectionWeightBuffer.Num() < NumTriangleVertices)
+		{
+			ProjectionWeightBuffer.SetNumUninitialized(2 * NumTriangleVertices);
+		}
 
-	TFunction<FVector3d(const FDynamicMesh3&, int, double)> UseSmoothFunc = GetSmoothFunction();
+		int MaxVertexID = Mesh->MaxVertexID();
+		if (TempWeightBuffer.Num() < MaxVertexID)
+		{
+			TempWeightBuffer.SetNumUninitialized(2 * MaxVertexID);
+		}
+		if (TempPosBuffer.Num() < MaxVertexID)
+		{
+			TempPosBuffer.SetNum(2 * MaxVertexID);
+		}
+		if (TempFlagBuffer.Num() < MaxVertexID)
+		{
+			TempFlagBuffer.SetNumUninitialized(2 * MaxVertexID);
+		}
+
+		int MaxEdgeIDs = Mesh->MaxEdgeID();
+		if (EdgeShouldBeQueuedBuffer.Num() < MaxEdgeIDs)
+		{
+			EdgeShouldBeQueuedBuffer.SetNumUninitialized(2 * MaxEdgeIDs);
+		}
+	}
+
+	const double MaxStepDistance = (bIsTuningIteration) ? (0.25 * MaxEdgeLength) : (0.5 * MaxEdgeLength);		// kind of arbitrary...
+	const double SmoothAreaDistance = FillAreaDistanceMultiplier * MaxEdgeLength;
+
+	const TFunction<FVector3d(const FDynamicMesh3&, int, double)> UseSmoothFunc = GetSmoothFunction();
 
 	// For each triangle, rotate it such that it aligns with closest normal on the target surface.
 	// Each vertex is then assigned a weighted combination of its corresponding triangle corner positions. The weighting
@@ -301,27 +333,42 @@ void FNormalFlowRemesher::TrackedFaceProjectionPass(double& MaxDistanceMoved, bo
 	});
 
 
-	MaxDistanceMoved = 0.0;
+	// We queue any edges that moved far enough to fall under min/max edge length thresholds
+	ParallelFor(Mesh->MaxEdgeID(), [this](int32 EdgeID)
+	{
+		EdgeShouldBeQueuedBuffer[EdgeID] = false;
+
+		if (!Mesh->IsEdge(EdgeID))
+		{
+			return;
+		}
+
+		const FIndex2i& EdgeVertices = Mesh->GetEdgeV(EdgeID);
+		const double NewEdgeLength = Distance(TempPosBuffer[EdgeVertices[0]], TempPosBuffer[EdgeVertices[1]]);
+		if (NewEdgeLength < MinEdgeLength || NewEdgeLength > MaxEdgeLength)
+		{
+			EdgeShouldBeQueuedBuffer[EdgeID] = true;
+		}
+	});
 
-	// We queue any edges that moved far enough to fall under min/max edge length thresholds. Also find the max distance
-	// moved by a vertex.
 	for (int EdgeID : Mesh->EdgeIndicesItr())
 	{
-		FIndex2i EdgeVertices = Mesh->GetEdgeV(EdgeID);
-		double NewEdgeLength = Distance(TempPosBuffer[EdgeVertices[0]], TempPosBuffer[EdgeVertices[1]]);
-		if (NewEdgeLength < MinEdgeLength || NewEdgeLength > MaxEdgeLength)
+		if (EdgeShouldBeQueuedBuffer[EdgeID])
 		{
 			QueueEdge(EdgeID);
 		}
+	}
+
 
-		for (int I = 0; I < 2; ++I)
+	// Return the maximum distance moved by a vertex
+	MaxDistanceMoved = 0.0;
+	for (int VertexID = 0; VertexID < Mesh->MaxVertexID(); ++VertexID)
+	{
+		if (TempFlagBuffer[VertexID] && Mesh->IsVertex(VertexID))
 		{
-			if (TempFlagBuffer[EdgeVertices[I]])
-			{
-				const FVector3d& CurrentPosition = Mesh->GetVertex(EdgeVertices[I]);
-				const FVector3d& ProjectedPosition = TempPosBuffer[EdgeVertices[I]];
-				MaxDistanceMoved = FMath::Max(MaxDistanceMoved, Distance(CurrentPosition, ProjectedPosition));
-			}
+			const FVector3d& CurrentPosition = Mesh->GetVertex(VertexID);
+			const FVector3d& ProjectedPosition = TempPosBuffer[VertexID];
+			MaxDistanceMoved = FMath::Max(MaxDistanceMoved, Distance(CurrentPosition, ProjectedPosition));
 		}
 	}
 
