GeometryCore: create new GeometryCore Module. Move GeometricObjects code from GeometryProcessing plugin to this new module, as well as core FDynamicMesh3 classes (attributes, aabbtree/octree, change tracking). Update module Build.cs files that reference GeometricObjects.
#rnx

diff --git a/Engine/Plugins/Experimental/Enterprise/DataprepGeometryOperations/Source/DataprepGeometryOperations.Build.cs b/Engine/Plugins/Experimental/Enterprise/DataprepGeometryOperations/Source/DataprepGeometryOperations.Build.cs
--- a/Engine/Plugins/Experimental/Enterprise/DataprepGeometryOperations/Source/DataprepGeometryOperations.Build.cs
+++ b/Engine/Plugins/Experimental/Enterprise/DataprepGeometryOperations/Source/DataprepGeometryOperations.Build.cs
@@ -16,7 +16,7 @@ namespace UnrealBuildTool.Rules
 					"DynamicMesh",
 					"EditorScriptingUtilities",
 					"Engine",
-					"GeometricObjects",
+					"GeometryCore",
 					"MeshConversion",
 					"MeshDescription",
 					"MeshModelingTools",
diff --git a/Engine/Plugins/Experimental/GeometryFlow/Source/GeometryFlowCore/GeometryFlowCore.Build.cs b/Engine/Plugins/Experimental/GeometryFlow/Source/GeometryFlowCore/GeometryFlowCore.Build.cs
--- a/Engine/Plugins/Experimental/GeometryFlow/Source/GeometryFlowCore/GeometryFlowCore.Build.cs
+++ b/Engine/Plugins/Experimental/GeometryFlow/Source/GeometryFlowCore/GeometryFlowCore.Build.cs
@@ -26,7 +26,7 @@ public class GeometryFlowCore : ModuleRules
 			new string[]
 			{
 				"Core",
-				"GeometricObjects"
+				"GeometryCore"
 			}
             );
 			
diff --git a/Engine/Plugins/Experimental/GeometryFlow/Source/GeometryFlowMeshProcessing/GeometryFlowMeshProcessing.Build.cs b/Engine/Plugins/Experimental/GeometryFlow/Source/GeometryFlowMeshProcessing/GeometryFlowMeshProcessing.Build.cs
--- a/Engine/Plugins/Experimental/GeometryFlow/Source/GeometryFlowMeshProcessing/GeometryFlowMeshProcessing.Build.cs
+++ b/Engine/Plugins/Experimental/GeometryFlow/Source/GeometryFlowMeshProcessing/GeometryFlowMeshProcessing.Build.cs
@@ -26,7 +26,7 @@ public class GeometryFlowMeshProcessing : ModuleRules
 			new string[]
 			{
 				"Core",
-				"GeometricObjects",
+				"GeometryCore",
 				"GeometryAlgorithms",
 				"DynamicMesh",
 				"GeometryFlowCore"
diff --git a/Engine/Plugins/Experimental/GeometryProcessing/Source/DynamicMesh/DynamicMesh.Build.cs b/Engine/Plugins/Experimental/GeometryProcessing/Source/DynamicMesh/DynamicMesh.Build.cs
--- a/Engine/Plugins/Experimental/GeometryProcessing/Source/DynamicMesh/DynamicMesh.Build.cs
+++ b/Engine/Plugins/Experimental/GeometryProcessing/Source/DynamicMesh/DynamicMesh.Build.cs
@@ -8,19 +8,10 @@ public class DynamicMesh : ModuleRules
 	{
         PCHUsage = ModuleRules.PCHUsageMode.UseExplicitOrSharedPCHs;
 
-        // For GPUSkinPublicDefs.h
-        PublicIncludePaths.Add("Runtime/Engine/Public");
-        
-        PublicIncludePathModuleNames.AddRange(
-			new string[] {
-				"AnimationCore",			// For the BoneWeights.h include
-			}
-		);
-
         PublicDependencyModuleNames.AddRange(
 			new string[] {
 				"Core",
-				"GeometricObjects",
+				"GeometryCore",
 				"GeometryAlgorithms"
 			}
 		);
diff --git a/Engine/Plugins/Experimental/GeometryProcessing/Source/DynamicMesh/Public/DynamicAttribute.h b/Engine/Plugins/Experimental/GeometryProcessing/Source/DynamicMesh/Public/DynamicAttribute.h
--- a/Engine/Plugins/Experimental/GeometryProcessing/Source/DynamicMesh/Public/DynamicAttribute.h
+++ b/Engine/Plugins/Experimental/GeometryProcessing/Source/DynamicMesh/Public/DynamicAttribute.h
@@ -2,351 +2,4 @@
 
 #pragma once
 
-#include "DynamicMesh3.h"
-#include "MeshIndexMappings.h"
-
-namespace UE
-{
-namespace Geometry
-{
-
-struct FMeshIndexMappings;
-
-template<typename ParentType>
-class TDynamicAttributeBase;
-
-/**
-* Generic base class for change tracking of an attribute layer
-*/
-template<typename ParentType>
-class TDynamicAttributeChangeBase
-{
-public:
-	virtual ~TDynamicAttributeChangeBase()
-	{
-	}
-
-	// default do-nothing implementations are provided because many attribute layers will only care about some kinds of elements and won't implement all of these
-
-	virtual void SaveInitialTriangle(const TDynamicAttributeBase<ParentType>* Attribute, int TriangleID)
-	{
-	}
-	virtual void SaveInitialVertex(const TDynamicAttributeBase<ParentType>* Attribute, int VertexID)
-	{
-	}
-
-	virtual void StoreAllFinalTriangles(const TDynamicAttributeBase<ParentType>* Attribute, const TArray<int>& TriangleIDs)
-	{
-	}
-	virtual void StoreAllFinalVertices(const TDynamicAttributeBase<ParentType>* Attribute, const TArray<int>& VertexIDs)
-	{
-	}
-
-	virtual bool Apply(TDynamicAttributeBase<ParentType>* Attribute, bool bRevert) const
-	{
-		return false;
-	}
-};
-
-typedef TDynamicAttributeChangeBase<FDynamicMesh3> FDynamicMeshAttributeChangeBase;
-
-
-/**
- * Base class for attributes that live on a dynamic mesh (or similar dynamic object)
- *
- * Subclasses can override the On* functions to ensure the attribute remains up to date through changes to the dynamic object
- */
-template<typename ParentType>
-class TDynamicAttributeBase
-{
-
-public:
-	virtual ~TDynamicAttributeBase()
-	{
-	}
-
-
-public:
-	/** Get optional identifier for this attribute set. */
-	FName GetName() const { return Name; }
-
-	/** Set optional identifier for this attribute set. */
-	void SetName(FName NameIn) { Name = NameIn; }
-
-protected:
-	/** Optional FName identifier for this attribute set. Not guaranteed to be unique. */
-	FName Name = FName();
-
-
-
-
-public:
-
-
-	/** Allocate a new copy of the attribute layer, optionally with a different parent */
-	virtual TDynamicAttributeBase* MakeCopy(ParentType* ParentIn) const = 0;
-	/** Allocate a new empty instance of the same type of attribute layer */
-	virtual TDynamicAttributeBase* MakeNew(ParentType* ParentIn) const = 0;
-	/**
-	 * Allocate a new compact copy of the attribute layer, optionally with a different parent.
-	 * Default implementation does a full copy and then compacts it, usually derived class will want to override this with a more efficient direct compact copy implementation
-	 */
-	virtual TDynamicAttributeBase* MakeCompactCopy(const FCompactMaps& CompactMaps, ParentType* ParentIn) const
-	{
-		TDynamicAttributeBase* Copy = MakeCopy(ParentIn);
-		Copy->CompactInPlace(CompactMaps);
-		return Copy;
-	}
-
-	/** Compact the attribute in place */
-	virtual void CompactInPlace(const FCompactMaps& CompactMaps) = 0;
-	
-	/** Update any held pointer to the parent */
-	virtual void Reparent(ParentType* NewParent) = 0;
-
-	/**
-	  * Copy data from a different attribute to this one, using the mesh index mapping to determine the correspondence 
-	  * @param Source copy attribute data from this source
-	  * @param Mapping the correspondence from Source's parent to this attribute's parent 
-	  * @return true if the copy succeeded, false otherwise (e.g. false if the data from the source attribute was not compatible and the CopyOut failed to copy across)
-	  */
-	virtual bool CopyThroughMapping(const TDynamicAttributeBase* Source, const FMeshIndexMappings& Mapping) = 0;
-
-	/** Generic function to copy data out of an attribute; it's up to the derived class to map RawID to chunks of attribute data */
-	virtual bool CopyOut(int RawID, void* Buffer, int BufferSize) const = 0;
-
-	/** Generic function to copy data in to an attribute; it's up to the derived class to map RawID to chunks of attribute data */
-	virtual bool CopyIn(int RawID, void* Buffer, int BufferSize) = 0;
-
-	virtual void OnNewVertex(int VertexID, bool bInserted)
-	{
-	}
-
-	virtual void OnRemoveVertex(int VertexID)
-	{
-	}
-
-	virtual void OnNewTriangle(int TriangleID, bool bInserted)
-	{
-	}
-
-	virtual void OnRemoveTriangle(int TriangleID)
-	{
-	}
-
-	virtual void OnReverseTriOrientation(int TriangleID)
-	{
-	}
-
-	/**
-	* Check validity of attribute
-	* 
-	* @param bAllowNonmanifold Accept non-manifold topology as valid. Note that this should almost always be true for attributes; non-manifold overlays are generally valid.
-	* @param FailMode Desired behavior if mesh is found invalid
-	*/
-	virtual bool CheckValidity(bool bAllowNonmanifold, EValidityCheckFailMode FailMode) const
-	{
-		// default impl just doesn't check anything; override with any useful sanity checks
-		return true;
-	}
-
-
-	virtual TUniquePtr<TDynamicAttributeChangeBase<ParentType>> NewBlankChange() = 0;
-
-
-
-	/** Update to reflect an edge split in the parent mesh */
-	virtual void OnSplitEdge(const DynamicMeshInfo::FEdgeSplitInfo& SplitInfo)
-	{
-	}
-
-	/** Update to reflect an edge flip in the parent mesh */
-	virtual void OnFlipEdge(const DynamicMeshInfo::FEdgeFlipInfo& FlipInfo)
-	{
-	}
-
-	/** Update to reflect an edge collapse in the parent mesh */
-	virtual void OnCollapseEdge(const DynamicMeshInfo::FEdgeCollapseInfo& CollapseInfo)
-	{
-	}
-
-	/** Update to reflect a face poke in the parent mesh */
-	virtual void OnPokeTriangle(const DynamicMeshInfo::FPokeTriangleInfo& PokeInfo)
-	{
-	}
-
-	/** Update to reflect an edge merge in the parent mesh */
-	virtual void OnMergeEdges(const DynamicMeshInfo::FMergeEdgesInfo& MergeInfo)
-	{
-	}
-
-	/** Update to reflect an edge merge in the parent mesh */
-	virtual void OnSplitVertex(const DynamicMeshInfo::FVertexSplitInfo& SplitInfo, const TArrayView<const int>& TrianglesToUpdate)
-	{
-	}
-
-protected:
-
-	/** 
-	 * Implementation of parent-class copy. MakeCopy() and MakeCompactCopy() implementations should call
-	 * this to transfer any custom data added by parent attribute set class.
-	 */
-	virtual void CopyParentClassData(const TDynamicAttributeBase<ParentType>& Other)
-	{
-		Name = Other.Name;
-	}
-};
-
-
-typedef TDynamicAttributeBase<FDynamicMesh3> FDynamicMeshAttributeBase;
-
-
-/**
-* Generic base class for managing a set of registered attributes that must all be kept up to date
-*/
-template<typename ParentType>
-class TDynamicAttributeSetBase
-{
-protected:
-	// not managed by the base class; we should be able to register any attributes here that we want to be automatically updated
-	TArray<TDynamicAttributeBase<ParentType>*> RegisteredAttributes;
-
-	/**
-	 * Stores the given attribute pointer in the attribute register, so that it will be updated with mesh changes, but does not take ownership of the attribute memory.
-	 */
-	void RegisterExternalAttribute(TDynamicAttributeBase<ParentType>* Attribute)
-	{
-		RegisteredAttributes.Add(Attribute);
-	}
-
-	void UnregisterExternalAttribute(TDynamicAttributeBase<ParentType>* Attribute)
-	{
-		RegisteredAttributes.Remove(Attribute);
-	}
-
-	void ResetRegisteredAttributes()
-	{
-		RegisteredAttributes.Reset();
-	}
-
-public:
-	virtual ~TDynamicAttributeSetBase()
-	{
-	}
-
-	int NumRegisteredAttributes() const
-	{
-		return RegisteredAttributes.Num();
-	}
-
-	TDynamicAttributeBase<ParentType>* GetRegisteredAttribute(int Idx) const
-	{
-		return RegisteredAttributes[Idx];
-	}
-
-	// These functions are called by the FDynamicMesh3 to update the various
-	// attributes when the parent mesh topology has been modified.
-	virtual void OnNewTriangle(int TriangleID, bool bInserted)
-	{
-		for (TDynamicAttributeBase<ParentType>* A : RegisteredAttributes)
-		{
-			A->OnNewTriangle(TriangleID, bInserted);
-		}
-	}
-	virtual void OnNewVertex(int VertexID, bool bInserted)
-	{
-		for (TDynamicAttributeBase<ParentType>* A : RegisteredAttributes)
-		{
-			A->OnNewVertex(VertexID, bInserted);
-		}
-	}
-	virtual void OnRemoveTriangle(int TriangleID)
-	{
-		for (TDynamicAttributeBase<ParentType>* A : RegisteredAttributes)
-		{
-			A->OnRemoveTriangle(TriangleID);
-		}
-	}
-	virtual void OnRemoveVertex(int VertexID)
-	{
-		for (TDynamicAttributeBase<ParentType>* A : RegisteredAttributes)
-		{
-			A->OnRemoveTriangle(VertexID);
-		}
-	}
-	virtual void OnReverseTriOrientation(int TriangleID)
-	{
-		for (TDynamicAttributeBase<ParentType>* A : RegisteredAttributes)
-		{
-			A->OnReverseTriOrientation(TriangleID);
-		}
-	}
-
-	/**
-	* Check validity of attributes
-	* 
-	* @param bAllowNonmanifold Accept non-manifold topology as valid. Note that this should almost always be true for attributes; non-manifold overlays are generally valid.
-	* @param FailMode Desired behavior if mesh is found invalid
-	*/
-	virtual bool CheckValidity(bool bAllowNonmanifold, EValidityCheckFailMode FailMode) const
-	{
-		bool bValid = true;
-		for (TDynamicAttributeBase<ParentType>* A : RegisteredAttributes)
-		{
-			bValid = A->CheckValidity(bAllowNonmanifold, FailMode) && bValid;
-		}
-		return bValid;
-	}
-
-
-	// mesh-specific on* functions; may be split out
-public:
-
-	virtual void OnSplitEdge(const DynamicMeshInfo::FEdgeSplitInfo& SplitInfo)
-	{
-		for (TDynamicAttributeBase<ParentType>* A : RegisteredAttributes)
-		{
-			A->OnSplitEdge(SplitInfo);
-		}
-	}
-	virtual void OnFlipEdge(const DynamicMeshInfo::FEdgeFlipInfo& FlipInfo)
-	{
-		for (TDynamicAttributeBase<ParentType>* A : RegisteredAttributes)
-		{
-			A->OnFlipEdge(FlipInfo);
-		}
-	}
-	virtual void OnCollapseEdge(const DynamicMeshInfo::FEdgeCollapseInfo& CollapseInfo)
-	{
-		for (TDynamicAttributeBase<ParentType>* A : RegisteredAttributes)
-		{
-			A->OnCollapseEdge(CollapseInfo);
-		}
-	}
-	virtual void OnPokeTriangle(const DynamicMeshInfo::FPokeTriangleInfo& PokeInfo)
-	{
-		for (TDynamicAttributeBase<ParentType>* A : RegisteredAttributes)
-		{
-			A->OnPokeTriangle(PokeInfo);
-		}
-	}
-	virtual void OnMergeEdges(const DynamicMeshInfo::FMergeEdgesInfo& MergeInfo)
-	{
-		for (TDynamicAttributeBase<ParentType>* A : RegisteredAttributes)
-		{
-			A->OnMergeEdges(MergeInfo);
-		}
-	}
-	virtual void OnSplitVertex(const DynamicMeshInfo::FVertexSplitInfo& SplitInfo, const TArrayView<const int>& TrianglesToUpdate)
-	{
-		for (TDynamicAttributeBase<ParentType>* A : RegisteredAttributes)
-		{
-			A->OnSplitVertex(SplitInfo, TrianglesToUpdate);
-		}
-	}
-};
-
-typedef TDynamicAttributeSetBase<FDynamicMesh3> FDynamicMeshAttributeSetBase;
-
-} // end namespace UE::Geometry
-} // end namespace UE
\ No newline at end of file
+#include "DynamicMesh/DynamicAttribute.h"
\ No newline at end of file
diff --git a/Engine/Plugins/Experimental/GeometryProcessing/Source/DynamicMesh/Public/DynamicMesh3.h b/Engine/Plugins/Experimental/GeometryProcessing/Source/DynamicMesh/Public/DynamicMesh3.h
--- a/Engine/Plugins/Experimental/GeometryProcessing/Source/DynamicMesh/Public/DynamicMesh3.h
+++ b/Engine/Plugins/Experimental/GeometryProcessing/Source/DynamicMesh/Public/DynamicMesh3.h
@@ -4,1482 +4,4 @@
 
 #pragma once
 
-#include "BoxTypes.h"
-#include "FrameTypes.h"
-#include "GeometryTypes.h"
-#include "HAL/Platform.h"
-#include "IndexTypes.h"
-#include "InfoTypes.h"
-#include "MathUtil.h"
-#include "Quaternion.h"
-#include "Util/CompactMaps.h"
-#include "Util/DynamicVector.h"
-#include "Util/IndexUtil.h"
-#include "Util/IteratorUtil.h"
-#include "Util/RefCountVector.h"
-#include "Util/SmallListSet.h"
-#include "VectorTypes.h"
-#include "VectorUtil.h"
-
-namespace UE
-{
-namespace Geometry
-{
-
-class FMeshShapeGenerator;
-class FDynamicMeshAttributeSet;
-
-enum class EMeshComponents : uint8
-{
-	None          = 0,
-	VertexNormals = 1,
-	VertexColors  = 2,
-	VertexUVs     = 4,
-	FaceGroups    = 8,
-	All           = 15
-};
-
-/**
-* FDynamicMesh3 is a dynamic triangle mesh class. The mesh has has connectivity,
-* is an indexed mesh, and allows for gaps in the index space.
-*
-* internally, all data is stored in POD-type buffers, except for the vertex->edge
-* links, which are stored as List<int>'s. The arrays of POD data are stored in
-* TDynamicVector's, so they grow in chunks, which is relatively efficient. The actual
-* blocks are arrays, so they can be efficiently mem-copied into larger buffers
-* if necessary.
-*
-* Reference counts for verts/tris/edges are stored as separate FRefCountVector
-* instances.
-*
-* Vertices are stored as doubles, although this should be easily changed
-* if necessary, as the internal data structure is not exposed
-*
-* Per-vertex Vertex Normals, Colors, and UVs are optional and stored as floats.
-*
-* For each vertex, VertexEdgeLists[i] is the unordered list of connected edges. The
-* elements of the list are indices into the edges list.
-* This list is unsorted but can be traversed in-order (ie cw/ccw) at some additional cost.
-*
-* Triangles are stored as 3 ints, with optionally a per-triangle integer group id.
-*
-* The edges of a triangle are similarly stored as 3 ints, in triangle_edes. If the
-* triangle is [v1,v2,v3], then the triangle edges [e1,e2,e3] are
-* e1=edge(v1,v2), e2=edge(v2,v3), e3=edge(v3,v1), where the e# are indexes into edges.
-*
-* Edges are stored as tuples of 4 ints. If the edge is between v1 and v2, with neighbour
-* tris t1 and t2, then the edge is [min(v1,v2), max(v1,v2), t1, t2]. For a boundary
-* edge, t2 is InvalidID. t1 is never InvalidID.
-*
-* Most of the class assumes that the mesh is manifold. Many functions will
-* work if the topology is non-manifold, but behavior of operators like Split/Flip/Collapse
-* edge is untested.
-*
-* The function CheckValidity() does extensive sanity checking on the mesh data structure.
-* Use this to test your code, both for mesh construction and editing!!
-*
-*
-* TODO:
-*  - Many of the iterators depend on lambda functions, can we replace these with calls to
-*    internal/static functions that do the same thing?
-*  - efficient TriTrianglesItr() implementation?
-*  - additional Topology timestamp?
-*  - CompactInPlace() does not compact VertexEdgeLists
-*  ? TDynamicVector w/ 'stride' option, so that we can guarantee that tuples are in single block.
-*    The can have custom accessor that looks up entire tuple
-*/
-class DYNAMICMESH_API FDynamicMesh3
-{
-public:
-	using FEdgeFlipInfo     = DynamicMeshInfo::FEdgeFlipInfo;
-	using FEdgeSplitInfo    = DynamicMeshInfo::FEdgeSplitInfo;
-	using FEdgeCollapseInfo = DynamicMeshInfo::FEdgeCollapseInfo;
-	using FMergeEdgesInfo   = DynamicMeshInfo::FMergeEdgesInfo;
-	using FPokeTriangleInfo = DynamicMeshInfo::FPokeTriangleInfo;
-	using FVertexSplitInfo  = DynamicMeshInfo::FVertexSplitInfo;
-
-	struct FEdge
-	{
-		FIndex2i Vert;
-		FIndex2i Tri;
-		friend bool operator!=(const FEdge& e0, const FEdge& e1)
-		{
-			return (e0.Vert != e1.Vert) || (e0.Tri != e1.Tri);
-		}
-
-		/**
-		 * Serialization operator for FEdge.
-		 *
-		 * @param Ar Archive to serialize with.
-		 * @param E Edge to serialize.
-		 * @returns Passing down serializing archive.
-		 */
-		friend FArchive& operator<<(FArchive& Ar, FEdge& E)
-		{
-			E.Serialize(Ar);
-			return Ar;
-		}
-
-		/** Serialize FEdge to an archive. */
-		void Serialize(FArchive& Ar)
-		{
-			Ar << Vert;
-			Ar << Tri;
-		}
-	};
-	/** InvalidID indicates that a vertex/edge/triangle ID is invalid */
-	constexpr static int InvalidID = IndexConstants::InvalidID;
-	/** NonManifoldID is returned by AppendTriangle() to indicate that the added triangle would result in nonmanifold geometry and hence was ignored */
-	constexpr static int NonManifoldID = -2;
-	/** DuplicateTriangleID is returned by AppendTriangle() to indicate that the added triangle already exists in the mesh, and was ignored because we do not support duplicate triangles */
-	constexpr static int DuplicateTriangleID = -3;
-
-	const static FVector3d InvalidVertex;
-	constexpr static FIndex3i InvalidTriangle{InvalidID, InvalidID, InvalidID};
-	constexpr static FIndex2i InvalidEdge{InvalidID, InvalidID};
-
-protected:
-	/** List of vertex positions */
-	TDynamicVector<FVector3d> Vertices{};
-	/** Reference counts of vertex indices. For vertices that exist, the count is 1 + num_triangle_using_vertex. Iterate over this to find out which vertex indices are valid. */
-	FRefCountVector VertexRefCounts{};
-	/** (optional) List of per-vertex normals */
-	TOptional<TDynamicVector<FVector3f>> VertexNormals{};
-	/** (optional) List of per-vertex colors */
-	TOptional<TDynamicVector<FVector3f>> VertexColors{};
-	/** (optional) List of per-vertex uv's */
-	TOptional<TDynamicVector<FVector2f>> VertexUVs{};
-	/** List of per-vertex edge one-rings */
-	FSmallListSet VertexEdgeLists;
-
-	/** List of triangle vertex-index triplets [Vert0 Vert1 Vert2]*/
-	TDynamicVector<FIndex3i> Triangles;
-	/** Reference counts of triangle indices. Ref count is always 1 if the triangle exists. Iterate over this to find out which triangle indices are valid. */
-	FRefCountVector TriangleRefCounts;
-	/** List of triangle edge triplets [Edge0 Edge1 Edge2] */
-	TDynamicVector<FIndex3i> TriangleEdges;
-	/** (optional) List of per-triangle group identifiers */
-	TOptional<TDynamicVector<int>> TriangleGroups{};
-	/** Upper bound on the triangle group IDs used in the mesh (may be larger than the actual maximum if triangles have been deleted) */
-	int GroupIDCounter = 0;
-
-	/** List of edge elements. An edge is four elements [VertA, VertB, Tri0, Tri1], where VertA < VertB, and Tri1 may be InvalidID (if the edge is a boundary edge) */
-	TDynamicVector<FEdge> Edges;
-	/** Reference counts of edge indices. Ref count is always 1 if the edge exists. Iterate over this to find out which edge indices are valid. */
-	FRefCountVector EdgeRefCounts;
-
-	TUniquePtr<FDynamicMeshAttributeSet> AttributeSet{};
-
-	/** The mesh timestamp is incremented any time a function that modifies the mesh is called */
-	int Timestamp = 0;
-	/** The shape timestamp is incremented any time a function that modifies the mesh shape or topology is called */
-	int ShapeTimestamp = 0;
-	/** The topology timestamp is incremented any time a function that modifies the mesh topology is called */
-	int TopologyTimestamp = 0;
-
-
-	/** Cached vertex bounding box (includes unreferenced vertices) */
-	FAxisAlignedBox3d CachedBoundingBox;
-	/** timestamp for CachedBoundingBox, if less than current timestamp, cache is invalid */
-	int CachedBoundingBoxTimestamp = -1;
-	/** Cached value of IsClosed() */
-	bool bIsClosedCached = false;
-	/** timestamp for bIsClosedCached, if less than current timestamp, cache is invalid */
-	int CachedIsClosedTimestamp = -1;
-
-public:
-	/** Default constructor */
-	FDynamicMesh3() : FDynamicMesh3(false, false, false, false) {}
-
-	/** Copy/Move construction */
-	FDynamicMesh3(const FDynamicMesh3& CopyMesh);
-	FDynamicMesh3(FDynamicMesh3&& MoveMesh);
-
-	/** Copy and move assignment */
-	const FDynamicMesh3& operator=(const FDynamicMesh3& CopyMesh);
-	const FDynamicMesh3& operator=(FDynamicMesh3&& MoveMesh);
-
-	/** Destructor */
-	virtual ~FDynamicMesh3();
-
-
-	/** Construct an empty mesh with specified attributes */
-	explicit FDynamicMesh3(bool bWantNormals, bool bWantColors, bool bWantUVs, bool bWantTriGroups);
-	explicit FDynamicMesh3(EMeshComponents flags)
-	    : FDynamicMesh3(((int)flags & (int)EMeshComponents::VertexNormals) != 0,
-	                    ((int)flags & (int)EMeshComponents::VertexColors) != 0,
-	                    ((int)flags & (int)EMeshComponents::VertexUVs) != 0,
-	                    ((int)flags & (int)EMeshComponents::FaceGroups) != 0)
-	{}
-
-	/** Construction from Mesh Generator */
-	FDynamicMesh3(const FMeshShapeGenerator* Generator);
-
-	/** Set internal data structures to be a copy of input mesh using the specified attributes*/
-	void Copy(const FDynamicMesh3& CopyMesh, bool bNormals = true, bool bColors = true, bool bUVs = true,
-	          bool bAttributes = true);
-
-	/** 
-	  * Initialize mesh from the output of a MeshShapeGenerator (assumes Generate() was already called) 
-	  * @return false if any triangles defined in Generator were skipped due to invalid topology
-	  */
-	bool Copy(const FMeshShapeGenerator* Generator);
-
-	/** 
-	 * Copy input mesh while compacting, i.e. removing unused vertices/triangles/edges. 
-	 * Note that it is currently unsafe to hold on to the attribute set pointer (returned from Attributes())
-	 * across the CompactCopy call, as the latter may rebuild an entirely new attribute set.
-	 */
-	void CompactCopy(const FDynamicMesh3& CopyMesh, bool bNormals = true, bool bColors = true, bool bUVs = true,
-	                 bool bAttributes = true, FCompactMaps* CompactInfo = nullptr);
-
-	/** Discard all data */
-	void Clear();
-
-	/**
-	 * Serialization operator for FDynamicMesh3.
-	 *
-	 * @param Ar Archive to serialize with.
-	 * @param Mesh Mesh to serialize.
-	 * @returns Passing down serializing archive.
-	 */
-	friend DYNAMICMESH_API FArchive& operator<<(FArchive& Ar, FDynamicMesh3& Mesh)
-	{
-		Mesh.Serialize(Ar);
-		return Ar;
-	}
-
-	/** Serialize the mesh to an archive. */
-	void Serialize(FArchive& Ar);
-
-public:
-	/** @return number of vertices in the mesh */
-	int VertexCount() const
-	{
-		return (int)VertexRefCounts.GetCount();
-	}
-	/** @return number of triangles in the mesh */
-	int TriangleCount() const
-	{
-		return (int)TriangleRefCounts.GetCount();
-	}
-	/** @return number of edges in the mesh */
-	int EdgeCount() const
-	{
-		return (int)EdgeRefCounts.GetCount();
-	}
-
-	/** @return upper bound on vertex IDs used in the mesh, i.e. all vertex IDs in use are < MaxVertexID */
-	int MaxVertexID() const
-	{
-		return (int)VertexRefCounts.GetMaxIndex();
-	}
-	/** @return upper bound on triangle IDs used in the mesh, i.e. all triangle IDs in use are < MaxTriangleID */
-	int MaxTriangleID() const
-	{
-		return (int)TriangleRefCounts.GetMaxIndex();
-	}
-	/** @return upper bound on edge IDs used in the mesh, i.e. all edge IDs in use are < MaxEdgeID */
-	int MaxEdgeID() const
-	{
-		return (int)EdgeRefCounts.GetMaxIndex();
-	}
-	/** @return upper bound on group IDs used in the mesh, i.e. all group IDs in use are < MaxGroupID */
-	int MaxGroupID() const
-	{
-		return GroupIDCounter;
-	}
-
-
-	/** @return true if this mesh has per-vertex normals */
-	bool HasVertexNormals() const
-	{
-		return VertexNormals.IsSet();
-	}
-	/** @return true if this mesh has per-vertex colors */
-	bool HasVertexColors() const
-	{
-		return VertexColors.IsSet();
-	}
-	/** @return true if this mesh has per-vertex UVs */
-	bool HasVertexUVs() const
-	{
-		return VertexUVs.IsSet();
-	}
-	/** @return true if this mesh has per-triangle groups */
-	bool HasTriangleGroups() const
-	{
-		return TriangleGroups.IsSet();
-	}
-	/** @return true if this mesh has attribute layers */
-	bool HasAttributes() const
-	{
-		return AttributeSet.IsValid();
-	}
-
-	/** @return bitwise-or of EMeshComponents flags specifying which extra data this mesh has */
-	int GetComponentsFlags() const;
-
-
-	/** @return true if VertexID is a valid vertex in this mesh */
-	inline bool IsVertex(int VertexID) const
-	{
-		return VertexRefCounts.IsValid(VertexID);
-	}
-	/** @return true if VertexID is a valid vertex in this mesh AND is used by at least one triangle */
-	inline bool IsReferencedVertex(int VertexID) const
-	{
-		return VertexID >= 0 && VertexID < (int)VertexRefCounts.GetMaxIndex() &&
-		    VertexRefCounts.GetRefCount(VertexID) > 1;
-	}
-	/** @return true if TriangleID is a valid triangle in this mesh */
-	inline bool IsTriangle(int TriangleID) const
-	{
-		return TriangleRefCounts.IsValid(TriangleID);
-	}
-	/** @return true if EdgeID is a valid edge in this mesh */
-	inline bool IsEdge(int EdgeID) const
-	{
-		return EdgeRefCounts.IsValid(EdgeID);
-	}
-
-
-	//
-	// Mesh Element Iterators
-	//   The functions VertexIndicesItr() / TriangleIndicesItr() / EdgeIndicesItr() allow you to do:
-	//      for ( int eid : EdgeIndicesItr() ) { ... }
-	//   and other related begin() / end() idioms
-public:
-	// simplify names for iterations
-	typedef typename FRefCountVector::IndexEnumerable vertex_iterator;
-	typedef typename FRefCountVector::IndexEnumerable triangle_iterator;
-	typedef typename FRefCountVector::IndexEnumerable edge_iterator;
-	template<typename T>
-	using value_iteration          = FRefCountVector::MappedEnumerable<T>;
-	using vtx_triangles_enumerable = TPairExpandEnumerable<FSmallListSet::ValueIterator>;
-
-	/** @return enumerable object for valid vertex indices suitable for use with range-based for, ie for ( int i : VertexIndicesItr() ) */
-	vertex_iterator VertexIndicesItr() const
-	{
-		return VertexRefCounts.Indices();
-	}
-
-	/** @return enumerable object for valid triangle indices suitable for use with range-based for, ie for ( int i : TriangleIndicesItr() ) */
-	triangle_iterator TriangleIndicesItr() const
-	{
-		return TriangleRefCounts.Indices();
-	}
-
-	/** @return enumerable object for valid edge indices suitable for use with range-based for, ie for ( int i : EdgeIndicesItr() ) */
-	edge_iterator EdgeIndicesItr() const
-	{
-		return EdgeRefCounts.Indices();
-	}
-
-	// TODO: write helper functions that allow us to do these iterations w/o lambdas
-
-	/** @return enumerable object for boundary edge indices suitable for use with range-based for, ie for ( int i : BoundaryEdgeIndicesItr() ) */
-	FRefCountVector::FilteredEnumerable BoundaryEdgeIndicesItr() const
-	{
-		return EdgeRefCounts.FilteredIndices([this](int EdgeID) { return Edges[EdgeID].Tri[1] == InvalidID; });
-	}
-
-	/** Enumerate positions of all vertices in mesh */
-	value_iteration<FVector3d> VerticesItr() const
-	{
-		return VertexRefCounts.MappedIndices<FVector3d>([this](int VertexID) { return Vertices[VertexID]; });
-	}
-
-	/** Enumerate all triangles in the mesh */
-	value_iteration<FIndex3i> TrianglesItr() const
-	{
-		return TriangleRefCounts.MappedIndices<FIndex3i>([this](int TriangleID) { return Triangles[TriangleID]; });
-	}
-
-	/** Enumerate edges. Each returned element is [v0,v1,t0,t1], where t1 will be InvalidID if this is a boundary edge */
-	value_iteration<FEdge> EdgesItr() const
-	{
-		return EdgeRefCounts.MappedIndices<FEdge>([this](int EdgeID) {
-			return Edges[EdgeID];
-		});
-	}
-
-	/** @return enumerable object for one-ring vertex neighbours of a vertex, suitable for use with range-based for, ie for ( int i : VtxVerticesItr(VertexID) ) */
-	FSmallListSet::MappedValueEnumerable VtxVerticesItr(int VertexID) const
-	{
-		checkSlow(VertexRefCounts.IsValid(VertexID));
-		return VertexEdgeLists.MappedValues(VertexID,
-			[VertexID, this](int eid) { return GetOtherEdgeVertex(eid, VertexID); });
-	}
-
-	/** Call VertexFunc for each one-ring vertex neighbour of a vertex. Currently this is more efficient than VtxVerticesItr() due to overhead in the Values() enumerable */
-	void EnumerateVertexVertices(int32 VertexID, TFunctionRef<void(int32)> VertexFunc) const
-	{
-		checkSlow(VertexRefCounts.IsValid(VertexID));
-		VertexEdgeLists.Enumerate(VertexID, [this, &VertexFunc, VertexID](int32 eid)
-		{
-			VertexFunc(GetOtherEdgeVertex(eid, VertexID));
-		});
-	}
-
-
-	/** @return enumerable object for one-ring edges of a vertex, suitable for use with range-based for, ie for ( int i : VtxEdgesItr(VertexID) ) */
-	FSmallListSet::ValueEnumerable VtxEdgesItr(int VertexID) const
-	{
-		checkSlow(VertexRefCounts.IsValid(VertexID));
-		return VertexEdgeLists.Values(VertexID);
-	}
-
-	/** Call EdgeFunc for each one-ring edge of a vertex. Currently this is more efficient than VtxEdgesItr() due to overhead in the Values() enumerable */
-	void EnumerateVertexEdges(int32 VertexID, TFunctionRef<void(int32)> EdgeFunc) const
-	{
-		checkSlow(VertexRefCounts.IsValid(VertexID));
-		VertexEdgeLists.Enumerate(VertexID, EdgeFunc);
-	}
-
-
-	/** @return enumerable object for one-ring triangles of a vertex, suitable for use with range-based for, ie for ( int i : VtxTrianglesItr(VertexID) ) */
-	vtx_triangles_enumerable VtxTrianglesItr(int VertexID) const
-	{
-		checkSlow(VertexRefCounts.IsValid(VertexID));
-		return vtx_triangles_enumerable(VertexEdgeLists.Values(VertexID), [this, VertexID](int EdgeID) {
-			return GetOrderedOneRingEdgeTris(VertexID, EdgeID);
-		});
-	}
-
-	/** Call ApplyFunc for each one-ring triangle of a vertex. Currently this is significantly more efficient than VtxTrianglesItr() in many use cases. */
-	void EnumerateVertexTriangles(int32 VertexID, TFunctionRef<void(int32)> ApplyFunc) const;
-
-	//
-	// Mesh Construction
-	//
-public:
-	/** Append vertex at position and other fields, returns vid */
-	int AppendVertex(const FVertexInfo& VertInfo);
-
-	/** Append vertex at position, returns vid */
-	int AppendVertex(const FVector3d& Position)
-	{
-		return AppendVertex(FVertexInfo(Position));
-	}
-
-	/** Copy vertex SourceVertexID from existing SourceMesh, returns new vertex id */
-	int AppendVertex(const FDynamicMesh3& SourceMesh, int SourceVertexID);
-
-	int AppendTriangle(const FIndex3i& TriVertices, int GroupID = 0);
-
-	inline int AppendTriangle(int Vertex0, int Vertex1, int Vertex2, int GroupID = 0)
-	{
-		return AppendTriangle(FIndex3i(Vertex0, Vertex1, Vertex2), GroupID);
-	}
-
-	//
-	// Support for inserting vertex and triangle at specific IDs. This is a bit tricky
-	// because we likely will need to update the free lists in the RefCountVectors, which
-	// can be expensive. If you are going to do many inserts (eg inside a loop), wrap in
-	// BeginUnsafe / EndUnsafe calls, and pass bUnsafe = true to the InsertX() calls, to
-	// the defer free list rebuild until you are done.
-	//
-
-	/** Call this before a set of unsafe InsertVertex() calls */
-	virtual void BeginUnsafeVerticesInsert()
-	{
-		// do nothing...
-	}
-
-	/** Call after a set of unsafe InsertVertex() calls to rebuild free list */
-	virtual void EndUnsafeVerticesInsert()
-	{
-		VertexRefCounts.RebuildFreeList();
-	}
-
-	/**
-	 * Insert vertex at given index, assuming it is unused.
-	 * If bUnsafe, we use fast id allocation that does not update free list.
-	 * You should only be using this between BeginUnsafeVerticesInsert() / EndUnsafeVerticesInsert() calls
-	 */
-	EMeshResult InsertVertex(int VertexID, const FVertexInfo& VertInfo, bool bUnsafe = false);
-
-	/** Call this before a set of unsafe InsertTriangle() calls */
-	virtual void BeginUnsafeTrianglesInsert()
-	{
-		// do nothing...
-	}
-
-	/** Call after a set of unsafe InsertTriangle() calls to rebuild free list */
-	virtual void EndUnsafeTrianglesInsert()
-	{
-		TriangleRefCounts.RebuildFreeList();
-	}
-
-	/**
-	 * Insert triangle at given index, assuming it is unused.
-	 * If bUnsafe, we use fast id allocation that does not update free list.
-	 * You should only be using this between BeginUnsafeTrianglesInsert() / EndUnsafeTrianglesInsert() calls
-	 */
-	EMeshResult InsertTriangle(int TriangleID, const FIndex3i& TriVertices, int GroupID = 0, bool bUnsafe = false);
-
-	//
-	// Vertex/Tri/Edge accessors
-	//
-public:
-	/** @return the vertex position */
-	inline FVector3d GetVertex(int VertexID) const
-	{
-		checkSlow(IsVertex(VertexID));
-		return Vertices[VertexID];
-	}
-
-	/** @return the vertex position */
-	inline const FVector3d& GetVertexRef(int VertexID) const
-	{
-		checkSlow(IsVertex(VertexID));
-		return Vertices[VertexID];
-	}
-
-	/** Set vertex position */
-	inline void SetVertex(int VertexID, const FVector3d& vNewPos)
-	{
-		checkSlow(VectorUtil::IsFinite(vNewPos));
-		checkSlow(IsVertex(VertexID));
-		if (VectorUtil::IsFinite(vNewPos))
-		{
-			Vertices[VertexID] = vNewPos;
-			UpdateTimeStamp(true, false);
-		}
-	}
-
-	/** Set vertex position. DOES NOT UPDATE TIMESTAMPS.
-	This is meant to be used in conjunction with IncrementTimeStamps, below. Set a bunch of vertex positions then 
-	increment the timestamps with the number of vertices that were moved, rather than incrementing the timestamps for
-	each vertex.
-	*/
-	void SetVertex_NoTimeStampUpdate(int VertexID, const FVector3d& vNewPos)
-	{
-		checkSlow(VectorUtil::IsFinite(vNewPos));
-		checkSlow(IsVertex(VertexID));
-
-		if (VectorUtil::IsFinite(vNewPos))
-		{
-			Vertices[VertexID] = vNewPos;
-		}
-	}
-
-	/** Increment the specified timestamps by the given amount. Thread-safe. 
-	This is meant to be used in conjunction with SetVertex_NoTimeStampUpdate, above. Set a bunch of vertex positions then
-	increment the timestamps with the number of vertices that were moved, rather than incrementing the timestamps for
-	each vertex.
-	*/
-	void IncrementTimeStamps(int Amount, bool bShapeChange, bool bTopologyChange)
-	{
-		FPlatformAtomics::InterlockedAdd(&Timestamp, Amount);
-		if (bShapeChange)
-		{
-			FPlatformAtomics::InterlockedAdd(&ShapeTimestamp, Amount);
-		}
-		if (bTopologyChange)
-		{
-			FPlatformAtomics::InterlockedAdd(&TopologyTimestamp, Amount);
-		}
-	}
-
-	/** Get extended vertex information */
-	bool GetVertex(int VertexID, FVertexInfo& VertInfo, bool bWantNormals, bool bWantColors, bool bWantUVs) const;
-
-	/** Get all vertex information available */
-	FVertexInfo GetVertexInfo(int VertexID) const;
-
-	/** @return the valence of a vertex (the number of connected edges) */
-	int GetVtxEdgeCount(int VertexID) const
-	{
-		return VertexRefCounts.IsValid(VertexID) ? VertexEdgeLists.GetCount(VertexID) : -1;
-	}
-
-	/** @return the max valence of all vertices in the mesh */
-	int GetMaxVtxEdgeCount() const;
-
-	/** adds triangles touching VertexID to the TrianglesOut list */
-	void GetVertexOneRingTriangles(int VertexID, TArray<int>& TrianglesOut) const;
-
-	/** Get triangle vertices */
-	inline FIndex3i GetTriangle(int TriangleID) const
-	{
-		checkSlow(IsTriangle(TriangleID));
-		return Triangles[TriangleID];
-	}
-
-	/** Get triangle vertices */
-	inline const FIndex3i& GetTriangleRef(int TriangleID) const
-	{
-		checkSlow(IsTriangle(TriangleID));
-		return Triangles[TriangleID];
-	}
-
-	/** Get triangle edges */
-	inline FIndex3i GetTriEdges(int TriangleID) const
-	{
-		checkSlow(IsTriangle(TriangleID));
-		return TriangleEdges[TriangleID];
-	}
-
-	/** Get triangle edges */
-	inline const FIndex3i& GetTriEdgesRef(int TriangleID) const
-	{
-		checkSlow(IsTriangle(TriangleID));
-		return TriangleEdges[TriangleID];
-	}
-
-	/** Get one of the edges of a triangle */
-	inline int GetTriEdge(int TriangleID, int j) const
-	{
-		checkSlow(IsTriangle(TriangleID));
-		return TriangleEdges[TriangleID][j];
-	}
-
-	/** Find the neighbour triangles of a triangle (any of them might be InvalidID) */
-	FIndex3i GetTriNeighbourTris(int TriangleID) const;
-
-	/** Get the three vertex positions of a triangle */
-	template<typename VecType>
-	inline void GetTriVertices(int TriangleID, VecType& v0, VecType& v1, VecType& v2) const
-	{
-		const FIndex3i& Triangle = Triangles[TriangleID];
-		v0                       = Vertices[Triangle[0]];
-		v1                       = Vertices[Triangle[1]];
-		v2                       = Vertices[Triangle[2]];
-	}
-
-	/** Get the position of one of the vertices of a triangle */
-	inline FVector3d GetTriVertex(int TriangleID, int j) const
-	{
-		return Vertices[Triangles[TriangleID][j]];
-	}
-
-	/** Get the vertices and triangles of an edge, returned as [v0,v1,t0,t1], where t1 may be InvalidID */
-	inline FEdge GetEdge(int EdgeID) const
-	{
-		checkSlow(IsEdge(EdgeID));
-		return Edges[EdgeID];
-	}
-
-	/** Get the vertices and triangles of an edge, returned as [v0,v1,t0,t1], where t1 may be InvalidID */
-	inline const FEdge& GetEdgeRef(int EdgeID) const
-	{
-		checkSlow(IsEdge(EdgeID));
-		return Edges[EdgeID];
-	}
-
-	/** Get the vertex pair for an edge */
-	inline FIndex2i GetEdgeV(int EdgeID) const
-	{
-		checkSlow(IsEdge(EdgeID));
-		return Edges[EdgeID].Vert;
-	}
-
-	/** Get the vertex positions of an edge */
-	inline bool GetEdgeV(int EdgeID, FVector3d& a, FVector3d& b) const
-	{
-		checkSlow(IsEdge(EdgeID));
-
-		const FIndex2i Verts = Edges[EdgeID].Vert;
-
-		a = Vertices[Verts[0]];
-		b = Vertices[Verts[1]];
-
-		return true;
-	}
-
-	/** Get the triangle pair for an edge. The second triangle may be InvalidID */
-	inline FIndex2i GetEdgeT(int EdgeID) const
-	{
-		checkSlow(IsEdge(EdgeID));
-		return Edges[EdgeID].Tri;
-	}
-
-	/** Return edge vertex indices, but oriented based on attached triangle (rather than min-sorted) */
-	FIndex2i GetOrientedBoundaryEdgeV(int EdgeID) const;
-
-	//
-	// Vertex and Triangle attribute arrays
-	//
-public:
-	/**
-	 * Enable requested set of mesh components (triangle groups and vertex normals/colors/UVs)
-	 * and discard any that are not requested
-	 * @param MeshComponentsFlags A 'bitwise or' of requested EMeshComponents flags
-	 */
-	void EnableMeshComponents(int MeshComponentsFlags);
-
-	void EnableVertexNormals(const FVector3f& InitialNormal);
-	void DiscardVertexNormals();
-
-	FVector3f GetVertexNormal(int vID) const
-	{
-		if (HasVertexNormals() == false)
-		{
-			return FVector3f::UnitY();
-		}
-		checkSlow(IsVertex(vID));
-		const TDynamicVector<FVector3f>& Normals = VertexNormals.GetValue();
-		return Normals[vID];
-	}
-
-	void SetVertexNormal(int vID, const FVector3f& vNewNormal)
-	{
-		if (HasVertexNormals())
-		{
-			checkSlow(IsVertex(vID));
-			TDynamicVector<FVector3f>& Normals = VertexNormals.GetValue();
-			Normals[vID]                       = vNewNormal;
-			UpdateTimeStamp(false, false);
-		}
-	}
-
-	void EnableVertexColors(const FVector3f& InitialColor);
-	void DiscardVertexColors();
-
-
-	FVector3f GetVertexColor(int vID) const
-	{
-		if (HasVertexColors() == false)
-		{
-			return FVector3f::One();
-		}
-		checkSlow(IsVertex(vID));
-
-		const TDynamicVector<FVector3f>& Colors = VertexColors.GetValue();
-		return Colors[vID];
-	}
-
-	void SetVertexColor(int vID, const FVector3f& vNewColor)
-	{
-		if (HasVertexColors())
-		{
-			checkSlow(IsVertex(vID));
-			TDynamicVector<FVector3f>& Colors = VertexColors.GetValue();
-			Colors[vID]                       = vNewColor;
-			UpdateTimeStamp(false, false);
-		}
-	}
-
-	void EnableVertexUVs(const FVector2f& InitialUV);
-	void DiscardVertexUVs();
-
-	FVector2f GetVertexUV(int vID) const
-	{
-		if (HasVertexUVs() == false)
-		{
-			return FVector2f::Zero();
-		}
-		checkSlow(IsVertex(vID));
-		const TDynamicVector<FVector2f>& UVs = VertexUVs.GetValue();
-		return UVs[vID];
-	}
-
-	void SetVertexUV(int vID, const FVector2f& vNewUV)
-	{
-		if (HasVertexUVs())
-		{
-			checkSlow(IsVertex(vID));
-			TDynamicVector<FVector2f>& UVs = VertexUVs.GetValue();
-			UVs[vID]                       = vNewUV;
-			UpdateTimeStamp(false, false);
-		}
-	}
-
-	void EnableTriangleGroups(int InitialGroupID = 0);
-	void DiscardTriangleGroups();
-
-	int AllocateTriangleGroup()
-	{
-		return GroupIDCounter++;
-	}
-
-	int GetTriangleGroup(int tID) const
-	{
-		return (HasTriangleGroups() == false) ? -1 :
-		                                        (TriangleRefCounts.IsValid(tID) ? TriangleGroups.GetValue()[tID] : 0);
-	}
-
-	void SetTriangleGroup(int tid, int group_id)
-	{
-		if (HasTriangleGroups())
-		{
-			checkSlow(IsTriangle(tid));
-			TriangleGroups.GetValue()[tid] = group_id;
-			GroupIDCounter                 = FMath::Max(GroupIDCounter, group_id + 1);
-			UpdateTimeStamp(false, false);
-		}
-	}
-
-	FDynamicMeshAttributeSet* Attributes()
-	{
-		return HasAttributes() ? AttributeSet.Get() : nullptr;
-	}
-	const FDynamicMeshAttributeSet* Attributes() const
-	{
-		return HasAttributes() ? AttributeSet.Get() : nullptr;
-	}
-
-	void EnableAttributes();
-	void DiscardAttributes();
-
-
-	//
-	// topological queries
-	//
-public:
-	/** Returns true if edge is on the mesh boundary, ie only connected to one triangle */
-	inline bool IsBoundaryEdge(int EdgeID) const
-	{
-		checkSlow(IsEdge(EdgeID));
-		return Edges[EdgeID].Tri[1] == InvalidID;
-	}
-
-	/** Returns true if the vertex is part of any boundary edges */
-	bool IsBoundaryVertex(int VertexID) const;
-
-	/** Returns true if any edge of triangle is a boundary edge */
-	bool IsBoundaryTriangle(int TriangleID) const;
-
-	/** Find id of edge connecting A and B */
-	int FindEdge(int VertexA, int VertexB) const;
-
-	/** Find edgeid for edge [a,b] from triangle that contains the edge. Faster than FindEdge() because it is constant-time. */
-	int FindEdgeFromTri(int VertexA, int VertexB, int TriangleID) const;
-
-	/** Find edgeid for edge connecting two triangles */
-	int FindEdgeFromTriPair(int TriangleA, int TriangleB) const;
-
-	/** Find triangle made up of any permutation of vertices [a,b,c] */
-	int FindTriangle(int A, int B, int C) const;
-
-	/**
-	 * If edge has vertices [a,b], and is connected two triangles [a,b,c] and [a,b,d],
-	 * this returns [c,d], or [c,InvalidID] for a boundary edge
-	 */
-	FIndex2i GetEdgeOpposingV(int EdgeID) const;
-
-	/**
-	 * Given an edge and vertex on that edge, returns other vertex of edge, the two opposing verts, and the two connected triangles (OppVert2Out and Tri2Out are be InvalidID for boundary edge)
-	 */
-	void GetVtxNbrhood(int EdgeID, int VertexID, int& OtherVertOut, int& OppVert1Out, int& OppVert2Out, int& Tri1Out,
-	                   int& Tri2Out) const;
-
-	/**
-	 * Returns count of boundary edges at vertex, and the first two boundary
-	 * edges if found. If return is > 2, call GetAllVtxBoundaryEdges
-	 */
-	int GetVtxBoundaryEdges(int VertexID, int& Edge0Out, int& Edge1Out) const;
-
-	/**
-	 * Find edge ids of boundary edges connected to vertex.
-	 * @param vID Vertex ID
-	 * @param EdgeListOut boundary edge IDs are appended to this list
-	 * @return count of number of elements of e that were filled
-	 */
-	int GetAllVtxBoundaryEdges(int VertexID, TArray<int>& EdgeListOut) const;
-
-	/**
-	 * return # of triangles attached to vID, or -1 if invalid vertex
-	 * if bBruteForce = true, explicitly checks, which creates a list and is expensive
-	 * default is false, uses orientation, no memory allocation
-	 */
-	int GetVtxTriangleCount(int VertexID, bool bBruteForce = false) const;
-
-	/**
-	 * Get triangle one-ring at vertex.
-	 * bUseOrientation is more efficient but returns incorrect result if vertex is a bowtie
-	 */
-	EMeshResult GetVtxTriangles(int VertexID, TArray<int>& TrianglesOut, bool bUseOrientation) const;
-
-	/**
-	 * @return Triangle ID for a single triangle connected to VertexID, or InvalidID if VertexID does not exist or has no attached triangles
-	 */
-	int GetVtxSingleTriangle(int VertexID) const;
-
-	/**
-	* Get triangles connected to vertex in contiguous order, with multiple groups if vertex is a bowtie.
-	* @param VertexID Vertex ID to search around
-	* @param TrianglesOut All triangles connected to the vertex, in contiguous order; if there are multiple contiguous groups they are packed one after another
-	* @param ContiguousGroupLengths Lengths of contiguous groups packed into TrianglesOut (if not a bowtie, this will just be a length-one array w/ {TrianglesOut.Num()})
-	* @param GroupIsLoop Indicates whether each contiguous group is a loop (first triangle connected to last) or not
-	*/
-	EMeshResult GetVtxContiguousTriangles(int VertexID, TArray<int>& TrianglesOut, TArray<int>& ContiguousGroupLengths,
-	                                      TArray<bool>& GroupIsLoop) const;
-
-	/** Returns true if the two triangles connected to edge have different group IDs */
-	bool IsGroupBoundaryEdge(int EdgeID) const;
-
-	/** Returns true if vertex has more than one tri group in its tri nbrhood */
-	bool IsGroupBoundaryVertex(int VertexID) const;
-
-	/** Returns true if more than two group boundary edges meet at vertex (ie 3+ groups meet at this vertex) */
-	bool IsGroupJunctionVertex(int VertexID) const;
-
-	/** Returns up to 4 group IDs at vertex. Returns false if > 4 encountered */
-	bool GetVertexGroups(int VertexID, FIndex4i& GroupsOut) const;
-
-	/** Returns all group IDs at vertex */
-	bool GetAllVertexGroups(int VertexID, TArray<int>& GroupsOut) const;
-
-	/** returns true if vID is a "bowtie" vertex, ie multiple disjoint triangle sets in one-ring */
-	bool IsBowtieVertex(int VertexID) const;
-
-	/** returns true if vertices, edges, and triangles are all dense (Count == MaxID) **/
-	bool IsCompact() const
-	{
-		return VertexRefCounts.IsDense() && EdgeRefCounts.IsDense() && TriangleRefCounts.IsDense();
-	}
-
-	/** @return true if vertex count == max vertex id */
-	bool IsCompactV() const
-	{
-		return VertexRefCounts.IsDense();
-	}
-
-	/** @return true if triangle count == max triangle id */
-	bool IsCompactT() const
-	{
-		return TriangleRefCounts.IsDense();
-	}
-
-	/** returns measure of compactness in range [0,1], where 1 is fully compacted */
-	double CompactMetric() const
-	{
-		return ((double)VertexCount() / (double)MaxVertexID() + (double)TriangleCount() / (double)MaxTriangleID()) *
-		    0.5;
-	}
-
-	/** @return true if mesh has no boundary edges */
-	bool IsClosed() const;
-
-	/** @return value of IsClosed() and cache result. cache is invalidated and recomputed if topology has changed since last call */
-	bool GetCachedIsClosed();
-
-
-	/** Timestamp is incremented any time any change is made to the mesh */
-	inline int GetTimestamp() const
-	{
-		return Timestamp;
-	}
-
-	/** ShapeTimestamp is incremented any time any vertex position is changed or the mesh topology is modified */
-	inline int GetShapeTimestamp() const
-	{
-		return ShapeTimestamp;
-	}
-
-	/** TopologyTimestamp is incremented any time any vertex position is changed or the mesh topology is modified */
-	inline int GetTopologyTimestamp() const
-	{
-		return TopologyTimestamp;
-	}
-
-	//
-	// Geometric queries
-	//
-public:
-	/** Returns bounding box of all mesh vertices (including unreferenced vertices) */
-	FAxisAlignedBox3d GetBounds(bool bParallel = false) const;
-
-	/** Returns GetBounds() and saves result, cache is invalidated and recomputed if topology has changed since last call */
-	FAxisAlignedBox3d GetCachedBounds();
-
-	/**
-	 * Compute a normal/tangent frame at vertex that is "stable" as long as
-	 * the mesh topology doesn't change, meaning that one axis of the frame
-	 * will be computed from projection of outgoing edge. Requires that vertex normals are available.
-	 * By default, frame.Z is normal, and .X points along mesh edge.
-	 * If bFrameNormalY, then frame.Y is normal (X still points along mesh edge)
-	 */
-	FFrame3d GetVertexFrame(int VertexID, bool bFrameNormalY = false) const;
-
-	/** Calculate face normal of triangle */
-	FVector3d GetTriNormal(int TriangleID) const;
-
-	/** Calculate area triangle */
-	double GetTriArea(int TriangleID) const;
-
-	/**
-	 * Compute triangle normal, area, and centroid all at once. Re-uses vertex
-	 * lookups and computes normal & area simultaneously. *However* does not produce
-	 * the same normal/area as separate calls, because of this.
-	 */
-	void GetTriInfo(int TriangleID, FVector3d& Normal, double& Area, FVector3d& Centroid) const;
-
-	/** Compute centroid of triangle */
-	FVector3d GetTriCentroid(int TriangleID) const;
-
-	/** Interpolate vertex positions of triangle using barycentric coordinates */
-	FVector3d GetTriBaryPoint(int TriangleID, double Bary0, double Bary1, double Bary2) const;
-
-	/** Interpolate vertex normals of triangle using barycentric coordinates */
-	FVector3d GetTriBaryNormal(int TriangleID, double Bary0, double Bary1, double Bary2) const;
-
-	/** Compute interpolated vertex attributes at point of triangle */
-	void GetTriBaryPoint(int TriangleID, double Bary0, double Bary1, double Bary2, FVertexInfo& VertInfo) const;
-
-	/** Construct bounding box of triangle as efficiently as possible */
-	FAxisAlignedBox3d GetTriBounds(int TriangleID) const;
-
-	/** Construct stable frame at triangle centroid, where frame.Z is face normal, and frame.X is aligned with edge nEdge of triangle. */
-	FFrame3d GetTriFrame(int TriangleID, int Edge = 0) const;
-
-	/** Compute solid angle of oriented triangle tID relative to point p - see WindingNumber() */
-	double GetTriSolidAngle(int TriangleID, const FVector3d& p) const;
-
-	/** Compute internal angle at vertex i of triangle (where i is 0,1,2); */
-	double GetTriInternalAngleR(int TriangleID, int i) const;
-
-	/** Compute internal angles at all vertices of triangle */
-	FVector3d GetTriInternalAnglesR(int TriangleID) const;
-
-	/** Returns average normal of connected face normals */
-	FVector3d GetEdgeNormal(int EdgeID) const;
-
-	/** Get point along edge, t clamped to range [0,1] */
-	FVector3d GetEdgePoint(int EdgeID, double ParameterT) const;
-
-	/**
-	 * Fastest possible one-ring centroid. This is used inside many other algorithms
-	 * so it helps to have it be maximally efficient
-	 */
-	void GetVtxOneRingCentroid(int VertexID, FVector3d& CentroidOut) const;
-
-	/**
-	 * Compute mesh winding number, from Jacobson et. al., Robust Inside-Outside Segmentation using Generalized Winding Numbers
-	 * http://igl.ethz.ch/projects/winding-number/
-	 * returns ~0 for points outside a closed, consistently oriented mesh, and a positive or negative integer
-	 * for points inside, with value > 1 depending on how many "times" the point inside the mesh (like in 2D polygon winding)
-	 */
-	double CalculateWindingNumber(const FVector3d& QueryPoint) const;
-
-	//
-	// direct buffer access
-	//
-public:
-	const TDynamicVector<FVector3d>& GetVerticesBuffer()
-	{
-		return Vertices;
-	}
-	const FRefCountVector& GetVerticesRefCounts()
-	{
-		return VertexRefCounts;
-	}
-	const TDynamicVector<FVector3f>* GetNormalsBuffer()
-	{
-		return HasVertexNormals() ? &VertexNormals.GetValue() : nullptr;
-	}
-	const TDynamicVector<FVector3f>* GetColorsBuffer()
-	{
-		return HasVertexColors() ? &VertexColors.GetValue() : nullptr;
-	}
-	const TDynamicVector<FVector2f>* GetUVBuffer()
-	{
-		return HasVertexUVs() ? &VertexUVs.GetValue() : nullptr;
-	}
-	const TDynamicVector<FIndex3i>& GetTrianglesBuffer()
-	{
-		return Triangles;
-	}
-	const FRefCountVector& GetTrianglesRefCounts()
-	{
-		return TriangleRefCounts;
-	}
-	const TDynamicVector<int>* GetTriangleGroupsBuffer()
-	{
-		return HasTriangleGroups() ? &TriangleGroups.GetValue() : nullptr;
-	}
-	const TDynamicVector<FEdge>& GetEdgesBuffer()
-	{
-		return Edges;
-	}
-	const FRefCountVector& GetEdgesRefCounts()
-	{
-		return EdgeRefCounts;
-	}
-	const FSmallListSet& GetVertexEdges()
-	{
-		return VertexEdgeLists;
-	}
-
-	//
-	// Mesh Edit operations
-	//
-public:
-	/**
-	 * Compact mesh in-place, by moving vertices around and rewriting indices.
-	 * Should be faster if the amount of compacting is not too significant, and is useful in some places.
-	 *
-	 * @param CompactInfo if not nullptr, will be filled with mapping indicating how vertex and (optionally) triangle IDs were changed during compaction
-	 * @todo VertexEdgeLists is not compacted. does not affect indices, but does keep memory.
-	 */
-	void CompactInPlace(FCompactMaps* CompactInfo = nullptr);
-
-	/**
-	 * Reverse the ccw/cw orientation of all triangles in the mesh, and
-	 * optionally flip the vertex normals if they exist
-	 */
-	void ReverseOrientation(bool bFlipNormals = true);
-
-	/**
-	 * Reverse the ccw/cw orientation of a triangle
-	 */
-	EMeshResult ReverseTriOrientation(int TriangleID);
-
-	/**
-	 * Remove vertex vID, and all connected triangles if bRemoveAllTriangles = true
-	 * Returns Failed_VertexStillReferenced if vertex is still referenced by triangles.
-	 * if bPreserveManifold, checks that we will not create a bow tie vertex first
-	 */
-	EMeshResult RemoveVertex(int VertexID, bool bRemoveAllTriangles = true, bool bPreserveManifold = false);
-
-	/**
-	* Remove a triangle from the mesh. Also removes any unreferenced edges after tri is removed.
-	* If bRemoveIsolatedVertices is false, then if you remove all tris from a vert, that vert is also removed.
-	* If bPreserveManifold, we check that you will not create a bow tie vertex (and return false).
-	* If this check is not done, you have to make sure you don't create a bow tie, because other
-	* code assumes we don't have bow ties, and will not handle it properly
-	*/
-	EMeshResult RemoveTriangle(int TriangleID, bool bRemoveIsolatedVertices = true, bool bPreserveManifold = false);
-
-	/**
-	 * Rewrite the triangle to reference the new tuple of vertices.
-	 *
-	 * @todo this function currently does not guarantee that the returned mesh is well-formed. Only call if you know it's OK.
-	 */
-	virtual EMeshResult SetTriangle(int TriangleID, const FIndex3i& NewVertices, bool bRemoveIsolatedVertices = true);
-
-	/**
-	 * Split an edge of the mesh by inserting a vertex. This creates a new triangle on either side of the edge (ie a 2-4 split).
-	 * If the original edge had vertices [a,b], with triangles t0=[a,b,c] and t1=[b,a,d],  then the split inserts new vertex f.
-	 * After the split t0=[a,f,c] and t1=[f,a,d], and we have t2=[f,b,c] and t3=[f,d,b]  (it's best to draw it out on paper...)
-	 * 
-	 * @param EdgeAB index of the edge to be split
-	 * @param SplitInfo returned information about new and modified mesh elements
-	 * @param SplitParameterT defines the position along the edge that we split at, must be between 0 and 1, and is assumed to be based on the order of vertices returned by GetEdgeV()
-	 * @return Ok on success, or enum value indicates why operation cannot be applied. Mesh remains unmodified on error.
-	 */
-	virtual EMeshResult SplitEdge(int EdgeAB, FEdgeSplitInfo& SplitInfo, double SplitParameterT = 0.5);
-
-	/** 
-	 * Splits the edge between two vertices at the midpoint, if this edge exists 
-	 * @param EdgeVertA index of first vertex
-	 * @param EdgeVertB index of second vertex
-	 * @param SplitInfo returned information about new and modified mesh elements
-	 * @return Ok on success, or enum value indicates why operation cannot be applied. Mesh remains unmodified on error.
-	 */
-	EMeshResult SplitEdge(int EdgeVertA, int EdgeVertB, FEdgeSplitInfo& SplitInfo);
-
-	/**
-	 * Flip/Rotate an edge of the mesh. This does not change the number of edges, vertices, or triangles.
-	 * Boundary edges of the mesh cannot be flipped.
-	 * @param EdgeAB index of edge to be flipped
-	 * @param FlipInfo returned information about new and modified mesh elements
-	 * @return Ok on success, or enum value indicates why operation cannot be applied. Mesh remains unmodified on error.
-	 */
-	virtual EMeshResult FlipEdge(int EdgeAB, FEdgeFlipInfo& FlipInfo);
-
-	/** calls FlipEdge() on the edge between two vertices, if it exists
-	 * @param EdgeVertA index of first vertex
-	 * @param EdgeVertB index of second vertex
-	 * @param FlipInfo returned information about new and modified mesh elements
-	 * @return Ok on success, or enum value indicates why operation cannot be applied. Mesh remains unmodified on error.
-	 */
-	virtual EMeshResult FlipEdge(int EdgeVertA, int EdgeVertB, FEdgeFlipInfo& FlipInfo);
-
-
-	/**
-	 * Clones the given vertex and updates any provided triangles to use the new vertex if/where they used the old one.
-	 * @param VertexID the vertex to split
-	 * @param TrianglesToUpdate triangles that should be updated to use the new vertex anywhere they previously had the old one
-	 * @param SplitInfo returned info about the new and modified mesh elements
-	 * @return Ok on success, or enum value indicates why operation cannot be applied. Mesh remains unmodified on error.
-	 */
-	virtual EMeshResult SplitVertex(int VertexID, const TArrayView<const int>& TrianglesToUpdate,
-	                                FVertexSplitInfo& SplitInfo);
-
-	/**
-	 * Tests whether splitting the given vertex with the given triangles would leave no triangles attached to the original vertex (creating an isolated vertex)
-	 * @param VertexID the vertex to split
-	 * @param TrianglesToUpdate triangles that should be updated to use the new vertex anywhere they previously had the old one
-	 * @return true if calling SplitVertex with these arguments would leave an isolated vertex at the original VertexID
-	 */
-	virtual bool SplitVertexWouldLeaveIsolated(int VertexID, const TArrayView<const int>& TrianglesToUpdate);
-
-
-	/**
-	 * Collapse the edge between the two vertices, if topologically possible.
-	 * @param KeepVertID index of the vertex that should be kept
-	 * @param RemoveVertID index of the vertex that should be removed
-	 * @param EdgeParameterT vKeep is moved to Lerp(KeepPos, RemovePos, collapse_t)
-	 * @param CollapseInfo returned information about new and modified mesh elements
-	 * @return Ok on success, or enum value indicates why operation cannot be applied. Mesh remains unmodified on error.
-	 */
-	virtual EMeshResult CollapseEdge(int KeepVertID, int RemoveVertID, double EdgeParameterT,
-	                                 FEdgeCollapseInfo& CollapseInfo);
-	virtual EMeshResult CollapseEdge(int KeepVertID, int RemoveVertID, FEdgeCollapseInfo& CollapseInfo)
-	{
-		return CollapseEdge(KeepVertID, RemoveVertID, 0, CollapseInfo);
-	}
-
-
-	/**
-	 * Given two edges of the mesh, weld both their vertices, so that one edge is removed.
-	 * This could result in one neighbour edge-pair attached to each vertex also collapsing,
-	 * so those cases are detected and handled (eg middle edge-pair in abysmal ascii drawing below)
-	 *
-	 *   ._._._.    (dots are vertices)
-	 *    \._./
-	 *
-	 * @param KeepEdgeID index of the edge that should be kept
-	 * @param DiscardEdgeID index of the edge that should be removed
-	 * @param MergeInfo returned information about new and modified mesh elements
-	 * @return Ok on success, or enum value indicates why operation cannot be applied. Mesh remains unmodified on error.
-	 */
-	virtual EMeshResult MergeEdges(int KeepEdgeID, int DiscardEdgeID, FMergeEdgesInfo& MergeInfo);
-
-
-	/**
-	 * Insert a new vertex inside a triangle, ie do a 1 to 3 triangle split
-	 * @param TriangleID index of triangle to poke
-	 * @param BaryCoordinates barycentric coordinates of poke position
-	 * @param PokeInfo returned information about new and modified mesh elements
-	 * @return Ok on success, or enum value indicates why operation cannot be applied. Mesh remains unmodified on error.
-	 */
-	virtual EMeshResult PokeTriangle(int TriangleID, const FVector3d& BaryCoordinates, FPokeTriangleInfo& PokeInfo);
-
-	/** Call PokeTriangle at the centroid of the triangle */
-	virtual EMeshResult PokeTriangle(int TriangleID, FPokeTriangleInfo& PokeInfo)
-	{
-		return PokeTriangle(TriangleID, FVector3d::One() / 3.0, PokeInfo);
-	}
-
-	//
-	// Debug utility functions
-	//
-public:
-	/**
-	 * Returns a debug string that contains mesh statistics and other information
-	 */
-	virtual FString MeshInfoString();
-
-	/**
-	 * Options for the IsSameAs check
-	 */
-	struct FSameAsOptions
-	{
-		bool bCheckConnectivity = true;
-		bool bCheckEdgeIDs = false;
-		bool bCheckNormals = false;
-		bool bCheckColors = false;
-		bool bCheckUVs = false;
-		bool bCheckGroups = false;
-		bool bCheckAttributes = false;
-		float Epsilon = TMathUtil<float>::Epsilon;
-	};
-
-	/**
-	 * Check if another mesh is the same as this mesh. By default only checks
-	 * vertices and triangles, turn on other parameters w/ flags
-	 */
-	virtual bool IsSameAs(const FDynamicMesh3& OtherMesh, const FSameAsOptions& Options);
-
-	/**
-	 * Options for what the validity check will permit
-	 */
-	struct FValidityOptions
-	{
-		bool bAllowNonManifoldVertices             = false;
-		bool bAllowAdjacentFacesReverseOrientation = false;
-
-		/**
-		 * Construct validity checking options
-		 */
-		FValidityOptions(bool bAllowNonManifoldVertices = false, bool bAllowAdjacentFacesReverseOrientation = false)
-		    : bAllowNonManifoldVertices(bAllowNonManifoldVertices)
-		    , bAllowAdjacentFacesReverseOrientation(bAllowAdjacentFacesReverseOrientation)
-		{}
-
-		/**
-		 * Construct with most-permissive options that we still consider valid for processing
-		 */
-		static FValidityOptions Permissive()
-		{
-			FValidityOptions ToRet;
-			ToRet.bAllowAdjacentFacesReverseOrientation = true;
-			ToRet.bAllowNonManifoldVertices             = true;
-			return ToRet;
-		}
-	};
-
-	/**
-	 * Checks that the mesh is well-formed, ie all internal data structures are consistent
-	 */
-	virtual bool CheckValidity(FValidityOptions ValidityOptions = FValidityOptions(),
-	                           EValidityCheckFailMode FailMode  = EValidityCheckFailMode::Check) const;
-
-	//
-	// Internal functions
-	//
-protected:
-	inline void SetTriangleInternal(int TriangleID, int v0, int v1, int v2)
-	{
-		Triangles[TriangleID] = FIndex3i(v0, v1, v2);
-	}
-	inline void SetTriangleEdgesInternal(int TriangleID, int e0, int e1, int e2)
-	{
-		TriangleEdges[TriangleID] = FIndex3i(e0, e1, e2);
-	}
-
-	int AddEdgeInternal(int vA, int vB, int tA, int tB = InvalidID);
-	int AddTriangleInternal(int a, int b, int c, int e0, int e1, int e2);
-
-	inline int ReplaceTriangleVertex(int TriangleID, int vOld, int vNew)
-	{
-		FIndex3i& Triangle = Triangles[TriangleID];
-		for (int i : {0, 1, 2})
-		{
-			if (Triangle[i] == vOld)
-			{
-				Triangle[i] = vNew;
-				return 0;
-			}
-		}
-		return -1;
-	}
-
-	inline void AllocateEdgesList(int VertexID)
-	{
-		if (VertexID < (int)VertexEdgeLists.Size())
-		{
-			VertexEdgeLists.Clear(VertexID);
-		}
-		VertexEdgeLists.AllocateAt(VertexID);
-	}
-
-	void GetVertexEdgesList(int VertexID, TArray<int>& EdgesOut) const
-	{
-		for (int eid : VertexEdgeLists.Values(VertexID))
-		{
-			EdgesOut.Add(eid);
-		}
-	}
-
-	inline void SetEdgeVerticesInternal(int EdgeID, int a, int b)
-	{
-		if (a > b)
-		{
-			Swap(a,b);
-		}
-		Edges[EdgeID].Vert[0] = a;
-		Edges[EdgeID].Vert[1] = b;
-	}
-
-	inline void SetEdgeTrianglesInternal(int EdgeID, int t0, int t1)
-	{
-		Edges[EdgeID].Tri[0] = t0;
-		Edges[EdgeID].Tri[1] = t1;
-	}
-
-	int ReplaceEdgeVertex(int EdgeID, int vOld, int vNew);
-	int ReplaceEdgeTriangle(int EdgeID, int tOld, int tNew);
-	int ReplaceTriangleEdge(int EdgeID, int eOld, int eNew);
-
-	inline bool TriangleHasVertex(int TriangleID, int VertexID) const
-	{
-		return Triangles[TriangleID][0] == VertexID || Triangles[TriangleID][1] == VertexID ||
-		    Triangles[TriangleID][2] == VertexID;
-	}
-
-	inline bool TriHasNeighbourTri(int CheckTriID, int NbrTriID) const
-	{
-		return EdgeHasTriangle(TriangleEdges[CheckTriID][0], NbrTriID) ||
-		    EdgeHasTriangle(TriangleEdges[CheckTriID][1], NbrTriID) ||
-		    EdgeHasTriangle(TriangleEdges[CheckTriID][2], NbrTriID);
-	}
-
-	inline bool TriHasSequentialVertices(int TriangleID, int vA, int vB) const
-	{
-		int v0 = Triangles[TriangleID][0], v1 = Triangles[TriangleID][1], v2 = Triangles[TriangleID][2];
-		return ((v0 == vA && v1 == vB) || (v1 == vA && v2 == vB) || (v2 == vA && v0 == vB));
-	}
-
-	int FindTriangleEdge(int TriangleID, int vA, int vB) const;
-
-	int32 FindEdgeInternal(int32 vA, int32 vB, bool& bIsBoundary) const;
-
-	inline bool EdgeHasVertex(int EdgeID, int VertexID) const
-	{
-		const FIndex2i Verts = Edges[EdgeID].Vert;
-		return (Verts[0] == VertexID) || (Verts[1] == VertexID);
-	}
-	inline bool EdgeHasTriangle(int EdgeID, int TriangleID) const
-	{
-		const FIndex2i Tris = Edges[EdgeID].Tri;
-		return (Tris[0] == TriangleID) || (Tris[1] == TriangleID);
-	}
-
-	inline int GetOtherEdgeVertex(int EdgeID, int VertexID) const
-	{
-		const FIndex2i Verts = Edges[EdgeID].Vert;
-		return (Verts[0] == VertexID) ? Verts[1] : ((Verts[1] == VertexID) ? Verts[0] : InvalidID);
-	}
-	inline int GetOtherEdgeTriangle(int EdgeID, int TriangleID) const
-	{
-		const FIndex2i Tris = Edges[EdgeID].Tri;
-		return (Tris[0] == TriangleID) ? Tris[1] : ((Tris[1] == TriangleID) ? Tris[0] : InvalidID);
-	}
-
-	inline void AddTriangleEdge(int TriangleID, int v0, int v1, int j, int EdgeID)
-	{
-		FIndex3i& TriEdges = TriangleEdges.ElementAt(TriangleID, FIndex3i(InvalidID, InvalidID, InvalidID));
-		if (EdgeID != InvalidID)
-		{
-			Edges[EdgeID].Tri[1] = TriangleID;
-			TriEdges[j] = EdgeID;
-		}
-		else
-		{
-			TriEdges[j] = AddEdgeInternal(v0, v1, TriangleID);
-		}
-	}
-
-	// utility function that returns one or two triangles of edge, used to enumerate vertex one-ring triangles
-	// The logic is a bit tricky to follow without drawing it out on paper, but this will only return
-	// each triangle once, for the 'outgoing' edge from the vertex, and each triangle only has one such edge
-	// at any vertex (including boundary triangles)
-	inline FIndex2i GetOrderedOneRingEdgeTris(int VertexID, int EdgeID) const
-	{
-		const FIndex2i Tris = Edges[EdgeID].Tri;
-
-		int vOther = GetOtherEdgeVertex(EdgeID, VertexID);
-		int et1    = Tris[1];
-		et1        = (et1 != InvalidID && TriHasSequentialVertices(et1, VertexID, vOther)) ? et1 : InvalidID;
-		int et0    = Tris[0];
-		return TriHasSequentialVertices(et0, VertexID, vOther) ? FIndex2i(et0, et1) : FIndex2i(et1, InvalidID);
-	}
-
-
-	void ReverseTriOrientationInternal(int TriangleID);
-
-	void UpdateTimeStamp(bool bShapeChange, bool bTopologyChange)
-	{
-		Timestamp++;
-		if (bShapeChange)
-		{
-			ShapeTimestamp++;
-		}
-		if (bTopologyChange)
-		{
-			checkSlow(bShapeChange); // we consider topology change to be a shape change!
-			TopologyTimestamp++;
-		}
-	}
-};
-
-
-
-} // end namespace UE::Geometry
-} // end namespace UE
+#include "DynamicMesh/DynamicMesh3.h"
\ No newline at end of file
diff --git a/Engine/Plugins/Experimental/GeometryProcessing/Source/DynamicMesh/Public/DynamicMeshAABBTree3.h b/Engine/Plugins/Experimental/GeometryProcessing/Source/DynamicMesh/Public/DynamicMeshAABBTree3.h
--- a/Engine/Plugins/Experimental/GeometryProcessing/Source/DynamicMesh/Public/DynamicMeshAABBTree3.h
+++ b/Engine/Plugins/Experimental/GeometryProcessing/Source/DynamicMesh/Public/DynamicMeshAABBTree3.h
@@ -2,15 +2,4 @@
 
 #pragma once
 
-#include "Spatial/MeshAABBTree3.h"
-#include "DynamicMesh3.h"
-
-namespace UE
-{
-namespace Geometry
-{
-
-typedef TMeshAABBTree3<FDynamicMesh3> FDynamicMeshAABBTree3;
-
-} // end namespace UE::Geometry
-} // end namespace UE
+#include "DynamicMesh/DynamicMeshAABBTree3.h"
\ No newline at end of file
diff --git a/Engine/Plugins/Experimental/GeometryProcessing/Source/DynamicMesh/Public/DynamicMeshAttributeSet.h b/Engine/Plugins/Experimental/GeometryProcessing/Source/DynamicMesh/Public/DynamicMeshAttributeSet.h
--- a/Engine/Plugins/Experimental/GeometryProcessing/Source/DynamicMesh/Public/DynamicMeshAttributeSet.h
+++ b/Engine/Plugins/Experimental/GeometryProcessing/Source/DynamicMesh/Public/DynamicMeshAttributeSet.h
@@ -3,476 +3,5 @@
 
 #pragma once
 
-#include "DynamicMeshOverlay.h"
-#include "DynamicMeshTriangleAttribute.h"
-#include "DynamicAttribute.h"
-#include "GeometryTypes.h"
-#include "InfoTypes.h"
-#include "Containers/IndirectArray.h"
-
-namespace UE
-{
-namespace Geometry
-{
-
-/** Standard UV overlay type - 2-element float */
-typedef TDynamicMeshVectorOverlay<float, 2, FVector2f> FDynamicMeshUVOverlay;
-/** Standard Normal overlay type - 3-element float */
-typedef TDynamicMeshVectorOverlay<float, 3, FVector3f> FDynamicMeshNormalOverlay;
-/** Standard Color overlay type - 4-element float (rbga) */
-typedef TDynamicMeshVectorOverlay<float, 4, FVector4f> FDynamicMeshColorOverlay;
-/** Standard per-triangle integer material ID */
-typedef TDynamicMeshScalarTriangleAttribute<int32> FDynamicMeshMaterialAttribute;
-
-/** Per-triangle integer polygroup ID */
-typedef TDynamicMeshScalarTriangleAttribute<int32> FDynamicMeshPolygroupAttribute;
-
-/** Forward declarations */
-template<typename ParentType>
-class TDynamicVertexSkinWeightsAttribute;
-using FDynamicMeshVertexSkinWeightsAttribute = TDynamicVertexSkinWeightsAttribute<FDynamicMesh3>;
-	
-/**
- * FDynamicMeshAttributeSet manages a set of extended attributes for a FDynamicMesh3.
- * This includes UV and Normal overlays, etc.
- * 
- * Currently the default is to always have one UV layer and one Normal layer, 
- * but the number of layers can be requested on construction.
- * 
- * @todo current internal structure is a work-in-progress
- */
-class DYNAMICMESH_API FDynamicMeshAttributeSet : public FDynamicMeshAttributeSetBase
-{
-public:
-	FDynamicMeshAttributeSet(FDynamicMesh3* Mesh);
-
-	FDynamicMeshAttributeSet(FDynamicMesh3* Mesh, int32 NumUVLayers, int32 NumNormalLayers);
-
-	virtual ~FDynamicMeshAttributeSet() override;
-
-	void Copy(const FDynamicMeshAttributeSet& Copy);
-
-	/** returns true if the attached overlays/attributes are compact */
-	bool IsCompact();
-
-	/**
-	 * Performs a CompactCopy of the attached overlays/attributes.
-	 * Called by the parent mesh CompactCopy function.
-	 *
-	 * @param CompactMaps Maps indicating how vertices and triangles were changes in the parent
-	 * @param Copy The attribute set to be copied
-	 */
-	void CompactCopy(const FCompactMaps& CompactMaps, const FDynamicMeshAttributeSet& Copy);
-
-	/**
-	 * Compacts the attribute set in place
-	 * Called by the parent mesh CompactInPlace function
-	 *
-	 * @param CompactMaps Maps of how the vertices and triangles were compacted in the parent
-	 */
-	void CompactInPlace(const FCompactMaps& CompactMaps);
-
-
-	/**
-	 * Split all bowtie vertices in all layers
-	 * @param bParallel if true, layers are processed in parallel
-	 */
-	void SplitAllBowties(bool bParallel = true);
-
-
-	/**
-	 * Enable the matching attributes and overlay layers as the reference Copy set, but do not copy any data across
-	 */
-	void EnableMatchingAttributes(const FDynamicMeshAttributeSet& ToMatch);
-
-	/** @return the parent mesh for this overlay */
-	const FDynamicMesh3* GetParentMesh() const { return ParentMesh; }
-	/** @return the parent mesh for this overlay */
-	FDynamicMesh3* GetParentMesh() { return ParentMesh; }
-
-private:
-	/** @set the parent mesh for this overlay.  Only safe for use during FDynamicMesh move */
-	void Reparent(FDynamicMesh3* NewParent);
-
-public:
-
-	/** @return true if the given edge is a seam edge in any overlay */
-	virtual bool IsSeamEdge(int EdgeID) const;
-
-	/** @return true if the given edge is the termination of a seam in any overlay*/
-	virtual bool IsSeamEndEdge(int EdgeID) const;
-
-	/** @return true if the given edge is a seam edge in any overlay */
-	virtual bool IsSeamEdge(int EdgeID, bool& bIsUVSeamOut, bool& bIsNormalSeamOut, bool& bIsColorSeamOut) const;
-
-	/** @return true if the given vertex is a seam vertex in any overlay */
-	virtual bool IsSeamVertex(int VertexID, bool bBoundaryIsSeam = true) const;
-
-	/** @return true if the given edge is a material ID boundary */
-	virtual bool IsMaterialBoundaryEdge(int EdgeID) const;
-
-	//
-	// UV Layers 
-	//
-
-	/** @return number of UV layers */
-	virtual int NumUVLayers() const
-	{
-		return UVLayers.Num();
-	}
-
-	/** Set number of UV (2-vector float overlay) layers */
-	virtual void SetNumUVLayers(int Num);
-
-	/** @return the UV layer at the given Index  if exists, else nullptr*/
-	FDynamicMeshUVOverlay* GetUVLayer(int Index)
-	{
-		if (Index < UVLayers.Num() && Index > -1)
-		{ 
-			return &UVLayers[Index];
-		}
-		return nullptr;
-	}
-
-	/** @return the UV layer at the given Index if exists, else nullptr */
-	const FDynamicMeshUVOverlay* GetUVLayer(int Index) const
-	{
-		if (Index < UVLayers.Num() && Index > -1)
-		{
-			return &UVLayers[Index];
-		}
-		return nullptr;
-	}
-
-	/** @return the primary UV layer (layer 0) */
-	FDynamicMeshUVOverlay* PrimaryUV()
-	{
-		return GetUVLayer(0);
-	}
-	/** @return the primary UV layer (layer 0) */
-	const FDynamicMeshUVOverlay* PrimaryUV() const
-	{
-		return GetUVLayer(0);
-	}
-
-
-	//
-	// Normal Layers 
-	//
-
-	/** @return number of Normals layers */
-	virtual int NumNormalLayers() const
-	{
-		return NormalLayers.Num();
-	}
-
-	/** Set number of Normals (3-vector float overlay) layers */
-	virtual void SetNumNormalLayers(int Num);
-
-	/** Enable Tangents overlays (Tangent = Normal layer 1, Bitangent = Normal layer 2) */
-	void EnableTangents();
-
-	/** Disable Tangents overlays */
-	void DisableTangents();
-
-	/** @return the Normal layer at the given Index if exists, else nullptr */
-	FDynamicMeshNormalOverlay* GetNormalLayer(int Index)
-	{
-		if (Index < NormalLayers.Num() && Index > -1)
-		{ 
-			return &NormalLayers[Index];
-		}
-		return nullptr;
-	}
-
-	/** @return the Normal layer at the given Index if exists, else nullptr */
-	const FDynamicMeshNormalOverlay* GetNormalLayer(int Index) const
-	{
-		if (Index < NormalLayers.Num() && Index > -1)
-		{
-			return &NormalLayers[Index];
-		}
-		return nullptr;
-	}
-
-
-	/** @return the primary Normal layer (normal layer 0) if exists, else nullptr */
-	FDynamicMeshNormalOverlay* PrimaryNormals()
-	{
-		return GetNormalLayer(0);
-	}
-	/** @return the primary Normal layer (normal layer 0) if exists, else nullptr */
-	const FDynamicMeshNormalOverlay* PrimaryNormals() const
-	{
-		return GetNormalLayer(0);
-	}
-	/** @return the primary tangent layer ( normal layer 1) if exists, else nullptr */
-	FDynamicMeshNormalOverlay* PrimaryTangents()
-	{
-		return GetNormalLayer(1);
-	}
-	/** @return the primary tangent layer ( normal layer 1) if exists, else nullptr */
-	const FDynamicMeshNormalOverlay* PrimaryTangents() const 
-	{
-		return GetNormalLayer(1);
-	}
-	/** @return the primary biTangent layer ( normal layer 2) if exists, else nullptr */
-	FDynamicMeshNormalOverlay* PrimaryBiTangents()
-	{
-		return GetNormalLayer(2);
-	}
-	/** @return the primary biTangent layer ( normal layer 2) if exists, else nullptr */
-	const FDynamicMeshNormalOverlay* PrimaryBiTangents() const
-	{
-		return GetNormalLayer(2);
-	}
-
-	/** @return true if normal layers exist for the normal, tangent, and bitangent */
-	bool HasTangentSpace() const
-	{
-		return (PrimaryNormals() != nullptr && PrimaryTangents()  != nullptr && PrimaryBiTangents() != nullptr);
-	}
-
-	bool HasPrimaryColors() const
-	{
-		return !!ColorLayer;
-	}
-
-	FDynamicMeshColorOverlay* PrimaryColors()
-	{
-		return ColorLayer.Get();
-	}
-
-	const FDynamicMeshColorOverlay* PrimaryColors() const
-	{
-		return ColorLayer.Get();
-	}
-
-	void EnablePrimaryColors();
-	
-	void DisablePrimaryColors();
-
-	//
-	// Polygroup layers
-	//
-
-	/** @return number of Polygroup layers */
-	virtual int32 NumPolygroupLayers() const;
-
-	/** Set the number of Polygroup layers */
-	virtual void SetNumPolygroupLayers(int32 Num);
-
-	/** @return the Polygroup layer at the given Index */
-	FDynamicMeshPolygroupAttribute* GetPolygroupLayer(int Index);
-
-	/** @return the Polygroup layer at the given Index */
-	const FDynamicMeshPolygroupAttribute* GetPolygroupLayer(int Index) const;
-
-
-
-
-	//
-	// Per-Triangle Material ID
-	//
-
-	bool HasMaterialID() const
-	{
-		return !!MaterialIDAttrib;
-	}
-
-
-	void EnableMaterialID();
-
-	void DisableMaterialID();
-
-	FDynamicMeshMaterialAttribute* GetMaterialID()
-	{
-		return MaterialIDAttrib.Get();
-	}
-
-	const FDynamicMeshMaterialAttribute* GetMaterialID() const
-	{
-		return MaterialIDAttrib.Get();
-	}
-
-	/** Skin weights */
-
-	/// Create a new skin weights attribute with a given skin weights profile name. If an attribute already exists with
-	/// that name, that existing attribute will be deleted.
-	void AttachSkinWeightsAttribute(FName InProfileName, FDynamicMeshVertexSkinWeightsAttribute* InAttribute);
-
-	/// Remove a skin weights attribute matching the given profile name.
-	void RemoveSkinWeightsAttribute(FName InProfileName);
-
-	/// Returns true if the list of skin weight attributes includes the given profile name.
-	bool HasSkinWeightsAttribute(FName InProfileName) const
-	{
-		return SkinWeightAttributes.Contains(InProfileName);
-	}
-
-	/// Returns a pointer to a skin weight attribute of the given profile name. If the attribute
-	/// does not exist, a nullptr is returned.
-	FDynamicMeshVertexSkinWeightsAttribute *GetSkinWeightsAttribute(FName InProfileName) const
-	{
-		if (SkinWeightAttributes.Contains(InProfileName))
-		{
-			return SkinWeightAttributes[InProfileName].Get();
-		}
-		else
-		{
-			return nullptr;
-		}
-	}
-	
-	/// Returns a map of all skin weight attributes.
-	const TMap<FName, TUniquePtr<FDynamicMeshVertexSkinWeightsAttribute>>& GetSkinWeightsAttributes() const
-	{
-		return SkinWeightAttributes;
-	}
-	
-	// Attach a new attribute (and transfer ownership of it to the attribute set)
-	void AttachAttribute(FName AttribName, FDynamicMeshAttributeBase* Attribute)
-	{
-		if (GenericAttributes.Contains(AttribName))
-		{
-			UnregisterExternalAttribute(GenericAttributes[AttribName].Get());
-		}
-		GenericAttributes.Add(AttribName, TUniquePtr<FDynamicMeshAttributeBase>(Attribute));
-		RegisterExternalAttribute(Attribute);
-	}
-
-	void RemoveAttribute(FName AttribName)
-	{
-		if (GenericAttributes.Contains(AttribName))
-		{
-			UnregisterExternalAttribute(GenericAttributes[AttribName].Get());
-			GenericAttributes.Remove(AttribName);
-		}
-	}
-
-	FDynamicMeshAttributeBase* GetAttachedAttribute(FName AttribName) const
-	{
-		return GenericAttributes[AttribName].Get();
-	}
-
-	int NumAttachedAttributes() const
-	{
-		return GenericAttributes.Num();
-	}
-
-	bool HasAttachedAttribute(FName AttribName) const
-	{
-		return GenericAttributes.Contains(AttribName);
-	}
-
-	const TMap<FName, TUniquePtr<FDynamicMeshAttributeBase>>& GetAttachedAttributes() const
-	{
-		return GenericAttributes;
-	}
-
-	/**
-	 * Returns true if this AttributeSet is the same as Other.
-	 */
-	bool IsSameAs(const FDynamicMeshAttributeSet& Other) const;
-
-	/**
-	 * Serialization operator for FDynamicMeshAttributeSet.
-	 *
-	 * @param Ar Archive to serialize with.
-	 * @param Set Attribute set to serialize.
-	 * @returns Passing down serializing archive.
-	 */
-	friend DYNAMICMESH_API FArchive& operator<<(FArchive& Ar, FDynamicMeshAttributeSet& Set)
-	{
-		Set.Serialize(Ar);
-		return Ar;
-	}
-
-	/** Serialize the set to an archive. */
-	void Serialize(FArchive& Ar);
-
-protected:
-	/** Parent mesh of this attribute set */
-	FDynamicMesh3* ParentMesh;
-	
-
-	TIndirectArray<FDynamicMeshUVOverlay> UVLayers;
-	TIndirectArray<FDynamicMeshNormalOverlay> NormalLayers;
-	TUniquePtr<FDynamicMeshColorOverlay> ColorLayer;
-
-	TUniquePtr<FDynamicMeshMaterialAttribute> MaterialIDAttrib;
-
-	TIndirectArray<FDynamicMeshPolygroupAttribute> PolygroupLayers;
-
-	TMap<FName, TUniquePtr<FDynamicMeshVertexSkinWeightsAttribute>> SkinWeightAttributes;
-	
-	TMap<FName, TUniquePtr<FDynamicMeshAttributeBase>> GenericAttributes;
-	
-protected:
-	friend class FDynamicMesh3;
-
-	/**
-	 * Initialize the existing attribute layers with the given vertex and triangle sizes
-	 */
-	void Initialize(int MaxVertexID, int MaxTriangleID)
-	{
-		for (FDynamicMeshUVOverlay& UVLayer : UVLayers)
-		{
-			UVLayer.InitializeTriangles(MaxTriangleID);
-		}
-		for (FDynamicMeshNormalOverlay& NormalLayer : NormalLayers)
-		{
-			NormalLayer.InitializeTriangles(MaxTriangleID);
-		}
-	}
-
-	// These functions are called by the FDynamicMesh3 to update the various
-	// attributes when the parent mesh topology has been modified.
-	// TODO: would it be better to register all the overlays and attributes with the base set and not overload these?  maybe!
-	virtual void OnNewTriangle(int TriangleID, bool bInserted);
-	virtual void OnNewVertex(int VertexID, bool bInserted);
-	virtual void OnRemoveTriangle(int TriangleID);
-	virtual void OnRemoveVertex(int VertexID);
-	virtual void OnReverseTriOrientation(int TriangleID);
-	virtual void OnSplitEdge(const DynamicMeshInfo::FEdgeSplitInfo & splitInfo);
-	virtual void OnFlipEdge(const DynamicMeshInfo::FEdgeFlipInfo & flipInfo);
-	virtual void OnCollapseEdge(const DynamicMeshInfo::FEdgeCollapseInfo & collapseInfo);
-	virtual void OnPokeTriangle(const DynamicMeshInfo::FPokeTriangleInfo & pokeInfo);
-	virtual void OnMergeEdges(const DynamicMeshInfo::FMergeEdgesInfo & mergeInfo);
-	virtual void OnSplitVertex(const DynamicMeshInfo::FVertexSplitInfo& SplitInfo, const TArrayView<const int>& TrianglesToUpdate);
-
-	/**
-	 * Check validity of attributes
-	 * 
-	 * @param bAllowNonmanifold Accept non-manifold topology as valid. Note that this should almost always be true for attributes; non-manifold overlays are generally valid.
-	 * @param FailMode Desired behavior if mesh is found invalid
-	 */
-	virtual bool CheckValidity(bool bAllowNonmanifold, EValidityCheckFailMode FailMode) const
-	{
-		bool bValid = FDynamicMeshAttributeSetBase::CheckValidity(bAllowNonmanifold, FailMode);
-		for (int UVLayerIndex = 0; UVLayerIndex < NumUVLayers(); UVLayerIndex++)
-		{
-			bValid = GetUVLayer(UVLayerIndex)->CheckValidity(bAllowNonmanifold, FailMode) && bValid;
-		}
-		bValid = PrimaryNormals()->CheckValidity(bAllowNonmanifold, FailMode) && bValid;
-		if (ColorLayer)
-		{
-			bValid = ColorLayer->CheckValidity(bAllowNonmanifold, FailMode) && bValid;
-		}
-		if (MaterialIDAttrib)
-		{
-			bValid = MaterialIDAttrib->CheckValidity(bAllowNonmanifold, FailMode) && bValid;
-		}
-		for (int PolygroupLayerIndex = 0; PolygroupLayerIndex < NumPolygroupLayers(); PolygroupLayerIndex++)
-		{
-			bValid = GetPolygroupLayer(PolygroupLayerIndex)->CheckValidity(bAllowNonmanifold, FailMode) && bValid;
-		}
-		return bValid;
-	}
-};
-
-
-
-} // end namespace UE::Geometry
-} // end namespace UE
+#include "DynamicMesh/DynamicMeshAttributeSet.h"
 
diff --git a/Engine/Plugins/Experimental/GeometryProcessing/Source/DynamicMesh/Public/DynamicMeshChangeTracker.h b/Engine/Plugins/Experimental/GeometryProcessing/Source/DynamicMesh/Public/DynamicMeshChangeTracker.h
--- a/Engine/Plugins/Experimental/GeometryProcessing/Source/DynamicMesh/Public/DynamicMeshChangeTracker.h
+++ b/Engine/Plugins/Experimental/GeometryProcessing/Source/DynamicMesh/Public/DynamicMeshChangeTracker.h
@@ -2,410 +2,5 @@
 
 #pragma once
 
-#include "DynamicMesh3.h"
-#include "DynamicMeshOverlay.h"
-#include "DynamicMeshAttributeSet.h"
+#include "DynamicMesh/DynamicMeshChangeTracker.h"
 
-
-//
-// Implementation of mesh change tracking for FDynamicMesh3.
-// 
-// The top-level class is FDynamicMeshChangeTracker, found at the bottom of this file.
-// You create an instance of this and then call BeginChange(), then call SaveVertex()/SaveTriangle()
-// before you make modifications to a vertex/triangle, and then call EndChange() at the end.
-// This function emits a FDynamicMeshChange instance, and you can call Apply() to apply/revert it.
-// So this is the object you would store in a higher-level undo/redo record, for example.
-// 
-// Attribute overlays make everything more complicated of course. Handling of these
-// follows a similar structure - the MeshChangeTracker creates a FDynamicMeshAttributeSetChangeTracker,
-// which in turn creates a TDynamicMeshAttributeChange for each UV and Normal overlay
-// (grouped together in a FDynamicMeshAttributeChangeSet). However you don't have to explicitly
-// do anything to get Attribute support, if the initial Mesh had attributes, then
-// this all happens automatically, and the attribute changes get attached to the FDynamicMeshChange.
-// 
-// @todo Currently the attribute support is hardcoded for UV (2-float) and Normal (3-float) overlays.
-// Perhaps it would be possible to make this more flexible to avoid hardcoding these?
-// Very little of the code depends on the element size or type.
-// 
-// 
-
-namespace UE
-{
-namespace Geometry
-{
-
-
-/**
- * TDynamicMeshAttributeChange represents a change to an attribute overlay of a FDynamicMesh3.
- * @warning This class is meant to be used via FDynamicMeshChange and is not fully functional
- * on its own (see comments in ApplyReplaceChange)
- */
-template<typename RealType, int ElementSize>
-class TDynamicMeshAttributeChange
-{
-public:
-	void SaveInitialElement(const TDynamicMeshOverlay<RealType,ElementSize>* Overlay, int ElementID);
-	void SaveInitialTriangle(const TDynamicMeshOverlay<RealType,ElementSize>* Overlay, int TriangleID);
-
-	void StoreFinalElement(const TDynamicMeshOverlay<RealType, ElementSize>* Overlay, int ElementID);
-	void StoreFinalTriangle(const TDynamicMeshOverlay<RealType, ElementSize>* Overlay, int TriangleID);
-
-	bool Apply(TDynamicMeshOverlay<RealType, ElementSize>* Overlay, bool bRevert) const;
-
-protected:
-	struct FChangeElement
-	{
-		int ElementID;
-		int DataIndex;
-		int ParentVertexID;
-	};
-
-	struct FChangeTriangle
-	{
-		int TriangleID;
-		FIndex3i Elements;
-	};
-
-	TArray<FChangeElement> OldElements;
-	TArray<RealType> OldElementData;
-	TArray<FChangeTriangle> OldTriangles;
-
-	TArray<FChangeElement> NewElements;
-	TArray<RealType> NewElementData;
-	TArray<FChangeTriangle> NewTriangles;
-
-	void ApplyReplaceChange(TDynamicMeshOverlay<RealType,ElementSize>* Overlay,
-		const TArray<FChangeTriangle>& RemoveTris,
-		const TArray<FChangeElement>& InsertElements,
-		const TArray<RealType>& InsertElementData,
-		const TArray<FChangeTriangle>& InsertTris) const;
-};
-
-
-
-/** Standard UV overlay change type - 2-element float */
-typedef TDynamicMeshAttributeChange<float,2> FDynamicMeshUVChange;
-/** Standard Normal overlay change type - 3-element float */
-typedef TDynamicMeshAttributeChange<float,3> FDynamicMeshNormalChange;
-/** Standard Color overlay change type - 4-element float */
-typedef TDynamicMeshAttributeChange<float, 4> FDynamicMeshColorChange;
-/** Standard per-triangle integer attribute change type */
-typedef FDynamicMeshTriangleAttributeChange<int32,1> FDynamicMeshTriGroupChange;
-
-/**
- * FDynamicMeshAttributeChangeSet stores a set of UV and Normal changes for a FDynamicMesh3
- */
-class DYNAMICMESH_API FDynamicMeshAttributeChangeSet
-{
-public:
-	UE_NONCOPYABLE(FDynamicMeshAttributeChangeSet);
-	FDynamicMeshAttributeChangeSet() = default;
-
-	TArray<FDynamicMeshUVChange> UVChanges;
-	TArray<FDynamicMeshNormalChange> NormalChanges;
-	TOptional<FDynamicMeshColorChange> ColorChange;
-	TOptional<FDynamicMeshTriGroupChange> MaterialIDAttribChange;
-	TArray<FDynamicMeshTriGroupChange> PolygroupChanges;
-	TArray<TUniquePtr<FDynamicMeshAttributeChangeBase>> RegisteredAttributeChanges;
-
-	/** call ::Apply() on all the UV and Normal changes */
-	bool Apply(FDynamicMeshAttributeSet* Attributes, bool bRevert) const;
-};
-
-
-/**
- * FDynamicMeshChange stores a "change" in a FDynamicMesh3, which in this context
- * means the replacement of one set of triangles with a second set, that may have
- * different vertices/attributes. The change can be applied and reverted via ::Apply()
- * 
- * Construction of a well-formed FDynamicMeshChange is quite complex and it is strongly
- * suggested that you do so via FDynamicMeshChangeTracker
- */
-class DYNAMICMESH_API FDynamicMeshChange
-{
-public:
-	~FDynamicMeshChange();
-
-	/** Store the initial state of a vertex */
-	void SaveInitialVertex(const FDynamicMesh3* Mesh, int VertexID);
-	/** Store the initial state of a triangle */
-	void SaveInitialTriangle(const FDynamicMesh3* Mesh, int TriangleID);
-
-	/** Store the final state of a vertex */
-	void StoreFinalVertex(const FDynamicMesh3* Mesh, int VertexID);
-	/** Store the final state of a triangle */
-	void StoreFinalTriangle(const FDynamicMesh3* Mesh, int TriangleID);
-
-	/** Attach an attribute change set to this mesh change, which will the be applied/reverted automatically */
-	void AttachAttributeChanges(TUniquePtr<FDynamicMeshAttributeChangeSet> AttribChanges)
-	{
-		this->AttributeChanges = MoveTemp(AttribChanges);
-	}
-
-	/** Apply or Revert this change using the given Mesh */
-	bool Apply(FDynamicMesh3* Mesh, bool bRevert) const;
-
-	/** Do (limited) sanity checks on this MeshChange to ensure it is well-formed */
-	void VerifySaveState() const;
-
-	/** @return true if this vertex was saved. Uses linear search. */
-	bool HasSavedVertex(int VertexID);
-
-	/** store IDs of saved triangles in TrianglesOut. if bInitial=true, old triangles are stored, otherwise new triangles */
-	void GetSavedTriangleList(TArray<int>& TrianglesOut, bool bInitial) const;
-
-	/** run self-validity checks on internal data structures to test if change is well-formed */
-	void CheckValidity(EValidityCheckFailMode FailMode = EValidityCheckFailMode::Check) const;
-
-protected:
-
-	struct FChangeVertex
-	{
-		int VertexID;
-		FVertexInfo Info;
-	};
-
-	struct FChangeTriangle
-	{
-		int TriangleID;
-		FIndex3i Vertices;
-		FIndex3i Edges;
-		int GroupID;
-	};
-
-	TArray<FChangeVertex> OldVertices;
-	TArray<FChangeTriangle> OldTriangles;
-
-	TArray<FChangeVertex> NewVertices;
-	TArray<FChangeTriangle> NewTriangles;
-
-	TUniquePtr<FDynamicMeshAttributeChangeSet> AttributeChanges;
-
-	void ApplyReplaceChange(FDynamicMesh3* Mesh,
-		const TArray<FChangeTriangle>& RemoveTris,
-		const TArray<FChangeVertex>& InsertVerts,
-		const TArray<FChangeTriangle>& InsertTris) const;
-
-};
-
-
-
-
-
-
-
-/**
- * FDynamicMeshAttributeSetChangeTracker constructs a well-formed set of TDynamicMeshAttributeChange
- * objects (stored in a FDynamicMeshAttributeChangeSet). You should not use this class
- * directly, it is intended to be used via FDynamicMeshChangeTracker
- */
-class DYNAMICMESH_API FDynamicMeshAttributeSetChangeTracker
-{
-public:
-	explicit FDynamicMeshAttributeSetChangeTracker(const FDynamicMeshAttributeSet* Attribs);
-
-	/** Start tracking a change */
-	void BeginChange();
-	/** End the change transaction and get the resulting change object */
-	TUniquePtr<FDynamicMeshAttributeChangeSet> EndChange();
-
-	/** Store the initial state of a triangle */
-	void SaveInitialTriangle(int TriangleID);
-
-	/** store the final state of a set of triangles */
-	void StoreAllFinalTriangles(const TArray<int>& TriangleIDs);
-
-	/** Store the initial state of a vertex */
-	void SaveInitialVertex(int VertexID);
-
-	/** store the final state of a set of vertices */
-	void StoreAllFinalVertices(const TArray<int>& VertexIDs);
-
-
-protected:
-	const FDynamicMeshAttributeSet* Attribs = nullptr;
-
-	FDynamicMeshAttributeChangeSet* Change = nullptr;
-
-	struct FElementState
-	{
-		int MaxElementID;
-		TBitArray<> StartElements;
-		TBitArray<> ChangedElements;
-	};
-	TArray<FElementState> UVStates;
-	TArray<FElementState> NormalStates;
-	FElementState ColorState;
-
-	template<typename AttribOverlayType, typename AttribChangeType>
-	void SaveElement(int ElementID, FElementState& State, const AttribOverlayType* Overlay, AttribChangeType& ChangeIn)
-	{
-		if (ElementID < State.MaxElementID && State.ChangedElements[ElementID] == false && State.StartElements[ElementID] == true)
-		{
-			ChangeIn.SaveInitialElement(Overlay, ElementID);
-			State.ChangedElements[ElementID] = true;
-		}
-	}
-
-};
-
-
-
-
-
-
-
-
-
-/**
- * FDynamicMeshChangeTracker tracks changes to a FDynamicMesh and returns a
- * FDynamicMeshChange instance that represents this change and allows it to be reverted/reapplied.
- * This is the top-level class you likely want to use to track mesh changes.
- * 
- * Call BeginChange() before making any changes to the mesh, then call SaveVertex()
- * and SaveTriangle() before modifying the respective elements. Then call EndChange()
- * to construct a FDynamicMeshChange that represents the mesh delta.
- *
- * @warning Currently only vertices that are part of saved triangles will be stored in the emitted FMeshChange!
- * 
- */
-class DYNAMICMESH_API FDynamicMeshChangeTracker
-{
-public:
-	explicit FDynamicMeshChangeTracker(const FDynamicMesh3* Mesh);
-	~FDynamicMeshChangeTracker();
-
-	/** Initialize the change-tracking process */
-	void BeginChange();
-	/** Construct a change object that represents the delta between the Begin and End states */
-	TUniquePtr<FDynamicMeshChange> EndChange();
-
-	/** Save necessary information about a triangle before it is modified */
-	void SaveTriangle(int32 TriangleID, bool bSaveVertices);
-
-	/** Save necessary information about an edge before it is modified */
-	inline void SaveEdge(int32 EdgeID, bool bVertices);
-
-	/** Save necessary information about a set of triangles before they are modified */
-	template<typename EnumerableType>
-	void SaveTriangles(EnumerableType TriangleIDs, bool bSaveVertices);
-
-	/** Save necessary information about a set of triangles before they are modified, and also include any direct triangle neighbours */
-	template<typename EnumerableType>
-	void SaveTrianglesAndNeighbourTris(EnumerableType TriangleIDs, bool bSaveVertices);
-
-	/** Save necessary information about a set of triangles in one-ring of a vertex */
-	inline void SaveVertexOneRingTriangles(int32 VertexID, bool bSaveVertices);
-
-	/** Save necessary information about a set of triangles in one-rings of a set of vertices */
-	template<typename EnumerableType>
-	void SaveVertexOneRingTriangles(EnumerableType VertexIDs, bool bSaveVertices);
-
-	/** Do (limited) sanity checking to make sure that the change is well-formed*/
-	void VerifySaveState();
-
-protected:
-
-	//
-	// Currently EndChange() only stores vertices that are part of modified triangles.
-	// So calling SaveVertex() independently, on triangles that are not saved, will be lost.
-	// (This needs to be fixed)
-	//
-
-	/** Save necessary information about a vertex before it is modified */
-	void SaveVertex(int32 VertexID);
-
-	/** Save necessary information about a set of vertices before they are modified */
-	template<typename EnumerableType>
-	void SaveVertices(EnumerableType VertexIDs);
-
-
-protected:
-	const FDynamicMesh3* Mesh = nullptr;
-
-	/** Active attribute tracker, if Mesh has attribute overlays */
-	FDynamicMeshAttributeSetChangeTracker* AttribChangeTracker = nullptr;
-
-	/** Active change that is being constructed */
-	FDynamicMeshChange* Change = nullptr;
-
-	int32 MaxTriangleID;
-	TBitArray<> StartTriangles;		// bit is 1 if triangle ID was in initial mesh on BeginChange()
-	TBitArray<> ChangedTriangles;	// bit is set to 1 if triangle was in StartTriangles and then had SaveTriangle() called for it
-
-	int32 MaxVertexID;
-	TBitArray<> StartVertices;		// bit is 1 if vertex ID was in initial mesh on BeginChange()
-	TBitArray<> ChangedVertices;	// bit is set to 1 if vertex was in StartVertices and then had SaveVertex() called for it
-};
-
-
-
-void FDynamicMeshChangeTracker::SaveEdge(int32 EdgeID, bool bVertices)
-{
-	const FDynamicMesh3::FEdge Edge = Mesh->GetEdge(EdgeID);
-	SaveTriangle(Edge.Tri[0], bVertices);
-	if (Edge.Tri[1] != FDynamicMesh3::InvalidID)
-	{
-		SaveTriangle(Edge.Tri[1], bVertices);
-	}
-}
-
-template<typename EnumerableType>
-void FDynamicMeshChangeTracker::SaveVertices(EnumerableType VertexIDs)
-{
-	for (int32 VertexID : VertexIDs)
-	{
-		SaveVertex(VertexID);
-	}
-}
-
-template<typename EnumerableType>
-void FDynamicMeshChangeTracker::SaveTriangles(EnumerableType TriangleIDs, bool bSaveVertices)
-{
-	for (int32 TriangleID : TriangleIDs)
-	{
-		SaveTriangle(TriangleID, bSaveVertices);
-	}
-}
-
-
-template<typename EnumerableType>
-void FDynamicMeshChangeTracker::SaveTrianglesAndNeighbourTris(EnumerableType TriangleIDs, bool bSaveVertices)
-{
-	for (int32 TriangleID : TriangleIDs)
-	{
-		SaveTriangle(TriangleID, bSaveVertices);
-
-		FIndex3i TriNbrs = Mesh->GetTriNeighbourTris(TriangleID);
-		for (int32 j = 0; j < 3; ++j)
-		{
-			if (TriNbrs[j] != FDynamicMesh3::InvalidID)
-			{
-				SaveTriangle(TriNbrs[j], bSaveVertices);
-			}
-		}
-	}
-}
-
-
-void FDynamicMeshChangeTracker::SaveVertexOneRingTriangles(int32 VertexID, bool bSaveVertices)
-{
-	Mesh->EnumerateVertexTriangles(VertexID, [this, bSaveVertices](int32 TriangleID)
-	{
-		SaveTriangle(TriangleID, bSaveVertices);
-	});
-}
-
-template<typename EnumerableType>
-void FDynamicMeshChangeTracker::SaveVertexOneRingTriangles(EnumerableType VertexIDs, bool bSaveVertices)
-{
-	for (int32 VertexID : VertexIDs)
-	{
-		SaveVertexOneRingTriangles(VertexID, bSaveVertices);
-	}
-}
-
-
-
-} // end namespace UE::Geometry
-} // end namespace UE
\ No newline at end of file
