GeometryCore: create new GeometryCore Module. Move GeometricObjects code from GeometryProcessing plugin to this new module, as well as core FDynamicMesh3 classes (attributes, aabbtree/octree, change t

diff --git a/Engine/Plugins/Experimental/GeometryProcessing/Source/DynamicMesh/Public/DynamicMeshOctree3.h b/Engine/Plugins/Experimental/GeometryProcessing/Source/DynamicMesh/Public/DynamicMeshOctree3.h
--- a/Engine/Plugins/Experimental/GeometryProcessing/Source/DynamicMesh/Public/DynamicMeshOctree3.h
+++ b/Engine/Plugins/Experimental/GeometryProcessing/Source/DynamicMesh/Public/DynamicMeshOctree3.h
@@ -2,455 +2,4 @@
 
 #pragma once
 
-#include "Spatial/MeshAABBTree3.h"
-#include "Spatial/SparseDynamicOctree3.h"
-#include "MeshQueries.h"
-#include "DynamicMesh3.h"
-#include "Async/ParallelFor.h"
-
-
-namespace UE
-{
-namespace Geometry
-{
-
-
-/**
- * FDynamicMeshOctree3 is an extension of FSparseDynamicOctree3 for the triangles of a FDynamicMesh3 instance.
- * This extension does several things:
- *   1) provides a simplified API based on triangle IDs to various Octree functions
- *   2) tracks ModifiedBounds box of modified areas
- *   3) support for computing/updating/querying a "Cut" of the octree, ie a set of cells which are roots of sub-branches
- *      that partition the tree. This is useful for splitting up mesh processing/rendering into spatially-coherent chunks.
- *      (This functionality should probably be extracted into a separate class...)
- */
-class FDynamicMeshOctree3 : public FSparseDynamicOctree3
-{
-	// potential optimizations:
-	//    - keep track of how many triangles are descendents of each cell? root cells?
-
-
-
-public:
-	/** parent mesh */
-	const FDynamicMesh3* Mesh;
-
-	/** bounding box of triangles that have been inserted/removed since last clear */
-	FAxisAlignedBox3d ModifiedBounds;
-
-	/** 
-	 * Add all triangles of MeshIn to the octree 
-	 */
-	void Initialize(const FDynamicMesh3* MeshIn)
-	{
-		this->Mesh = MeshIn;
-
-		for (int tid : Mesh->TriangleIndicesItr())
-		{
-			InsertTriangle(tid);
-		}
-	}
-
-	/**
-	 * Reset the internal ModifiedBounds box that tracks modified triangle bounds
-	 */
-	void ResetModifiedBounds()
-	{
-		ModifiedBounds = FAxisAlignedBox3d::Empty();
-	}
-
-	/**
-	 * Insert a triangle into the tree
-	 */
-	void InsertTriangle(int32 TriangleID)
-	{
-		if (Mesh->IsTriangle(TriangleID))
-		{
-			FAxisAlignedBox3d Bounds = Mesh->GetTriBounds(TriangleID);
-			ModifiedBounds.Contain(Bounds);
-			InsertObject(TriangleID, Bounds);
-		}
-	}
-
-	/**
-	 * Insert a list of triangles into the tree
-	 */
-	void InsertTriangles(const TArray<int>& Triangles)
-	{
-		int N = Triangles.Num();
-		for (int i = 0; i < N; ++i)
-		{
-			if (Mesh->IsTriangle(Triangles[i]))
-			{
-				FAxisAlignedBox3d Bounds = Mesh->GetTriBounds(Triangles[i]);
-				ModifiedBounds.Contain(Bounds);
-				InsertObject(Triangles[i], Bounds);
-			}
-		}
-	}
-
-	/**
-	 * Insert a set of triangles into the tree
-	 */
-	void InsertTriangles(const TSet<int>& Triangles)
-	{
-		for (int TriangleID : Triangles)
-		{
-			if (Mesh->IsTriangle(TriangleID))
-			{
-				FAxisAlignedBox3d Bounds = Mesh->GetTriBounds(TriangleID);
-				ModifiedBounds.Contain(Bounds);
-				InsertObject(TriangleID, Bounds);
-			}
-		}
-	}
-
-
-	/**
-	 * Remove a triangle from the tree
-	 */
-	bool RemoveTriangle(int32 TriangleID)
-	{
-		if (Mesh->IsTriangle(TriangleID))
-		{
-			FAxisAlignedBox3d Bounds = Mesh->GetTriBounds(TriangleID);
-			ModifiedBounds.Contain(Bounds);
-		}
-		return RemoveObject(TriangleID);		// will ignore if we do not contain this triangle
-	}
-
-
-	/**
-	 * Remove a set of triangles into the tree
-	 */
-	template<typename EnumerableType>
-	void RemoveTriangles(const EnumerableType& Triangles, bool bMarkModifiedBounds = true)
-	{
-		for (int TriangleID : Triangles)
-		{
-			if (RemoveObject(TriangleID) && bMarkModifiedBounds && Mesh->IsTriangle(TriangleID))
-			{
-				FAxisAlignedBox3d Bounds = Mesh->GetTriBounds(TriangleID);
-				ModifiedBounds.Contain(Bounds);
-			}
-		}
-	}
-
-
-	/**
-	 * Reinsert a set of triangles into the tree
-	 */
-	template<typename EnumerableType>
-	void ReinsertTriangles(const EnumerableType& Triangles)
-	{
-		for (int TriangleID : Triangles)
-		{
-			if (Mesh->IsTriangle(TriangleID))
-			{
-				FAxisAlignedBox3d Bounds = Mesh->GetTriBounds(TriangleID);
-				ModifiedBounds.Contain(Bounds);
-				ReinsertObject(TriangleID, Bounds);
-			}
-			else
-			{
-				RemoveObject(TriangleID);		// can only remove, will ignore if we do not contain this triangle
-			}
-		}
-	}
-
-
-	/**
-	 * Reinsert a set of triangles into the tree. Internally precomputes which triangles need
-	 * re-inserting, which can be done in parallel and generally saves time as some triangles can be skipped.
-	 */
-	void ReinsertTrianglesParallel(const TArray<int32>& Triangles, TArray<uint32>& TempBuffer, TArray<bool>& TempFlagBuffer)
-	{
-		int32 NumTriangles = Triangles.Num();
-		TempBuffer.SetNum(NumTriangles, false);
-		TempFlagBuffer.SetNum(NumTriangles, false);
-
-		// can check which triangles need reinsertion in parallel. This will also return which
-		// CellID the triangle is in, which saves time in the Reinsert function
-		ParallelFor(NumTriangles, [&](int k)
-		{
-			FAxisAlignedBox3d Bounds = Mesh->GetTriBounds(Triangles[k]);
-			TempFlagBuffer[k] = CheckIfObjectNeedsReinsert(Triangles[k], Bounds, TempBuffer[k]);
-		});
-
-		// now reinsert all necessary triangles
-		for (int32 k = 0; k < NumTriangles; ++k)
-		{
-			if (TempFlagBuffer[k])
-			{
-				ReinsertObject(Triangles[k], Mesh->GetTriBounds(Triangles[k]), TempBuffer[k]);
-			}
-		}
-	}
-
-
-
-	/**
-	 * Include the current bounds of a triangle in the ModifiedBounds box
-	 */
-	void NotifyPendingModification(int TriangleID)
-	{
-		if (Mesh->IsTriangle(TriangleID))
-		{
-			FAxisAlignedBox3d Bounds = Mesh->GetTriBounds(TriangleID);
-			ModifiedBounds.Contain(Bounds);
-		}
-	}
-
-	/**
-	 * Include the current bounds of a set of triangles in the ModifiedBounds box
-	 */
-	template<typename EnumerableType>
-	void NotifyPendingModification(const EnumerableType& Triangles)
-	{
-		for (int TriangleID : Triangles)
-		{
-			if (Mesh->IsTriangle(TriangleID))
-			{
-				FAxisAlignedBox3d Bounds = Mesh->GetTriBounds(TriangleID);
-				ModifiedBounds.Contain(Bounds);
-			}
-		}
-	}
-
-
-	/**
-	 * Find the nearest triangle of the mesh that is hit by the ray
-	 */
-	int32 FindNearestHitObject(const FRay3d& Ray,
-		double MaxDistance = TNumericLimits<double>::Max()) const
-	{
-		return FSparseDynamicOctree3::FindNearestHitObject(Ray,
-			[&](int tid) {return Mesh->GetTriBounds(tid); },
-			[&](int tid, const FRay3d& Ray) {
-			FIntrRay3Triangle3d Intr = TMeshQueries<FDynamicMesh3>::TriangleIntersection(*Mesh, tid, Ray);
-			return (Intr.IntersectionType == EIntersectionType::Point) ?
-				Intr.RayParameter : TNumericLimits<double>::Max();
-
-		}, MaxDistance);
-	}
-
-
-	/**
-	 * Find the nearest triangle of the mesh that is hit by the ray
-	 * @param IncludeTriangleIDFunc predicate function that must return true for given TriangleID for it to be considered
-	 */
-	int32 FindNearestHitObject(const FRay3d& Ray,
-		TFunctionRef<bool(int)> IncludeTriangleIDFunc,
-		double MaxDistance = TNumericLimits<double>::Max()) const
-	{
-		return FSparseDynamicOctree3::FindNearestHitObject(Ray,
-			[&](int tid) {return Mesh->GetTriBounds(tid); },
-			[&](int tid, const FRay3d& Ray) {
-				if (IncludeTriangleIDFunc(tid) == false)
-				{
-					return TNumericLimits<double>::Max();
-				}
-				FIntrRay3Triangle3d Intr = TMeshQueries<FDynamicMesh3>::TriangleIntersection(*Mesh, tid, Ray);
-				return (Intr.IntersectionType == EIntersectionType::Point) ?
-					Intr.RayParameter : TNumericLimits<double>::Max();
-			}, 
-			MaxDistance);
-	}
-
-
-	/**
-	 * Check that the Octree is internally valid
-	 */
-	void CheckValidity(
-		EValidityCheckFailMode FailMode = EValidityCheckFailMode::Check,
-		bool bVerbose = false,
-		bool bFailOnMissingObjects = false) const
-	{
-		FSparseDynamicOctree3::CheckValidity(
-			[&](int tid) {return Mesh->IsTriangle(tid); },
-			[&](int tid) {return Mesh->GetTriBounds(tid); },
-			FailMode, bVerbose, bFailOnMissingObjects);
-	}
-
-
-
-
-
-
-	//
-	// Support for building "cuts" of the octree, which are sets
-	// of internal nodes which can be used to decompose the tree
-	// (eg into spatially-coherent chunks of triangles, for example)
-	// 
-
-	/** FCellReference is a handle to an internal cell of the octree */
-	struct FCellReference
-	{
-		uint32 CellID;
-	};
-
-	/** FTreeCutSet is a cut of the tree, ie a set of internal cells which are "parents" of separate branches */
-	class FTreeCutSet
-	{
-	public:
-		TArray<FCellReference> CutCells;
-
-	protected:
-		TSet<uint32> CutCellIDs;			// just a set version of CutCells for internal use - maybe use set there instead, and add compare ops to FCellReference?
-		int FixedCutLevel;
-
-		friend class FDynamicMeshOctree3;
-	};
-
-
-	/**
-	 * @return a cut of the tree at a fixed level
-	 */
-	FTreeCutSet BuildLevelCutSet(uint32 CutLevel = 5) const
-	{
-		FTreeCutSet CutSet;
-		CutSet.FixedCutLevel = CutLevel;
-
-		for (const FSparseOctreeCell& Cell : Cells)
-		{
-			if (Cell.Level == CutSet.FixedCutLevel)
-			{
-				FCellReference CellRef;
-				CellRef.CellID = Cell.CellID;
-				CutSet.CutCells.Add(CellRef);
-
-				CutSet.CutCellIDs.Add(Cell.CellID);
-			}
-		}
-
-		return CutSet;
-	}
-
-
-	/**
-	 * For a fixed-level cut set created by BuildLevelCutSet, check that all current cells at that level are in the cut set
-	 * (call this after adding/removing to the tree to make sure the CutSet is up to date)
-	 * @param CutSet the tree cut set to check/update
-	 * @param NewCutCellsOut list of new cut cells added by the update
-	 */
-	void UpdateLevelCutSet(FTreeCutSet& CutSet, TArray<FCellReference>& NewCutCellsOut) const
-	{
-		for (const FSparseOctreeCell& Cell : Cells)
-		{
-			if (Cell.Level == CutSet.FixedCutLevel)
-			{
-				if (CutSet.CutCellIDs.Contains(Cell.CellID) == false)
-				{
-					FCellReference CellRef;
-					CellRef.CellID = Cell.CellID;
-					CutSet.CutCells.Add(CellRef);
-					CutSet.CutCellIDs.Add(Cell.CellID);
-					NewCutCellsOut.Add(CellRef);
-				}
-			}
-		}
-	}
-
-	/**
-	 * @return true if the cell identified by CellRef intersects the Bounds box
-	 */
-	bool TestCellIntersection(const FCellReference& CellRef, const FAxisAlignedBox3d& Bounds) const
-	{
-		check(CellRefCounts.IsValid(CellRef.CellID));
-		const FSparseOctreeCell& Cell = Cells[CellRef.CellID];
-		return GetCellBox(Cell, MaxExpandFactor).Intersects(Bounds);
-	}
-
-	/**
-	 * Call TriangleFunc on any triangles in the branch of the tree starting at CellRef
-	 */
-	void CollectTriangles(
-		const FCellReference& CellRef, 
-		TFunctionRef<void(int)> TriangleFunc) const
-	{
-		check(CellRefCounts.IsValid(CellRef.CellID));
-
-		TArray<const FSparseOctreeCell*> Queue;
-		Queue.Add(&Cells[CellRef.CellID]);
-		while (Queue.Num() > 0)
-		{
-			const FSparseOctreeCell* CurCell = Queue.Pop(false);
-
-			// process elements
-			for (int ObjectID : CellObjectLists.Values(CurCell->CellID))
-			{
-				TriangleFunc(ObjectID);
-			}
-
-			for (int k = 0; k < 8; ++k)
-			{
-				if (CurCell->HasChild(k))
-				{
-					const FSparseOctreeCell* ChildCell = &Cells[CurCell->GetChildCellID(k)];
-					Queue.Add(ChildCell);
-				}
-			}
-		}
-	}
-
-
-	/**
-	 * Call TriangleFunc for all triangles in the octree "above" the CutSet
-	 * (ie at tree cells that are not children of any cut cells)
-	 */
-	void CollectRootTriangles(const FTreeCutSet& CutSet,
-		TFunctionRef<void(int)> TriangleFunc) const
-	{
-		TArray<const FSparseOctreeCell*> Queue;
-
-		// start at root cells
-		RootCells.AllocatedIteration([&](const uint32* RootCellID)
-		{
-			const FSparseOctreeCell* RootCell = &Cells[*RootCellID];
-			Queue.Add(&Cells[*RootCellID]);
-		});
-
-		while (Queue.Num() > 0)
-		{
-			const FSparseOctreeCell* CurCell = Queue.Pop(false);
-			if (CutSet.CutCellIDs.Contains(CurCell->CellID))
-			{
-				continue;
-			}
-
-			// process elements
-			for (int TriangleID : CellObjectLists.Values(CurCell->CellID))
-			{
-				TriangleFunc(TriangleID);
-			}
-
-			for (int k = 0; k < 8; ++k)
-			{
-				if (CurCell->HasChild(k))
-				{
-					const FSparseOctreeCell* ChildCell = &Cells[CurCell->GetChildCellID(k)];
-					Queue.Add(ChildCell);
-				}
-			}
-		}
-	}
-
-	/**
-	 * Call TriangleFunc for any triangles in the spill set (ie not contained in any Root cell)
-	 */
-	void CollectSpillTriangles(
-		TFunctionRef<void(int)> TriangleFunc) const
-	{
-		for (int tid : SpillObjectSet)
-		{
-			TriangleFunc(tid);
-		}
-	}
-
-
-};
-
-
-
-} // end namespace UE::Geometry
-} // end namespace UE
+#include "DynamicMesh/DynamicMeshOctree3.h"
diff --git a/Engine/Plugins/Experimental/GeometryProcessing/Source/DynamicMesh/Public/DynamicMeshOverlay.h b/Engine/Plugins/Experimental/GeometryProcessing/Source/DynamicMesh/Public/DynamicMeshOverlay.h
--- a/Engine/Plugins/Experimental/GeometryProcessing/Source/DynamicMesh/Public/DynamicMeshOverlay.h
+++ b/Engine/Plugins/Experimental/GeometryProcessing/Source/DynamicMesh/Public/DynamicMeshOverlay.h
@@ -3,590 +3,5 @@
 
 #pragma once
 
-#include "Util/DynamicVector.h"
-#include "Util/RefCountVector.h"
-#include "Util/CompactMaps.h"
-#include "VectorTypes.h"
-#include "GeometryTypes.h"
-#include "InfoTypes.h"
-#include "IndexTypes.h"
+#include "DynamicMesh/DynamicMeshOverlay.h"
 
-namespace UE
-{
-namespace Geometry
-{
-
-class FDynamicMesh3;
-
-/**
- * TDynamicMeshOverlay is an add-on to a FDynamicMesh3 that allows for per-triangle storage
- * of an "element" (eg like a per-triangle UV or normal). However the elements can be shared
- * between triangles because they are stored in a separate indexable list. As a result,
- * the overlay has it's own topology, IE there may be "seam" boundary edges in the 
- * overlay topology that are not mesh boundary edges in the associated/parent mesh.
- *
- * A "seam" edge is one where at least one of the elements of the triangles on either
- * side of the edge is not shared between the two triangles.
- *
- * The FDynamicMesh3 mesh topology operations (eg split/flip/collapse edge, poke face, etc) 
- * can be mirrored to the overlay via OnSplitEdge(), etc. 
- *
- * Note that although this is a template, many of the functions are defined in the .cpp file.
- * As a result you need to explicitly instantiate and export the instance of the template that 
- * you wish to use in the block at the top of DynamicMeshOverlay.cpp
- */
-template<typename RealType, int ElementSize>
-class TDynamicMeshOverlay
-{
-
-protected:
-	/** The parent mesh this overlay belongs to */
-	FDynamicMesh3* ParentMesh;
-
-	/** Reference counts of element indices. Iterate over this to find out which elements are valid. */
-	FRefCountVector ElementsRefCounts;
-	/** List of element values */
-	TDynamicVector<RealType> Elements;
-	/** List of parent vertex indices, one per element */
-	TDynamicVector<int> ParentVertices;
-
-	/** List of triangle element-index triplets [Elem0 Elem1 Elem2]*/
-	TDynamicVector<int> ElementTriangles;
-
-	friend class FDynamicMesh3;
-	friend class FDynamicMeshAttributeSet;
-
-public:
-	/** Create an empty overlay */
-	TDynamicMeshOverlay()
-	{
-		ParentMesh = nullptr;
-	}
-
-	/** Create an overlay for the given parent mesh */
-	TDynamicMeshOverlay(FDynamicMesh3* ParentMeshIn)
-	{
-		ParentMesh = ParentMeshIn;
-	}
-private:
-	/** @set the parent mesh for this overlay.  Only safe for use during FDynamicMesh move */
-	void Reparent(FDynamicMesh3* ParentMeshIn)
-	{
-		ParentMesh = ParentMeshIn;
-	}
-public:
-	/** @return the parent mesh for this overlay */
-	const FDynamicMesh3* GetParentMesh() const { return ParentMesh; }
-	/** @return the parent mesh for this overlay */
-	FDynamicMesh3* GetParentMesh() { return ParentMesh; }
-
-	/** Set this overlay to contain the same arrays as the copy overlay */
-	void Copy(const TDynamicMeshOverlay<RealType,ElementSize>& Copy)
-	{
-		ElementsRefCounts = FRefCountVector(Copy.ElementsRefCounts);
-		Elements = Copy.Elements;
-		ParentVertices = Copy.ParentVertices;
-		ElementTriangles = Copy.ElementTriangles;
-	}
-
-	/** Copy the Copy overlay to a compact rep, also updating parent references based on the CompactMaps */
-	void CompactCopy(const FCompactMaps& CompactMaps, const TDynamicMeshOverlay<RealType, ElementSize>& Copy)
-	{
-		ClearElements();
-
-		// map of element IDs
-		TArray<int> MapE; MapE.SetNumUninitialized(Copy.MaxElementID());
-
-		// copy elements across
-		RealType Data[ElementSize];
-		for (int EID = 0; EID < Copy.MaxElementID(); EID++)
-		{
-			if (Copy.IsElement(EID))
-			{
-				Copy.GetElement(EID, Data);
-				MapE[EID] = AppendElement(Data);
-			}
-			else
-			{
-				MapE[EID] = -1;
-			}
-		}
-
-		// copy triangles across
-		check(CompactMaps.bKeepTriangleMap && CompactMaps.MapT.Num() == Copy.GetParentMesh()->MaxTriangleID()); // must have valid triangle map
-		for (int FromTID : Copy.GetParentMesh()->TriangleIndicesItr())
-		{
-			if (!Copy.IsSetTriangle(FromTID))
-			{
-				continue;
-			}
-			int ToTID = CompactMaps.MapT[FromTID];
-			FIndex3i FromTriElements = Copy.GetTriangle(FromTID);
-			SetTriangle(ToTID, FIndex3i(MapE[FromTriElements.A], MapE[FromTriElements.B], MapE[FromTriElements.C]));
-		}
-	}
-
-	/** Compact overlay and update links to parent based on CompactMaps */
-	void CompactInPlace(const FCompactMaps& CompactMaps)
-	{
-		int iLastE = MaxElementID() - 1, iCurE = 0;
-		while (iLastE >= 0 && ElementsRefCounts.IsValidUnsafe(iLastE) == false)
-		{
-			iLastE--;
-		}
-		while (iCurE < iLastE && ElementsRefCounts.IsValidUnsafe(iCurE))
-		{
-			iCurE++;
-		}
-
-		// make a map to track element index changes, to use to update element triangles later
-		// TODO: it may be faster to not construct this and to do the remapping per element as we go (by iterating the one ring of each parent vertex for each element)
-		TArray<int> MapE; MapE.SetNumUninitialized(MaxElementID());
-		for (int ID = 0; ID < MapE.Num(); ID++)
-		{
-			// mapping is 1:1 by default; sparsely re-mapped below
-			MapE[ID] = ID;
-			// remap all parents
-			if (ParentVertices[ID] >= 0)
-			{
-				ParentVertices[ID] = CompactMaps.GetVertex(ParentVertices[ID]);
-			}
-		}
-
-		TDynamicVector<unsigned short>& ERef = ElementsRefCounts.GetRawRefCountsUnsafe();
-		RealType Data[ElementSize];
-		while (iCurE < iLastE)
-		{
-			// remap the element data
-			GetElement(iLastE, Data);
-			SetElement(iCurE, Data);
-			ParentVertices[iCurE] = ParentVertices[iLastE];
-			ERef[iCurE] = ERef[iLastE];
-			ERef[iLastE] = FRefCountVector::INVALID_REF_COUNT;
-			MapE[iLastE] = iCurE;
-
-			// move cur forward one, last back one, and  then search for next valid
-			iLastE--; iCurE++;
-			while (iLastE >= 0 && ElementsRefCounts.IsValidUnsafe(iLastE) == false)
-			{
-				iLastE--;
-			}
-			while (iCurE < iLastE && ElementsRefCounts.IsValidUnsafe(iCurE))
-			{
-				iCurE++;
-			}
-		}
-		ElementsRefCounts.Trim(ElementCount());
-		Elements.Resize(ElementCount() * ElementSize);
-		ParentVertices.Resize(ElementCount());
-
-		for (int TID = 0, OldMaxTID = ElementTriangles.Num() / 3; TID < OldMaxTID; TID++)
-		{
-			int OldStart = TID * 3;
-			if (ElementTriangles[OldStart] == -1)
-			{
-				continue; // triangle was not set; skip it
-			}
-			int NewStart = CompactMaps.GetTriangle(TID) * 3;
-			for (int SubIdx = 0; SubIdx < 3; SubIdx++)
-			{
-				ElementTriangles[NewStart + SubIdx] = MapE[ElementTriangles[OldStart + SubIdx]];
-			}
-		}
-
-		checkSlow(IsCompact());
-	}
-
-	/** Discard current set of elements, but keep triangles */
-	void ClearElements();
-
-	/** @return the number of in-use Elements in the overlay */
-	int ElementCount() const { return (int)ElementsRefCounts.GetCount(); }
-	/** @return the maximum element index in the overlay. This may be larger than the count if Elements have been deleted. */
-	int MaxElementID() const { return (int)ElementsRefCounts.GetMaxIndex(); }
-	/** @return true if this element index is in use */
-	inline bool IsElement(int vID) const { return ElementsRefCounts.IsValid(vID); }
-
-	/** @return true if the elements are compact */
-	bool IsCompact() const { return ElementsRefCounts.IsDense(); }
-
-
-	typedef typename FRefCountVector::IndexEnumerable element_iterator;
-
-	/** @return enumerator for valid element indices suitable for use with range-based for */
-	element_iterator ElementIndicesItr() const { return ElementsRefCounts.Indices(); }
-
-
-	/** Allocate a new element with the given constant value */
-	int AppendElement(RealType ConstantValue);
-	/** Allocate a new element with the given value */
-	int AppendElement(const RealType* Value);
-
-	void SetParentVertex(int ElementIndex, int ParentVertexIndex)
-	{
-		ParentVertices[ElementIndex] = ParentVertexIndex;
-	}
-
-	/** Initialize the triangle list to the given size, and set all triangles to InvalidID */
-	void InitializeTriangles(int MaxTriangleID);
-	/** Set the triangle to the given Element index tuple, and increment element reference counts */
-	EMeshResult SetTriangle(int TriangleID, const FIndex3i& TriElements);
-
-	/**
-	 * Set the triangle to have InvalidID element IDs, decrementing element reference counts if needed.
-	 * Deletes elements that are no longer used after the triangle is unset.
-	 */
-	void UnsetTriangle(int TriangleID);
-
-	/** @return true if this triangle was set */
-	bool IsSetTriangle(int TID) const
-	{
-		bool bIsSet = ElementTriangles[3 * TID] >= 0;
-		// we require that triangle elements either be all set or all unset
-		checkSlow(ElementTriangles[3 * TID + 1] >= 0 == bIsSet);
-		checkSlow(ElementTriangles[3 * TID + 2] >= 0 == bIsSet);
-		return bIsSet;
-	}
-
-
-	/**
-	 * Build overlay topology from a predicate function, e.g. to build topology for sharp normals
-	 * 
-	 * @param TrisCanShareVertexPredicate Indicator function returns true if the given vertex can be shared for the given pair of triangles
-	 *									  Note if a vertex can be shared between tris A and B, and B and C, it will be shared between all three
-	 * @param InitElementValue Initial element value, copied into all created elements
-	 */
-	void CreateFromPredicate(TFunctionRef<bool(int ParentVertexIdx, int TriIDA, int TriIDB)> TrisCanShareVertexPredicate, RealType InitElementValue);
-
-	/**
-	 * Refine an existing overlay topology.  For any element on a given triangle, if the predicate returns true, it gets topologically split out so it isn't shared by any other triangle.
-	 * Used for creating sharp vertices in the normals overlay.
-	 *
-	 * @param ShouldSplitOutVertex predicate returns true of the element should be split out and not shared w/ any other triangle
-	 * @param GetNewElementValue function to assign a new value to any element that is split out
-	 */
-	void SplitVerticesWithPredicate(TFunctionRef<bool(int ElementIdx, int TriID)> ShouldSplitOutVertex, TFunctionRef<void(int ElementIdx, int TriID, RealType* FillVect)> GetNewElementValue);
-
-
-	/**
-	 * Create a new copy of ElementID, and update connected triangles in the TrianglesToUpdate array to reference the copy of ElementID where they used to reference ElementID
-	 * (Note: This just calls "SplitElementWithNewParent" with the existing element's parent id.)
-	 *
-	 * @param ElementID the element to copy
-	 * @param TrianglesToUpdate the triangles that should now reference the new element
-	 * @return the ID of the newly created element
-	 */
-	int SplitElement(int ElementID, const TArrayView<const int>& TrianglesToUpdate);
-
-	/**
-	* Create a new copy of ElementID, and update connected triangles in the TrianglesToUpdate array to reference the copy of ElementID where they used to reference ElementID.  The new element will have the given parent vertex ID.
-	* Deletes any elements that are no longer used after the triangles are changed.
-	*
-	* @param ElementID the element to copy
-	* @param SplitParentVertexID the new parent vertex for copied elements
-	* @param TrianglesToUpdate the triangles that should now reference the new element.  Note: this is allowed to include triangles that do not have the element at all; sometimes you may want to do so to avoid creating a new array for each call.
-	* @return the ID of the newly created element
-	*/
-	int SplitElementWithNewParent(int ElementID, int SplitParentVertexID, const TArrayView<const int>& TrianglesToUpdate);
-
-
-	/**
-	* Refine an existing overlay topology by splitting any bow ties
-	*/
-	void SplitBowties();
-
-
-	//
-	// Support for inserting element at specific ID. This is a bit tricky
-	// because we likely will need to update the free list in the RefCountVector, which
-	// can be expensive. If you are going to do many inserts (eg inside a loop), wrap in
-	// BeginUnsafe / EndUnsafe calls, and pass bUnsafe = true to the InsertElement() calls,
-	// to the defer free list rebuild until you are done.
-	//
-
-	/** Call this before a set of unsafe InsertVertex() calls */
-	void BeginUnsafeElementsInsert()
-	{
-		// do nothing...
-	}
-
-	/** Call after a set of unsafe InsertVertex() calls to rebuild free list */
-	void EndUnsafeElementsInsert()
-	{
-		ElementsRefCounts.RebuildFreeList();
-	}
-
-	/**
-	 * Insert element at given index, assuming it is unused.
-	 * If bUnsafe, we use fast id allocation that does not update free list.
-	 * You should only be using this between BeginUnsafeElementsInsert() / EndUnsafeElementsInsert() calls
-	 */
-	EMeshResult InsertElement(int ElementID, const RealType* Value, bool bUnsafe = false);
-
-
-	//
-	// Accessors/Queries
-	//  
-
-
-	/** Get the element at a given index */
-	inline void GetElement(int ElementID, RealType* Data) const
-	{
-		int k = ElementID * ElementSize;
-		for (int i = 0; i < ElementSize; ++i) 
-		{
-			Data[i] = Elements[k + i];
-		}
-	}
-
-	/** Get the element at a given index */
-	template<typename AsType>
-	void GetElement(int ElementID, AsType& Data) const
-	{
-		int k = ElementID * ElementSize;
-		for (int i = 0; i < ElementSize; ++i) 
-		{
-			Data[i] = Elements[k + i];
-		}
-	}
-
-
-	/** Get the parent vertex id for the element at a given index */
-	inline int GetParentVertex(int ElementID) const
-	{
-		return ParentVertices[ElementID];
-	}
-
-
-	/** Get the element index tuple for a triangle */
-	inline FIndex3i GetTriangle(int TriangleID) const 
-	{
-		int i = 3 * TriangleID;
-		return FIndex3i(ElementTriangles[i], ElementTriangles[i + 1], ElementTriangles[i + 2]);
-	}
-
-
-	/** Set the element at a given index */
-	inline void SetElement(int ElementID, const RealType* Data)
-	{
-		int k = ElementID * ElementSize;
-		for (int i = 0; i < ElementSize; ++i)
-		{
-			Elements[k + i] = Data[i];
-		}
-	}
-
-	/** Set the element at a given index */
-	template<typename AsType>
-	void SetElement(int ElementID, const AsType& Data)
-	{
-		int k = ElementID * ElementSize;
-		for (int i = 0; i < ElementSize; ++i)
-		{
-			Elements[k + i] = Data[i];
-		}
-	}
-
-	/** @return true if triangle contains element */
-	inline bool TriangleHasElement(int TriangleID, int ElementID) const
-	{
-		int i = 3 * TriangleID;
-		return (ElementTriangles[i] == ElementID || ElementTriangles[i+1] == ElementID || ElementTriangles[i+2] == ElementID);
-	}
-
-
-	/** Returns true if the parent-mesh edge is a "Seam" in this overlay.  
-	*   If present, bIsNonIntersectingOut will be true only if this is a seam edge 
-	*   that does not intersect with another seam or the end of the seam. 
-	*/
-	bool IsSeamEdge(int EdgeID, bool* bIsNonIntersectingOut = nullptr) const;
-	/** Returns true if the parent-mesh edge is a "Seam End" in this overlay, meaning the adjacent element triangles share one element, not two */
-	bool IsSeamEndEdge(int EdgeID) const;
-	/** Returns true if the parent-mesh vertex is connected to any seam edges */
-	bool IsSeamVertex(int VertexID, bool bBoundaryIsSeam = true) const;
-
-	/**
-	 * Determines whether the base-mesh vertex has "bowtie" topology in the Overlay.
-	 * Bowtie topology means that one or more elements at the vertex are shared across disconnected UV-components.
-	 * @return true if the base-mesh vertex has "bowtie" topology in the overlay
-	 */
-	bool IsBowtieInOverlay(int32 VertexID) const;
-
-	/** @return true if the two triangles are connected, ie shared edge exists and is not a seam edge */
-	bool AreTrianglesConnected(int TriangleID0, int TriangleID1) const;
-
-	/** find the elements associated with a given parent-mesh vertex */
-	void GetVertexElements(int VertexID, TArray<int>& OutElements) const;
-	/** Count the number of unique elements for a given parent-mesh vertex */
-	int CountVertexElements(int VertexID, bool bBruteForce = false) const;
-
-	/** find the triangles connected to an element */
-	void GetElementTriangles(int ElementID, TArray<int>& OutTriangles) const;
-
-	/** @return true if overlay has any interior seam edges. This requires an O(N) search unless it early-outs. */
-	bool HasInteriorSeamEdges() const;
-
-	/**
-	 * Compute interpolated parameter value inside triangle using barycentric coordinates
-	 * @param TriangleID index of triangle
-	 * @param BaryCoords 3 barycentric coordinates inside triangle
-	 * @param DataOut resulting interpolated overlay parameter value (of size ElementSize)
-	 */
-	template<typename AsType>
-	void GetTriBaryInterpolate(int32 TriangleID, const AsType* BaryCoords, AsType* DataOut) const
-	{
-		int32 TriIndex = 3 * TriangleID;
-		int32 ElemIndex0 = ElementTriangles[TriIndex] * ElementSize;
-		int32 ElemIndex1 = ElementTriangles[TriIndex+1] * ElementSize;
-		int32 ElemIndex2 = ElementTriangles[TriIndex+2] * ElementSize;
-		const AsType Bary0 = (AsType)BaryCoords[0], Bary1 = (AsType)BaryCoords[1], Bary2 = (AsType)BaryCoords[2];
-		for (int32 i = 0; i < ElementSize; ++i)
-		{
-			DataOut[i] = Bary0*(AsType)Elements[ElemIndex0+i] + Bary1*(AsType)Elements[ElemIndex1+i] + Bary2*(AsType)Elements[ElemIndex2+i];
-		}
-	}
-
-	/**
-	 * Checks that the overlay mesh is well-formed, ie all internal data structures are consistent
-	 */
-	bool CheckValidity(bool bAllowNonManifoldVertices = true, EValidityCheckFailMode FailMode = EValidityCheckFailMode::Check) const;
-
-	/**
-	 * Returns true if this overlay is the same as Other.
-	 */
-	bool IsSameAs(const TDynamicMeshOverlay<RealType, ElementSize>& Other) const;
-
-	/**
-	 * Serialization operator for FDynamicMeshOverlay.
-	 *
-	 * @param Ar Archive to serialize with.
-	 * @param Overlay Mesh overlay to serialize.
-	 * @returns Passing down serializing archive.
-	 */
-	friend FArchive& operator<<(FArchive& Ar, TDynamicMeshOverlay<RealType,ElementSize>& Overlay)
-	{
-		Overlay.Serialize(Ar);
-		return Ar;
-	}
-
-	/** Serialize the overlay to an archive. */
-	void Serialize(FArchive& Ar);
-
-public:
-	/** Set a triangle's element indices to InvalidID */
-	void InitializeNewTriangle(int TriangleID);
-	/** Remove a triangle from the overlay */
-	void OnRemoveTriangle(int TriangleID);
-	/** Reverse the orientation of a triangle's elements */
-	void OnReverseTriOrientation(int TriangleID);
-	/** Update the overlay to reflect an edge split in the parent mesh */
-	void OnSplitEdge(const DynamicMeshInfo::FEdgeSplitInfo& SplitInfo);
-	/** Update the overlay to reflect an edge flip in the parent mesh */
-	void OnFlipEdge(const DynamicMeshInfo::FEdgeFlipInfo& FlipInfo);
-	/** Update the overlay to reflect an edge collapse in the parent mesh */
-	void OnCollapseEdge(const DynamicMeshInfo::FEdgeCollapseInfo& CollapseInfo);
-	/** Update the overlay to reflect a face poke in the parent mesh */
-	void OnPokeTriangle(const DynamicMeshInfo::FPokeTriangleInfo& PokeInfo);
-	/** Update the overlay to reflect an edge merge in the parent mesh */
-	void OnMergeEdges(const DynamicMeshInfo::FMergeEdgesInfo& MergeInfo);
-	/** Update the overlay to reflect a vertex split in the parent mesh */
-	void OnSplitVertex(const DynamicMeshInfo::FVertexSplitInfo& SplitInfo, const TArrayView<const int>& TrianglesToUpdate);
-
-protected:
-	/** Set the value at an Element to be a linear interpolation of two other Elements */
-	void SetElementFromLerp(int SetElement, int ElementA, int ElementB, double Alpha);
-	/** Set the value at an Element to be a barycentric interpolation of three other Elements */
-	void SetElementFromBary(int SetElement, int ElementA, int ElementB, int ElementC, const FVector3<double>& BaryCoords);
-
-	/** updates the triangles array and optionally the element reference counts */
-	void InternalSetTriangle(int TriangleID, const FIndex3i& TriElements, bool bIncrementRefCounts);
-
-};
-
-
-
-
-
-
-/**
- * TDynamicMeshVectorOverlay is an convenient extension of TDynamicMeshOverlay that adds
- * a specific N-element Vector type to the template, and adds accessor functions
- * that convert between that N-element vector type and the N-element arrays used by TDynamicMeshOverlay.
- */
-template<typename RealType, int ElementSize, typename VectorType>
-class TDynamicMeshVectorOverlay : public TDynamicMeshOverlay<RealType, ElementSize>
-{
-public:
-	using BaseType = TDynamicMeshOverlay<RealType, ElementSize>;
-
-	TDynamicMeshVectorOverlay()
-		: TDynamicMeshOverlay<RealType, ElementSize>()
-	{
-	}
-
-	TDynamicMeshVectorOverlay(FDynamicMesh3* parentMesh) 
-		: TDynamicMeshOverlay<RealType, ElementSize>(parentMesh)
-	{
-	}
-
-	/**
-	 * Append a new Element to the overlay
-	 */
-	inline int AppendElement(const VectorType& Value)
-	{
-		// Cannot use cast operator here because Core Vector types do not define it.
-		// However assuming that vector has .X member is also not good...
-		//return BaseType::AppendElement((const RealType*)Value);
-		return BaseType::AppendElement(&Value.X);
-	}
-
-	/**
-	 * Append a new Element to the overlay
-	 */
-	inline int AppendElement(const RealType* Value)
-	{
-		return BaseType::AppendElement(Value);
-	}
-
-	/**
-	 * Get Element at a specific ID
-	 */
-	inline VectorType GetElement(int ElementID) const
-	{
-		VectorType V;
-		BaseType::GetElement(ElementID, V);
-		return V;
-	}
-
-	/**
-	 * Get Element at a specific ID
-	 */
-	inline void GetElement(int ElementID, VectorType& V) const
-	{
-		BaseType::GetElement(ElementID, V);
-	}
-
-	/**
-	 * Get the three Elements associated with a triangle
-	 */
-	inline void GetTriElements(int TriangleID, VectorType& A, VectorType& B, VectorType& C) const
-	{
-		int i = 3 * TriangleID;
-		GetElement(BaseType::ElementTriangles[i], A);
-		GetElement(BaseType::ElementTriangles[i+1], B);
-		GetElement(BaseType::ElementTriangles[i+2], C);
-	}
-
-
-	/**
-	 * Set Element at a specific ID
-	 */
-	inline void SetElement(int ElementID, const VectorType& Value)
-	{
-		BaseType::SetElement(ElementID, Value);
-	}
-};
-
-
-} // end namespace UE::Geometry
-} // end namespace UE
\ No newline at end of file
diff --git a/Engine/Plugins/Experimental/GeometryProcessing/Source/DynamicMesh/Public/DynamicMeshTriangleAttribute.h b/Engine/Plugins/Experimental/GeometryProcessing/Source/DynamicMesh/Public/DynamicMeshTriangleAttribute.h
--- a/Engine/Plugins/Experimental/GeometryProcessing/Source/DynamicMesh/Public/DynamicMeshTriangleAttribute.h
+++ b/Engine/Plugins/Experimental/GeometryProcessing/Source/DynamicMesh/Public/DynamicMeshTriangleAttribute.h
@@ -2,498 +2,5 @@
 
 #pragma once
 
-#include "DynamicMesh3.h"
-#include "DynamicAttribute.h"
+#include "DynamicMesh/DynamicMeshTriangleAttribute.h"
 
-namespace UE
-{
-namespace Geometry
-{
-
-template<typename AttribValueType, int AttribDimension>
-class TDynamicMeshTriangleAttribute;
-
-
-template<typename AttribValueType, int AttribDimension>
-class FDynamicMeshTriangleAttributeChange : public FDynamicMeshAttributeChangeBase
-{
-private:
-	struct FChangeTriangleAttribute
-	{
-		AttribValueType Data[AttribDimension];
-		int TriangleID;
-	};
-
-	TArray<FChangeTriangleAttribute> OldTriangleAttributes, NewTriangleAttributes;
-
-public:
-	FDynamicMeshTriangleAttributeChange()
-	{}
-
-	virtual ~FDynamicMeshTriangleAttributeChange()
-	{}
-
-	inline virtual void SaveInitialTriangle(const FDynamicMeshAttributeBase* Attribute, int TriangleID) override;
-
-	inline virtual void StoreAllFinalTriangles(const FDynamicMeshAttributeBase* Attribute, const TArray<int>& TriangleIDs) override;
-
-	inline virtual bool Apply(FDynamicMeshAttributeBase* Attribute, bool bRevert) const override;
-};
-
-
-
-/**
- * TDynamicMeshTriangleAttribute is an add-on to a FDynamicMesh3 that allows for 
- * per-triangle storage of an attribute value.
- *
- * The FDynamicMesh3 mesh topology operations (eg split/flip/collapse edge, poke face, etc)
- * can be mirrored to the overlay via OnSplitEdge(), etc.
- */
-template<typename AttribValueType, int AttribDimension>
-class TDynamicMeshTriangleAttribute : public FDynamicMeshAttributeBase
-{
-
-protected:
-	/** The parent mesh this overlay belongs to */
-	FDynamicMesh3* ParentMesh;
-
-	/** List of per-triangle attribute values */
-	TDynamicVector<AttribValueType> AttribValues;
-
-	friend class FDynamicMesh3;
-	friend class FDynamicMeshAttributeSet;
-
-public:
-	/** Create an empty overlay */
-	TDynamicMeshTriangleAttribute()
-	{
-		ParentMesh = nullptr;
-	}
-
-	/** Create an overlay for the given parent mesh */
-	TDynamicMeshTriangleAttribute(FDynamicMesh3* ParentMeshIn, bool bAutoInit = true)
-	{
-		ParentMesh = ParentMeshIn;
-		if (bAutoInit)
-		{
-			Initialize();
-		}
-	}
-
-private:
-	/** @set the parent mesh for this overlay.  Only safe for use during FDynamicMesh move */
-	void Reparent(FDynamicMesh3* ParentMeshIn)
-	{
-		ParentMesh = ParentMeshIn;
-	}
-
-public:
-	/** @return the parent mesh for this overlay */
-	const FDynamicMesh3* GetParentMesh() const { return ParentMesh; }
-	/** @return the parent mesh for this overlay */
-	FDynamicMesh3* GetParentMesh() { return ParentMesh; }
-
-	virtual FDynamicMeshAttributeBase* MakeNew(FDynamicMesh3* ParentMeshIn) const override
-	{
-		TDynamicMeshTriangleAttribute<AttribValueType, AttribDimension>* Matching = new TDynamicMeshTriangleAttribute<AttribValueType, AttribDimension>(ParentMeshIn);
-		Matching->Initialize();
-		return Matching;
-	}
-
-	virtual FDynamicMeshAttributeBase* MakeCopy(FDynamicMesh3* ParentMeshIn) const override
-	{
-		TDynamicMeshTriangleAttribute<AttribValueType, AttribDimension>* ToFill = new TDynamicMeshTriangleAttribute<AttribValueType, AttribDimension>(ParentMeshIn);
-		ToFill->Copy(*this);
-		return ToFill;
-	}
-
-	/** Set this overlay to contain the same arrays as the copy overlay */
-	void Copy(const TDynamicMeshTriangleAttribute<AttribValueType, AttribDimension>& Copy)
-	{
-		CopyParentClassData(Copy);
-		AttribValues = Copy.AttribValues;
-	}
-
-	virtual FDynamicMeshAttributeBase* MakeCompactCopy(const FCompactMaps& CompactMaps, FDynamicMesh3* ParentMeshIn) const override
-	{
-		TDynamicMeshTriangleAttribute<AttribValueType, AttribDimension>* ToFill = new TDynamicMeshTriangleAttribute<AttribValueType, AttribDimension>(ParentMeshIn);
-		ToFill->Initialize();
-		ToFill->CompactCopy(CompactMaps, *this);
-		return ToFill;
-	}
-
-	void CompactInPlace(const FCompactMaps& CompactMaps)
-	{
-		for (int TID = 0; TID < CompactMaps.MapT.Num(); TID++)
-		{
-			int ToTID = CompactMaps.MapT[TID];
-			if (ToTID < 0)
-			{
-				continue;
-			}
-			if (ensure(ToTID <= TID))
-			{
-				CopyValue(TID, ToTID);
-			}
-		}
-		AttribValues.Resize(ParentMesh->MaxTriangleID() * AttribDimension);
-	}
-
-	void CompactCopy(const FCompactMaps& CompactMaps, const TDynamicMeshTriangleAttribute<AttribValueType, AttribDimension>& ToCopy)
-	{
-		CopyParentClassData(ToCopy);
-		check(CompactMaps.MapT.Num() <= int(ToCopy.AttribValues.Num() / AttribDimension));
-		AttribValueType Data[AttribDimension];
-		for (int TID = 0; TID < CompactMaps.MapT.Num(); TID++)
-		{
-			int ToTID = CompactMaps.MapT[TID];
-			if (ToTID == -1)
-			{
-				continue;
-			}
-			ToCopy.GetValue(TID, Data);
-			SetValue(ToTID, Data);
-		}
-	}
-
-
-	/** Initialize the attribute values to the given max triangle ID */
-	void Initialize(AttribValueType InitialValue = (AttribValueType)0)
-	{
-		check(ParentMesh != nullptr);
-		AttribValues.Resize(ParentMesh->MaxTriangleID() * AttribDimension);
-		AttribValues.Fill(InitialValue);
-	}
-
-	void SetNewValue(int NewTriangleID, const AttribValueType* Data)
-	{
-		int k = NewTriangleID * AttribDimension;
-		for (int i = 0; i < AttribDimension; ++i)
-		{
-			AttribValues.InsertAt(Data[i], k + i);
-		}
-	}
-
-
-
-	//
-	// Accessors/Queries
-	//  
-
-	virtual bool CopyThroughMapping(const TDynamicAttributeBase<FDynamicMesh3>* Source, const FMeshIndexMappings& Mapping) override
-	{
-		AttribValueType BufferData[AttribDimension];
-		int BufferSize = sizeof(BufferData);
-		for (const TPair<int32, int32>& MapTID : Mapping.GetTriangleMap().GetForwardMap())
-		{
-			if (!ensure(Source->CopyOut(MapTID.Key, BufferData, BufferSize)))
-			{
-				return false;
-			}
-			SetValue(MapTID.Value, BufferData);
-		}
-		return true;
-	}
-	virtual bool CopyOut(int RawID, void* Buffer, int BufferSize) const override
-	{
-		if (sizeof(AttribValueType) * AttribDimension != BufferSize)
-		{
-			return false;
-		}
-		AttribValueType* BufferData = static_cast<AttribValueType*>(Buffer);
-		int k = RawID * AttribDimension;
-		for (int i = 0; i < AttribDimension; ++i)
-		{
-			BufferData[i] = AttribValues[k + i];
-		}
-		return true;
-	}
-	virtual bool CopyIn(int RawID, void* Buffer, int BufferSize) override
-	{
-		if (sizeof(AttribValueType) * AttribDimension != BufferSize)
-		{
-			return false;
-		}
-		AttribValueType* BufferData = static_cast<AttribValueType*>(Buffer);
-		int k = RawID * AttribDimension;
-		for (int i = 0; i < AttribDimension; ++i)
-		{
-			AttribValues[k + i] = BufferData[i];
-		}
-		return true;
-	}
-
-	/** Get the element at a given index */
-	inline void GetValue(int TriangleID, AttribValueType* Data) const
-	{
-		int k = TriangleID * AttribDimension;
-		for (int i = 0; i < AttribDimension; ++i)
-		{
-			Data[i] = AttribValues[k + i];
-		}
-	}
-
-	/** Get the element at a given index */
-	template<typename AsType>
-	void GetValue(int TriangleID, AsType& Data) const
-	{
-		int k = TriangleID * AttribDimension;
-		for (int i = 0; i < AttribDimension; ++i)
-		{
-			Data[i] = AttribValues[k + i];
-		}
-	}
-
-
-	/** Set the element at a given index */
-	inline void SetValue(int TriangleID, const AttribValueType* Data)
-	{
-		int k = TriangleID * AttribDimension;
-		for (int i = 0; i < AttribDimension; ++i)
-		{
-			AttribValues[k + i] = Data[i];
-		}
-	}
-
-	/** Set the element at a given index */
-	template<typename AsType>
-	void SetValue(int TriangleID, const AsType& Data)
-	{
-		int k = TriangleID * AttribDimension;
-		for (int i = 0; i < AttribDimension; ++i)
-		{
-			AttribValues[k + i] = Data[i];
-		}
-	}
-
-	/**
-	 * Copy the attribute value at FromTriangleID to ToTriangleID
-	 */
-	inline void CopyValue(int FromTriangleID, int ToTriangleID)
-	{
-		int kA = FromTriangleID * AttribDimension;
-		int kB = ToTriangleID * AttribDimension;
-		for (int i = 0; i < AttribDimension; ++i)
-		{
-			AttribValues.InsertAt(AttribValues[kA+i], kB+i);
-		}
-	}
-
-
-	/** Returns true if the parent-mesh edge is a "Seam" in this overlay */
-	bool IsBorderEdge(int EdgeID, bool bMeshBoundaryIsBorder = true) const
-	{
-		FIndex2i EdgeTris = ParentMesh->GetEdgeT(EdgeID);
-		if (EdgeTris.B == FDynamicMesh3::InvalidID)
-		{
-			return bMeshBoundaryIsBorder;
-		}
-		int kA = EdgeTris.A * AttribDimension;
-		int kB = EdgeTris.B * AttribDimension;
-		for (int i = 0; i < AttribDimension; ++i)
-		{
-			if (AttribValues[kA+i] != AttribValues[kB+i])
-			{
-				return true;
-			}
-		}
-	}
-
-
-	virtual TUniquePtr<FDynamicMeshAttributeChangeBase> NewBlankChange() override
-	{
-		return MakeUnique<FDynamicMeshTriangleAttributeChange<AttribValueType, AttribDimension>>();
-	}
-
-
-public:
-
-	/** Update the overlay to reflect an edge split in the parent mesh */
-	void OnSplitEdge(const DynamicMeshInfo::FEdgeSplitInfo& SplitInfo) override
-	{
-		CopyValue(SplitInfo.OriginalTriangles.A, SplitInfo.NewTriangles.A);
-		if (SplitInfo.OriginalTriangles.B != FDynamicMesh3::InvalidID)
-		{
-			CopyValue(SplitInfo.OriginalTriangles.B, SplitInfo.NewTriangles.B);
-		}
-	}
-
-	/** Update the overlay to reflect an edge flip in the parent mesh */
-	void OnFlipEdge(const DynamicMeshInfo::FEdgeFlipInfo& FlipInfo) override
-	{
-		// yikes! triangles did not actually change so we will leave attrib unmodified
-	}
-
-	/** Update the overlay to reflect an edge collapse in the parent mesh */
-	void OnCollapseEdge(const DynamicMeshInfo::FEdgeCollapseInfo& CollapseInfo) override
-	{
-		// nothing to do here, triangles were only deleted
-	}
-
-	/** Update the overlay to reflect a face poke in the parent mesh */
-	void OnPokeTriangle(const DynamicMeshInfo::FPokeTriangleInfo& PokeInfo) override
-	{
-		CopyValue(PokeInfo.OriginalTriangle, PokeInfo.NewTriangles.A);
-		CopyValue(PokeInfo.OriginalTriangle, PokeInfo.NewTriangles.B);
-	}
-
-	/** Update the overlay to reflect an edge merge in the parent mesh */
-	void OnMergeEdges(const DynamicMeshInfo::FMergeEdgesInfo& MergeInfo) override
-	{
-		// nothing to do here because triangles did not change
-	}
-
-	/** Update the overlay to reflect a vertex split in the parent */
-	void OnSplitVertex(const DynamicMeshInfo::FVertexSplitInfo& SplitInfo, const TArrayView<const int>& TrianglesToUpdate) override
-	{
-		// nothing to do here because triangles did not change
-	}
-
-	virtual AttribValueType GetDefaultAttributeValue()
-	{
-		return (AttribValueType)0;
-	}
-
-	inline void ResizeAttribStoreIfNeeded(int TriangleID)
-	{
-		if (!ensure(TriangleID >= 0))
-		{
-			return;
-		}
-		size_t NeededSize = (((size_t)TriangleID+1) * AttribDimension);
-		if (NeededSize > AttribValues.Num())
-		{
-			AttribValues.Resize(NeededSize, GetDefaultAttributeValue());
-		}
-	}
-
-	virtual void OnNewTriangle(int TriangleID, bool bInserted) override
-	{
-		ResizeAttribStoreIfNeeded(TriangleID);
-	}
-
-public:
-	/**
-	 * Returns true if this AttributeSet is the same as Other.
-	 */
-	bool IsSameAs(const TDynamicMeshTriangleAttribute<AttribValueType, AttribDimension>& Other) const
-	{
-		if (AttribValues.Num() != Other.AttribValues.Num())
-		{
-			return false;
-		}
-
-		for (int Idx = 0, NumValues = AttribValues.Num(); Idx < NumValues; Idx++)
-		{
-			if (AttribValues[Idx] != Other.AttribValues[Idx])
-			{
-				return false;
-			}
-		}
-
-		return true;
-	}
-
-	/**
-	 * Serialization operator for TDynamicMeshTriangleAttribute.
-	 *
-	 * @param Ar Archive to serialize with.
-	 * @param Attr Mesh triangle attribute to serialize.
-	 * @returns Passing down serializing archive.
-	 */
-	friend FArchive& operator<<(FArchive& Ar, TDynamicMeshTriangleAttribute<AttribValueType, AttribDimension>& Attr)
-	{
-		Attr.Serialize(Ar);
-		return Ar;
-	}
-
-	/** Serialize the attribute to an archive. */
-	void Serialize(FArchive& Ar)
-	{
-		Ar << AttribValues;
-	}
-};
-
-
-
-template<typename AttribValueType, int AttribDimension>
-void FDynamicMeshTriangleAttributeChange<AttribValueType, AttribDimension>::SaveInitialTriangle(const FDynamicMeshAttributeBase* Attribute, int TriangleID)
-{
-	FChangeTriangleAttribute& Change = OldTriangleAttributes.Emplace_GetRef();
-	Change.TriangleID = TriangleID;
-	const TDynamicMeshTriangleAttribute<AttribValueType, AttribDimension>* AttribCast = static_cast<const TDynamicMeshTriangleAttribute<AttribValueType, AttribDimension>*>(Attribute);
-	AttribCast->GetValue(TriangleID, Change.Data);
-}
-
-template<typename AttribValueType, int AttribDimension>
-void FDynamicMeshTriangleAttributeChange<AttribValueType, AttribDimension>::StoreAllFinalTriangles(const FDynamicMeshAttributeBase* Attribute, const TArray<int>& TriangleIDs)
-{
-	const TDynamicMeshTriangleAttribute<AttribValueType, AttribDimension>* AttribCast = static_cast<const TDynamicMeshTriangleAttribute<AttribValueType, AttribDimension>*>(Attribute);
-	NewTriangleAttributes.Reserve(NewTriangleAttributes.Num() + TriangleIDs.Num());
-	for (int TriangleID : TriangleIDs)
-	{
-		FChangeTriangleAttribute& Change = NewTriangleAttributes.Emplace_GetRef();
-		Change.TriangleID = TriangleID;
-		AttribCast->GetValue(TriangleID, Change.Data);
-	}
-}
-
-template<typename AttribValueType, int AttribDimension>
-bool FDynamicMeshTriangleAttributeChange<AttribValueType, AttribDimension>::Apply(FDynamicMeshAttributeBase* Attribute, bool bRevert) const
-{
-	TDynamicMeshTriangleAttribute<AttribValueType, AttribDimension>* AttribCast = static_cast<TDynamicMeshTriangleAttribute<AttribValueType, AttribDimension>*>(Attribute);
-	const TArray<FChangeTriangleAttribute> *Changes = bRevert ? &OldTriangleAttributes : &NewTriangleAttributes;
-	for (const FChangeTriangleAttribute& Change : *Changes)
-	{
-		bool bIsTriangle = AttribCast->GetParentMesh()->IsTriangle(Change.TriangleID);
-		checkSlow(bIsTriangle);
-		if (bIsTriangle)
-		{
-			AttribCast->SetValue(Change.TriangleID, Change.Data);
-		}
-	}
-	return true;
-}
-
-
-
-/**
- * TDynamicMeshScalarTriangleAttribute is an extension of TDynamicMeshTriangleAttribute for scalar-valued attributes.
- * Adds some convenience functions to simplify get/set code.
- */
-template<typename RealType>
-class TDynamicMeshScalarTriangleAttribute : public TDynamicMeshTriangleAttribute<RealType, 1>
-{
-public:
-	using BaseType = TDynamicMeshTriangleAttribute<RealType, 1>;
-	using BaseType::SetNewValue;
-	using BaseType::GetValue;
-	using BaseType::SetValue;
-
-	TDynamicMeshScalarTriangleAttribute() : BaseType()
-	{
-	}
-
-	TDynamicMeshScalarTriangleAttribute(FDynamicMesh3* ParentMeshIn) : BaseType(ParentMeshIn)
-	{
-	}
-
-	inline void SetNewValue(int NewTriangleID, RealType Value)
-	{
-		this->AttribValues.InsertAt(Value, NewTriangleID);
-	}
-
-	inline RealType GetValue(int TriangleID) const
-	{
-		return this->AttribValues[TriangleID];
-	}
-
-	inline void SetValue(int TriangleID, RealType Value)
-	{
-		this->AttribValues[TriangleID] = Value;
-	}
-};
-
-
-} // end namespace UE::Geometry
-} // end namespace UE
\ No newline at end of file
diff --git a/Engine/Plugins/Experimental/GeometryProcessing/Source/DynamicMesh/Public/DynamicPointSet3.h b/Engine/Plugins/Experimental/GeometryProcessing/Source/DynamicMesh/Public/DynamicPointSet3.h
--- a/Engine/Plugins/Experimental/GeometryProcessing/Source/DynamicMesh/Public/DynamicPointSet3.h
+++ b/Engine/Plugins/Experimental/GeometryProcessing/Source/DynamicMesh/Public/DynamicPointSet3.h
@@ -2,332 +2,4 @@
 
 #pragma once
 
-#include "BoxTypes.h"
-#include "GeometryTypes.h"
-#include "MathUtil.h"
-#include "Util/DynamicVector.h"
-#include "Util/IndexUtil.h"
-#include "Util/IteratorUtil.h"
-#include "Util/RefCountVector.h"
-#include "VectorTypes.h"
-#include "VectorUtil.h"
-#include "DynamicAttribute.h"
-
-class FDynamicAttributeSetBase;
-
-namespace UE
-{
-namespace Geometry
-{
-
-/**
- * TDynamicPointSet3 implements a dynamic 3D point set, templated on real-value type (float or double).
- * The points are indexed and the class allows for gaps in the index space.
- * 
- * The points are referred to as "Vertices" for consistency with higher-level Mesh and Graph classes.
- * This class can be used interchangeably with the Vertices of those classes in some cases
- * (for example in spatial data structures like TDynamicVerticesOctree3)
- *
- * The point data is stored in via TDynamicVector<RealType>, which internally uses chunked storage
- * and so is relatively efficient to grow/shrink/etc
- *
- * Reference counts for the points are stored via a FRefCountVector instance
- *
- * A per-point attribute set can optionally be attached. Currently the client is
- * responsible for managing the memory for this attribute set. The attribute set
- * will be automatically updated when points are added or removed.
- */
-template<typename RealType>
-class TDynamicPointSet3
-{
-protected:
-	/** Reference counts of vertex indices. Iterate over this to find out which vertex indices are valid. */
-	FRefCountVector VertexRefCounts;
-	/** List of vertex positions */
-	TDynamicVector<RealType> Vertices;
-
-	/** The timestamp is incremented any time a function that modifies the point set is called */
-	int Timestamp = 0;
-
-
-	/** Base attribute set */
-	TDynamicAttributeSetBase<TDynamicPointSet3<RealType>>* BaseAttributeSet = nullptr;
-
-
-	//
-	// Attribute Set Support. At the level of this class, we do not have a specific attribute set.
-	// Clients can provide a pointer to their attribute set so that functions like AppendVertex()
-	// or DeleteVertex() will work.
-	//
-public:
-
-	void SetExternallyManagedAttributes(TDynamicAttributeSetBase<TDynamicPointSet3<RealType>>* AttributeSet)
-	{
-		BaseAttributeSet = AttributeSet;
-	}
-
-	TDynamicAttributeSetBase<TDynamicPointSet3<RealType>>* GetBaseAttributeSet() const
-	{
-		return BaseAttributeSet;
-	}
-
-
-public:
-
-
-
-	//
-	// Indices
-	//
-public:
-
-	/** @return number of vertices in the point set */
-	int VertexCount() const
-	{
-		return (int)VertexRefCounts.GetCount();
-	}
-
-	/** @return upper bound on vertex IDs used in the point set, ie all vertex IDs in use are < MaxVertexID */
-	int MaxVertexID() const
-	{
-		return (int)VertexRefCounts.GetMaxIndex();
-	}
-
-
-	/** @return true if VertexID is a valid vertex in this point set */
-	inline bool IsVertex(int VertexID) const
-	{
-		return VertexRefCounts.IsValid(VertexID);
-	}
-
-
-
-	//
-	// Construction
-	//
-public:
-
-	/** Discard all data */
-	void Clear()
-	{
-		Vertices.Clear();
-		VertexRefCounts = FRefCountVector();
-		BaseAttributeSet = nullptr;
-		Timestamp = 0;
-	}
-
-
-	/** Append vertex at position, returns vid */
-	int AppendVertex(const FVector3<RealType>& Position)
-	{
-		int vid = VertexRefCounts.Allocate();
-		int i = 3 * vid;
-		Vertices.InsertAt(Position[2], i + 2);
-		Vertices.InsertAt(Position[1], i + 1);
-		Vertices.InsertAt(Position[0], i);
-
-		if (GetBaseAttributeSet() != nullptr)
-		{
-			GetBaseAttributeSet()->OnNewVertex(vid, false);
-		}
-
-		UpdateTimeStamp();
-		return vid;
-	}
-
-	/** Call this before a set of unsafe InsertVertex() calls */
-	void BeginUnsafeVerticesInsert()
-	{
-		// do nothing...
-	}
-
-	/** Call after a set of unsafe InsertVertex() calls to rebuild free list */
-	void EndUnsafeVerticesInsert()
-	{
-		VertexRefCounts.RebuildFreeList();
-	}
-
-	/**
-	 * Insert vertex at given index, assuming it is unused.
-	 * If bUnsafe, we use fast id allocation that does not update free list.
-	 * You should only be using this between BeginUnsafeVerticesInsert() / EndUnsafeVerticesInsert() calls
-	 */
-	EMeshResult InsertVertex(int VertexID, const FVector3<RealType>& Position, bool bUnsafe = false)
-	{
-		if (VertexRefCounts.IsValid(VertexID))
-		{
-			return EMeshResult::Failed_VertexAlreadyExists;
-		}
-
-		bool bOK = (bUnsafe) ? VertexRefCounts.AllocateAtUnsafe(VertexID) : VertexRefCounts.AllocateAt(VertexID);
-		if (bOK == false)
-		{
-			return EMeshResult::Failed_CannotAllocateVertex;
-		}
-
-		int i = 3 * VertexID;
-		Vertices.InsertAt(Position[2], i + 2);
-		Vertices.InsertAt(Position[1], i + 1);
-		Vertices.InsertAt(Position[0], i);
-
-		if (GetBaseAttributeSet() != nullptr)
-		{
-			GetBaseAttributeSet()->OnNewVertex(VertexID, true);
-		}
-
-		UpdateTimeStamp();
-		return EMeshResult::Ok;
-	}
-
-
-
-	//
-	// get/set
-	//
-public:
-
-	/** @return the vertex position */
-	inline FVector3<RealType> GetVertex(int VertexID) const
-	{
-		check(IsVertex(VertexID));
-		int i = 3 * VertexID;
-		return FVector3<RealType>(Vertices[i], Vertices[i + 1], Vertices[i + 2]);
-	}
-
-	/** Set vertex position */
-	inline void SetVertex(int VertexID, const FVector3<RealType>& vNewPos)
-	{
-		//check(VectorUtil::IsFinite(vNewPos));
-		check(IsVertex(VertexID));
-		int i = 3 * VertexID;
-		Vertices[i] = vNewPos.X;
-		Vertices[i + 1] = vNewPos.Y;
-		Vertices[i + 2] = vNewPos.Z;
-		UpdateTimeStamp();
-	}
-
-
-
-	//
-	// iterators
-	//
-public:
-	/** @return enumerable object for valid vertex indices suitable for use with range-based for, ie for ( int i : VertexIndicesItr() ) */
-	FRefCountVector::IndexEnumerable VertexIndicesItr() const
-	{
-		return VertexRefCounts.Indices();
-	}
-
-	/** Enumerate positions of all points */
-	FRefCountVector::MappedEnumerable<FVector3<RealType>> VerticesItr() const
-	{
-		return VertexRefCounts.MappedIndices<FVector3<RealType>>([this](int VertexID) {
-			int i = 3 * VertexID;
-			return FVector3<RealType>(Vertices[i], Vertices[i + 1], Vertices[i + 2]);
-		});
-	}
-
-
-
-
-	/** @return true if vertex count == max vertex id */
-	bool IsCompact() const
-	{
-		return VertexRefCounts.IsDense();
-	}
-
-	/** returns measure of compactness in range [0,1], where 1 is fully compacted */
-	RealType CompactMetric() const
-	{
-		return ((RealType)VertexCount() / (RealType)MaxVertexID());
-	}
-
-	/** Timestamp is incremented any time any change is made to the point set */
-	inline int GetTimestamp() const
-	{
-		return Timestamp;
-	}
-
-	//
-	// Geometric queries
-	//
-public:
-
-	/** Returns bounding box of all points */
-	TAxisAlignedBox3<RealType> GetBounds() const
-	{
-		TAxisAlignedBox3<RealType> Box = TAxisAlignedBox3<RealType>::Empty();
-		for (int vi : VertexIndicesItr())
-		{
-			int k = 3 * vi;
-			Box.Contain( FVector3<RealType>(Vertices[k], Vertices[k + 1], Vertices[k + 2]) );
-		}
-		return Box;
-	}
-
-
-
-	//
-	// direct buffer access  (not a good idea to use)
-	//
-public:
-	const TDynamicVector<double>& GetVerticesBuffer()
-	{
-		return Vertices;
-	}
-	const FRefCountVector& GetVerticesRefCounts()
-	{
-		return VertexRefCounts;
-	}
-
-
-
-	//
-	// Edit operations
-	//
-public:
-
-	// @todo port this from DynamicMesh3
-	//void CompactInPlace(FCompactMaps* CompactInfo = nullptr);
-
-	/**
-	 * Remove point
-	 */
-	EMeshResult RemoveVertex(int VertexID)
-	{
-		if (VertexRefCounts.IsValid(VertexID) == false)
-		{
-			return EMeshResult::Failed_NotAVertex;
-		}
-		VertexRefCounts.Decrement(VertexID);
-		check(VertexRefCounts.IsValid(VertexID) == false);
-
-		if (GetBaseAttributeSet() != nullptr)
-		{
-			GetBaseAttributeSet()->OnRemoveVertex(VertexID);
-		}
-
-		UpdateTimeStamp();
-		return EMeshResult::Ok;
-	}
-
-
-
-protected:
-	void UpdateTimeStamp()
-	{
-		Timestamp++;
-	}
-
-
-
-};
-
-
-
-typedef TDynamicPointSet3<float> FDynamicPointSet3f;
-typedef TDynamicPointSet3<double> FDynamicPointSet3d;
-
-
-} // end namespace UE::Geometry
-} // end namespace UE
\ No newline at end of file
+#include "DynamicMesh/DynamicPointSet3.h"
diff --git a/Engine/Plugins/Experimental/GeometryProcessing/Source/DynamicMesh/Public/DynamicVertexAttribute.h b/Engine/Plugins/Experimental/GeometryProcessing/Source/DynamicMesh/Public/DynamicVertexAttribute.h
--- a/Engine/Plugins/Experimental/GeometryProcessing/Source/DynamicMesh/Public/DynamicVertexAttribute.h
+++ b/Engine/Plugins/Experimental/GeometryProcessing/Source/DynamicMesh/Public/DynamicVertexAttribute.h
@@ -2,451 +2,4 @@
 
 #pragma once
 
-#include "DynamicMesh3.h"
-#include "DynamicAttribute.h"
-
-namespace UE
-{
-namespace Geometry
-{
-
-template<typename AttribValueType, int AttribDimension, typename ParentType>
-class TDynamicVertexAttribute;
-
-
-template<typename AttribValueType, int AttribDimension, typename ParentType>
-class TDynamicVertexAttributeChange : public TDynamicAttributeChangeBase<ParentType>
-{
-private:
-	struct FChangeVertexAttribute
-	{
-		AttribValueType Data[AttribDimension];
-		int VertexID;
-	};
-
-	TArray<FChangeVertexAttribute> OldVertexAttributes, NewVertexAttributes;
-
-public:
-	TDynamicVertexAttributeChange()
-	{}
-
-	virtual ~TDynamicVertexAttributeChange()
-	{}
-
-	inline virtual void SaveInitialVertex(const TDynamicAttributeBase<ParentType>* Attribute, int VertexID) override;
-
-	inline virtual void StoreAllFinalVertices(const TDynamicAttributeBase<ParentType>* Attribute, const TArray<int>& VertexIDs) override;
-
-	inline virtual bool Apply(TDynamicAttributeBase<ParentType>* Attribute, bool bRevert) const override;
-};
-
-
-/**
- * TDynamicVertexAttribute provides per-vertex storage of an attribute value
- */
-template<typename AttribValueType, int AttribDimension, typename ParentType>
-class TDynamicVertexAttribute : public TDynamicAttributeBase<ParentType>
-{
-
-protected:
-	/** The parent object (e.g. mesh, point set) this attribute belongs to */
-	ParentType* Parent;
-
-	/** List of per-vertex attribute values */
-	TDynamicVector<AttribValueType> AttribValues;
-
-public:
-	/** Create an empty overlay */
-	TDynamicVertexAttribute() : Parent(nullptr)
-	{
-	}
-
-	/** Create an attribute for the given parent */
-	TDynamicVertexAttribute(ParentType* ParentIn, bool bAutoInit = true) : Parent(ParentIn)
-	{
-		if (bAutoInit)
-		{
-			Initialize();
-		}
-	}
-
-	virtual ~TDynamicVertexAttribute()
-	{
-	}
-
-	/** @return the parent for this attribute */
-	const ParentType* GetParent() const { return Parent; }
-	/** @return the parent for this attribute */
-	ParentType* GetParent() { return Parent; }
-private:
-	void Reparent(ParentType* NewParent) override { Parent = NewParent;  }
-public:
-	virtual TDynamicAttributeBase<ParentType>* MakeNew(ParentType* ParentIn) const override
-	{
-		TDynamicVertexAttribute<AttribValueType, AttribDimension, ParentType>* Matching = new TDynamicVertexAttribute<AttribValueType, AttribDimension, ParentType>(ParentIn);
-		Matching->Initialize();
-		return Matching;
-	}
-	virtual TDynamicAttributeBase<ParentType>* MakeCopy(ParentType* ParentIn) const override
-	{
-		TDynamicVertexAttribute<AttribValueType, AttribDimension, ParentType>* ToFill = new TDynamicVertexAttribute<AttribValueType, AttribDimension, ParentType>(ParentIn);
-		ToFill->Copy(*this);
-		return ToFill;
-	}
-
-	/** Set this overlay to contain the same arrays as the copy overlay */
-	void Copy(const TDynamicVertexAttribute<AttribValueType, AttribDimension, ParentType>& Copy)
-	{
-		TDynamicAttributeBase<ParentType>::CopyParentClassData(Copy);
-		AttribValues = Copy.AttribValues;
-	}
-
-	virtual TDynamicAttributeBase<ParentType>* MakeCompactCopy(const FCompactMaps& CompactMaps, ParentType* ParentTypeIn) const override
-	{
-		TDynamicVertexAttribute<AttribValueType, AttribDimension, ParentType>* ToFill = new TDynamicVertexAttribute<AttribValueType, AttribDimension, ParentType>(ParentTypeIn);
-		ToFill->Initialize();
-		ToFill->CompactCopy(CompactMaps, *this);
-		return ToFill;
-	}
-
-	void CompactInPlace(const FCompactMaps& CompactMaps)
-	{
-		for (int VID = 0; VID < CompactMaps.MapV.Num(); VID++)
-		{
-			int ToVID = CompactMaps.MapV[VID];
-			if (ToVID < 0)
-			{
-				continue;
-			}
-			if (ensure(ToVID <= VID))
-			{
-				CopyValue(VID, ToVID);
-			}
-		}
-		AttribValues.Resize(Parent->MaxVertexID() * AttribDimension);
-	}
-
-	void CompactCopy(const FCompactMaps& CompactMaps, const TDynamicVertexAttribute<AttribValueType, AttribDimension, ParentType>& ToCopy)
-	{
-		TDynamicAttributeBase<ParentType>::CopyParentClassData(ToCopy);
-		check(CompactMaps.MapV.Num() >= 0 && static_cast<size_t>(CompactMaps.MapV.Num()) <= AttribValues.Num() / AttribDimension);
-
-		AttribValueType Data[AttribDimension];
-		for (int VID = 0; VID < CompactMaps.MapV.Num(); VID++)
-		{
-			int ToVID = CompactMaps.MapV[VID];
-			if (ToVID == -1)
-			{
-				continue;
-			}
-			ToCopy.GetValue(VID, Data);
-			SetValue(ToVID, Data);
-		}
-	}
-
-	/** Initialize the attribute values to the given max triangle ID */
-	void Initialize(AttribValueType InitialValue = (AttribValueType)0)
-	{
-		check(Parent != nullptr);
-		AttribValues.Resize(0);
-		AttribValues.Resize( Parent->MaxVertexID() * AttribDimension, InitialValue );
-	}
-
-	void SetNewValue(int NewVertexID, const AttribValueType* Data)
-	{
-		int k = NewVertexID * AttribDimension;
-		for (int i = 0; i < AttribDimension; ++i)
-		{
-			AttribValues.InsertAt(Data[i], k + i);
-		}
-	}
-
-
-
-	//
-	// Accessors/Queries
-	//
-
-	virtual bool CopyThroughMapping(const TDynamicAttributeBase<ParentType>* Source, const FMeshIndexMappings& Mapping) override
-	{
-		AttribValueType BufferData[AttribDimension];
-		int BufferSize = sizeof(BufferData);
-		for (const TPair<int32, int32>& MapVID : Mapping.GetVertexMap().GetForwardMap())
-		{
-			if (!ensure(Source->CopyOut(MapVID.Key, BufferData, BufferSize)))
-			{
-				return false;
-			}
-			SetValue(MapVID.Value, BufferData);
-		}
-		return true;
-	}
-	virtual bool CopyOut(int RawID, void* Buffer, int BufferSize) const override
-	{
-		if (sizeof(AttribValueType)*AttribDimension != BufferSize)
-		{
-			return false;
-		}
-		AttribValueType* BufferData = static_cast<AttribValueType*>(Buffer);
-		int k = RawID * AttribDimension;
-		for (int i = 0; i < AttribDimension; ++i)
-		{
-			BufferData[i] = AttribValues[k + i];
-		}
-		return true;
-	}
-	virtual bool CopyIn(int RawID, void* Buffer, int BufferSize) override
-	{
-		if (sizeof(AttribValueType) * AttribDimension != BufferSize)
-		{
-			return false;
-		}
-		AttribValueType* BufferData = static_cast<AttribValueType*>(Buffer);
-		int k = RawID * AttribDimension;
-		for (int i = 0; i < AttribDimension; ++i)
-		{
-			AttribValues[k + i] = BufferData[i];
-		}
-		return true;
-	}
-
-	/** Get the element at a given index */
-	inline void GetValue(int VertexID, AttribValueType* Data) const
-	{
-		int k = VertexID * AttribDimension;
-		for (int i = 0; i < AttribDimension; ++i)
-		{
-			Data[i] = AttribValues[k + i];
-		}
-	}
-
-	/** Get the element at a given index */
-	template<typename AsType>
-	void GetValue(int VertexID, AsType& Data) const
-	{
-		int k = VertexID * AttribDimension;
-		for (int i = 0; i < AttribDimension; ++i)
-		{
-			Data[i] = AttribValues[k + i];
-		}
-	}
-
-
-	/** Set the element at a given index */
-	inline void SetValue(int VertexID, const AttribValueType* Data)
-	{
-		int k = VertexID * AttribDimension;
-		for (int i = 0; i < AttribDimension; ++i)
-		{
-			AttribValues[k + i] = Data[i];
-		}
-	}
-
-	/** Set the element at a given index */
-	template<typename AsType>
-	void SetValue(int VertexID, const AsType& Data)
-	{
-		int k = VertexID * AttribDimension;
-		for (int i = 0; i < AttribDimension; ++i)
-		{
-			AttribValues[k + i] = Data[i];
-		}
-	}
-
-	/**
-	 * Copy the attribute value at FromVertexID to ToVertexID
-	 */
-	inline void CopyValue(int FromVertexID, int ToVertexID)
-	{
-		int kA = FromVertexID * AttribDimension;
-		int kB = ToVertexID * AttribDimension;
-		for (int i = 0; i < AttribDimension; ++i)
-		{
-			AttribValues.InsertAt(AttribValues[kA+i], kB+i);
-		}
-	}
-
-
-public:
-
-	/** Update the overlay to reflect an edge split in the parent */
-	void OnSplitEdge(const FDynamicMesh3::FEdgeSplitInfo& SplitInfo) override
-	{
-		ResizeAttribStoreIfNeeded(SplitInfo.NewVertex);
-		SetAttributeFromLerp(SplitInfo.NewVertex, SplitInfo.OriginalVertices.A, SplitInfo.OriginalVertices.B, SplitInfo.SplitT);
-	}
-
-	/** Update the overlay to reflect an edge flip in the parent */
-	void OnFlipEdge(const FDynamicMesh3::FEdgeFlipInfo& FlipInfo) override
-	{
-		// vertices unchanged
-	}
-
-	/** Update the overlay to reflect an edge collapse in the parent */
-	void OnCollapseEdge(const FDynamicMesh3::FEdgeCollapseInfo& CollapseInfo) override
-	{
-		SetAttributeFromLerp(CollapseInfo.KeptVertex, CollapseInfo.KeptVertex, CollapseInfo.RemovedVertex, CollapseInfo.CollapseT);
-	}
-
-	virtual AttribValueType GetDefaultAttributeValue()
-	{
-		return (AttribValueType)0;
-	}
-
-	inline void ResizeAttribStoreIfNeeded(int VertexID)
-	{
-		if (!ensure(VertexID >= 0))
-		{
-			return;
-		}
-		size_t NeededSize = (size_t(VertexID)+1) * AttribDimension;
-		if (NeededSize > AttribValues.Num())
-		{
-			AttribValues.Resize(NeededSize, GetDefaultAttributeValue());
-		}
-	}
-
-	virtual void OnNewVertex(int VertexID, bool bInserted) override
-	{
-		ResizeAttribStoreIfNeeded(VertexID);
-	}
-
-	/** Update the overlay to reflect a face poke in the parent */
-	void OnPokeTriangle(const FDynamicMesh3::FPokeTriangleInfo& PokeInfo) override
-	{
-		FIndex3i Tri = PokeInfo.TriVertices;
-		ResizeAttribStoreIfNeeded(PokeInfo.NewVertex);
-		SetAttributeFromBary(PokeInfo.NewVertex, Tri.A, Tri.B, Tri.C, PokeInfo.BaryCoords);
-	}
-
-	/** Update the overlay to reflect an edge merge in the parent */
-	void OnMergeEdges(const FDynamicMesh3::FMergeEdgesInfo& MergeInfo) override
-	{
-		// just blend the attributes?
-		if (MergeInfo.RemovedVerts.A != FDynamicMesh3::InvalidID)
-		{
-			SetAttributeFromLerp(MergeInfo.KeptVerts.A, MergeInfo.KeptVerts.A, MergeInfo.RemovedVerts.A, .5);
-		}
-		if (MergeInfo.RemovedVerts.B != FDynamicMesh3::InvalidID)
-		{
-			SetAttributeFromLerp(MergeInfo.KeptVerts.B, MergeInfo.KeptVerts.B, MergeInfo.RemovedVerts.B, .5);
-		}
-	}
-
-	/** Update the overlay to reflect a vertex split in the parent */
-	void OnSplitVertex(const FDynamicMesh3::FVertexSplitInfo& SplitInfo, const TArrayView<const int>& TrianglesToUpdate) override
-	{
-		CopyValue(SplitInfo.OriginalVertex, SplitInfo.NewVertex);
-	}
-
-	virtual TUniquePtr<TDynamicAttributeChangeBase<ParentType>> NewBlankChange() override
-	{
-		return MakeUnique<TDynamicVertexAttributeChange<AttribValueType, AttribDimension, ParentType>>();
-	}
-
-	/**
-	* Check validity of attribute
-	* 
-	* @param bAllowNonmanifold Accept non-manifold topology as valid. Note that this should almost always be true for attributes; non-manifold overlays are generally valid.
-	* @param FailMode Desired behavior if mesh is found invalid
-	*/
-	virtual bool CheckValidity(bool bAllowNonmanifold, EValidityCheckFailMode FailMode) const override
-	{
-		// just check that the values buffer is big enough
-		if (Parent->MaxVertexID() < 0 || static_cast<size_t>(Parent->MaxVertexID()*AttribDimension) > AttribValues.Num())
-		{
-			switch (FailMode)
-			{
-			case EValidityCheckFailMode::Check:
-				check(false);
-			case EValidityCheckFailMode::Ensure:
-				ensure(false);
-			default:
-				return false;
-			}
-		}
-
-		return true;
-	}
-
-
-protected:
-
-	// interpolation functions; default implementation assumes your attributes can be interpolated as reals
-
-	/** Set the value at an Attribute to be a linear interpolation of two other Attributes */
-	virtual void SetAttributeFromLerp(int SetAttribute, int AttributeA, int AttributeB, double Alpha)
-	{
-		int IndexSet = AttribDimension * SetAttribute;
-		int IndexA = AttribDimension * AttributeA;
-		int IndexB = AttribDimension * AttributeB;
-		AttribValueType Beta = ((AttribValueType)1 - Alpha);
-		for (int i = 0; i < AttribDimension; ++i)
-		{
-			AttribValues[IndexSet+i] = Beta*AttribValues[IndexA+i] + Alpha*AttribValues[IndexB+i];
-		}
-	}
-
-	/** Set the value at an Attribute to be a barycentric interpolation of three other Attributes */
-	virtual void SetAttributeFromBary(int SetAttribute, int AttributeA, int AttributeB, int AttributeC, const FVector3d& BaryCoords)
-	{
-		int IndexSet = AttribDimension * SetAttribute;
-		int IndexA = AttribDimension * AttributeA;
-		int IndexB = AttribDimension * AttributeB;
-		int IndexC = AttribDimension * AttributeC;
-		for (int i = 0; i < AttribDimension; ++i)
-		{
-			AttribValues[IndexSet + i] = 
-				BaryCoords.X*AttribValues[IndexA+i] + BaryCoords.Y*AttribValues[IndexB+i] + BaryCoords.Z*AttribValues[IndexC+i];
-		}
-	}
-
-};
-
-
-template<typename AttribValueType, int AttribDimension>
-using TDynamicMeshVertexAttribute = TDynamicVertexAttribute<AttribValueType, AttribDimension, FDynamicMesh3>;
-
-
-template<typename AttribValueType, int AttribDimension, typename ParentType>
-void TDynamicVertexAttributeChange<AttribValueType, AttribDimension, ParentType>::SaveInitialVertex(const TDynamicAttributeBase<ParentType>* Attribute, int VertexID)
-{
-	FChangeVertexAttribute& Change = OldVertexAttributes.Emplace_GetRef();
-	Change.VertexID = VertexID;
-	const TDynamicVertexAttribute<AttribValueType, AttribDimension, ParentType>* AttribCast = static_cast<const TDynamicVertexAttribute<AttribValueType, AttribDimension, ParentType>*>(Attribute);
-	AttribCast->GetValue(VertexID, Change.Data);
-}
-
-template<typename AttribValueType, int AttribDimension, typename ParentType>
-void TDynamicVertexAttributeChange<AttribValueType, AttribDimension, ParentType>::StoreAllFinalVertices(const TDynamicAttributeBase<ParentType>* Attribute, const TArray<int>& VertexIDs)
-{
-	NewVertexAttributes.Reserve(NewVertexAttributes.Num() + VertexIDs.Num());
-	const TDynamicVertexAttribute<AttribValueType, AttribDimension, ParentType>* AttribCast = static_cast<const TDynamicVertexAttribute<AttribValueType, AttribDimension, ParentType>*>(Attribute);
-	for (int VertexID : VertexIDs)
-	{
-		FChangeVertexAttribute& Change = NewVertexAttributes.Emplace_GetRef();
-		Change.VertexID = VertexID;
-		AttribCast->GetValue(VertexID, Change.Data);
-	}
-}
-
-template<typename AttribValueType, int AttribDimension, typename ParentType>
-bool TDynamicVertexAttributeChange<AttribValueType, AttribDimension, ParentType>::Apply(TDynamicAttributeBase<ParentType>* Attribute, bool bRevert) const
-{
-	const TArray<FChangeVertexAttribute> *Changes = bRevert ? &OldVertexAttributes : &NewVertexAttributes;
-	TDynamicVertexAttribute<AttribValueType, AttribDimension, ParentType>* AttribCast = static_cast<TDynamicVertexAttribute<AttribValueType, AttribDimension, ParentType>*>(Attribute);
-	for (const FChangeVertexAttribute& Change : *Changes)
-	{
-		bool bIsVertex = AttribCast->GetParent()->IsVertex(Change.VertexID);
-		check(bIsVertex);
-		if (bIsVertex)
-		{
-			AttribCast->SetValue(Change.VertexID, Change.Data);
-		}
-	}
-	return true;
-}
-
-
-} // end namespace UE::Geometry
-} // end namespace UE
+#include "DynamicMesh/DynamicVertexAttribute.h"
diff --git a/Engine/Plugins/Experimental/GeometryProcessing/Source/DynamicMesh/Public/DynamicVertexSkinWeightsAttribute.h b/Engine/Plugins/Experimental/GeometryProcessing/Source/DynamicMesh/Public/DynamicVertexSkinWeightsAttribute.h
--- a/Engine/Plugins/Experimental/GeometryProcessing/Source/DynamicMesh/Public/DynamicVertexSkinWeightsAttribute.h
+++ b/Engine/Plugins/Experimental/GeometryProcessing/Source/DynamicMesh/Public/DynamicVertexSkinWeightsAttribute.h
@@ -2,410 +2,4 @@
 
 #pragma once
 
-#include "DynamicMesh3.h"
-#include "DynamicAttribute.h"
-
-#include "BoneWeights.h"
-#include "HAL/UnrealMemory.h"
-
-
-namespace UE
-{
-namespace Geometry
-{
-
-template<typename ParentType>
-class TDynamicVertexSkinWeightsAttributeChange final : 
-	public TDynamicAttributeChangeBase<ParentType>
-{
-	struct FChangeVertexBoneWeights
-	{
-		UE::AnimationCore::FBoneWeights Weights;
-		int VertexID;
-	};
-
-	TArray<FChangeVertexBoneWeights> OldVertexBoneWeights, NewVertexBoneWeights;
-
-public:
-	TDynamicVertexSkinWeightsAttributeChange() = default;
-
-	virtual ~TDynamicVertexSkinWeightsAttributeChange() override = default;
-
-	void SaveInitialVertex(const TDynamicAttributeBase<ParentType>* Attribute, int VertexID) override;
-
-	void StoreAllFinalVertices(const TDynamicAttributeBase<ParentType>* Attribute, const TArray<int>& VertexIDs) override;
-
-	bool Apply(TDynamicAttributeBase<ParentType>* Attribute, bool bRevert) const override;
-};
-
-
-/**
- * TDynamicVertexAttribute provides per-vertex storage of bone weights (skin weights)
- */
-template<typename ParentType>
-class TDynamicVertexSkinWeightsAttribute final : 
-	public TDynamicAttributeBase<ParentType>
-{
-public:
-	using FBoneWeights = UE::AnimationCore::FBoneWeights;
-
-protected:
-	friend class FDynamicMeshAttributeSet;
-
-	/** The parent object (e.g. mesh, point set) this attribute belongs to */
-	ParentType* Parent = nullptr;
-
-	/** List of per-vertex bone weights values */
-	// FIXME: Replace with a local slab allocator storage to avoid per-item TInlineAllocator
-	TDynamicVector<FBoneWeights> VertexBoneWeights;
-
-public:
-
-	/** Create an empty overlay */
-	TDynamicVertexSkinWeightsAttribute() = default;
-
-	/** Create an attribute for the given parent */
-	TDynamicVertexSkinWeightsAttribute(ParentType* ParentIn, bool bAutoInit = true) : 
-		Parent(ParentIn)
-	{
-		if (bAutoInit)
-		{
-			Initialize();
-		}
-	}
-
-	virtual ~TDynamicVertexSkinWeightsAttribute() = default;
-
-	/** @return the parent for this attribute */
-	const ParentType* GetParent() const { return Parent; }
-	/** @return the parent for this attribute */
-	ParentType* GetParent() { return Parent; }
-private:
-	void Reparent(ParentType* NewParent) override { Parent = NewParent;  }
-public:
-	TDynamicAttributeBase<ParentType>* MakeNew(ParentType* ParentIn) const override
-	{
-		TDynamicVertexSkinWeightsAttribute<ParentType>* Matching = new TDynamicVertexSkinWeightsAttribute<ParentType>(ParentIn);
-		Matching->Initialize();
-		return Matching;
-	}
-	TDynamicAttributeBase<ParentType>* MakeCopy(ParentType* ParentIn) const override
-	{
-		TDynamicVertexSkinWeightsAttribute<ParentType>* ToFill = new TDynamicVertexSkinWeightsAttribute<ParentType>(ParentIn);
-		ToFill->Copy(*this);
-		return ToFill;
-	}
-
-	/** Set this overlay to contain the same arrays as the copy overlay */
-	void Copy(const TDynamicVertexSkinWeightsAttribute<ParentType>& Copy)
-	{
-		TDynamicAttributeBase<ParentType>::CopyParentClassData(Copy);
-		VertexBoneWeights = Copy.VertexBoneWeights;
-	}
-
-	TDynamicAttributeBase<ParentType>* MakeCompactCopy(const FCompactMaps& CompactMaps, ParentType* ParentTypeIn) const override
-	{
-		TDynamicVertexSkinWeightsAttribute<ParentType>* ToFill = new TDynamicVertexSkinWeightsAttribute<ParentType>(ParentTypeIn);
-		ToFill->Initialize();
-		ToFill->CompactCopy(CompactMaps, *this);
-		return ToFill;
-	}
-
-	void CompactInPlace(const FCompactMaps& CompactMaps) override
-	{
-		for (int32 VID = 0; VID < CompactMaps.MapV.Num(); VID++)
-		{
-			const int32 ToVID = CompactMaps.MapV[VID];
-			if (ToVID < 0)
-			{
-				continue;
-			}
-			if (ensure(ToVID <= VID))
-			{
-				CopyValue(VID, ToVID);
-			}
-		}
-		VertexBoneWeights.Resize(Parent->MaxVertexID());
-	}
-
-	void CompactCopy(const FCompactMaps& CompactMaps, const TDynamicVertexSkinWeightsAttribute<ParentType>& ToCopy)
-	{
-		TDynamicAttributeBase<ParentType>::CopyParentClassData(ToCopy);
-		check(CompactMaps.MapV.Num() <= VertexBoneWeights.Num());
-		FBoneWeights Data;
-		for (int32 VID = 0; VID < CompactMaps.MapV.Num(); VID++)
-		{
-			const int32 ToVID = CompactMaps.MapV[VID];
-			if (ToVID == -1)
-			{
-				continue;
-			}
-			ToCopy.GetValue(VID, Data);
-			SetValue(ToVID, Data);
-		}
-	}
-
-	/** Initialize the attribute values to the given max triangle ID */
-	void Initialize(const FBoneWeights InitialValue = {})
-	{
-		check(Parent != nullptr);
-		VertexBoneWeights.Resize(Parent->MaxVertexID());
-		VertexBoneWeights.Fill(InitialValue);
-	}
-
-	void SetNewValue(int32 InNewVertexID, const FBoneWeights InBoneWeights)
-	{
-		VertexBoneWeights.InsertAt(InBoneWeights, InNewVertexID);
-	}
-
-
-
-	//
-	// Accessors/Queries
-	//
-
-	bool CopyThroughMapping(const TDynamicAttributeBase<ParentType>* Source, const FMeshIndexMappings& Mapping) override
-	{
-		// Don't snarf the FBoneWeight as a concrete object, since it _may_ contain a pointer
-		// and we don't want it to destruct.
-		constexpr int32 BufferSize = sizeof(FBoneWeights);
-		int8 BufferData[BufferSize];
-		for (const TPair<int32, int32>& MapVID : Mapping.GetVertexMap().GetForwardMap())
-		{
-			if (!ensure(Source->CopyOut(MapVID.Key, &BufferData, BufferSize)))
-			{
-				return false;
-			}
-			SetValue(MapVID.Value, *reinterpret_cast<FBoneWeights *>(BufferData));
-		}
-		return true;
-	}
-	bool CopyOut(int RawID, void* Buffer, int BufferSize) const override
-	{
-		if (sizeof(FBoneWeights) != BufferSize)
-		{
-			return false;
-		}
-
-		// Avoid the copy constructor.
-		FMemory::Memcpy(Buffer, &VertexBoneWeights[RawID], sizeof(FBoneWeights));
-		return true;
-	}
-	bool CopyIn(int RawID, void* Buffer, int BufferSize) override
-	{
-		if (sizeof(FBoneWeights) != BufferSize)
-		{
-			return false;
-		}
-		
-		// Ensure the copy constructor is called.
-		VertexBoneWeights[RawID] = *static_cast<FBoneWeights*>(Buffer);
-		return true;
-	}
-
-	/** Get the element at a given index */
-	void GetValue(int VertexID, FBoneWeights& Data) const
-	{
-		Data = VertexBoneWeights[VertexID];
-	}
-
-	/** Get the element at a given index */
-	template<typename AsType>
-	void GetValue(int VertexID, AsType& Data) const
-	{
-		Data = VertexBoneWeights[VertexID];
-	}
-
-
-	/** Set the element at a given index */
-	void SetValue(int VertexID, const FBoneWeights& Data)
-	{
-		VertexBoneWeights[VertexID] = Data;
-	}
-
-	/** Set the element at a given index */
-	template<typename ContainerAdapter>
-	void SetValue(int VertexID, const UE::AnimationCore::TBoneWeights<ContainerAdapter>& Data)
-	{
-		VertexBoneWeights[VertexID] = FBoneWeights::Create(Data);
-	}
-
-	/**
-	 * Copy the attribute value at FromVertexID to ToVertexID
-	 */
-	void CopyValue(int FromVertexID, int ToVertexID)
-	{
-		VertexBoneWeights.InsertAt(VertexBoneWeights[FromVertexID], ToVertexID);
-	}
-
-
-public:
-
-	/** Update the overlay to reflect an edge split in the parent */
-	void OnSplitEdge(const FDynamicMesh3::FEdgeSplitInfo& SplitInfo) override
-	{
-		ResizeAttribStoreIfNeeded(SplitInfo.NewVertex);
-		SetBoneWeightsFromLerp(SplitInfo.NewVertex, SplitInfo.OriginalVertices.A, SplitInfo.OriginalVertices.B, SplitInfo.SplitT);
-	}
-
-	/** Update the overlay to reflect an edge flip in the parent */
-	void OnFlipEdge(const FDynamicMesh3::FEdgeFlipInfo& FlipInfo) override
-	{
-		// vertices unchanged
-	}
-
-	/** Update the overlay to reflect an edge collapse in the parent */
-	void OnCollapseEdge(const FDynamicMesh3::FEdgeCollapseInfo& CollapseInfo) override
-	{
-		SetBoneWeightsFromLerp(CollapseInfo.KeptVertex, CollapseInfo.KeptVertex, CollapseInfo.RemovedVertex, CollapseInfo.CollapseT);
-	}
-
-	void ResizeAttribStoreIfNeeded(int VertexID)
-	{
-		if (!ensure(VertexID >= 0))
-		{
-			return;
-		}
-		size_t NeededSize = (size_t(VertexID)+1);
-		if (NeededSize > VertexBoneWeights.Num())
-		{
-			VertexBoneWeights.Resize(NeededSize, FBoneWeights{});
-		}
-	}
-
-	void OnNewVertex(int VertexID, bool bInserted) override
-	{
-		ResizeAttribStoreIfNeeded(VertexID);
-	}
-
-	/** Update the overlay to reflect a face poke in the parent */
-	void OnPokeTriangle(const FDynamicMesh3::FPokeTriangleInfo& PokeInfo) override
-	{
-		FIndex3i Tri = PokeInfo.TriVertices;
-		ResizeAttribStoreIfNeeded(PokeInfo.NewVertex);
-		SetBoneWeightsFromBary(PokeInfo.NewVertex, Tri.A, Tri.B, Tri.C, PokeInfo.BaryCoords);
-	}
-
-	/** Update the overlay to reflect an edge merge in the parent */
-	void OnMergeEdges(const FDynamicMesh3::FMergeEdgesInfo& MergeInfo) override
-	{
-		// just blend the attributes?
-		if (MergeInfo.RemovedVerts.A != FDynamicMesh3::InvalidID)
-		{
-			SetBoneWeightsFromLerp(MergeInfo.KeptVerts.A, MergeInfo.KeptVerts.A, MergeInfo.RemovedVerts.A, .5);
-		}
-		if (MergeInfo.RemovedVerts.B != FDynamicMesh3::InvalidID)
-		{
-			SetBoneWeightsFromLerp(MergeInfo.KeptVerts.B, MergeInfo.KeptVerts.B, MergeInfo.RemovedVerts.B, .5);
-		}
-	}
-
-	/** Update the overlay to reflect a vertex split in the parent */
-	void OnSplitVertex(const FDynamicMesh3::FVertexSplitInfo& SplitInfo, const TArrayView<const int>& TrianglesToUpdate) override
-	{
-		CopyValue(SplitInfo.OriginalVertex, SplitInfo.NewVertex);
-	}
-
-	TUniquePtr<TDynamicAttributeChangeBase<ParentType>> NewBlankChange() override
-	{
-		return MakeUnique<TDynamicVertexSkinWeightsAttributeChange<ParentType>>();
-	}
-
-	/**
-	* Check validity of attribute
-	* 
-	* @param bAllowNonmanifold Accept non-manifold topology as valid. Note that this should almost always be true for attributes; non-manifold overlays are generally valid.
-	* @param FailMode Desired behavior if mesh is found invalid
-	*/
-	bool CheckValidity(bool bAllowNonmanifold, EValidityCheckFailMode FailMode) const override
-	{
-		// just check that the values buffer is big enough
-		if (Parent->MaxVertexID() > VertexBoneWeights.Num())
-		{
-			switch (FailMode)
-			{
-			case EValidityCheckFailMode::Check:
-				check(false);
-			case EValidityCheckFailMode::Ensure:
-				ensure(false);
-			default:
-				return false;
-			}
-		}
-
-		return true;
-	}
-
-
-protected:
-
-	// interpolation functions; default implementation assumes your attributes can be interpolated as reals
-
-	/** Set the value at an Attribute to be a linear interpolation of two other Attributes */
-	void SetBoneWeightsFromLerp(int SetAttribute, int AttributeA, int AttributeB, double Alpha)
-	{
-		Alpha = FMath::Clamp(Alpha, 0.0, 1.0);
-		VertexBoneWeights[SetAttribute] = FBoneWeights::Blend(VertexBoneWeights[AttributeA], VertexBoneWeights[AttributeB], Alpha);
-	}
-
-	/** Set the value at an Attribute to be a barycentric interpolation of three other Attributes */
-	void SetBoneWeightsFromBary(int SetAttribute, int AttributeA, int AttributeB, int AttributeC, const FVector3d& BaryCoords)
-	{
-		// Since FBoneWeights only defines Blend for two inputs, we need to split the barycentric coordinate interpolation 
-		// into two blends.
-		if (!FMath::IsNearlyZero(BaryCoords.Y + BaryCoords.Z))
-		{
-			const double BCW = BaryCoords.Y / (BaryCoords.Y + BaryCoords.Z);
-			const FBoneWeights BC = FBoneWeights::Blend(VertexBoneWeights[AttributeB], VertexBoneWeights[AttributeC], BCW);
-
-			VertexBoneWeights[SetAttribute] = FBoneWeights::Blend(VertexBoneWeights[AttributeA], BC, BaryCoords.X);
-		}
-		else if (SetAttribute != AttributeA)
-		{
-			VertexBoneWeights[SetAttribute] = VertexBoneWeights[AttributeA];
-		}
-	}
-};
-
-
-class FDynamicMesh3;
-using FDynamicMeshVertexSkinWeightsAttribute = TDynamicVertexSkinWeightsAttribute<FDynamicMesh3>;
-
-
-template<typename ParentType>
-void TDynamicVertexSkinWeightsAttributeChange<ParentType>::SaveInitialVertex(const TDynamicAttributeBase<ParentType>* Attribute, int VertexID)
-{
-	FChangeVertexBoneWeights& Change = OldVertexBoneWeights.Emplace_GetRef();
-	Change.VertexID = VertexID;
-	const TDynamicVertexSkinWeightsAttribute<ParentType>* AttribCast = static_cast<const TDynamicVertexSkinWeightsAttribute<ParentType>*>(Attribute);
-	AttribCast->GetValue(VertexID, Change.Weights);
-}
-
-template<typename ParentType>
-void TDynamicVertexSkinWeightsAttributeChange<ParentType>::StoreAllFinalVertices(const TDynamicAttributeBase<ParentType>* Attribute, const TArray<int>& VertexIDs)
-{
-	NewVertexBoneWeights.Reserve(NewVertexBoneWeights.Num() + VertexIDs.Num());
-	const TDynamicVertexSkinWeightsAttribute<ParentType>* AttribCast = static_cast<const TDynamicVertexSkinWeightsAttribute<ParentType>*>(Attribute);
-	for (int VertexID : VertexIDs)
-	{
-		FChangeVertexBoneWeights& Change = NewVertexBoneWeights.Emplace_GetRef();
-		Change.VertexID = VertexID;
-		AttribCast->GetValue(VertexID, Change.Weights);
-	}
-}
-
-template<typename ParentType>
-bool TDynamicVertexSkinWeightsAttributeChange<ParentType>::Apply(TDynamicAttributeBase<ParentType>* Attribute, bool bRevert) const
-{
-	const TArray<FChangeVertexBoneWeights> *Changes = bRevert ? &OldVertexBoneWeights : &NewVertexBoneWeights;
-	TDynamicVertexSkinWeightsAttribute<ParentType>* AttribCast = static_cast<TDynamicVertexSkinWeightsAttribute<ParentType>*>(Attribute);
-	for (const FChangeVertexBoneWeights& Change : *Changes)
-	{
-		check(AttribCast->GetParent()->IsVertex(Change.VertexID));
-		AttribCast->SetValue(Change.VertexID, Change.Weights);
-	}
-	return true;
-}
-
-}
-}
+#include "DynamicMesh/DynamicVertexSkinWeightsAttribute.h"
diff --git a/Engine/Plugins/Experimental/GeometryProcessing/Source/DynamicMesh/Public/DynamicVerticesOctree3.h b/Engine/Plugins/Experimental/GeometryProcessing/Source/DynamicMesh/Public/DynamicVerticesOctree3.h
--- a/Engine/Plugins/Experimental/GeometryProcessing/Source/DynamicMesh/Public/DynamicVerticesOctree3.h
+++ b/Engine/Plugins/Experimental/GeometryProcessing/Source/DynamicMesh/Public/DynamicVerticesOctree3.h
@@ -2,230 +2,5 @@
 
 #pragma once
 
-#include "Spatial/MeshAABBTree3.h"
-#include "Spatial/SparseDynamicPointOctree3.h"
-#include "MeshQueries.h"
-#include "Intersection/IntersectionUtil.h"
-#include "DynamicMesh3.h"
+#include "DynamicMesh/DynamicVerticesOctree3.h"
 
-namespace UE
-{
-namespace Geometry
-{
-
-/**
- * TDynamicVerticesOctree3 is an extension of FSparseDynamicPointOctree3 for the vertices of 
- * a FDynamicXYZ3 instance (eg FDynamicMesh3, TDynamicPointSet3, FDynamicGraph, etc)
- * This extension does several things:
- *   1) provides a simplified API based on vertex IDs to various Octree functions
- *   2) tracks ModifiedBounds box of modified areas
- *
- * The template expansion requires that FDynamicXYZ3 has the following API:
- *   1) a function int MaxVertexID() that returns the maximum vertex ID/index
- *   2) a function bool IsVertex(int) that returns true if the index is valid
- *   3) a function FVector3d GetVertex(int index) that returns the position of a vertex
- *   4) a function FAxisAlignedBox3d GetBounds() that returns a bounding-box of the point set
- */
-template<typename SourceType>
-class TDynamicVerticesOctree3 : public FSparseDynamicPointOctree3
-{
-	// potential optimizations:
-	//    - keep track of how many vertices are descendents of each cell? root cells?
-
-
-
-public:
-	/** parent mesh */
-	const SourceType* VertexSource;
-
-	/** bounding box of vertices that have been inserted/removed since last clear */
-	FAxisAlignedBox3d ModifiedBounds;
-
-	double HitSphereRadius = 1.0f;
-
-	/** 
-	 * Add all vertices of MeshIn to the octree 
-	 */
-	void Initialize(const SourceType* VertexSourceIn, bool bDynamicExpand)
-	{
-		this->VertexSource = VertexSourceIn;
-
-		HitSphereRadius = VertexSourceIn->GetBounds().Diagonal().Length() * 0.001;
-
-		int MaxVertexID = VertexSource->MaxVertexID();
-		for ( int VertexID = 0; VertexID < MaxVertexID; VertexID++ )
-		{
-			if (VertexSource->IsVertex(VertexID))
-			{
-				if (bDynamicExpand)
-				{
-					ModifiedBounds.Contain(VertexSource->GetVertex(VertexID));
-					const SourceType* VtxSource = VertexSource;
-					FSparseDynamicPointOctree3::InsertPoint_DynamicExpand(
-						VertexID,
-						[VtxSource](int k) { return VtxSource->GetVertex(k); });
-				}
-				else
-				{
-					InsertVertex(VertexID);
-				}
-			}
-		}
-	}
-
-	/**
-	 * Reset the internal ModifiedBounds box that tracks modified triangle bounds
-	 */
-	void ResetModifiedBounds()
-	{
-		ModifiedBounds = FAxisAlignedBox3d::Empty();
-	}
-
-	/**
-	 * Insert a triangle into the tree
-	 */
-	void InsertVertex(int32 VertexID)
-	{
-		FVector3d Point = VertexSource->GetVertex(VertexID);
-		ModifiedBounds.Contain(Point);
-		FSparseDynamicPointOctree3::InsertPoint(VertexID, Point);
-	}
-
-	/**
-	 * Insert a list of vertices into the tree
-	 */
-	void InsertVertices(const TArray<int>& Vertices)
-	{
-		int N = Vertices.Num();
-		for (int i = 0; i < N; ++i)
-		{
-			FVector3d Point = VertexSource->GetVertex(Vertices[i]);
-			ModifiedBounds.Contain(Point);
-			FSparseDynamicPointOctree3::InsertPoint(Vertices[i], Point);
-		}
-	}
-
-	/**
-	 * Insert a set of vertices into the tree
-	 */
-	void InsertVertices(const TSet<int>& Vertices)
-	{
-		for (int VertexID : Vertices)
-		{
-			FVector3d Point = VertexSource->GetVertex(VertexID);
-			ModifiedBounds.Contain(Point);
-			FSparseDynamicPointOctree3::InsertPoint(VertexID, Point);
-		}
-	}
-
-
-	/**
-	 * Remove a triangle from the tree
-	 */
-	bool RemovePoint(int32 VertexID)
-	{
-		FVector3d Point = VertexSource->GetVertex(VertexID);
-		ModifiedBounds.Contain(Point);
-		return FSparseDynamicPointOctree3::RemovePoint(VertexID);
-	}
-
-	/**
-	 * Remove a list of vertices from the tree
-	 */
-	void RemoveVertices(const TArray<int>& Vertices)
-	{
-		int N = Vertices.Num();
-		for ( int i = 0; i < N; ++i )
-		{
-			FVector3d Point = VertexSource->GetVertex(Vertices[i]);
-			ModifiedBounds.Contain(Point);
-			FSparseDynamicPointOctree3::RemovePoint(Vertices[i]);
-		}
-	}
-
-	/**
-	 * Remove a set of vertices from the tree
-	 */
-	void RemoveVertices(const TSet<int>& Vertices)
-	{
-		for (int VertexID : Vertices)
-		{
-			FVector3d Point = VertexSource->GetVertex(VertexID);
-			ModifiedBounds.Contain(Point);
-			FSparseDynamicPointOctree3::RemovePoint(VertexID);
-		}
-	}
-
-
-	/**
-	 * Reinsert a set of vertices into the tree
-	 */
-	void ReinsertVertices(const TSet<int>& Vertices)
-	{
-		for (int VertexID : Vertices)
-		{
-			FVector3d Point = VertexSource->GetVertex(VertexID);
-			ModifiedBounds.Contain(Point);
-			FSparseDynamicPointOctree3::ReinsertPoint(VertexID, ModifiedBounds);
-		}
-	}
-
-
-	/**
-	 * Include the current bounds of a triangle in the ModifiedBounds box
-	 */
-	void NotifyPendingModification(int VertexID)
-	{
-		FVector3d Point = VertexSource->GetVertex(VertexID);
-		ModifiedBounds.Contain(Point);
-	}
-
-	/**
-	 * Include the current bounds of a set of vertices in the ModifiedBounds box
-	 */
-	void NotifyPendingModification(const TSet<int>& Vertices)
-	{
-		for (int VertexID : Vertices)
-		{
-			FVector3d Point = VertexSource->GetVertex(VertexID);
-			ModifiedBounds.Contain(Point);
-		}
-	}
-
-
-	/**
-	 * Find the nearest triangle of the VertexSource that is hit by the ray
-	 */
-	int32 FindNearestHitVertex(const FRay3d& Ray,
-		double MaxDistance = TNumericLimits<double>::Max()) const
-	{
-		return FSparseDynamicPointOctree3::FindNearestHitPoint(Ray,
-			[this](int vid, const FRay3d& Ray) {
-				FVector3d Point = VertexSource->GetVertex(vid);
-				FLinearIntersection Hit = IntersectionUtil::RaySphereIntersection(Ray.Origin, Ray.Direction, Point, HitSphereRadius);
-				return (Hit.intersects) ? Hit.parameter.Min : TNumericLimits<double>::Max();
-			}, MaxDistance);
-	}
-
-
-	/**
-	 * Check that the Octree is internally valid
-	 */
-	void CheckValidity(
-		EValidityCheckFailMode FailMode = EValidityCheckFailMode::Check,
-		bool bVerbose = false,
-		bool bFailOnMissingPoints = false) const
-	{
-		FSparseDynamicPointOctree3::CheckValidity(
-			[&](int vid) { return VertexSource->IsVertex(vid); },
-			[this](int vid) { return VertexSource->GetVertex(vid); },
-			FailMode, bVerbose, bFailOnMissingPoints);
-	}
-
-
-};
-
-
-
-} // end namespace UE::Geometry
-} // end namespace UE
diff --git a/Engine/Plugins/Experimental/GeometryProcessing/Source/DynamicMesh/Public/InfoTypes.h b/Engine/Plugins/Experimental/GeometryProcessing/Source/DynamicMesh/Public/InfoTypes.h
--- a/Engine/Plugins/Experimental/GeometryProcessing/Source/DynamicMesh/Public/InfoTypes.h
+++ b/Engine/Plugins/Experimental/GeometryProcessing/Source/DynamicMesh/Public/InfoTypes.h
@@ -2,120 +2,4 @@
 
 #pragma once
 
-#include "IndexTypes.h"
-#include "VectorTypes.h"
-
-namespace UE
-{
-namespace Geometry
-{
-
-/**
- * FVertexInfo stores information about vertex attributes - position, normal, color, UV
- */
-struct FVertexInfo
-{
-	FVector3d Position{ FVector3d::Zero() };
-	FVector3f Normal{ FVector3f::Zero() };
-	FVector3f Color{ FVector3f::Zero() };
-	FVector2f UV{ FVector2f::Zero() };
-	bool bHaveN{}, bHaveC{}, bHaveUV{};
-
-	FVertexInfo() = default;
-	FVertexInfo(const FVector3d& PositionIn)
-		: Position{ PositionIn }{}
-	FVertexInfo(const FVector3d& PositionIn, const FVector3f& NormalIn)
-		: Position{ PositionIn }, Normal{ NormalIn }, bHaveN{ true }{}
-	FVertexInfo(const FVector3d& PositionIn, const FVector3f& NormalIn, const FVector3f& ColorIn)
-		: Position{ PositionIn }, Normal{ NormalIn }, Color{ ColorIn }, bHaveN{ true }, bHaveC{true}{}
-	FVertexInfo(const FVector3d& PositionIn, const FVector3f& NormalIn, const FVector3f& ColorIn, const FVector2f& UVIn)
-		: Position{ PositionIn }, Normal{ NormalIn }, Color{ ColorIn }, UV{ UVIn }, bHaveN{ true }, bHaveC{true}, bHaveUV{true}{}
-};
-
-
-} // end namespace UE::Geometry
-} // end namespace UE
-
-
-
-namespace DynamicMeshInfo
-{
-using namespace UE::Geometry;
-
-/** Information about the mesh elements created by a call to SplitEdge() */
-struct FEdgeSplitInfo
-{
-	int OriginalEdge;					// the edge that was split
-	FIndex2i OriginalVertices;			// original edge vertices [a,b]
-	FIndex2i OtherVertices;				// original opposing vertices [c,d] - d is InvalidID for boundary edges
-	FIndex2i OriginalTriangles;			// original edge triangles [t0,t1]
-	bool bIsBoundary;					// was the split edge a boundary edge?  (redundant)
-
-	int NewVertex;						// new vertex f that was created
-	FIndex2i NewTriangles;				// new triangles [t2,t3], oriented as explained below
-	FIndex3i NewEdges;					// new edges are [f,b], [f,c] and [f,d] if this is not a boundary edge
-
-	double SplitT;						// parameter value for NewVertex along original edge
-};
-
-/** Information about the mesh elements modified by a call to FlipEdge() */
-struct FEdgeFlipInfo
-{
-	int EdgeID;						// the edge that was flipped
-	FIndex2i OriginalVerts;			// original verts of the flipped edge, that are no longer connected
-	FIndex2i OpposingVerts;			// the opposing verts of the flipped edge, that are now connected
-	FIndex2i Triangles;				// the two triangle IDs. Original tris vert [Vert0,Vert1,OtherVert0] and [Vert1,Vert0,OtherVert1].
-										// New triangles are [OtherVert0, OtherVert1, Vert1] and [OtherVert1, OtherVert0, Vert0]
-};
-
-/** Information about mesh elements modified/removed by CollapseEdge() */
-struct FEdgeCollapseInfo
-{
-	int KeptVertex;					// the vertex that was kept (ie collapsed "to")
-	int RemovedVertex;				// the vertex that was removed
-	FIndex2i OpposingVerts;			// the opposing vertices [c,d]. If the edge was a boundary edge, d is InvalidID
-	bool bIsBoundary;				// was the edge a boundary edge
-
-	int CollapsedEdge;				// the edge that was collapsed/removed
-	FIndex2i RemovedTris;			// the triangles that were removed in the collapse (second is InvalidID for boundary edge)
-	FIndex2i RemovedEdges;			// the edges that were removed (second is InvalidID for boundary edge)
-	FIndex2i KeptEdges;				// the edges that were kept (second is InvalidID for boundary edge)
-
-	double CollapseT;				// interpolation parameter along edge for new vertex in range [0,1]
-};
-
-/** Information about mesh elements modified by MergeEdges() */
-struct FMergeEdgesInfo
-{
-	int KeptEdge;				// the edge that was kept
-	int RemovedEdge;			// the edge that was removed
-
-	FIndex2i KeptVerts;			// The two vertices that were kept (redundant w/ KeptEdge?)
-	FIndex2i RemovedVerts;		// The removed vertices of RemovedEdge. Either may be InvalidID if it was same as the paired KeptVert
-
-	FIndex2i ExtraRemovedEdges; // extra removed edges, see description below. Either may be or InvalidID
-	FIndex2i ExtraKeptEdges;	// extra kept edges, paired with ExtraRemovedEdges
-};
-
-/** Information about mesh elements modified/created by PokeTriangle() */
-struct FPokeTriangleInfo
-{
-	int OriginalTriangle;				// the triangle that was poked
-	FIndex3i TriVertices;				// vertices of the original triangle
-
-	int NewVertex;						// the new vertex that was inserted
-	FIndex2i NewTriangles;				// the two new triangles that were added (OriginalTriangle is re-used, see code for vertex orders)
-	FIndex3i NewEdges;					// the three new edges connected to NewVertex
-
-	FVector3d BaryCoords;				// barycentric coords that NewVertex was inserted at
-};
-
-/** Information about mesh elements modified/created by SplitVertex() */
-struct FVertexSplitInfo
-{
-	int OriginalVertex;
-	int NewVertex;
-	// if needed could possibly add information about added edges?  but it would be a dynamic array, and there is no use for it yet.
-	// modified triangles are passed as input to the function, no need to store those here.
-};
-}
+#include "DynamicMesh/InfoTypes.h"
\ No newline at end of file
