GeometryProcessing: Move MeshConversion module from GeometryProcessing plugin to Runtime. Transfer necessary additional files from DynamicMesh to GeometryCore. Move functions from MeshIndexUtil and MeshAdapterUtil namespaces to UE::Geometry namespace.
#rnx

diff --git a/Engine/Plugins/Experimental/Enterprise/DataprepGeometryOperations/Source/Private/DataprepGeometryOperations.cpp b/Engine/Plugins/Experimental/Enterprise/DataprepGeometryOperations/Source/Private/DataprepGeometryOperations.cpp
--- a/Engine/Plugins/Experimental/Enterprise/DataprepGeometryOperations/Source/Private/DataprepGeometryOperations.cpp
+++ b/Engine/Plugins/Experimental/Enterprise/DataprepGeometryOperations/Source/Private/DataprepGeometryOperations.cpp
@@ -18,7 +18,7 @@
 #include "MeshAdapterTransforms.h"
 #include "MeshDescriptionAdapter.h"
 #include "MeshDescriptionToDynamicMesh.h"
-#include "Operations/MergeCoincidentMeshEdges.h"
+#include "DynamicMesh/Operations/MergeCoincidentMeshEdges.h"
 #include "PhysicsEngine/BodySetup.h"
 
 #define LOCTEXT_NAMESPACE "DatasmithEditingOperationsExperimental"
diff --git a/Engine/Plugins/Experimental/GeometryFlow/Source/GeometryFlowMeshProcessing/Private/PhysicsNodes/GenerateSimpleCollisionNode.cpp b/Engine/Plugins/Experimental/GeometryFlow/Source/GeometryFlowMeshProcessing/Private/PhysicsNodes/GenerateSimpleCollisionNode.cpp
--- a/Engine/Plugins/Experimental/GeometryFlow/Source/GeometryFlowMeshProcessing/Private/PhysicsNodes/GenerateSimpleCollisionNode.cpp
+++ b/Engine/Plugins/Experimental/GeometryFlow/Source/GeometryFlowMeshProcessing/Private/PhysicsNodes/GenerateSimpleCollisionNode.cpp
@@ -2,7 +2,7 @@
 
 #include "PhysicsNodes/GenerateSimpleCollisionNode.h"
 #include "Operations/MeshConvexHull.h"
-#include "MeshIndexUtil.h"
+#include "DynamicMesh/MeshIndexUtil.h"
 #include "Async/ParallelFor.h"
 #include "DynamicSubmesh3.h"
 #include "Util/ProgressCancel.h"
@@ -68,12 +68,12 @@ namespace GenerateSimpleCollisionNodeHelpers
 					TArray<int> SampledVertices;
 					FMeshConvexHull::GridSample(Mesh, Settings.ConvexHullSettings.PrefilterGridResolution, SampledVertices);
 					TArray<int> IndexSetVertices;
-					MeshIndexUtil::TriangleToVertexIDs(&Mesh, IndexData.IndexSets[k], IndexSetVertices);
+					UE::Geometry::TriangleToVertexIDs(&Mesh, IndexData.IndexSets[k], IndexSetVertices);
 					Hull.VertexSet = SetIntersection(SampledVertices, IndexSetVertices);
 				}
 				else
 				{
-					MeshIndexUtil::TriangleToVertexIDs(&Mesh, IndexData.IndexSets[k], Hull.VertexSet);
+					UE::Geometry::TriangleToVertexIDs(&Mesh, IndexData.IndexSets[k], Hull.VertexSet);
 				}
 
 				Hull.bPostSimplify = Settings.ConvexHullSettings.SimplifyToTriangleCount > 0;
diff --git a/Engine/Plugins/Experimental/GeometryProcessing/GeometryProcessing.uplugin b/Engine/Plugins/Experimental/GeometryProcessing/GeometryProcessing.uplugin
--- a/Engine/Plugins/Experimental/GeometryProcessing/GeometryProcessing.uplugin
+++ b/Engine/Plugins/Experimental/GeometryProcessing/GeometryProcessing.uplugin
@@ -28,11 +28,6 @@
 			"Name": "DynamicMesh",
 			"Type": "Runtime",
 			"LoadingPhase": "Default"
-		},
-		{
-			"Name": "MeshConversion",
-			"Type": "Runtime",
-			"LoadingPhase": "Default"
 		}
 	]
 }
\ No newline at end of file
diff --git a/Engine/Plugins/Experimental/GeometryProcessing/Source/DynamicMesh/Private/Operations/ExtrudeMesh.cpp b/Engine/Plugins/Experimental/GeometryProcessing/Source/DynamicMesh/Private/Operations/ExtrudeMesh.cpp
--- a/Engine/Plugins/Experimental/GeometryProcessing/Source/DynamicMesh/Private/Operations/ExtrudeMesh.cpp
+++ b/Engine/Plugins/Experimental/GeometryProcessing/Source/DynamicMesh/Private/Operations/ExtrudeMesh.cpp
@@ -1,10 +1,10 @@
 // Copyright Epic Games, Inc. All Rights Reserved.
 
 #include "Operations/ExtrudeMesh.h"
-#include "MeshNormals.h"
+#include "DynamicMesh/MeshNormals.h"
 #include "DynamicMeshEditor.h"
 #include "Selections/MeshConnectedComponents.h"
-#include "MeshIndexUtil.h"
+#include "DynamicMesh/MeshIndexUtil.h"
 
 #include "ExplicitUseGeometryMathTypes.h"		// using UE::Geometry::(math types)
 using namespace UE::Geometry;
@@ -64,7 +64,7 @@ bool FExtrudeMesh::ApplyExtrude(FExtrusionInfo& Region, FMeshNormals* UseNormals
 		return true;
 	}
 
-	MeshIndexUtil::TriangleToVertexIDs(Mesh, Region.InitialTriangles, Region.InitialVertices);
+	UE::Geometry::TriangleToVertexIDs(Mesh, Region.InitialTriangles, Region.InitialVertices);
 
 	// duplicate triangles of mesh
 
diff --git a/Engine/Plugins/Experimental/GeometryProcessing/Source/DynamicMesh/Private/Operations/GroupEdgeInserter.cpp b/Engine/Plugins/Experimental/GeometryProcessing/Source/DynamicMesh/Private/Operations/GroupEdgeInserter.cpp
--- a/Engine/Plugins/Experimental/GeometryProcessing/Source/DynamicMesh/Private/Operations/GroupEdgeInserter.cpp
+++ b/Engine/Plugins/Experimental/GeometryProcessing/Source/DynamicMesh/Private/Operations/GroupEdgeInserter.cpp
@@ -8,7 +8,7 @@
 #include "DynamicMeshChangeTracker.h"
 #include "DynamicMeshEditor.h"
 #include "FrameTypes.h"
-#include "MeshIndexUtil.h"
+#include "DynamicMesh/MeshIndexUtil.h"
 #include "MeshRegionBoundaryLoops.h"
 #include "Operations/GroupEdgeInserter.h"
 #include "Operations/MeshPlaneCut.h"
@@ -784,7 +784,7 @@ bool DeleteGroupTrianglesAndGetLoop(FDynamicMesh3& Mesh, const FGroupTopology& T
 
 	// Get all verts in the component, and the verts on the boundary
 	TArray<int32> ComponentVids;
-	MeshIndexUtil::TriangleToVertexIDs(&Mesh, Component.Indices, ComponentVids);
+	UE::Geometry::TriangleToVertexIDs(&Mesh, Component.Indices, ComponentVids);
 	TSet<int32> BoundaryVidSet(RegionLoops.Loops[0].Vertices);
 
 	// Delete the triangles
diff --git a/Engine/Plugins/Experimental/GeometryProcessing/Source/DynamicMesh/Private/Operations/JoinMeshLoops.cpp b/Engine/Plugins/Experimental/GeometryProcessing/Source/DynamicMesh/Private/Operations/JoinMeshLoops.cpp
--- a/Engine/Plugins/Experimental/GeometryProcessing/Source/DynamicMesh/Private/Operations/JoinMeshLoops.cpp
+++ b/Engine/Plugins/Experimental/GeometryProcessing/Source/DynamicMesh/Private/Operations/JoinMeshLoops.cpp
@@ -1,9 +1,9 @@
 // Copyright Epic Games, Inc. All Rights Reserved.
 
 #include "Operations/JoinMeshLoops.h"
-#include "MeshNormals.h"
+#include "DynamicMesh/MeshNormals.h"
 #include "DynamicMeshEditor.h"
-#include "MeshIndexUtil.h"
+#include "DynamicMesh/MeshIndexUtil.h"
 
 #include "ExplicitUseGeometryMathTypes.h"		// using UE::Geometry::(math types)
 using namespace UE::Geometry;
diff --git a/Engine/Plugins/Experimental/GeometryProcessing/Source/DynamicMesh/Private/Parameterization/DynamicMeshUVEditor.cpp b/Engine/Plugins/Experimental/GeometryProcessing/Source/DynamicMesh/Private/Parameterization/DynamicMeshUVEditor.cpp
--- a/Engine/Plugins/Experimental/GeometryProcessing/Source/DynamicMesh/Private/Parameterization/DynamicMeshUVEditor.cpp
+++ b/Engine/Plugins/Experimental/GeometryProcessing/Source/DynamicMesh/Private/Parameterization/DynamicMeshUVEditor.cpp
@@ -4,9 +4,9 @@
 #include "Parameterization/DynamicMeshUVEditor.h"
 #include "DynamicSubmesh3.h"
 
-#include "MeshNormals.h"
+#include "DynamicMesh/MeshNormals.h"
 #include "MeshBoundaryLoops.h"
-#include "MeshIndexUtil.h"
+#include "DynamicMesh/MeshIndexUtil.h"
 
 #include "Parameterization/MeshDijkstra.h"
 #include "Parameterization/MeshLocalParam.h"
@@ -375,7 +375,7 @@ static bool FindSeamTriSplitSets_BoundaryVtx(const FDynamicMesh3* Mesh, const FD
 		bool bDone = false;
 		while (!bDone)
 		{
-			FIndex3i NextTri = UE::MeshIndexUtil::FindNextAdjacentTriangleAroundVtx(Mesh, SplitVtx, CurTri, PrevTri,
+			FIndex3i NextTri = UE::Geometry::FindNextAdjacentTriangleAroundVtx(Mesh, SplitVtx, CurTri, PrevTri,
 				[&](int32 Tri0, int32 Tri1, int32 Edge) { return UVOverlay->AreTrianglesConnected(Tri0, Tri1); }
 			);
 			if (NextTri.A != IndexConstants::InvalidID)
@@ -425,7 +425,7 @@ static bool FindSeamTriSplitSets_InteriorVtx(const FDynamicMesh3* Mesh, const FD
 		bool bDone = false;
 		while (!bDone)
 		{
-			FIndex3i NextTri = UE::MeshIndexUtil::FindNextAdjacentTriangleAroundVtx(Mesh, SplitVtx, CurTri, PrevTri,
+			FIndex3i NextTri = UE::Geometry::FindNextAdjacentTriangleAroundVtx(Mesh, SplitVtx, CurTri, PrevTri,
 				[&](int32 Tri0, int32 Tri1, int32 Edge) { return UVOverlay->AreTrianglesConnected(Tri0, Tri1) && Edge != PrevBaseEdgeID && Edge != NextBaseEdgeID; }
 			);
 			if (NextTri.A != IndexConstants::InvalidID)
diff --git a/Engine/Plugins/Experimental/GeometryProcessing/Source/DynamicMesh/Public/MeshIndexUtil.h b/Engine/Plugins/Experimental/GeometryProcessing/Source/DynamicMesh/Public/MeshIndexUtil.h
--- a/Engine/Plugins/Experimental/GeometryProcessing/Source/DynamicMesh/Public/MeshIndexUtil.h
+++ b/Engine/Plugins/Experimental/GeometryProcessing/Source/DynamicMesh/Public/MeshIndexUtil.h
@@ -2,102 +2,4 @@
 
 #pragma once
 
-#include "DynamicMesh3.h"
-
-/**
- * Utility functions for dealing with mesh indices
- */
-namespace MeshIndexUtil
-{
-	using namespace UE::Geometry;
-
-	//
-	// Do not add any more functions to this namespace. Use UE::MeshIndexUtil instead
-	//
-
-	/**
-	 * Find list of unique vertices that are contained in one or more triangles
-	 * @param Mesh input mesh
-	 * @param TriangleIDs list of triangle IDs of Mesh
-	 * @param VertexIDsOut list of vertices contained by triangles
-	 */
-	DYNAMICMESH_API void TriangleToVertexIDs(const FDynamicMesh3* Mesh, const TArray<int>& TriangleIDs, TArray<int>& VertexIDsOut);
-
-
-	/**
-	 * Find all the triangles in all the one rings of a set of vertices
-	 * @param Mesh input mesh
-	 * @param VertexIDs list of Vertex IDs of Mesh
-	 * @param TriangleIDsOut list of triangle IDs where any vertex of triangle touches an element of VertexIDs
-	 */
-	DYNAMICMESH_API void VertexToTriangleOneRing(const FDynamicMesh3* Mesh, const TArray<int>& VertexIDs, TSet<int>& TriangleIDsOut);
-
-}
-
-namespace UE
-{
-	namespace MeshIndexUtil
-	{
-		using namespace UE::Geometry;
-
-		/**
-		 * Find list of unique vertices that are contained in one or more triangles
-		 * @param Mesh input mesh
-		 * @param TriangleIDs list of triangle IDs of Mesh
-		 * @param VertexIDsOut list of vertices contained by triangles
-		 */
-		DYNAMICMESH_API void TriangleToVertexIDs(const FDynamicMesh3* Mesh, const TArray<int>& TriangleIDs, TArray<int>& VertexIDsOut);
-
-
-		/**
-		 * Find all the triangles in all the one rings of a set of vertices
-		 * @param Mesh input mesh
-		 * @param VertexIDs list of Vertex IDs of Mesh
-		 * @param TriangleIDsOut list of triangle IDs where any vertex of triangle touches an element of VertexIDs
-		 */
-		DYNAMICMESH_API void VertexToTriangleOneRing(const FDynamicMesh3* Mesh, const TArray<int>& VertexIDs, TSet<int>& TriangleIDsOut);
-
-		/**
-		 * Walk around VertexID from FromTriangleID to next connected triangle if it exists, walking "away" from PrevTriangleID.
-		 * @param TrisConnectedFunc returns true if two triangles should be considered connected, to support breaking at seams/etc that are not in base mesh topology
-		 * @return triplet of values (FoundTriangleID, SharedEdgeID, IndexOfEdgeInFromTri), or all IndexConstants::InvalidID if not found
-		 */
-		template<typename TrisConnectedPredicate>
-		FIndex3i FindNextAdjacentTriangleAroundVtx(const FDynamicMesh3* Mesh, 
-			int32 VertexID, int32 FromTriangleID, int32 PrevTriangleID,
-			TrisConnectedPredicate TrisConnectedTest)
-		{
-			check(Mesh);
-
-			// find neigbour edges and tris for triangle
-			FIndex3i TriEdges = Mesh->GetTriEdges(FromTriangleID);
-			FIndex3i TriNbrTris;
-			for (int32 j = 0; j < 3; ++j)
-			{
-				FIndex2i EdgeT = Mesh->GetEdgeT(TriEdges[j]);
-				TriNbrTris[j] = (EdgeT.A == FromTriangleID) ? EdgeT.B : EdgeT.A;
-			}
-
-			// Search for the neighbour tri that is not PrevTriangleID, and is also connected to VertexID.
-			// This is our next triangle around the ring
-			for (int32 j = 0; j < 3; ++j)
-			{
-				if (TriNbrTris[j] != PrevTriangleID && Mesh->IsTriangle(TriNbrTris[j]))
-				{
-					FIndex3i TriVerts = Mesh->GetTriangle(TriNbrTris[j]);
-					if (TriVerts.A == VertexID || TriVerts.B == VertexID || TriVerts.C == VertexID)
-					{
-						// test if predicate allows this connection
-						if (TrisConnectedTest(FromTriangleID, TriNbrTris[j], TriEdges[j]) == false)
-						{
-							break;
-						}
-						return FIndex3i(TriNbrTris[j], TriEdges[j], j);
-					}
-				}
-			}
-
-			return FIndex3i(IndexConstants::InvalidID, IndexConstants::InvalidID, IndexConstants::InvalidID);
-		}
-	}
-}
+#include "DynamicMesh/MeshIndexUtil.h"
diff --git a/Engine/Plugins/Experimental/GeometryProcessing/Source/DynamicMesh/Public/MeshNormals.h b/Engine/Plugins/Experimental/GeometryProcessing/Source/DynamicMesh/Public/MeshNormals.h
--- a/Engine/Plugins/Experimental/GeometryProcessing/Source/DynamicMesh/Public/MeshNormals.h
+++ b/Engine/Plugins/Experimental/GeometryProcessing/Source/DynamicMesh/Public/MeshNormals.h
@@ -1,217 +1,5 @@
 // Copyright Epic Games, Inc. All Rights Reserved.
 
-// Port of geometry3Sharp MeshNormals
-
 #pragma once
 
-#include "DynamicMesh3.h"
-#include "DynamicMeshAttributeSet.h"
-
-namespace UE
-{
-namespace Geometry
-{
-
-
-/**
- * FMeshNormals is a utility class that can calculate and store various types of
- * normal vectors for a FDynamicMesh. 
- */
-class DYNAMICMESH_API FMeshNormals
-{
-protected:
-	/** Target Mesh */
-	const FDynamicMesh3* Mesh;
-	/** Set of computed normals */
-	TArray<FVector3d> Normals;
-
-public:
-	FMeshNormals()
-	{
-		Mesh = nullptr;
-	}
-
-	FMeshNormals(const FDynamicMesh3* Mesh)
-	{
-		SetMesh(Mesh);
-	}
-
-
-	void SetMesh(const FDynamicMesh3* MeshIn)
-	{
-		this->Mesh = MeshIn;
-	}
-
-	const TArray<FVector3d>& GetNormals() const { return Normals; }
-
-	FVector3d& operator[](int i) { return Normals[i]; }
-	const FVector3d& operator[](int i) const { return Normals[i]; }
-
-
-	/**
-	 * Set the size of the Normals array to Count, and optionally clear all values to (0,0,0)
-	 */
-	void SetCount(int Count, bool bClearToZero);
-
-	/**
-	 * Compute standard per-vertex normals by averaging one-ring face normals
-	 */
-	void ComputeVertexNormals(bool bWeightByArea = true, bool bWeightByAngle = true)
-	{
-		Compute_FaceAvg(bWeightByArea, bWeightByAngle);
-	}
-
-	/**
-	 * Compute per-triangle normals
-	 */
-	void ComputeTriangleNormals()
-	{
-		Compute_Triangle();
-	}
-
-	/**
-	 * Assumes that ComputeTriangleNormals() has already been run.
-	 * 
-	 * Looks through the normals for any that are zero, and tries to set them to the normal of 
-	 * some neighboring triangle, preferring the triangle on its longer side if possible. This is 
-	 * useful for finding connected components based on normals, as it prevents degenerate triangles
-	 * with 0 normals from connecting otherwise disconnected components.
-	 * This won't help if a whole mesh componenent is entirely made of such degenerate triangles
-	 * (which is usually ok).
-	 */
-	void SetDegenerateTriangleNormalsToNeighborNormal();
-
-	/**
-	 * Recompute the per-element normals of the given overlay by averaging one-ring face normals
-	 * @warning NormalOverlay must be attached to ParentMesh or an exact copy
-	 */
-	void RecomputeOverlayNormals(const FDynamicMeshNormalOverlay* NormalOverlay, bool bWeightByArea = true, bool bWeightByAngle = true)
-	{
-		Compute_Overlay_FaceAvg(NormalOverlay, bWeightByArea, bWeightByAngle);
-	}
-
-
-	/**
-	 * Copy the current set of normals to the vertex normals of SetMesh
-	 * @warning assumes that the computed normals are vertex normals
-	 * @param bInvert if true, normals are flipped
-	 */
-	void CopyToVertexNormals(FDynamicMesh3* SetMesh, bool bInvert = false) const;
-
-	/**
-	 * Copy the current set of normals to the NormalOverlay attribute layer
-	 * @warning assumes that the computed normals are attribute normals
-	 * @param bInvert if true, normals are flipped
-	 */
-	void CopyToOverlay(FDynamicMeshNormalOverlay* NormalOverlay, bool bInvert = false) const;
-
-
-	/**
-	 * Compute per-vertex normals for the given Mesh
-	 * @param bInvert if true, normals are flipped
-	 */
-	static void QuickComputeVertexNormals(FDynamicMesh3& Mesh, bool bInvert = false);
-
-	/**
-	 * Compute per-vertex normals for the vertices of a set of triangles of a Mesh
-	 * @param bWeightByArea weight neighbor triangles by area
-	 * @param bWeightByAngle weight neighbor triangles by angle
-	 * @param bInvert if true, normals are flipped
-	 */
-	static void QuickComputeVertexNormalsForTriangles(FDynamicMesh3& Mesh, const TArray<int32>& Triangles, bool bWeightByArea = true, bool bWeightByAngle = true, bool bInvert = false);
-
-
-	/**
-	 * Compute normal at mesh vertex by weighted sum of one-ring triangle normals. Can optionally weight by area, angle, or both (averaged)
-	 * @param bWeightByArea weight neighbor triangles by area
-	 * @param bWeightByAngle weight neighbor triangles by angle
-	 * @return the vertex normal at vertex VertIdx of Mesh.
-	 */
-	static FVector3d ComputeVertexNormal(const FDynamicMesh3& Mesh, int VertIdx, bool bWeightByArea = true, bool bWeightByAngle = true);
-
-	/**
-	 * Compute normal at mesh vertex by weighted sum of subset of one-ring triangle normals. Can optionally weight by area, angle, or both (averaged)
-	 * @param TriangleFilterFunc Only one-ring triangles for which this function returns true will be included
-	 * @param bWeightByArea weight neighbor triangles by area
-	 * @param bWeightByAngle weight neighbor triangles by angle
-	 * @return the vertex normal at vertex VertIdx of Mesh.
-	 */
-	static FVector3d ComputeVertexNormal(const FDynamicMesh3& Mesh, int32 VertIdx, TFunctionRef<bool(int32)> TriangleFilterFunc, bool bWeightByArea = true, bool bWeightByAngle = true);
-
-
-	/**
-	 * @return the computed overlay normal at an element of the overlay (ie based on normals of triangles connected to this element)
-	 */
-	static FVector3d ComputeOverlayNormal(const FDynamicMesh3& Mesh, const FDynamicMeshNormalOverlay* NormalOverlay, int ElemIdx);
-
-	/**
-	 * Initialize the given NormalOverlay with per-vertex normals, ie single overlay element for each mesh vertex.
-	 * @param bUseMeshVertexNormalsIfAvailable if true and the parent mesh has per-vertex normals, use them instead of calculating new ones
-	 */
-	static void InitializeOverlayToPerVertexNormals(FDynamicMeshNormalOverlay* NormalOverlay, bool bUseMeshVertexNormalsIfAvailable = true);
-
-
-	/**
-	 * Initialize the given NormalOverlay with per-face normals, ie separate overlay element for each vertex of each triangle.
-	 */
-	static void InitializeOverlayToPerTriangleNormals(FDynamicMeshNormalOverlay* NormalOverlay);
-
-
-	static void InitializeOverlayTopologyFromOpeningAngle(const FDynamicMesh3* Mesh, FDynamicMeshNormalOverlay* NormalOverlay, 
-		double AngleThresholdDeg);
-
-	static void InitializeOverlayTopologyFromFaceGroups(const FDynamicMesh3* Mesh, FDynamicMeshNormalOverlay* NormalOverlay);
-
-
-	/**
-	 * Initialize the given Mesh with per-face normals, ie separate overlay element for each vertex of each triangle.
-	 */
-	static void InitializeMeshToPerTriangleNormals(FDynamicMesh3* Mesh);
-
-
-	/**
-	 * Initialize the given triangles of NormalOverlay with per-vertex normals, ie single overlay element for each mesh vertex.
-	 * Only the triangles included in the region are considered when calculating per-vertex normals.
-	 */
-	static void InitializeOverlayRegionToPerVertexNormals(FDynamicMeshNormalOverlay* NormalOverlay, const TArray<int32>& Triangles);
-
-	/**
-	 * Compute overlay normals for the given mesh
-	 * @param bInvert if true, normals are flipped
-	 */
-	static bool QuickRecomputeOverlayNormals(FDynamicMesh3& Mesh, bool bInvert = false, bool bWeightByArea = true, bool bWeightByAngle = true);
-
-
-protected:
-	/** Compute per-vertex normals using area-weighted averaging of one-ring triangle normals */
-	void Compute_FaceAvg_AreaWeighted();
-	/** Compute per-vertex normals using a custom combination of area-weighted and angle-weighted averaging of one-ring triangle normals */
-	void Compute_FaceAvg(bool bWeightByArea, bool bWeightByAngle);
-
-	/** Compute per-triangle normals */
-	void Compute_Triangle();
-
-	/** Recompute the element Normals of the given attribute overlay using area-weighted averaging of one-ring triangle normals */
-	void Compute_Overlay_FaceAvg_AreaWeighted(const FDynamicMeshNormalOverlay* NormalOverlay);
-	/** Recompute the element Normals of the given attribute overlay using a custom combination of area-weighted and angle-weighted averaging of one-ring triangle normals */
-	void Compute_Overlay_FaceAvg(const FDynamicMeshNormalOverlay* NormalOverlay, bool bWeightByArea, bool bWeightByAngle);
-
-	static FVector3d GetVertexWeightsOnTriangle(const FDynamicMesh3* Mesh, int TriID, double TriArea, bool bWeightByArea, bool bWeightByAngle)
-	{
-		FVector3d TriNormalWeights = FVector3d::One();
-		if (bWeightByAngle)
-		{
-			TriNormalWeights = Mesh->GetTriInternalAnglesR(TriID); // component-wise multiply by per-vertex internal angles
-		}
-		if (bWeightByArea)
-		{
-			TriNormalWeights *= TriArea;
-		}
-		return TriNormalWeights;
-	}
-
-};
-
-
-} // end namespace UE::Geometry
-} // end namespace UE
\ No newline at end of file
+#include "DynamicMesh/MeshNormals.h"
diff --git a/Engine/Plugins/Experimental/GeometryProcessing/Source/DynamicMesh/Public/MeshTangents.h b/Engine/Plugins/Experimental/GeometryProcessing/Source/DynamicMesh/Public/MeshTangents.h
--- a/Engine/Plugins/Experimental/GeometryProcessing/Source/DynamicMesh/Public/MeshTangents.h
+++ b/Engine/Plugins/Experimental/GeometryProcessing/Source/DynamicMesh/Public/MeshTangents.h
@@ -1,236 +1,5 @@
 // Copyright Epic Games, Inc. All Rights Reserved.
 
 #pragma once
-
-#include "DynamicMesh3.h"
-#include "DynamicMeshAttributeSet.h"
-
-
-namespace UE
-{
-namespace Geometry
-{
-
-
-/**
- * Options used by TMeshTangents for tangents computation
- */
-struct DYNAMICMESH_API FComputeTangentsOptions
-{
-	/** if true, per-face tangents are blended to create triange-vertex tangents */
-	bool bAveraged = true;		
-	/** if true, per-face tangents are weighted by opening angle at triangle-vertices when blending */
-	bool bAngleWeighted = true;
-};
-
-
-/**
- * TMeshTangents is a utility class that can calculate and store various types of
- * tangent vectors for a FDynamicMesh.
- */
-template<typename RealType>
-class TMeshTangents
-{
-protected:
-	/** Target Mesh */
-	const FDynamicMesh3* Mesh;
-	/** Set of computed tangents */
-	TArray<FVector3<RealType>> Tangents;
-	/** Set of computed bitangents */
-	TArray<FVector3<RealType>> Bitangents;
-	/** Indices of degenerate triangles. This may be useful to know externally, as those tangets are often meaningless */
-	TArray<int32> AllDegenerateTris;
-
-public:
-	TMeshTangents()
-	{
-		Mesh = nullptr;
-	}
-
-	TMeshTangents(const FDynamicMesh3* Mesh)
-	{
-		SetMesh(Mesh);
-	}
-
-	void SetMesh(const FDynamicMesh3* MeshIn)
-	{
-		this->Mesh = MeshIn;
-	}
-
-	const TArray<FVector3<RealType>>& GetTangents() const { return Tangents; }
-
-	const TArray<FVector3<RealType>>& GetBitangents() const { return Bitangents; }
-
-	const TArray<int32>& GetDegenerateTris() const { return AllDegenerateTris; }
-
-
-	//
-	// Per-Triangle-Vertex Tangents
-	// Number of Tangents is NumTriangles * 3
-	//
-
-	/**
-	 * Set internal buffer sizes suitable for calculating per-triangle tangents.
-	 * This is intended to be used if you wish to calculate your own tangents and use SetPerTriangleTangent()
-	 */
-	void InitializeTriVertexTangents(bool bClearToZero)
-	{
-		SetTangentCount(Mesh->MaxTriangleID() * 3, bClearToZero);
-	}
-
-	/**
-	 * Initialize Tangents from other Tangents set
-	 */
-	template<typename OtherRealType>
-	void CopyTriVertexTangents(const TMeshTangents<OtherRealType>& OtherTangents);
-
-	/**
-	 * Initialize Tangents from the FDynamicMeshAttributeSet of SourceMesh 
-	 * @return false if tangent attributes do not exist
-	 */
-	bool CopyTriVertexTangents(const FDynamicMesh3& SourceMesh);
-
-
-	/**
-	 * Calculate per-triangle tangent spaces based on the given per-triangle normal and UV overlays.
-	 * In this mode there is no averaging of tangents across triangles. So if we have N triangles
-	 * in the mesh, then 3*N tangents are generated. These tangents are computed in parallel.
-	 */
-	void ComputeTriVertexTangents(const FDynamicMeshNormalOverlay* NormalOverlay, const FDynamicMeshUVOverlay* UVOverlay, const FComputeTangentsOptions& Options);
-
-	/**
-	 * Return tangent and bitangent at a vertex of triangle for per-triangle computed tangents
-	 * @param TriangleID triangle index in mesh
-	 * @param TriVertIdx vertex index in range 0,1,2
-	 */
-	void GetPerTriangleTangent(int32 TriangleID, int32 TriVertIdx, UE::Math::TVector<RealType>& TangentOut, UE::Math::TVector<RealType>& BitangentOut) const
-	{
-		int32 k = TriangleID * 3 + TriVertIdx;
-		TangentOut = Tangents[k];
-		BitangentOut = Bitangents[k];
-	}
-
-
-	/**
-	 * Return tangent and bitangent at a vertex of triangle for per-triangle computed tangents
-	 * @param TriangleID triangle index in mesh
-	 * @param TriVertIdx vertex index in range 0,1,2
-	 */
-	template<typename VectorType, typename OtherRealType>
-	void GetPerTriangleTangent(int32 TriangleID, int32 TriVertIdx, VectorType& TangentOut, VectorType& BitangentOut) const
-	{
-		int32 k = TriangleID * 3 + TriVertIdx;
-		TangentOut.X = (OtherRealType)Tangents[k].X;
-		TangentOut.Y = (OtherRealType)Tangents[k].Y;
-		TangentOut.Z = (OtherRealType)Tangents[k].Z;
-		BitangentOut.X = (OtherRealType)Bitangents[k].X;
-		BitangentOut.Y = (OtherRealType)Bitangents[k].Y;
-		BitangentOut.Z = (OtherRealType)Bitangents[k].Z;
-	}
-
-
-	/**
-	 * @return Interpolated tangent and bitangent at barycentric coordinates inside a triangle, for per-triangle computed tangents
-	 * @param TriangleID triangle index in mesh
-	 * @param BaryCoords barycentric coordinates in triangle
-	 * @param TangentOut interpolated tangent
-	 * @param BitangentOut interpolated bitangent
-	 */
-	void GetInterpolatedTriangleTangent(int32 TriangleID, 
-		const UE::Math::TVector<RealType>& BaryCoords,
-		UE::Math::TVector<RealType>& TangentOut,
-		UE::Math::TVector<RealType>& BitangentOut) const
-	{
-		int32 k = TriangleID * 3;
-		if (k >= 0 && (k+2) < Tangents.Num())
-		{
-			TangentOut = Blend3(Tangents[k], Tangents[k+1], Tangents[k+2], BaryCoords.X, BaryCoords.Y, BaryCoords.Z);
-			Normalize(TangentOut);
-			BitangentOut = Blend3(Bitangents[k], Bitangents[k+1], Bitangents[k+2], BaryCoords.X, BaryCoords.Y, BaryCoords.Z);
-			Normalize(BitangentOut);
-		}
-		else
-		{
-			TangentOut = FVector3<RealType>::UnitX();
-			BitangentOut = FVector3<RealType>::UnitY();
-		}
-	}
-
-
-	/**
-	 * Set tangent and bitangent at a vertex of triangle for per-triangle computed tangents.
-	 * @param TriangleID triangle index in mesh
-	 * @param TriVertIdx vertex index in range 0,1,2
-	 */
-	void SetPerTriangleTangent(int TriangleID, int TriVertIdx, const FVector3<RealType>& Tangent, const FVector3<RealType>& Bitangent)
-	{
-		int k = TriangleID * 3 + TriVertIdx;
-		Tangents[k] = Tangent;
-		Bitangents[k] = Bitangent;
-	}
-
-
-	//
-	// Per-Triangle Tangents
-	// Number of Tangents is NumTriangles
-	//
-
-	/** Initialize buffer sizes to one tangent/bitangent per mesh triangle */
-	void InitializeTriangleTangents(bool bClearToZero)
-	{
-		SetTangentCount(Mesh->MaxTriangleID(), bClearToZero);
-	}
-
-	/**
-	 * Compute per-triangle tangents for the given UV Overlay 
-	 * @param bOrthogonalize if true, we orthogonalize the tangents by computing the Bitangent sign and regenerating it. Otherwise the tangents are left non-orthogonal.
-	 */
-	void ComputeTriangleTangents(const FDynamicMeshUVOverlay* UVOverlay, bool bOrthogonalize = true);
-
-
-	/**
-	 * Set Tangents on mesh overlays
-	 * @param MeshToSet Mesh to copy overlays to; does not need to be the same as the Mesh member of this class
-	 */
-	bool CopyToOverlays(FDynamicMesh3& MeshToSet);
-
-protected:
-	/**
-	 * Set the size of the Tangents array to Count, and optionally clear all values to (0,0,0)
-	 */
-	void SetTangentCount(int Count, bool bClearToZero);
-
-	// calculate per-triangle tangents and then projected to overlay normals at each triangle-vertex
-	void ComputeSeparatePerTriangleTangents(const FDynamicMeshNormalOverlay* NormalOverlay, const FDynamicMeshUVOverlay* UVOverlay, const FComputeTangentsOptions& Options);
-
-	// Calculate MikkT-space-style tangents
-	void ComputeMikkTStyleTangents(const FDynamicMeshNormalOverlay* NormalOverlay, const FDynamicMeshUVOverlay* UVOverlay, const FComputeTangentsOptions& Options);
-};
-
-typedef TMeshTangents<float> FMeshTangentsf;
-typedef TMeshTangents<double> FMeshTangentsd;
-
-
-
-
-template<typename RealType>
-template<typename OtherRealType>
-void TMeshTangents<RealType>::CopyTriVertexTangents(const TMeshTangents<OtherRealType>& OtherMeshTangents)
-{
-	const TArray<FVector3<OtherRealType>>& OtherTangents = OtherMeshTangents.GetTangents();
-	const TArray<FVector3<OtherRealType>>& OtherBitangents = OtherMeshTangents.GetBitangents();
-
-	int32 Num = OtherTangents.Num();
-	check(Mesh->MaxTriangleID() * 3 == Num);
-	InitializeTriVertexTangents(false);
-	for (int32 k = 0; k < Num; ++k)
-	{
-		Tangents[k] = FVector3<RealType>(OtherTangents[k]);
-		Bitangents[k] = FVector3<RealType>(OtherBitangents[k]);
-	}
-	AllDegenerateTris = OtherMeshTangents.GetDegenerateTris();
-}
-
-
-} // end namespace UE::Geometry
-} // end namespace UE
\ No newline at end of file
+ 
+#include "DynamicMesh/MeshTangents.h"
diff --git a/Engine/Plugins/Experimental/GeometryProcessing/Source/DynamicMesh/Public/MeshTransforms.h b/Engine/Plugins/Experimental/GeometryProcessing/Source/DynamicMesh/Public/MeshTransforms.h
--- a/Engine/Plugins/Experimental/GeometryProcessing/Source/DynamicMesh/Public/MeshTransforms.h
+++ b/Engine/Plugins/Experimental/GeometryProcessing/Source/DynamicMesh/Public/MeshTransforms.h
@@ -1,63 +1,5 @@
 // Copyright Epic Games, Inc. All Rights Reserved.
 
-// Port of geometry3cpp MeshWeights
-
 #pragma once
 
-#include "DynamicMesh3.h"
-
-
-/**
- * Utility functions for applying transformations to meshes
- */
-namespace MeshTransforms
-{
-	using namespace UE::Geometry;
-	using UE::Geometry::FTransform3d;
-
-	/**
-	 * Apply Translation to vertex positions of Mesh. Does not modify any other attributes.
-	 */
-	DYNAMICMESH_API void Translate(FDynamicMesh3& Mesh, const FVector3d& Translation);
-
-	/**
-	 * Apply Scale to vertex positions of Mesh, relative to given Origin. Does not modify any other attributes (normals/etc)
-	 */
-	DYNAMICMESH_API void Scale(FDynamicMesh3& Mesh, const FVector3d& Scale, const FVector3d& Origin);
-
-	/**
-	 * Transform Mesh into local coordinates of Frame
-	 */
-	DYNAMICMESH_API void WorldToFrameCoords(FDynamicMesh3& Mesh, const FFrame3d& Frame);
-
-	/**
-	 * Transform Mesh out of local coordinates of Frame
-	 */
-	DYNAMICMESH_API void FrameCoordsToWorld(FDynamicMesh3& Mesh, const FFrame3d& Frame);
-
-
-	/**
-	 * Apply given Transform to a Mesh.
-	 * Modifies Vertex Positions and Normals, and any Per-Triangle Normal Overlays
-	 */
-	DYNAMICMESH_API void ApplyTransform(FDynamicMesh3& Mesh, const FTransform3d& Transform);
-
-
-	/**
-	 * Apply inverse of given Transform to a Mesh.
-	 * Modifies Vertex Positions and Normals, and any Per-Triangle Normal Overlays
-	 */
-	DYNAMICMESH_API void ApplyTransformInverse(FDynamicMesh3& Mesh, const FTransform3d& Transform);
-
-
-	/**
-	 * Apply given Transform to a Mesh.
-	 * Modifies Vertex Positions and Normals, and any Per-Triangle Normal Overlays
-	 */
-	DYNAMICMESH_API void ApplyTransform(FDynamicMesh3& Mesh, 
-		TFunctionRef<FVector3d(const FVector3d&)> PositionTransform,
-		TFunctionRef<FVector3f(const FVector3f&)> NormalTransform);
-
-
-
-};
\ No newline at end of file
+#include "DynamicMesh/MeshTransforms.h"
diff --git a/Engine/Plugins/Experimental/MeshModelingToolset/Source/MeshModelingTools/Private/ConversionUtils/VolumeToDynamicMesh.cpp b/Engine/Plugins/Experimental/MeshModelingToolset/Source/MeshModelingTools/Private/ConversionUtils/VolumeToDynamicMesh.cpp
--- a/Engine/Plugins/Experimental/MeshModelingToolset/Source/MeshModelingTools/Private/ConversionUtils/VolumeToDynamicMesh.cpp
+++ b/Engine/Plugins/Experimental/MeshModelingToolset/Source/MeshModelingTools/Private/ConversionUtils/VolumeToDynamicMesh.cpp
@@ -10,7 +10,7 @@
 #include "MeshQueries.h"
 #include "MeshRegionBoundaryLoops.h"
 #include "Model.h"
-#include "Operations/MergeCoincidentMeshEdges.h"
+#include "DynamicMesh/Operations/MergeCoincidentMeshEdges.h"
 #include "Operations/MinimalHoleFiller.h"
 #include "Operations/PlanarFlipsOptimization.h"
 
diff --git a/Engine/Plugins/Experimental/MeshModelingToolset/Source/MeshModelingTools/Private/DrawPolyPathTool.cpp b/Engine/Plugins/Experimental/MeshModelingToolset/Source/MeshModelingTools/Private/DrawPolyPathTool.cpp
--- a/Engine/Plugins/Experimental/MeshModelingToolset/Source/MeshModelingTools/Private/DrawPolyPathTool.cpp
+++ b/Engine/Plugins/Experimental/MeshModelingToolset/Source/MeshModelingTools/Private/DrawPolyPathTool.cpp
@@ -11,7 +11,7 @@
 #include "Intersection/IntersectionUtil.h"
 #include "Util/ColorConstants.h"
 #include "ToolSetupUtil.h"
-#include "MeshIndexUtil.h"
+#include "DynamicMesh/MeshIndexUtil.h"
 #include "Generators/RectangleMeshGenerator.h"
 #include "Distance/DistLine3Line3.h"
 #include "AssetGenerationUtil.h"
diff --git a/Engine/Plugins/Experimental/MeshModelingToolset/Source/MeshModelingTools/Private/DynamicMeshSculptTool.cpp b/Engine/Plugins/Experimental/MeshModelingToolset/Source/MeshModelingTools/Private/DynamicMeshSculptTool.cpp
--- a/Engine/Plugins/Experimental/MeshModelingToolset/Source/MeshModelingTools/Private/DynamicMeshSculptTool.cpp
+++ b/Engine/Plugins/Experimental/MeshModelingToolset/Source/MeshModelingTools/Private/DynamicMeshSculptTool.cpp
@@ -12,8 +12,8 @@
 #include "MeshConstraints.h"
 #include "MeshConstraintsUtil.h"
 #include "MeshWeights.h"
-#include "MeshNormals.h"
-#include "MeshIndexUtil.h"
+#include "DynamicMesh/MeshNormals.h"
+#include "DynamicMesh/MeshIndexUtil.h"
 #include "Drawing/MeshDebugDrawing.h"
 #include "PreviewMesh.h"
 #include "ToolSetupUtil.h"
diff --git a/Engine/Plugins/Experimental/MeshModelingToolset/Source/MeshModelingTools/Private/EditMeshPolygonsTool.cpp b/Engine/Plugins/Experimental/MeshModelingToolset/Source/MeshModelingTools/Private/EditMeshPolygonsTool.cpp
--- a/Engine/Plugins/Experimental/MeshModelingToolset/Source/MeshModelingTools/Private/EditMeshPolygonsTool.cpp
+++ b/Engine/Plugins/Experimental/MeshModelingToolset/Source/MeshModelingTools/Private/EditMeshPolygonsTool.cpp
@@ -8,7 +8,7 @@
 #include "SegmentTypes.h"
 #include "DynamicMeshAttributeSet.h"
 #include "Mechanics/DragAlignmentMechanic.h"
-#include "MeshNormals.h"
+#include "DynamicMesh/MeshNormals.h"
 #include "ToolSceneQueriesUtil.h"
 #include "Intersection/IntersectionUtil.h"
 #include "Transforms/MultiTransformer.h"
@@ -26,7 +26,7 @@
 #include "DynamicMeshEditor.h"
 #include "DynamicMeshChangeTracker.h"
 #include "Changes/MeshChange.h"
-#include "MeshIndexUtil.h"
+#include "DynamicMesh/MeshIndexUtil.h"
 #include "MeshRegionBoundaryLoops.h"
 
 #include "Operations/OffsetMeshRegion.h"
@@ -1498,7 +1498,7 @@ void UEditMeshPolygonsTool::ApplyCutFaces()
 	FDynamicMeshChangeTracker ChangeTracker(Mesh);
 	ChangeTracker.BeginChange();
 	TArray<int32> VertexSelection;
-	MeshIndexUtil::TriangleToVertexIDs(Mesh, ActiveTriangleSelection, VertexSelection);
+	UE::Geometry::TriangleToVertexIDs(Mesh, ActiveTriangleSelection, VertexSelection);
 	ChangeTracker.SaveVertexOneRingTriangles(VertexSelection, true);
 
 	// apply the cut to edges of selected triangles
@@ -1801,7 +1801,7 @@ void UEditMeshPolygonsTool::ApplyRetriangulate()
 			// want to throw away boundary verts. However, this means that we'll have to go back
 			// through these vertices later to throw away isolated internal verts.
 			TArray<int32> OldVertices;
-			MeshIndexUtil::TriangleToVertexIDs(Mesh, Triangles, OldVertices);
+			UE::Geometry::TriangleToVertexIDs(Mesh, Triangles, OldVertices);
 			Editor.RemoveTriangles(Topology->GetGroupTriangles(GroupID), false);
 
 			RegionLoops.Loops[0].Reverse();
diff --git a/Engine/Plugins/Experimental/MeshModelingToolset/Source/MeshModelingTools/Private/MeshGroupPaintTool.cpp b/Engine/Plugins/Experimental/MeshModelingToolset/Source/MeshModelingTools/Private/MeshGroupPaintTool.cpp
--- a/Engine/Plugins/Experimental/MeshModelingToolset/Source/MeshModelingTools/Private/MeshGroupPaintTool.cpp
+++ b/Engine/Plugins/Experimental/MeshModelingToolset/Source/MeshModelingTools/Private/MeshGroupPaintTool.cpp
@@ -8,8 +8,8 @@
 #include "Async/Async.h"
 
 #include "MeshWeights.h"
-#include "MeshNormals.h"
-#include "MeshIndexUtil.h"
+#include "DynamicMesh/MeshNormals.h"
+#include "DynamicMesh/MeshIndexUtil.h"
 #include "Util/BufferUtil.h"
 #include "Util/ColorConstants.h"
 #include "Selections/MeshConnectedComponents.h"
@@ -1402,12 +1402,12 @@ void UMeshGroupPaintTool::OnDynamicMeshComponentChanged(USimpleDynamicMeshCompon
 		WaitForPendingUndoRedo();
 
 		// this is not right because now we are going to do extra recomputation, but it's very messy otherwise...
-		MeshIndexUtil::VertexToTriangleOneRing(Mesh, Change->Vertices, AccumulatedTriangleROI);
+		UE::Geometry::VertexToTriangleOneRing(Mesh, Change->Vertices, AccumulatedTriangleROI);
 	}
 	else
 	{
 		AccumulatedTriangleROI.Reset();
-		MeshIndexUtil::VertexToTriangleOneRing(Mesh, Change->Vertices, AccumulatedTriangleROI);
+		UE::Geometry::VertexToTriangleOneRing(Mesh, Change->Vertices, AccumulatedTriangleROI);
 	}
 
 	// note that we have a pending update
diff --git a/Engine/Plugins/Experimental/MeshModelingToolset/Source/MeshModelingTools/Private/MeshSelectionTool.cpp b/Engine/Plugins/Experimental/MeshModelingToolset/Source/MeshModelingTools/Private/MeshSelectionTool.cpp
--- a/Engine/Plugins/Experimental/MeshModelingToolset/Source/MeshModelingTools/Private/MeshSelectionTool.cpp
+++ b/Engine/Plugins/Experimental/MeshModelingToolset/Source/MeshModelingTools/Private/MeshSelectionTool.cpp
@@ -11,7 +11,7 @@
 #include "Util/ColorConstants.h"
 #include "Selections/MeshConnectedComponents.h"
 #include "MeshRegionBoundaryLoops.h"
-#include "MeshIndexUtil.h"
+#include "DynamicMesh/MeshIndexUtil.h"
 #include "AssetGenerationUtil.h"
 #include "ToolSetupUtil.h"
 #include "Selections/MeshConnectedComponents.h"
@@ -1048,7 +1048,7 @@ void UMeshSelectionTool::GrowShrinkSelection(bool bGrow)
 
 	const FDynamicMesh3* Mesh = PreviewMesh->GetPreviewDynamicMesh();
 	TArray<int32> Vertices;
-	MeshIndexUtil::TriangleToVertexIDs(Mesh, SelectedFaces, Vertices);
+	UE::Geometry::TriangleToVertexIDs(Mesh, SelectedFaces, Vertices);
 
 	TSet<int32> ChangeFaces;
 	for (int vid : Vertices)
diff --git a/Engine/Plugins/Experimental/MeshModelingToolset/Source/MeshModelingTools/Private/MeshVertexSculptTool.cpp b/Engine/Plugins/Experimental/MeshModelingToolset/Source/MeshModelingTools/Private/MeshVertexSculptTool.cpp
--- a/Engine/Plugins/Experimental/MeshModelingToolset/Source/MeshModelingTools/Private/MeshVertexSculptTool.cpp
+++ b/Engine/Plugins/Experimental/MeshModelingToolset/Source/MeshModelingTools/Private/MeshVertexSculptTool.cpp
@@ -8,8 +8,8 @@
 #include "Selections/MeshConnectedComponents.h"
 
 #include "MeshWeights.h"
-#include "MeshNormals.h"
-#include "MeshIndexUtil.h"
+#include "DynamicMesh/MeshNormals.h"
+#include "DynamicMesh/MeshIndexUtil.h"
 #include "Util/BufferUtil.h"
 #include "AssetUtils/Texture2DUtil.h"
 
@@ -1039,7 +1039,7 @@ void UMeshVertexSculptTool::OnDynamicMeshComponentChanged(USimpleDynamicMeshComp
 
 	// figure out the set of modified triangles
 	AccumulatedTriangleROI.Reset();
-	MeshIndexUtil::VertexToTriangleOneRing(Mesh, Change->Vertices, AccumulatedTriangleROI);
+	UE::Geometry::VertexToTriangleOneRing(Mesh, Change->Vertices, AccumulatedTriangleROI);
 
 	// start the normal recomputation
 	UndoNormalsFuture = Async(VertexSculptToolAsyncExecTarget, [this, Mesh]()
diff --git a/Engine/Plugins/Experimental/MeshModelingToolset/Source/MeshModelingTools/Private/Properties/MeshAnalysisProperties.cpp b/Engine/Plugins/Experimental/MeshModelingToolset/Source/MeshModelingTools/Private/Properties/MeshAnalysisProperties.cpp
--- a/Engine/Plugins/Experimental/MeshModelingToolset/Source/MeshModelingTools/Private/Properties/MeshAnalysisProperties.cpp
+++ b/Engine/Plugins/Experimental/MeshModelingToolset/Source/MeshModelingTools/Private/Properties/MeshAnalysisProperties.cpp
@@ -2,10 +2,10 @@
 
 #include "Properties/MeshAnalysisProperties.h"
 
-#include "DynamicMesh3.h"
+#include "DynamicMesh/DynamicMesh3.h"
 #include "MeshQueries.h"
 #include "MeshAdapter.h"
-#include "MeshAdapterUtil.h"
+#include "DynamicMesh/MeshAdapterUtil.h"
 
 #define LOCTEXT_NAMESPACE "UMeshAnalysisProperites"
 
@@ -13,7 +13,7 @@ using namespace UE::Geometry;
 
 void UMeshAnalysisProperties::Update(const FDynamicMesh3& MeshIn, const FTransform& Transform)
 {
-	FTriangleMeshAdapterd TransformedMesh = MeshAdapterUtil::MakeTransformedDynamicMeshAdapter(&MeshIn, Transform);
+	FTriangleMeshAdapterd TransformedMesh = UE::Geometry::MakeTransformedDynamicMeshAdapter(&MeshIn, Transform);
 	FVector2d VolArea = TMeshQueries<FTriangleMeshAdapterd>::GetVolumeArea(TransformedMesh);
 	this->SurfaceArea = FString::Printf(TEXT("%.2f m^2"), VolArea.Y / 10000);
 	this->Volume = FString::Printf(TEXT("%.2f m^3"), VolArea.X / 1000000);
diff --git a/Engine/Plugins/Experimental/MeshModelingToolset/Source/MeshModelingTools/Private/SeamSculptTool.cpp b/Engine/Plugins/Experimental/MeshModelingToolset/Source/MeshModelingTools/Private/SeamSculptTool.cpp
--- a/Engine/Plugins/Experimental/MeshModelingToolset/Source/MeshModelingTools/Private/SeamSculptTool.cpp
+++ b/Engine/Plugins/Experimental/MeshModelingToolset/Source/MeshModelingTools/Private/SeamSculptTool.cpp
@@ -17,8 +17,8 @@
 #include "Parameterization/MeshUVTransforms.h"
 #include "Parameterization/MeshDijkstra.h"
 #include "Selections/MeshConnectedComponents.h"
-#include "MeshNormals.h"
-#include "MeshIndexUtil.h"
+#include "DynamicMesh/MeshNormals.h"
+#include "DynamicMesh/MeshIndexUtil.h"
 #include "DynamicMeshChangeTracker.h"
 #include "DynamicMeshToMeshDescription.h"
 
diff --git a/Engine/Plugins/Experimental/MeshModelingToolset/Source/MeshModelingTools/Private/WeldMeshEdgesTool.cpp b/Engine/Plugins/Experimental/MeshModelingToolset/Source/MeshModelingTools/Private/WeldMeshEdgesTool.cpp
--- a/Engine/Plugins/Experimental/MeshModelingToolset/Source/MeshModelingTools/Private/WeldMeshEdgesTool.cpp
+++ b/Engine/Plugins/Experimental/MeshModelingToolset/Source/MeshModelingTools/Private/WeldMeshEdgesTool.cpp
@@ -4,9 +4,9 @@
 #include "InteractiveToolManager.h"
 #include "ToolBuilderUtil.h"
 
-#include "DynamicMesh3.h"
-#include "DynamicMeshAttributeSet.h"
-#include "Operations/MergeCoincidentMeshEdges.h"
+#include "DynamicMesh/DynamicMesh3.h"
+#include "DynamicMesh/DynamicMeshAttributeSet.h"
+#include "DynamicMesh/Operations/MergeCoincidentMeshEdges.h"
 
 #include "SimpleDynamicMeshComponent.h"
 
diff --git a/Engine/Plugins/Experimental/MeshModelingToolset/Source/MeshModelingToolsEditorOnly/Private/EditMeshMaterialsTool.cpp b/Engine/Plugins/Experimental/MeshModelingToolset/Source/MeshModelingToolsEditorOnly/Private/EditMeshMaterialsTool.cpp
--- a/Engine/Plugins/Experimental/MeshModelingToolset/Source/MeshModelingToolsEditorOnly/Private/EditMeshMaterialsTool.cpp
+++ b/Engine/Plugins/Experimental/MeshModelingToolset/Source/MeshModelingToolsEditorOnly/Private/EditMeshMaterialsTool.cpp
@@ -11,7 +11,7 @@
 #include "Util/ColorConstants.h"
 #include "Selections/MeshConnectedComponents.h"
 #include "MeshRegionBoundaryLoops.h"
-#include "MeshIndexUtil.h"
+#include "DynamicMesh/MeshIndexUtil.h"
 #include "AssetGenerationUtil.h"
 #include "ToolSetupUtil.h"
 
diff --git a/Engine/Plugins/Experimental/MeshModelingToolset/Source/ModelingOperatorsEditorOnly/Private/CleaningOps/SimplifyMeshOp.cpp b/Engine/Plugins/Experimental/MeshModelingToolset/Source/ModelingOperatorsEditorOnly/Private/CleaningOps/SimplifyMeshOp.cpp
--- a/Engine/Plugins/Experimental/MeshModelingToolset/Source/ModelingOperatorsEditorOnly/Private/CleaningOps/SimplifyMeshOp.cpp
+++ b/Engine/Plugins/Experimental/MeshModelingToolset/Source/ModelingOperatorsEditorOnly/Private/CleaningOps/SimplifyMeshOp.cpp
@@ -2,21 +2,19 @@
 
 #include "CleaningOps/SimplifyMeshOp.h"
 
-#include "DynamicMeshAttributeSet.h"
 #include "MeshConstraints.h"
 #include "MeshDescriptionToDynamicMesh.h"
 
-
-#include "DynamicMesh3.h"
-#include "DynamicMeshAttributeSet.h"
+#include "DynamicMesh/DynamicMesh3.h"
+#include "DynamicMesh/DynamicMeshAttributeSet.h"
 #include "MeshSimplification.h"
 #include "MeshConstraintsUtil.h"
 #include "ProjectionTargets.h"
 
 #include "IMeshReductionInterfaces.h"
 #include "MeshDescription.h"
-#include "MeshNormals.h"
-#include "Operations/MergeCoincidentMeshEdges.h"
+#include "DynamicMesh/MeshNormals.h"
+#include "DynamicMesh/Operations/MergeCoincidentMeshEdges.h"
 #include "OverlappingCorners.h"
 #include "StaticMeshOperations.h"
 
diff --git a/Engine/Plugins/Experimental/PlanarCutPlugin/Source/PlanarCut/Private/GeometryMeshConversion.cpp b/Engine/Plugins/Experimental/PlanarCutPlugin/Source/PlanarCut/Private/GeometryMeshConversion.cpp
--- a/Engine/Plugins/Experimental/PlanarCutPlugin/Source/PlanarCut/Private/GeometryMeshConversion.cpp
+++ b/Engine/Plugins/Experimental/PlanarCutPlugin/Source/PlanarCut/Private/GeometryMeshConversion.cpp
@@ -21,19 +21,19 @@
 
 #include "DisjointSet.h"
 
-#include "DynamicMesh3.h"
+#include "DynamicMesh/DynamicMesh3.h"
 #include "DynamicMeshEditor.h"
-#include "DynamicMeshAABBTree3.h"
+#include "DynamicMesh/DynamicMeshAABBTree3.h"
 #include "Selections/MeshConnectedComponents.h"
 #include "MeshTransforms.h"
 #include "Operations/MeshBoolean.h"
 #include "Operations/MeshSelfUnion.h"
-#include "Operations/MergeCoincidentMeshEdges.h"
+#include "DynamicMesh/Operations/MergeCoincidentMeshEdges.h"
 #include "MeshBoundaryLoops.h"
 #include "QueueRemesher.h"
 #include "DynamicVertexAttribute.h"
-#include "MeshNormals.h"
-#include "MeshTangents.h"
+#include "DynamicMesh/MeshNormals.h"
+#include "DynamicMesh/MeshTangents.h"
 #include "ConstrainedDelaunay2.h"
 
 #include "Engine/EngineTypes.h"
diff --git a/Engine/Plugins/Experimental/PlanarCutPlugin/Source/PlanarCut/Private/PlanarCut.cpp b/Engine/Plugins/Experimental/PlanarCutPlugin/Source/PlanarCut/Private/PlanarCut.cpp
--- a/Engine/Plugins/Experimental/PlanarCutPlugin/Source/PlanarCut/Private/PlanarCut.cpp
+++ b/Engine/Plugins/Experimental/PlanarCutPlugin/Source/PlanarCut/Private/PlanarCut.cpp
@@ -18,19 +18,19 @@
 
 #include "DisjointSet.h"
 
-#include "DynamicMesh3.h"
+#include "DynamicMesh/DynamicMesh3.h"
 #include "DynamicMeshEditor.h"
-#include "DynamicMeshAABBTree3.h"
+#include "DynamicMesh/DynamicMeshAABBTree3.h"
 #include "Selections/MeshConnectedComponents.h"
 #include "MeshTransforms.h"
 #include "Operations/MeshBoolean.h"
 #include "Operations/MeshSelfUnion.h"
-#include "Operations/MergeCoincidentMeshEdges.h"
+#include "DynamicMesh/Operations/MergeCoincidentMeshEdges.h"
 #include "MeshBoundaryLoops.h"
 #include "QueueRemesher.h"
 #include "DynamicVertexAttribute.h"
-#include "MeshNormals.h"
-#include "MeshTangents.h"
+#include "DynamicMesh/MeshNormals.h"
+#include "DynamicMesh/MeshTangents.h"
 #include "ConstrainedDelaunay2.h"
 
 #include "Engine/EngineTypes.h"
diff --git a/Engine/Plugins/Experimental/Water/Source/Runtime/Private/WaterBooleanUtils.h b/Engine/Plugins/Experimental/Water/Source/Runtime/Private/WaterBooleanUtils.h
--- a/Engine/Plugins/Experimental/Water/Source/Runtime/Private/WaterBooleanUtils.h
+++ b/Engine/Plugins/Experimental/Water/Source/Runtime/Private/WaterBooleanUtils.h
@@ -3,12 +3,12 @@
 #pragma once
 
 
-#include "DynamicMesh3.h"
+#include "DynamicMesh/DynamicMesh3.h"
 #include "TransformTypes.h"
 #include "VectorUtil.h"
 #include "Generators/MinimalBoxMeshGenerator.h"
 #include "Operations/MeshBoolean.h"
-#include "Operations/MergeCoincidentMeshEdges.h"
+#include "DynamicMesh/Operations/MergeCoincidentMeshEdges.h"
 #include "Selections/MeshConnectedComponents.h"
 #include "DynamicMeshEditor.h"
 #include "MeshTransforms.h"
diff --git a/Engine/Source/Runtime/GeometryCore/Private/DynamicMesh/MeshTransforms.cpp b/Engine/Source/Runtime/GeometryCore/Private/DynamicMesh/MeshTransforms.cpp
new file mode 100644
--- /dev/null
+++ b/Engine/Source/Runtime/GeometryCore/Private/DynamicMesh/MeshTransforms.cpp
@@ -0,0 +1,231 @@
+// Copyright Epic Games, Inc. All Rights Reserved.
+
+
+#include "DynamicMesh/MeshTransforms.h"
+#include "DynamicMesh/DynamicMeshAttributeSet.h"
+#include "DynamicMesh/DynamicMeshOverlay.h"
+#include "Async/ParallelFor.h"
+
+#include "ExplicitUseGeometryMathTypes.h"		// using UE::Geometry::(math types)
+using namespace UE::Geometry;
+
+void MeshTransforms::Translate(FDynamicMesh3& Mesh, const FVector3d& Translation)
+{
+	int NumVertices = Mesh.MaxVertexID();
+	ParallelFor(NumVertices, [&](int vid) 
+	{
+		if (Mesh.IsVertex(vid))
+		{
+			Mesh.SetVertex(vid, Mesh.GetVertex(vid) + Translation);
+		}
+	});
+}
+
+
+void MeshTransforms::Scale(FDynamicMesh3& Mesh, const FVector3d& Scale, const FVector3d& Origin)
+{
+	int NumVertices = Mesh.MaxVertexID();
+	ParallelFor(NumVertices, [&](int32 vid)
+	{
+		if (Mesh.IsVertex(vid))
+		{
+			Mesh.SetVertex(vid, (Mesh.GetVertex(vid) - Origin) * Scale + Origin);
+		}
+	});
+}
+
+
+
+void MeshTransforms::WorldToFrameCoords(FDynamicMesh3& Mesh, const FFrame3d& Frame)
+{
+	bool bVertexNormals = Mesh.HasVertexNormals();
+
+	int NumVertices = Mesh.MaxVertexID();
+	ParallelFor(NumVertices, [&](int vid)
+	{
+		if (Mesh.IsVertex(vid))
+		{
+			FVector3d Position = Mesh.GetVertex(vid);
+			Mesh.SetVertex(vid, Frame.ToFramePoint(Position));
+
+			if (bVertexNormals)
+			{
+				FVector3f Normal = Mesh.GetVertexNormal(vid);
+				Mesh.SetVertexNormal(vid, (FVector3f)Frame.ToFrameVector((FVector3d)Normal));
+			}
+		}
+	});
+
+	if (Mesh.HasAttributes())
+	{
+		FDynamicMeshNormalOverlay* Normals = Mesh.Attributes()->PrimaryNormals();
+		int NumNormals = Normals->MaxElementID();
+		ParallelFor(NumNormals, [&](int elemid)
+		{
+			if (Normals->IsElement(elemid))
+			{
+				FVector3f Normal = Normals->GetElement(elemid);
+				Normals->SetElement(elemid, (FVector3f)Frame.ToFrameVector((FVector3d)Normal));
+			}
+		});
+	}
+}
+
+
+
+
+
+void MeshTransforms::FrameCoordsToWorld(FDynamicMesh3& Mesh, const FFrame3d& Frame)
+{
+	bool bVertexNormals = Mesh.HasVertexNormals();
+
+	int NumVertices = Mesh.MaxVertexID();
+	ParallelFor(NumVertices, [&](int vid)
+	{
+		if (Mesh.IsVertex(vid))
+		{
+			FVector3d Position = Mesh.GetVertex(vid);
+			Mesh.SetVertex(vid, Frame.FromFramePoint(Position));
+
+			if (bVertexNormals)
+			{
+				FVector3f Normal = Mesh.GetVertexNormal(vid);
+				Mesh.SetVertexNormal(vid, (FVector3f)Frame.FromFrameVector((FVector3d)Normal));
+			}
+		}
+	});
+
+	if (Mesh.HasAttributes())
+	{
+		FDynamicMeshNormalOverlay* Normals = Mesh.Attributes()->PrimaryNormals();
+		int NumNormals = Normals->MaxElementID();
+		ParallelFor(NumNormals, [&](int elemid)
+		{
+			if (Normals->IsElement(elemid))
+			{
+				FVector3f Normal = Normals->GetElement(elemid);
+				Normals->SetElement(elemid, (FVector3f)Frame.FromFrameVector((FVector3d)Normal));
+			}
+		});
+	}
+}
+
+
+
+
+void MeshTransforms::ApplyTransform(FDynamicMesh3& Mesh, const FTransform3d& Transform)
+{
+	bool bVertexNormals = Mesh.HasVertexNormals();
+
+	for (int vid : Mesh.VertexIndicesItr())
+	{
+		FVector3d Position = Mesh.GetVertex(vid);
+		Position = Transform.TransformPosition(Position);
+		Mesh.SetVertex(vid, Position);
+
+		if (bVertexNormals)
+		{
+			FVector3f Normal = Mesh.GetVertexNormal(vid);
+			Normal = (FVector3f)Transform.TransformNormal((FVector3d)Normal);
+			Mesh.SetVertexNormal(vid, Normalized(Normal));
+		}
+	}
+	if (Mesh.HasAttributes())
+	{
+		FDynamicMeshNormalOverlay* Normals = Mesh.Attributes()->PrimaryNormals();
+		for (int elemid : Normals->ElementIndicesItr())
+		{
+			FVector3f Normal = Normals->GetElement(elemid);
+			Normal = (FVector3f)Transform.TransformNormal((FVector3d)Normal);
+			Normals->SetElement(elemid, Normalized(Normal));
+		}
+	}
+}
+
+
+
+void MeshTransforms::ApplyTransformInverse(FDynamicMesh3& Mesh, const FTransform3d& Transform)
+{
+	bool bVertexNormals = Mesh.HasVertexNormals();
+
+	int NumVertices = Mesh.MaxVertexID();
+	ParallelFor(NumVertices, [&](int vid)
+	{
+		if (Mesh.IsVertex(vid)) 
+		{
+			FVector3d Position = Mesh.GetVertex(vid);
+			Position = Transform.InverseTransformPosition(Position);
+			Mesh.SetVertex(vid, Position);
+
+			if (bVertexNormals)
+			{
+				FVector3f Normal = Mesh.GetVertexNormal(vid);
+				Normal = (FVector3f)Transform.InverseTransformNormal((FVector3d)Normal);
+				Mesh.SetVertexNormal(vid, Normalized(Normal));
+			}
+		}
+	});
+
+	if (Mesh.HasAttributes())
+	{
+		FDynamicMeshNormalOverlay* Normals = Mesh.Attributes()->PrimaryNormals();
+		if (Normals)
+		{
+			int NumNormals = Normals->MaxElementID();
+			ParallelFor(NumNormals, [&](int elemid)
+			{
+				if (Normals->IsElement(elemid))
+				{
+					FVector3f Normal = Normals->GetElement(elemid);
+					Normal = (FVector3f)Transform.InverseTransformNormal((FVector3d)Normal);
+					Normals->SetElement(elemid, Normalized(Normal));
+				}
+			});
+		}
+	}
+}
+
+
+
+void MeshTransforms::ApplyTransform(FDynamicMesh3& Mesh,
+	TFunctionRef<FVector3d(const FVector3d&)> PositionTransform,
+	TFunctionRef<FVector3f(const FVector3f&)> NormalTransform)
+{
+	bool bVertexNormals = Mesh.HasVertexNormals();
+
+	int NumVertices = Mesh.MaxVertexID();
+	ParallelFor(NumVertices, [&](int vid)
+	{
+		if (Mesh.IsVertex(vid))
+		{
+			FVector3d Position = Mesh.GetVertex(vid);
+			Position = PositionTransform(Position);
+			Mesh.SetVertex(vid, Position);
+
+			if (bVertexNormals)
+			{
+				FVector3f Normal = Mesh.GetVertexNormal(vid);
+				Normal = NormalTransform(Normal);
+				Mesh.SetVertexNormal(vid, Normalized(Normal));
+			}
+		}
+	});
+
+	if (Mesh.HasAttributes())
+	{
+		FDynamicMeshNormalOverlay* Normals = Mesh.Attributes()->PrimaryNormals();
+		if (Normals)
+		{
+			int NumNormals = Normals->MaxElementID();
+			ParallelFor(NumNormals, [&](int elemid)
+			{
+				if (Normals->IsElement(elemid))
+				{
+					FVector3f Normal = Normals->GetElement(elemid);
+					Normal = NormalTransform(Normal);
+					Normals->SetElement(elemid, Normalized(Normal));
+				}
+			});
+		}
+	}
+}
diff --git a/Engine/Source/Runtime/GeometryCore/Public/DynamicMesh/MeshNormals.h b/Engine/Source/Runtime/GeometryCore/Public/DynamicMesh/MeshNormals.h
new file mode 100644
--- /dev/null
+++ b/Engine/Source/Runtime/GeometryCore/Public/DynamicMesh/MeshNormals.h
@@ -0,0 +1,217 @@
+// Copyright Epic Games, Inc. All Rights Reserved.
+
+// Port of geometry3Sharp MeshNormals
+
+#pragma once
+
+#include "DynamicMesh3.h"
+#include "DynamicMeshAttributeSet.h"
+
+namespace UE
+{
+namespace Geometry
+{
+
+
+/**
+ * FMeshNormals is a utility class that can calculate and store various types of
+ * normal vectors for a FDynamicMesh. 
+ */
+class GEOMETRYCORE_API FMeshNormals
+{
+protected:
+	/** Target Mesh */
+	const FDynamicMesh3* Mesh;
+	/** Set of computed normals */
+	TArray<FVector3d> Normals;
+
+public:
+	FMeshNormals()
+	{
+		Mesh = nullptr;
+	}
+
+	FMeshNormals(const FDynamicMesh3* Mesh)
+	{
+		SetMesh(Mesh);
+	}
+
+
+	void SetMesh(const FDynamicMesh3* MeshIn)
+	{
+		this->Mesh = MeshIn;
+	}
+
+	const TArray<FVector3d>& GetNormals() const { return Normals; }
+
+	FVector3d& operator[](int i) { return Normals[i]; }
+	const FVector3d& operator[](int i) const { return Normals[i]; }
+
+
+	/**
+	 * Set the size of the Normals array to Count, and optionally clear all values to (0,0,0)
+	 */
+	void SetCount(int Count, bool bClearToZero);
+
+	/**
+	 * Compute standard per-vertex normals by averaging one-ring face normals
+	 */
+	void ComputeVertexNormals(bool bWeightByArea = true, bool bWeightByAngle = true)
+	{
+		Compute_FaceAvg(bWeightByArea, bWeightByAngle);
+	}
+
+	/**
+	 * Compute per-triangle normals
+	 */
+	void ComputeTriangleNormals()
+	{
+		Compute_Triangle();
+	}
+
+	/**
+	 * Assumes that ComputeTriangleNormals() has already been run.
+	 * 
+	 * Looks through the normals for any that are zero, and tries to set them to the normal of 
+	 * some neighboring triangle, preferring the triangle on its longer side if possible. This is 
+	 * useful for finding connected components based on normals, as it prevents degenerate triangles
+	 * with 0 normals from connecting otherwise disconnected components.
+	 * This won't help if a whole mesh componenent is entirely made of such degenerate triangles
+	 * (which is usually ok).
+	 */
+	void SetDegenerateTriangleNormalsToNeighborNormal();
+
+	/**
+	 * Recompute the per-element normals of the given overlay by averaging one-ring face normals
+	 * @warning NormalOverlay must be attached to ParentMesh or an exact copy
+	 */
+	void RecomputeOverlayNormals(const FDynamicMeshNormalOverlay* NormalOverlay, bool bWeightByArea = true, bool bWeightByAngle = true)
+	{
+		Compute_Overlay_FaceAvg(NormalOverlay, bWeightByArea, bWeightByAngle);
+	}
+
+
+	/**
+	 * Copy the current set of normals to the vertex normals of SetMesh
+	 * @warning assumes that the computed normals are vertex normals
+	 * @param bInvert if true, normals are flipped
+	 */
+	void CopyToVertexNormals(FDynamicMesh3* SetMesh, bool bInvert = false) const;
+
+	/**
+	 * Copy the current set of normals to the NormalOverlay attribute layer
+	 * @warning assumes that the computed normals are attribute normals
+	 * @param bInvert if true, normals are flipped
+	 */
+	void CopyToOverlay(FDynamicMeshNormalOverlay* NormalOverlay, bool bInvert = false) const;
+
+
+	/**
+	 * Compute per-vertex normals for the given Mesh
+	 * @param bInvert if true, normals are flipped
+	 */
+	static void QuickComputeVertexNormals(FDynamicMesh3& Mesh, bool bInvert = false);
+
+	/**
+	 * Compute per-vertex normals for the vertices of a set of triangles of a Mesh
+	 * @param bWeightByArea weight neighbor triangles by area
+	 * @param bWeightByAngle weight neighbor triangles by angle
+	 * @param bInvert if true, normals are flipped
+	 */
+	static void QuickComputeVertexNormalsForTriangles(FDynamicMesh3& Mesh, const TArray<int32>& Triangles, bool bWeightByArea = true, bool bWeightByAngle = true, bool bInvert = false);
+
+
+	/**
+	 * Compute normal at mesh vertex by weighted sum of one-ring triangle normals. Can optionally weight by area, angle, or both (averaged)
+	 * @param bWeightByArea weight neighbor triangles by area
+	 * @param bWeightByAngle weight neighbor triangles by angle
+	 * @return the vertex normal at vertex VertIdx of Mesh.
+	 */
+	static FVector3d ComputeVertexNormal(const FDynamicMesh3& Mesh, int VertIdx, bool bWeightByArea = true, bool bWeightByAngle = true);
+
+	/**
+	 * Compute normal at mesh vertex by weighted sum of subset of one-ring triangle normals. Can optionally weight by area, angle, or both (averaged)
+	 * @param TriangleFilterFunc Only one-ring triangles for which this function returns true will be included
+	 * @param bWeightByArea weight neighbor triangles by area
+	 * @param bWeightByAngle weight neighbor triangles by angle
+	 * @return the vertex normal at vertex VertIdx of Mesh.
+	 */
+	static FVector3d ComputeVertexNormal(const FDynamicMesh3& Mesh, int32 VertIdx, TFunctionRef<bool(int32)> TriangleFilterFunc, bool bWeightByArea = true, bool bWeightByAngle = true);
+
+
+	/**
+	 * @return the computed overlay normal at an element of the overlay (ie based on normals of triangles connected to this element)
+	 */
+	static FVector3d ComputeOverlayNormal(const FDynamicMesh3& Mesh, const FDynamicMeshNormalOverlay* NormalOverlay, int ElemIdx);
+
+	/**
+	 * Initialize the given NormalOverlay with per-vertex normals, ie single overlay element for each mesh vertex.
+	 * @param bUseMeshVertexNormalsIfAvailable if true and the parent mesh has per-vertex normals, use them instead of calculating new ones
+	 */
+	static void InitializeOverlayToPerVertexNormals(FDynamicMeshNormalOverlay* NormalOverlay, bool bUseMeshVertexNormalsIfAvailable = true);
+
+
+	/**
+	 * Initialize the given NormalOverlay with per-face normals, ie separate overlay element for each vertex of each triangle.
+	 */
+	static void InitializeOverlayToPerTriangleNormals(FDynamicMeshNormalOverlay* NormalOverlay);
+
+
+	static void InitializeOverlayTopologyFromOpeningAngle(const FDynamicMesh3* Mesh, FDynamicMeshNormalOverlay* NormalOverlay, 
+		double AngleThresholdDeg);
+
+	static void InitializeOverlayTopologyFromFaceGroups(const FDynamicMesh3* Mesh, FDynamicMeshNormalOverlay* NormalOverlay);
+
+
+	/**
+	 * Initialize the given Mesh with per-face normals, ie separate overlay element for each vertex of each triangle.
+	 */
+	static void InitializeMeshToPerTriangleNormals(FDynamicMesh3* Mesh);
+
+
+	/**
+	 * Initialize the given triangles of NormalOverlay with per-vertex normals, ie single overlay element for each mesh vertex.
+	 * Only the triangles included in the region are considered when calculating per-vertex normals.
+	 */
+	static void InitializeOverlayRegionToPerVertexNormals(FDynamicMeshNormalOverlay* NormalOverlay, const TArray<int32>& Triangles);
+
+	/**
+	 * Compute overlay normals for the given mesh
+	 * @param bInvert if true, normals are flipped
+	 */
+	static bool QuickRecomputeOverlayNormals(FDynamicMesh3& Mesh, bool bInvert = false, bool bWeightByArea = true, bool bWeightByAngle = true);
+
+
+protected:
+	/** Compute per-vertex normals using area-weighted averaging of one-ring triangle normals */
+	void Compute_FaceAvg_AreaWeighted();
+	/** Compute per-vertex normals using a custom combination of area-weighted and angle-weighted averaging of one-ring triangle normals */
+	void Compute_FaceAvg(bool bWeightByArea, bool bWeightByAngle);
+
+	/** Compute per-triangle normals */
+	void Compute_Triangle();
+
+	/** Recompute the element Normals of the given attribute overlay using area-weighted averaging of one-ring triangle normals */
+	void Compute_Overlay_FaceAvg_AreaWeighted(const FDynamicMeshNormalOverlay* NormalOverlay);
+	/** Recompute the element Normals of the given attribute overlay using a custom combination of area-weighted and angle-weighted averaging of one-ring triangle normals */
+	void Compute_Overlay_FaceAvg(const FDynamicMeshNormalOverlay* NormalOverlay, bool bWeightByArea, bool bWeightByAngle);
+
+	static FVector3d GetVertexWeightsOnTriangle(const FDynamicMesh3* Mesh, int TriID, double TriArea, bool bWeightByArea, bool bWeightByAngle)
+	{
+		FVector3d TriNormalWeights = FVector3d::One();
+		if (bWeightByAngle)
+		{
+			TriNormalWeights = Mesh->GetTriInternalAnglesR(TriID); // component-wise multiply by per-vertex internal angles
+		}
+		if (bWeightByArea)
+		{
+			TriNormalWeights *= TriArea;
+		}
+		return TriNormalWeights;
+	}
+
+};
+
+
+} // end namespace UE::Geometry
+} // end namespace UE
\ No newline at end of file
diff --git a/Engine/Source/Runtime/GeometryCore/Public/DynamicMesh/MeshTransforms.h b/Engine/Source/Runtime/GeometryCore/Public/DynamicMesh/MeshTransforms.h
new file mode 100644
--- /dev/null
+++ b/Engine/Source/Runtime/GeometryCore/Public/DynamicMesh/MeshTransforms.h
@@ -0,0 +1,63 @@
+// Copyright Epic Games, Inc. All Rights Reserved.
+
+// Port of geometry3cpp MeshWeights
+
+#pragma once
+
+#include "DynamicMesh/DynamicMesh3.h"
+
+
+/**
+ * Utility functions for applying transformations to meshes
+ */
+namespace MeshTransforms
+{
+	using namespace UE::Geometry;
+	using UE::Geometry::FTransform3d;
+
+	/**
+	 * Apply Translation to vertex positions of Mesh. Does not modify any other attributes.
+	 */
+	GEOMETRYCORE_API void Translate(FDynamicMesh3& Mesh, const FVector3d& Translation);
+
+	/**
+	 * Apply Scale to vertex positions of Mesh, relative to given Origin. Does not modify any other attributes (normals/etc)
+	 */
+	GEOMETRYCORE_API void Scale(FDynamicMesh3& Mesh, const FVector3d& Scale, const FVector3d& Origin);
+
+	/**
+	 * Transform Mesh into local coordinates of Frame
+	 */
+	GEOMETRYCORE_API void WorldToFrameCoords(FDynamicMesh3& Mesh, const FFrame3d& Frame);
+
+	/**
+	 * Transform Mesh out of local coordinates of Frame
+	 */
+	GEOMETRYCORE_API void FrameCoordsToWorld(FDynamicMesh3& Mesh, const FFrame3d& Frame);
+
+
+	/**
+	 * Apply given Transform to a Mesh.
+	 * Modifies Vertex Positions and Normals, and any Per-Triangle Normal Overlays
+	 */
+	GEOMETRYCORE_API void ApplyTransform(FDynamicMesh3& Mesh, const FTransform3d& Transform);
+
+
+	/**
+	 * Apply inverse of given Transform to a Mesh.
+	 * Modifies Vertex Positions and Normals, and any Per-Triangle Normal Overlays
+	 */
+	GEOMETRYCORE_API void ApplyTransformInverse(FDynamicMesh3& Mesh, const FTransform3d& Transform);
+
+
+	/**
+	 * Apply given Transform to a Mesh.
+	 * Modifies Vertex Positions and Normals, and any Per-Triangle Normal Overlays
+	 */
+	GEOMETRYCORE_API void ApplyTransform(FDynamicMesh3& Mesh, 
+		TFunctionRef<FVector3d(const FVector3d&)> PositionTransform,
+		TFunctionRef<FVector3f(const FVector3f&)> NormalTransform);
+
+
+
+};
\ No newline at end of file
diff --git a/Engine/Source/Runtime/MeshConversion/Public/MeshDescriptionBuilder.h b/Engine/Source/Runtime/MeshConversion/Public/MeshDescriptionBuilder.h
new file mode 100644
--- /dev/null
+++ b/Engine/Source/Runtime/MeshConversion/Public/MeshDescriptionBuilder.h
@@ -0,0 +1,136 @@
+// Copyright Epic Games, Inc. All Rights Reserved.
+
+#pragma once
+
+#include "CoreMinimal.h"
+#include "MeshDescription.h"
+#include "GeometryBase.h"
+
+PREDECLARE_USE_GEOMETRY_CLASS(FDynamicMesh3);
+
+
+/**
+ * These attributes are used to store custom modeling tools data on a MeshDescription
+ */
+namespace ExtendedMeshAttribute
+{
+	extern MESHCONVERSION_API const FName PolyTriGroups;
+}
+
+
+/**
+ * Utility class to construct MeshDescription instances.
+ * NB: this will add a vertex-instance UV to the Description, if it has none. 
+ */
+class MESHCONVERSION_API FMeshDescriptionBuilder
+{
+public:
+	void SetMeshDescription(FMeshDescription* Description);
+
+	/** Pre-allocate space in the mesh description */
+	void ReserveNewVertices(int32 Count);
+
+	/** Append vertex and return new vertex ID */
+	FVertexID AppendVertex(const FVector& Position);
+
+	/** Return position of vertex */
+	FVector GetPosition(const FVertexID& VertexID);
+
+	/** Return position of vertex parent of instance */
+	FVector GetPosition(const FVertexInstanceID& InstanceID);
+
+	/** Set the position of a vertex */
+	void SetPosition(const FVertexID& VertexID, const FVector& NewPosition);
+
+	/** Set the number of UV layers */
+	void SetNumUVLayers(int32 NumUVLayers);
+
+	/** Pre-allocate space in the mesh description for UVs in the indicated UVLayer */
+	void ReserveNewUVs(int32 Count, int UVLayerIndex);
+
+	/** Append a UV 'vertex' and return a new UV ID*/
+	FUVID AppendUV(const FVector2D& UVvalue, int32 UVLayerIndex);
+
+	/** Append new vertex instance and return ID */
+	FVertexInstanceID AppendInstance(const FVertexID& VertexID);
+
+	/** 
+	* Set the UV of a vertex instance.
+	* Note: this generally shouldn't be called directly because it alters the instance UVs but not the shared UVs
+	* When setting UVs use AppendUVTriangle	  
+	*/
+	void SetInstanceUV(const FVertexInstanceID& InstanceID, const FVector2D& InstanceUV, int32 UVLayerIndex = 0);
+
+	/** Set the Normal of a vertex instance*/
+	void SetInstanceNormal(const FVertexInstanceID& InstanceID, const FVector& Normal);
+
+	/** Set the full tangent space of a vertex instance, this is stored as a vec3 normal, vec3 tangent, and a bitangent sign */
+	void SetInstanceTangentSpace(const FVertexInstanceID& InstanceID, const FVector& Normal, const FVector& Tangent, float Sign);
+
+	/** Set the Color of a vertex instance*/
+	void SetInstanceColor(const FVertexInstanceID& InstanceID, const FVector4& Color);
+
+	/** Enable per-triangle integer attribute named PolyTriGroups */
+	void EnablePolyGroups();
+
+	/** Create a new polygon group and return it's ID */
+	FPolygonGroupID AppendPolygonGroup();
+
+	/** Set the PolyTriGroups attribute value to a specific GroupID for a Triangle */
+	void SetPolyGroupID(const FTriangleID& TriangleID, int GroupID);
+
+	/** Append a UV triangle to the specified UV layer. This will use both shared and per-instance UV storage*/
+	void AppendUVTriangle(const FTriangleID& TriangleID, const FUVID UVverterxID0, const FUVID UVvertexID1, const FUVID UVvertexID2, int32 UVLayerIndex);
+
+	/** Append a triangle to the mesh with the given PolygonGroup ID */
+	FTriangleID AppendTriangle(const FVertexID& Vertex0, const FVertexID& Vertex1, const FVertexID& Vertex2, const FPolygonGroupID& PolygonGroup);
+
+
+	/** Append a triangle to the mesh with the given PolygonGroup ID */
+	FTriangleID AppendTriangle(const FVertexID* Triangle, const FPolygonGroupID& PolygonGroup);
+
+
+	/**
+	 * Append a triangle to the mesh using the given vertex instances and PolygonGroup ID
+	 */
+	FTriangleID AppendTriangle(const FVertexInstanceID& Instance0, const FVertexInstanceID& Instance1, const FVertexInstanceID& Instance2, const FPolygonGroupID& PolygonGroup);
+
+	/**
+	 * Append an arbitrary polygon to the mesh with the given PolygonGroup ID
+	 * Unique Vertex instances will be created for each polygon-vertex.
+	 */
+	FPolygonID AppendPolygon(const TArray<FVertexID>& Vertices, const FPolygonGroupID& PolygonGroup);
+
+	/** Set MeshAttribute::Edge::IsHard to true for all edges */
+	void SetAllEdgesHardness(bool bHard);
+
+	/** Translate the MeshDescription vertex positions */
+	void Translate(const FVector& Translation);
+
+
+	/** Return the current bounding box of the mesh */
+	FBox ComputeBoundingBox() const;
+
+	/** Disable the construction of secondary data structures in the mesh description */
+	void SuspendMeshDescriptionIndexing();
+
+	/** Enable the construction of secondary data structures in the mesh description */
+	void ResumeMeshDescriptionIndexing();
+
+protected:
+
+	FMeshDescription* MeshDescription;
+
+	TVertexAttributesRef<FVector3f> VertexPositions;
+	TVertexInstanceAttributesRef<FVector2D> InstanceUVs;
+	TVertexInstanceAttributesRef<FVector3f> InstanceNormals;
+	TVertexInstanceAttributesRef<FVector3f> InstanceTangents;
+	TVertexInstanceAttributesRef<float> InstanceBiTangentSign;
+	TVertexInstanceAttributesRef<FVector4> InstanceColors;
+
+	TArray<TUVAttributesRef<FVector2D>> UVCoordinateLayers; 
+	TArray<FVertexID> TempBuffer;
+	TArray<FUVID> TempUVBuffer;
+
+	TPolygonAttributesRef<int> PolyGroups;
+};
\ No newline at end of file
