GeometryProcessing:
- Add MakeBoundsFromIndices() utility functions to TAxisAlignedBox3 to simplify computing bounding-box of point sets
- FDynamicMesh3::GetVertexFrame() can now take an externally-provided normal, and so now can be used without computing VertexNormals attribute
- Add FModuloIteration to IteratorUtil for doing quasi-random-ish iterations over index ranges
- Add CollectVertexPositions() to MeshIndexUtil.h, for creating arrays/sets of vertex positions as a one-liner
- Add FMeshWeights::CotanWeightsBlendSafe() to support cotan-based smoothing of arbitrary properties
- Add TMinVolumeBox3::SolveSubsample() which will (quasi-randomly) subsample a set of the input points to reduce cost of box computation
#rnx

diff --git a/Engine/Plugins/Experimental/GeometryProcessing/Source/DynamicMesh/Private/MeshWeights.cpp b/Engine/Plugins/Experimental/GeometryProcessing/Source/DynamicMesh/Private/MeshWeights.cpp
--- a/Engine/Plugins/Experimental/GeometryProcessing/Source/DynamicMesh/Private/MeshWeights.cpp
+++ b/Engine/Plugins/Experimental/GeometryProcessing/Source/DynamicMesh/Private/MeshWeights.cpp
@@ -244,6 +244,81 @@ FVector3d FMeshWeights::CotanCentroidSafe(const FDynamicMesh3& mesh, int32 v_i,
 
 
 
+
+template<typename PositionFuncType>
+void TCotanWeightsBlendSafe(const FDynamicMesh3& mesh, int32 v_i, PositionFuncType PositionFunc, TFunctionRef<void(int32, double)> BlendingFunc, double DegenerateTol = 100.0, bool* bOutputIsUniform = nullptr)
+{
+	// based on equations in http://www.geometry.caltech.edu/pubs/DMSB_III.pdf
+
+	FVector3d Vi = PositionFunc(v_i);
+
+	bool bDegenerated = false;
+	
+	TArray<int32, TInlineAllocator<32>> Neighbours;
+	TArray<double, TInlineAllocator<32>> Weights;
+	double wSum = 0;
+
+	int v_j = FDynamicMesh3::InvalidID, opp_v1 = FDynamicMesh3::InvalidID, opp_v2 = FDynamicMesh3::InvalidID;
+	int t1 = FDynamicMesh3::InvalidID, t2 = FDynamicMesh3::InvalidID;
+	for (int eid : mesh.VtxEdgesItr(v_i))
+	{
+		opp_v2 = FDynamicMesh3::InvalidID;
+		mesh.GetVtxNbrhood(eid, v_i, v_j, opp_v1, opp_v2, t1, t2);
+		FVector3d Vj = PositionFunc(v_j);
+
+		FVector3d Vo1 = PositionFunc(opp_v1);
+		double cot_alpha_ij = VectorUtil::VectorCot((Vi - Vo1), (Vj - Vo1));
+		double w_ij = cot_alpha_ij;
+
+		double cot_beta_ij = 0;
+		if (opp_v2 != FDynamicMesh3::InvalidID)
+		{
+			FVector3d Vo2 = PositionFunc(opp_v2);
+			cot_beta_ij = VectorUtil::VectorCot((Vi - Vo2), (Vj - Vo2));
+		}
+		w_ij += cot_beta_ij;
+
+		if (FMathd::Abs(w_ij) > DegenerateTol)
+		{
+			bDegenerated = true;
+			break;
+		}
+
+		Neighbours.Add(v_j);
+		Weights.Add(w_ij);
+		wSum += w_ij;
+	}
+
+	if (bOutputIsUniform != nullptr)
+	{
+		*bOutputIsUniform = bDegenerated;
+	}
+
+	int32 N = Neighbours.Num();
+	if (bDegenerated || fabs(wSum) < FMathd::ZeroTolerance)
+	{
+		double Weight = 1.0 / (double)N;
+		for (int32 k = 0; k < N; ++k)
+		{
+			BlendingFunc(Neighbours[k], Weight);
+		}
+	}
+	else
+	{
+		for (int32 k = 0; k < N; ++k )
+		{
+			BlendingFunc(Neighbours[k], Weights[k] / wSum);
+		}
+	}
+}
+
+
+void FMeshWeights::CotanWeightsBlendSafe(const FDynamicMesh3& Mesh, int32 VertexIndex, TFunctionRef<void(int32, double)> BlendingFunc, double DegenerateTol, bool* bFailedToUniform)
+{
+	TCotanWeightsBlendSafe(Mesh, VertexIndex, [&](int32 vid) { return Mesh.GetVertex(vid); }, BlendingFunc, DegenerateTol, bFailedToUniform);
+}
+
+
 template<typename GetPositionFuncType>
 double TVoronoiArea(const FDynamicMesh3& mesh, int32 v_i, GetPositionFuncType GetPositionFunc, double CotClampRange = 1000 /* about 0.05 deg */ )
 {
diff --git a/Engine/Plugins/Experimental/GeometryProcessing/Source/DynamicMesh/Public/MeshWeights.h b/Engine/Plugins/Experimental/GeometryProcessing/Source/DynamicMesh/Public/MeshWeights.h
--- a/Engine/Plugins/Experimental/GeometryProcessing/Source/DynamicMesh/Public/MeshWeights.h
+++ b/Engine/Plugins/Experimental/GeometryProcessing/Source/DynamicMesh/Public/MeshWeights.h
@@ -82,6 +82,17 @@ public:
 	static FVector3d CotanCentroidSafe(const FDynamicMesh3& Mesh, int32 VertexIndex, TFunctionRef<FVector3d(int32)> VertexPositionFunc, double DegenerateTol = 100.0, bool* bFailedToUniform = nullptr);
 
 
+	/**
+	 * Compute cotan-weighted blend for a vertex one-ring, with some weight analysis/clamping to avoid vertices getting "stuck"
+	 * in explicit integration/iterations. If failure is detected, Uniform centroid is returned, which does cause some tangential flow.
+	 * Equivalent to CotanCentroidSafe() if the weights are used to blend vertex positions, but can be used to weight other properties.
+	 * @param BlendingFunc BlendingFunc(NbrVertexID, Weight) will be called for every one-ring neighbour vertex
+	 * @param DegenerateTol if any weights are larger than this value, return uniform weights instead. Should be > 1.
+	 * @param bFailedToUniform will be set to true if non-null and result was clamped
+	 */
+	static void CotanWeightsBlendSafe(const FDynamicMesh3& Mesh, int32 VertexIndex, TFunctionRef<void(int32,double)> BlendingFunc, double DegenerateTol = 100.0, bool* bFailedToUniform = nullptr);
+
+
 
 	/**
 	 * Compute the Mixed Voronoi Area associated with a vertex.
diff --git a/Engine/Plugins/Experimental/GeometryProcessing/Source/GeometryAlgorithms/Private/MinVolumeBox3.cpp b/Engine/Plugins/Experimental/GeometryProcessing/Source/GeometryAlgorithms/Private/MinVolumeBox3.cpp
--- a/Engine/Plugins/Experimental/GeometryProcessing/Source/GeometryAlgorithms/Private/MinVolumeBox3.cpp
+++ b/Engine/Plugins/Experimental/GeometryProcessing/Source/GeometryAlgorithms/Private/MinVolumeBox3.cpp
@@ -22,6 +22,7 @@ THIRD_PARTY_INCLUDES_END
 
 #include "GteUtil.h"
 #include "Util/ProgressCancel.h"
+#include "Util/IteratorUtil.h"
 
 #include "ExplicitUseGeometryMathTypes.h"		// using UE::Geometry::(math types)
 using namespace UE::Geometry;
@@ -127,6 +128,31 @@ bool TMinVolumeBox3<RealType>::Solve(int32 NumPoints, TFunctionRef<FVector3<Real
 	return Internal->ComputeResult(Progress);
 }
 
+
+template<typename RealType>
+bool TMinVolumeBox3<RealType>::SolveSubsample(int32 NumPoints, int32 MaxPoints, TFunctionRef<FVector3<RealType>(int32)> GetPointFunc, bool bUseExactBox, FProgressCancel* Progress)
+{
+	if (NumPoints <= MaxPoints)
+	{
+		return Solve(NumPoints, GetPointFunc, bUseExactBox, Progress);
+	}
+
+	Initialize(MaxPoints, bUseExactBox);
+	check(Internal);
+
+	int32 k = 0;
+	FModuloIteration Iter(NumPoints);
+	int32 Index;
+	while (Iter.GetNextIndex(Index) && k < MaxPoints)
+	{
+		FVector3<RealType> Point = GetPointFunc(Index);
+		Internal->SetPoint(k, Point);
+		k++;
+	}
+
+	return Internal->ComputeResult(Progress);
+}
+
 template<typename RealType>
 bool TMinVolumeBox3<RealType>::IsSolutionAvailable() const
 {
diff --git a/Engine/Plugins/Experimental/GeometryProcessing/Source/GeometryAlgorithms/Public/MinVolumeBox3.h b/Engine/Plugins/Experimental/GeometryProcessing/Source/GeometryAlgorithms/Public/MinVolumeBox3.h
--- a/Engine/Plugins/Experimental/GeometryProcessing/Source/GeometryAlgorithms/Public/MinVolumeBox3.h
+++ b/Engine/Plugins/Experimental/GeometryProcessing/Source/GeometryAlgorithms/Public/MinVolumeBox3.h
@@ -26,11 +26,24 @@ class TMinVolumeBox3
 public:
 	/**
 	 * Calculate the minimal box for the given point set.
+	 * @param NumPoints number of points in the set, ie GetPointFunc can be called for any integer in range [0...NumPoints)
+	 * @param GetPointFunc function that returns a 3D point for a valid Index
 	 * @param bUseExactBox If true, high-precision number types are used for the minimal-box calculation, rather than doubles. This is *much* slower but more accurate (but not recommended).
 	 * @return true if minimal box was found
 	 */
 	bool Solve(int32 NumPoints, TFunctionRef<FVector3<RealType>(int32)> GetPointFunc, bool bUseExactBox = false, FProgressCancel* Progress = nullptr);
 
+	/**
+	 * Calculate the minimal box for a Subsampling of MaxPoints points of a point set
+	 * @param NumPoints number of points in the set, ie GetPointFunc can be called for any integer in range [0...NumPoints)
+	 * @param NumSamplePoints maximum number of points to sample from NumPoints. Currently a random-ish subset is selected.
+	 * @param GetPointFunc function that returns a 3D point for a valid Index
+	 * @param bUseExactBox If true, high-precision number types are used for the minimal-box calculation, rather than doubles. This is *much* slower but more accurate (but not recommended).
+	 * @return true if minimal box was found
+	 */
+	bool SolveSubsample(int32 NumPoints, int32 NumSamplePoints, TFunctionRef<FVector3<RealType>(int32)> GetPointFunc, bool bUseExactBox = false, FProgressCancel* Progress = nullptr);
+
+
 	/** @return true if minimal box is available */
 	bool IsSolutionAvailable() const;
 
diff --git a/Engine/Source/Runtime/GeometryCore/Private/DynamicMesh/DynamicMesh3_Queries.cpp b/Engine/Source/Runtime/GeometryCore/Private/DynamicMesh/DynamicMesh3_Queries.cpp
--- a/Engine/Source/Runtime/GeometryCore/Private/DynamicMesh/DynamicMesh3_Queries.cpp
+++ b/Engine/Source/Runtime/GeometryCore/Private/DynamicMesh/DynamicMesh3_Queries.cpp
@@ -743,16 +743,15 @@ void FDynamicMesh3::GetVtxOneRingCentroid(int vID, FVector3d& centroid) const
 }
 
 
-FFrame3d FDynamicMesh3::GetVertexFrame(int vID, bool bFrameNormalY) const
+FFrame3d FDynamicMesh3::GetVertexFrame(int vID, bool bFrameNormalY, FVector3d* UseNormal) const
 {
-	if (ensure(HasVertexNormals()) == false)
+	if (HasVertexNormals() == false && UseNormal == nullptr)
 	{
 		return FFrame3d();
 	}
 
 	FVector3d v = Vertices[vID];
-	const TDynamicVector<FVector3f>& Normals = VertexNormals.GetValue();
-	FVector3d normal = FVector3d(Normals[vID]);
+	FVector3d normal = (UseNormal != nullptr) ? Normalized(*UseNormal) : FVector3d(VertexNormals.GetValue()[vID]);
 	int eid = VertexEdgeLists.First(vID);
 	int ovi = GetOtherEdgeVertex(eid, vID);
 	FVector3d ov = Vertices[ovi];
diff --git a/Engine/Source/Runtime/GeometryCore/Public/BoxTypes.h b/Engine/Source/Runtime/GeometryCore/Public/BoxTypes.h
--- a/Engine/Source/Runtime/GeometryCore/Public/BoxTypes.h
+++ b/Engine/Source/Runtime/GeometryCore/Public/BoxTypes.h
@@ -333,6 +333,35 @@ struct TAxisAlignedBox3
 			FVector3<RealType>(TNumericLimits<RealType>::Max(), TNumericLimits<RealType>::Max(), TNumericLimits<RealType>::Max()) );
 	}
 
+	/**
+	 * Compute bounding box of 3D points returned by GetPoint(Index) for indices in range [0...MaxIndex)
+	 */
+	template<typename PointFunc>
+	static TAxisAlignedBox3<RealType> MakeBoundsFromIndices(int32 MaxIndex, PointFunc GetPoint)
+	{
+		TAxisAlignedBox3<RealType> Result = TAxisAlignedBox3<RealType>::Empty();
+		for (int32 Index = 0; Index < MaxIndex; ++Index)
+		{
+			Result.Contain(GetPoint(Index));
+		}
+		return Result;
+	}
+
+	/**
+	 * Compute bounding box of 3D points returned by GetPoint(Index) for indices in range for_each(IndexEnumerable)
+	 */
+	template<typename EnumerableIntType, typename PointFunc>
+	static TAxisAlignedBox3<RealType> MakeBoundsFromIndices(EnumerableIntType IndexEnumerable, PointFunc GetPoint)
+	{
+		TAxisAlignedBox3<RealType> Result = TAxisAlignedBox3<RealType>::Empty();
+		for (int32 Index : IndexEnumerable)
+		{
+			Result.Contain(GetPoint(Index));
+		}
+		return Result;
+	}
+
+
 	FVector3<RealType> Center() const
 	{
 		return FVector3<RealType>(
@@ -446,6 +475,11 @@ struct TAxisAlignedBox3
 		return (RealType)0.25 * (delta_x * delta_x + delta_y * delta_y + delta_z * delta_z);
 	}
 
+	RealType Dimension(int32 Index) const
+	{
+		return TMathUtil<RealType>::Max(Max[Index] - Min[Index], (RealType)0);
+	}
+
 	RealType Width() const
 	{
 		return TMathUtil<RealType>::Max(Max.X - Min.X, (RealType)0);
diff --git a/Engine/Source/Runtime/GeometryCore/Public/DynamicMesh/DynamicMesh3.h b/Engine/Source/Runtime/GeometryCore/Public/DynamicMesh/DynamicMesh3.h
--- a/Engine/Source/Runtime/GeometryCore/Public/DynamicMesh/DynamicMesh3.h
+++ b/Engine/Source/Runtime/GeometryCore/Public/DynamicMesh/DynamicMesh3.h
@@ -989,11 +989,13 @@ public:
 	/**
 	 * Compute a normal/tangent frame at vertex that is "stable" as long as
 	 * the mesh topology doesn't change, meaning that one axis of the frame
-	 * will be computed from projection of outgoing edge. Requires that vertex normals are available.
+	 * will be computed from projection of outgoing edge. 
 	 * By default, frame.Z is normal, and .X points along mesh edge.
-	 * If bFrameNormalY, then frame.Y is normal (X still points along mesh edge)
+	 * Requires that VertexNormals are available, or UseNormal != nullptr, otherwise returns default Frame3d
+	 * @param bFrameNormalY if true, then frame.Y is normal (X still points along mesh edge)
+	 * @param UseNormal if defined, this normal is used instead of VertexNormals normal
 	 */
-	FFrame3d GetVertexFrame(int VertexID, bool bFrameNormalY = false) const;
+	FFrame3d GetVertexFrame(int VertexID, bool bFrameNormalY = false, FVector3d* UseNormal = nullptr) const;
 
 	/** Calculate face normal of triangle */
 	FVector3d GetTriNormal(int TriangleID) const;
diff --git a/Engine/Source/Runtime/GeometryCore/Public/DynamicMesh/MeshIndexUtil.h b/Engine/Source/Runtime/GeometryCore/Public/DynamicMesh/MeshIndexUtil.h
--- a/Engine/Source/Runtime/GeometryCore/Public/DynamicMesh/MeshIndexUtil.h
+++ b/Engine/Source/Runtime/GeometryCore/Public/DynamicMesh/MeshIndexUtil.h
@@ -25,6 +25,23 @@ namespace UE
 		 */
 		GEOMETRYCORE_API void VertexToTriangleOneRing(const FDynamicMesh3* Mesh, const TArray<int>& VertexIDs, TSet<int>& TriangleIDsOut);
 
+
+		/**
+		 * Call SetType.Add(VertexPosition) for all valid vertex indices in for_each(Enumeration)
+		 */
+		template<typename EnumeratorType, typename SetType>
+		void CollectVertexPositions(const FDynamicMesh3& Mesh, EnumeratorType Enumeration, SetType& Output)
+		{
+			for (int32 vid : Enumeration)
+			{
+				if (Mesh.IsVertex(vid))
+				{
+					Output.Add(Mesh.GetVertex(vid));
+				}
+			}
+		}
+
+
 		/**
 		 * Walk around VertexID from FromTriangleID to next connected triangle if it exists, walking "away" from PrevTriangleID.
 		 * @param TrisConnectedFunc returns true if two triangles should be considered connected, to support breaking at seams/etc that are not in base mesh topology
diff --git a/Engine/Source/Runtime/GeometryCore/Public/Util/IteratorUtil.h b/Engine/Source/Runtime/GeometryCore/Public/Util/IteratorUtil.h
--- a/Engine/Source/Runtime/GeometryCore/Public/Util/IteratorUtil.h
+++ b/Engine/Source/Runtime/GeometryCore/Public/Util/IteratorUtil.h
@@ -400,5 +400,65 @@ public:
 };
 
 
+
+/**
+ * FModuloIteration is used to iterate over a range of indices [0,N) using modulo-arithmetic.
+ * The iteration proceeds as NextValue = (CurValue + ModuloValue) % N.
+ * As long as the ModuloValue is a prime number > N/2, then every integer in the sequence 0...N-1
+ * will appear exactly once before 0 re-appears (and in fact any index in the range can be used
+ * as the starting value). 
+ * 
+ * FModuloIteration computes in 64-bit with (by default) a large enough prime that will work
+ * for any 32-bit unsigned integer. If 64-bit iterations are needed, some larger primes can 
+ * be found here: https://en.wikipedia.org/wiki/P%C3%A9pin%27s_test
+ * 
+ * (The prime does not strictly need to be > N/2, any prime will work as long as it is not a divisor of N. 
+ *  And it doesn't even need to be a prime number, just a co-prime of N, ie GCD(N, ModuloValue) = 1. 
+ *  It is possible to check GCD relatively quickly to search for valid constants, for example if
+ *  many values were needed to use as seeds/etc)
+ * 
+ * Usage:
+ * 
+	FModuloIteration Iter(N);
+	uint32 Index;
+	while (Iter.GetNextIndex(Index)) { ... }
+ */
+struct FModuloIteration
+{
+	uint64 MaxIndex = 0;
+	uint64 ModuloPrime = 3208642561;		// first elite prime > max_unsigned_int/2
+	uint64 CurIndex = 0;
+	uint64 StartIndex = 0;
+	uint64 Count = 0;
+	uint64 ModuloNum = 1;
+
+	FModuloIteration(uint32 MaxIndexIn, uint32 StartIndexIn = 0, uint64 ModuloPrimeIn = 3208642561)
+	{
+		MaxIndex = (uint64)FMath::Max((uint32)0, MaxIndexIn);
+		StartIndex = (uint64)FMath::Max((uint32)0, StartIndexIn);
+		CurIndex = 0;
+		Count = 0;
+		ModuloNum = FMath::Max((uint64)1, MaxIndex);  // can't be zero or we hit integer-divide. If MaxIndex is 0 we will terminate on first iteration anyway
+		ModuloPrime = ModuloPrimeIn;
+		check(ModuloPrime > MaxIndex);
+	}
+
+	bool GetNextIndex(uint32& NextIndexOut)
+	{
+		NextIndexOut = (uint32)CurIndex;
+		CurIndex = (CurIndex + ModuloPrime) % ModuloNum;
+		return (Count++ != MaxIndex);
+	}
+
+	bool GetNextIndex(int32& NextIndexOut)
+	{
+		NextIndexOut = (int32)CurIndex;
+		CurIndex = (CurIndex + ModuloPrime) % ModuloNum;
+		return (Count++ != MaxIndex);
+	}
+};
+
+
+
 } // end namespace UE::Geometry
 } // end namespace UE
\ No newline at end of file
