ModelingTools: Refactor new bake path to bake evaluators.
Add support for variable float size and bake results per evaluator.

#rnx

diff --git a/Engine/Plugins/Experimental/GeometryProcessing/Source/DynamicMesh/Private/Sampling/MeshCurvatureMapBaker.cpp b/Engine/Plugins/Experimental/GeometryProcessing/Source/DynamicMesh/Private/Sampling/MeshCurvatureMapBaker.cpp
--- a/Engine/Plugins/Experimental/GeometryProcessing/Source/DynamicMesh/Private/Sampling/MeshCurvatureMapBaker.cpp
+++ b/Engine/Plugins/Experimental/GeometryProcessing/Source/DynamicMesh/Private/Sampling/MeshCurvatureMapBaker.cpp
@@ -1,7 +1,6 @@
 // Copyright Epic Games, Inc. All Rights Reserved.
 
 #include "Sampling/MeshCurvatureMapBaker.h"
-#include "Sampling/MeshMapBaker.h"
 #include "Sampling/SphericalFibonacci.h"
 #include "Sampling/Gaussians.h"
 #include "MeshCurvature.h"
@@ -13,14 +12,14 @@
 #include "ExplicitUseGeometryMathTypes.h"		// using UE::Geometry::(math types)
 using namespace UE::Geometry;
 
-void FMeshCurvatureMapBaker::CacheDetailCurvatures(const FDynamicMesh3* DetailMeshIn)
+void FMeshCurvatureMapBaker::CacheDetailCurvatures(const FDynamicMesh3* DetailMesh)
 {
 	if (! Curvatures)
 	{
 		Curvatures = MakeShared<FMeshVertexCurvatureCache>();
-		Curvatures->BuildAll(*DetailMeshIn);
+		Curvatures->BuildAll(*DetailMesh);
 	}
-	check(Curvatures->Num() == DetailMeshIn->MaxVertexID());
+	check(Curvatures->Num() == DetailMesh->MaxVertexID());
 
 }
 
@@ -29,7 +28,7 @@ void FMeshCurvatureMapBaker::Bake()
 {
 	const FMeshImageBakingCache* BakeCache = GetCache();
 	check(BakeCache);
-	DetailMesh = BakeCache->GetDetailMesh();
+	const FDynamicMesh3* DetailMesh = BakeCache->GetDetailMesh();
 
 	CacheDetailCurvatures(DetailMesh);
 
@@ -70,10 +69,12 @@ void FMeshCurvatureMapBaker::Bake()
 void FMeshCurvatureMapBaker::Bake_Single()
 {
 	const FMeshImageBakingCache* BakeCache = GetCache();
+	const FDynamicMesh3* DetailMesh = BakeCache->GetDetailMesh();
 	const FMeshVertexCurvatureCache& Cache = *Curvatures;
 
-	MinPreClamp = -TNumericLimits<double>::Max();
-	MaxPreClamp = TNumericLimits<double>::Max();
+
+	double MinPreClamp = -TNumericLimits<double>::Max();
+	double MaxPreClamp = TNumericLimits<double>::Max();
 	if (UseClampMode == EClampMode::Positive)
 	{
 		MinPreClamp = 0;
@@ -83,6 +84,23 @@ void FMeshCurvatureMapBaker::Bake_Single()
 		MaxPreClamp = 0;
 	}
 
+
+
+	auto GetCurvature = [&](int32 vid) {
+		switch (UseCurvatureType)
+		{
+		default:
+		case ECurvatureType::Mean:
+			return FMathd::Clamp(Cache[vid].Mean, MinPreClamp, MaxPreClamp);
+		case ECurvatureType::Gaussian:
+			return FMathd::Clamp(Cache[vid].Gaussian, MinPreClamp, MaxPreClamp);
+		case ECurvatureType::MaxPrincipal:
+			return FMathd::Clamp(Cache[vid].MaxPrincipal, MinPreClamp, MaxPreClamp);
+		case ECurvatureType::MinPrincipal:
+			return FMathd::Clamp(Cache[vid].MinPrincipal, MinPreClamp, MaxPreClamp);
+		}
+	};
+
 	FSampleSetStatisticsd CurvatureStats = Cache.MeanStats;
 	switch (UseCurvatureType)
 	{
@@ -107,13 +125,35 @@ void FMeshCurvatureMapBaker::Bake_Single()
 		ClampMax = RangeScale * OverrideRangeMax;
 	}
 	double MinClamp = MinRangeScale * ClampMax;
-	ClampRange = FInterval1d(MinClamp, ClampMax);
+	FInterval1d ClampRange(MinClamp, ClampMax);
+
+	// sample curvature of detail surface in tangent-space of base surface
+	auto CurvatureSampleFunction = [&](const FMeshImageBakingCache::FCorrespondenceSample& SampleData)
+	{
+		int32 DetailTriID = SampleData.DetailTriID;
+		if (DetailMesh->IsTriangle(DetailTriID))
+		{
+			FIndex3i DetailTri = DetailMesh->GetTriangle(DetailTriID);
+			double CurvatureA = GetCurvature(DetailTri.A);
+			double CurvatureB = GetCurvature(DetailTri.B);
+			double CurvatureC = GetCurvature(DetailTri.C);
+			double InterpCurvature = SampleData.DetailBaryCoords.X * CurvatureA
+				+ SampleData.DetailBaryCoords.Y * CurvatureB
+				+ SampleData.DetailBaryCoords.Z * CurvatureC;
+
+			return InterpCurvature;
+		}
+		return 0.0;
+	};
+
+
 
+	FVector3f NegativeColor, ZeroColor, PositiveColor;
 	GetColorMapRange(NegativeColor, ZeroColor, PositiveColor);
 
-	BakeCache->EvaluateSamples([this](const FVector2i& Coords, const FMeshImageBakingCache::FCorrespondenceSample& Sample)
+	BakeCache->EvaluateSamples([&](const FVector2i& Coords, const FMeshImageBakingCache::FCorrespondenceSample& Sample)
 	{
-		double Curvature = SampleFunction<FMeshImageBakingCache::FCorrespondenceSample>(Sample);
+		double Curvature = CurvatureSampleFunction(Sample);
 
 		double Sign = FMathd::Sign(Curvature);
 		float T = (float)ClampRange.GetT(FMathd::Abs(Curvature));
@@ -131,117 +171,26 @@ void FMeshCurvatureMapBaker::Bake_Multi()
 
 }
 
-double FMeshCurvatureMapBaker::GetCurvature(int32 vid)
-{
-	const FMeshVertexCurvatureCache& Cache = *Curvatures;
-	switch (UseCurvatureType)
-	{
-	default:
-	case ECurvatureType::Mean:
-		return FMathd::Clamp(Cache[vid].Mean, MinPreClamp, MaxPreClamp);
-	case ECurvatureType::Gaussian:
-		return FMathd::Clamp(Cache[vid].Gaussian, MinPreClamp, MaxPreClamp);
-	case ECurvatureType::MaxPrincipal:
-		return FMathd::Clamp(Cache[vid].MaxPrincipal, MinPreClamp, MaxPreClamp);
-	case ECurvatureType::MinPrincipal:
-		return FMathd::Clamp(Cache[vid].MinPrincipal, MinPreClamp, MaxPreClamp);
-	}
-}
 
-void FMeshCurvatureMapBaker::GetColorMapRange(FVector3f& NegativeColorOut, FVector3f& ZeroColorOut, FVector3f& PositiveColorOut)
+void FMeshCurvatureMapBaker::GetColorMapRange(FVector3f& NegativeColor, FVector3f& ZeroColor, FVector3f& PositiveColor)
 {
 	switch (UseColorMode)
 	{
 	default:
 	case EColorMode::BlackGrayWhite:
-		NegativeColorOut = FVector3f(0, 0, 0);
-		ZeroColorOut = FVector3f(0.5, 0.5, 0.5);
-		PositiveColorOut = FVector3f(1, 1, 1);
+		NegativeColor = FVector3f(0, 0, 0);
+		ZeroColor = FVector3f(0.5, 0.5, 0.5);
+		PositiveColor = FVector3f(1, 1, 1);
 		break;
 	case EColorMode::RedGreenBlue:
-		NegativeColorOut = FVector3f(1, 0, 0);
-		ZeroColorOut = FVector3f(0, 1, 0);
-		PositiveColorOut = FVector3f(0, 0, 1);
+		NegativeColor = FVector3f(1, 0, 0);
+		ZeroColor = FVector3f(0, 1, 0);
+		PositiveColor = FVector3f(0, 0, 1);
 		break;
 	case EColorMode::RedBlue:
-		NegativeColorOut = FVector3f(1, 0, 0);
-		ZeroColorOut = FVector3f(0, 0, 0);
-		PositiveColorOut = FVector3f(0, 0, 1);
-		break;
-	}
-}
-
-void FMeshCurvatureMapBaker::PreEvaluate(const FMeshMapBaker& Baker)
-{
-	DetailMesh = Baker.GetDetailMesh();
-	CacheDetailCurvatures(DetailMesh);
-
-	MinPreClamp = -TNumericLimits<double>::Max();
-	MaxPreClamp = TNumericLimits<double>::Max();
-	if (UseClampMode == EClampMode::Positive)
-	{
-		MinPreClamp = 0;
-	}
-	else if (UseClampMode == EClampMode::Negative)
-	{
-		MaxPreClamp = 0;
-	}
-
-	const FMeshVertexCurvatureCache& Cache = *Curvatures;
-	FSampleSetStatisticsd CurvatureStats = Cache.MeanStats;
-	switch (UseCurvatureType)
-	{
-	default:
-	case ECurvatureType::Mean:
-		CurvatureStats = Cache.MeanStats;
-		break;
-	case ECurvatureType::Gaussian:
-		CurvatureStats = Cache.GaussianStats;
-		break;
-	case ECurvatureType::MaxPrincipal:
-		CurvatureStats = Cache.MaxPrincipalStats;
-		break;
-	case ECurvatureType::MinPrincipal:
-		CurvatureStats = Cache.MinPrincipalStats;
+		NegativeColor = FVector3f(1, 0, 0);
+		ZeroColor = FVector3f(0, 0, 0);
+		PositiveColor = FVector3f(0, 0, 1);
 		break;
 	}
-
-	double ClampMax = RangeScale * (CurvatureStats.Mean + CurvatureStats.StandardDeviation);
-	if (bOverrideCurvatureRange)
-	{
-		ClampMax = RangeScale * OverrideRangeMax;
-	}
-	double MinClamp = MinRangeScale * ClampMax;
-	ClampRange = FInterval1d(MinClamp, ClampMax);
-
-	GetColorMapRange(NegativeColor, ZeroColor, PositiveColor);
-}
-
-FVector4f FMeshCurvatureMapBaker::EvaluateSample(const FMeshMapBaker& Baker, const FCorrespondenceSample& Sample)
-{
-	double Curvature = SampleFunction<FCorrespondenceSample>(Sample);
-	double Sign = FMathd::Sign(Curvature);
-	float T = (float)ClampRange.GetT(FMathd::Abs(Curvature));
-	FVector3f CurvatureColor = (Sign < 0) ?
-		Lerp(ZeroColor, NegativeColor, T) : Lerp(ZeroColor, PositiveColor, T);
-	return FVector4f(CurvatureColor.X, CurvatureColor.Y, CurvatureColor.Z, 1.0f);
-}
-
-template<class SampleType>
-double FMeshCurvatureMapBaker::SampleFunction(const SampleType& SampleData)
-{
-	int32 DetailTriID = SampleData.DetailTriID;
-	if (DetailMesh->IsTriangle(DetailTriID))
-	{
-		FIndex3i DetailTri = DetailMesh->GetTriangle(DetailTriID);
-		double CurvatureA = GetCurvature(DetailTri.A);
-		double CurvatureB = GetCurvature(DetailTri.B);
-		double CurvatureC = GetCurvature(DetailTri.C);
-		double InterpCurvature = SampleData.DetailBaryCoords.X * CurvatureA
-			+ SampleData.DetailBaryCoords.Y * CurvatureB
-			+ SampleData.DetailBaryCoords.Z * CurvatureC;
-
-		return InterpCurvature;
-	}
-	return 0.0;
-}
+}
\ No newline at end of file
diff --git a/Engine/Plugins/Experimental/GeometryProcessing/Source/DynamicMesh/Private/Sampling/MeshCurvatureMapEvaluator.cpp b/Engine/Plugins/Experimental/GeometryProcessing/Source/DynamicMesh/Private/Sampling/MeshCurvatureMapEvaluator.cpp
new file mode 100644
--- /dev/null
+++ b/Engine/Plugins/Experimental/GeometryProcessing/Source/DynamicMesh/Private/Sampling/MeshCurvatureMapEvaluator.cpp
@@ -0,0 +1,146 @@
+// Copyright Epic Games, Inc. All Rights Reserved.
+
+#include "Sampling/MeshCurvatureMapEvaluator.h"
+#include "Sampling/MeshMapBaker.h"
+#include "MeshCurvature.h"
+#include "MeshWeights.h"
+
+#include "ExplicitUseGeometryMathTypes.h"		// using UE::Geometry::(math types)
+using namespace UE::Geometry;
+
+void FMeshCurvatureMapEvaluator::CacheDetailCurvatures(const FDynamicMesh3* DetailMeshIn)
+{
+	if (!Curvatures)
+	{
+		Curvatures = MakeShared<FMeshVertexCurvatureCache>();
+		Curvatures->BuildAll(*DetailMeshIn);
+	}
+	check(Curvatures->Num() == DetailMeshIn->MaxVertexID());
+}
+
+void FMeshCurvatureMapEvaluator::Setup(const FMeshMapBaker& Baker, FEvaluationContext& Context)
+{
+	Context.Evaluate = &EvaluateSample;
+	Context.EvaluateDefault = &EvaluateDefault;
+	Context.EvalData = this;
+	Context.AccumulateMode = EAccumulateMode::Add;
+	Context.DataLayout = { EComponents::Float3 };
+
+	// Cache data from the baker
+	DetailMesh = Baker.GetDetailMesh();
+	CacheDetailCurvatures(DetailMesh);
+
+	MinPreClamp = -TNumericLimits<double>::Max();
+	MaxPreClamp = TNumericLimits<double>::Max();
+	if (UseClampMode == EClampMode::Positive)
+	{
+		MinPreClamp = 0;
+	}
+	else if (UseClampMode == EClampMode::Negative)
+	{
+		MaxPreClamp = 0;
+	}
+
+	const FMeshVertexCurvatureCache& Cache = *Curvatures;
+	FSampleSetStatisticsd CurvatureStats = Cache.MeanStats;
+	switch (UseCurvatureType)
+	{
+	default:
+	case ECurvatureType::Mean:
+		CurvatureStats = Cache.MeanStats;
+		break;
+	case ECurvatureType::Gaussian:
+		CurvatureStats = Cache.GaussianStats;
+		break;
+	case ECurvatureType::MaxPrincipal:
+		CurvatureStats = Cache.MaxPrincipalStats;
+		break;
+	case ECurvatureType::MinPrincipal:
+		CurvatureStats = Cache.MinPrincipalStats;
+		break;
+	}
+
+	double ClampMax = RangeScale * (CurvatureStats.Mean + CurvatureStats.StandardDeviation);
+	if (bOverrideCurvatureRange)
+	{
+		ClampMax = RangeScale * OverrideRangeMax;
+	}
+	double MinClamp = MinRangeScale * ClampMax;
+	ClampRange = FInterval1d(MinClamp, ClampMax);
+
+	GetColorMapRange(NegativeColor, ZeroColor, PositiveColor);
+}
+
+void FMeshCurvatureMapEvaluator::EvaluateSample(float*& Out, const FCorrespondenceSample& Sample, void* EvalData)
+{
+	FMeshCurvatureMapEvaluator* Eval = static_cast<FMeshCurvatureMapEvaluator*>(EvalData);
+	double Curvature = Eval->SampleFunction(Sample);
+	double Sign = FMathd::Sign(Curvature);
+	float T = (float)Eval->ClampRange.GetT(FMathd::Abs(Curvature));
+	FVector3f CurvatureColor = (Sign < 0) ?
+		Lerp(Eval->ZeroColor, Eval->NegativeColor, T) : Lerp(Eval->ZeroColor, Eval->PositiveColor, T);
+	WriteToBuffer(Out, CurvatureColor);
+}
+
+void FMeshCurvatureMapEvaluator::EvaluateDefault(float*& Out, void* EvalData)
+{
+	WriteToBuffer(Out, FVector3f::Zero());
+}
+
+double FMeshCurvatureMapEvaluator::SampleFunction(const FCorrespondenceSample& SampleData)
+{
+	int32 DetailTriID = SampleData.DetailTriID;
+	if (DetailMesh->IsTriangle(DetailTriID))
+	{
+		FIndex3i DetailTri = DetailMesh->GetTriangle(DetailTriID);
+		double CurvatureA = GetCurvature(DetailTri.A);
+		double CurvatureB = GetCurvature(DetailTri.B);
+		double CurvatureC = GetCurvature(DetailTri.C);
+		double InterpCurvature = SampleData.DetailBaryCoords.X * CurvatureA
+			+ SampleData.DetailBaryCoords.Y * CurvatureB
+			+ SampleData.DetailBaryCoords.Z * CurvatureC;
+
+		return InterpCurvature;
+	}
+	return 0.0;
+}
+
+double FMeshCurvatureMapEvaluator::GetCurvature(int32 vid)
+{
+	const FMeshVertexCurvatureCache& Cache = *Curvatures;
+	switch (UseCurvatureType)
+	{
+	default:
+	case ECurvatureType::Mean:
+		return FMathd::Clamp(Cache[vid].Mean, MinPreClamp, MaxPreClamp);
+	case ECurvatureType::Gaussian:
+		return FMathd::Clamp(Cache[vid].Gaussian, MinPreClamp, MaxPreClamp);
+	case ECurvatureType::MaxPrincipal:
+		return FMathd::Clamp(Cache[vid].MaxPrincipal, MinPreClamp, MaxPreClamp);
+	case ECurvatureType::MinPrincipal:
+		return FMathd::Clamp(Cache[vid].MinPrincipal, MinPreClamp, MaxPreClamp);
+	}
+}
+
+void FMeshCurvatureMapEvaluator::GetColorMapRange(FVector3f& NegativeColorOut, FVector3f& ZeroColorOut, FVector3f& PositiveColorOut)
+{
+	switch (UseColorMode)
+	{
+	default:
+	case EColorMode::BlackGrayWhite:
+		NegativeColorOut = FVector3f(0, 0, 0);
+		ZeroColorOut = FVector3f(0.5, 0.5, 0.5);
+		PositiveColorOut = FVector3f(1, 1, 1);
+		break;
+	case EColorMode::RedGreenBlue:
+		NegativeColorOut = FVector3f(1, 0, 0);
+		ZeroColorOut = FVector3f(0, 1, 0);
+		PositiveColorOut = FVector3f(0, 0, 1);
+		break;
+	case EColorMode::RedBlue:
+		NegativeColorOut = FVector3f(1, 0, 0);
+		ZeroColorOut = FVector3f(0, 0, 0);
+		PositiveColorOut = FVector3f(0, 0, 1);
+		break;
+	}
+}
diff --git a/Engine/Plugins/Experimental/GeometryProcessing/Source/DynamicMesh/Private/Sampling/MeshMapBaker.cpp b/Engine/Plugins/Experimental/GeometryProcessing/Source/DynamicMesh/Private/Sampling/MeshMapBaker.cpp
--- a/Engine/Plugins/Experimental/GeometryProcessing/Source/DynamicMesh/Private/Sampling/MeshMapBaker.cpp
+++ b/Engine/Plugins/Experimental/GeometryProcessing/Source/DynamicMesh/Private/Sampling/MeshMapBaker.cpp
@@ -127,6 +127,99 @@ static bool GetDetailMeshTrianglePoint_Nearest(
 	return false;
 }
 
+
+//
+// FMeshMapBaker
+//
+
+
+void FMeshMapBaker::InitBake()
+{
+	// Retrieve evaluation contexts and cache:
+	// - index lists of accumulation modes (BakeAccumulateLists)
+	// - evaluator to bake result offsets (BakeOffsets)
+	// - buffer size per sample (BakeSampleBufferSize)
+	int32 NumBakers = Bakers.Num();
+	BakeContexts.SetNum(NumBakers);
+	BakeOffsets.SetNumUninitialized(NumBakers + 1);
+	BakeAccumulateLists.SetNum((int32) FMeshMapEvaluator::EAccumulateMode::Last);
+	BakeSampleBufferSize = 0;
+	int32 Offset = 0;
+	for (int32 Idx = 0; Idx < NumBakers; ++Idx)
+	{
+		Bakers[Idx]->Setup(*this, BakeContexts[Idx]);
+		checkSlow(BakeContexts[Idx].Evaluate != nullptr && BakeContexts[Idx].EvaluateDefault != nullptr);
+		checkSlow(BakeContexts[Idx].DataLayout.Num() > 0);
+		int32 NumData = BakeContexts[Idx].DataLayout.Num();
+		for (int32 DataIdx = 0; DataIdx < NumData; ++DataIdx)
+		{
+			BakeSampleBufferSize += (int32) BakeContexts[Idx].DataLayout[DataIdx];
+		}
+		BakeOffsets[Idx] = Offset;
+		Offset += NumData;
+		BakeAccumulateLists[(int32) BakeContexts[Idx].AccumulateMode].Add(Idx);
+	}
+	BakeOffsets[NumBakers] = Offset;
+
+	// Initialize our BakeResults list and cache offsets into the sample buffer
+	// per bake result
+	const int32 NumResults = Offset;
+	BakeResults.SetNum(NumResults);
+	BakeSampleOffsets.SetNumUninitialized(NumResults + 1);
+	int32 SampleOffset = 0;
+	for (int32 Idx = 0; Idx < NumBakers; ++Idx)
+	{
+		int32 NumData = BakeContexts[Idx].DataLayout.Num();
+		for (int32 DataIdx = 0; DataIdx < NumData; ++DataIdx)
+		{
+			int32 ResultIdx = BakeOffsets[Idx] + DataIdx;
+			BakeResults[ResultIdx] = MakeUnique<TImageBuilder<FVector4f>>();
+			BakeResults[ResultIdx]->SetDimensions(Dimensions);
+
+			BakeSampleOffsets[ResultIdx] = SampleOffset;
+			SampleOffset += (int32) BakeContexts[Idx].DataLayout[DataIdx];
+		}
+	}
+	BakeSampleOffsets[NumResults] = SampleOffset;
+
+	InitBakeDefaults();
+
+	for (int32 Idx = 0; Idx < NumResults; ++Idx)
+	{
+		BakeResults[Idx]->Clear(BakeDefaultColors[Idx]);
+	}
+}
+
+void FMeshMapBaker::InitBakeDefaults()
+{
+	// Cache default float buffer and colors for each bake result.
+	checkSlow(BakeSampleBufferSize > 0);
+	BakeDefaults.SetNumUninitialized(BakeSampleBufferSize);
+	float* Buffer = BakeDefaults.GetData();
+	float* BufferPtr = Buffer;
+
+	const int32 NumBakers = Bakers.Num();
+	for (int32 Idx = 0; Idx < NumBakers; ++Idx)
+	{
+		BakeContexts[Idx].EvaluateDefault(BufferPtr, BakeContexts[Idx].EvalData);
+	}
+	checkSlow((BufferPtr - Buffer) == BakeSampleBufferSize);
+
+	BufferPtr = Buffer;
+	const int32 NumBakeResults = BakeResults.Num();
+	BakeDefaultColors.SetNumUninitialized(NumBakeResults);
+	for (int32 Idx = 0; Idx < NumBakers; ++Idx)
+	{
+		const int32 NumData = BakeContexts[Idx].DataLayout.Num();
+		for (int32 DataIdx = 0; DataIdx < NumData; ++DataIdx)
+		{
+			const int32 ResultIdx = BakeOffsets[Idx] + DataIdx;
+			BakeDefaultColors[ResultIdx] = FloatToPixel(BufferPtr, BakeContexts[Idx].DataLayout[DataIdx], 1.0f);
+		}
+	}
+	checkSlow((BufferPtr - Buffer) == BakeSampleBufferSize);
+}
+
 void FMeshMapBaker::Bake()
 {
 	if (Bakers.IsEmpty())
@@ -134,16 +227,7 @@ void FMeshMapBaker::Bake()
 		return;
 	}
 
-	// Initialize Bake targets and results.
-	BakeResults.SetNum(Bakers.Num());
-	for (int32 Idx = 0; Idx < Bakers.Num(); Idx++)
-	{
-		BakeResults[Idx] = MakeShared<TImageBuilder<FVector4f>, ESPMode::ThreadSafe>();
-		BakeResults[Idx]->SetDimensions(Dimensions);
-		FVector4f ClearColor = Bakers[Idx]->DefaultSample();
-		BakeResults[Idx]->Clear(ClearColor);
-		Bakers[Idx]->PreEvaluate(*this);
-	}
+	InitBake();
 
 	// Generate UV space mesh
 	const FDynamicMesh3* Mesh = TargetMesh;
@@ -172,9 +256,9 @@ void FMeshMapBaker::Bake()
 		UseStrategy = ECorrespondenceStrategy::NearestPoint;
 	}
 
-	// this sampler finds the correspondence between base surface and detail surface
-	DetailMeshSampler.Initialize(Mesh, UVOverlay, EMeshSurfaceSamplerQueryType::TriangleAndUV, FMeshImageBaker::FCorrespondenceSample(),
-		[Mesh, NormalOverlay, UseStrategy, this](const FMeshUVSampleInfo& SampleInfo, FMeshImageBaker::FCorrespondenceSample& ValueOut)
+	// This sampler finds the correspondence between target surface and detail surface.
+	DetailMeshSampler.Initialize(Mesh, UVOverlay, EMeshSurfaceSamplerQueryType::TriangleAndUV, FMeshMapEvaluator::FCorrespondenceSample(),
+		[Mesh, NormalOverlay, UseStrategy, this](const FMeshUVSampleInfo& SampleInfo, FMeshMapEvaluator::FCorrespondenceSample& ValueOut)
 	{
 		NormalOverlay->GetTriBaryInterpolate<double>(SampleInfo.TriangleIndex, &SampleInfo.BaryCoords.X, &ValueOut.BaseNormal.X);
 		Normalize(ValueOut.BaseNormal);
@@ -193,21 +277,22 @@ void FMeshMapBaker::Bake()
 			bool bFoundTri = GetDetailMeshTrianglePoint_Nearest(*DetailMesh, *DetailSpatial, SampleInfo.SurfacePoint,
 				ValueOut.DetailTriID, ValueOut.DetailBaryCoords);
 		}
-		else	// fall back to raycast strategy
+		else	// Fall back to raycast strategy
 		{
 			double SampleThickness = this->GetThickness();		// could modulate w/ a map here...
 
-			// find detail mesh triangle point
+			// Find detail mesh triangle point
 			bool bFoundTri = GetDetailMeshTrianglePoint_Raycast(*DetailMesh, *DetailSpatial, SampleInfo.SurfacePoint, RayDir,
 				ValueOut.DetailTriID, ValueOut.DetailBaryCoords, SampleThickness,
 				(UseStrategy == ECorrespondenceStrategy::RaycastStandardThenNearest));
 		}
-
 	});
 
 	// Setup image tiling
 	FImageTiling Tiles(Dimensions, TileSize, TileSize, 0);
-	ParallelFor(Tiles.Num(), [this, &Tiles](int32 TileIdx)
+	const int32 NumBakers = Bakers.Num();
+	const int32 NumResults = BakeResults.Num();
+	ParallelFor(Tiles.Num(), [this, &Tiles, &NumBakers, &NumResults](int32 TileIdx)
 	{
 		FImageDimensions Tile = Tiles.GetTile(TileIdx);
 
@@ -216,102 +301,216 @@ void FMeshMapBaker::Bake()
 		OccupancyMap.Initialize(Tile, Multisampling);
 		OccupancyMap.ComputeFromUVSpaceMesh(FlatMesh, [this](int32 TriangleID) { return FlatMesh.GetTriangleGroup(TriangleID); });
 
-		// calculate interior texels
-		for (int32 ImgY = 0; ImgY < Tile.GetHeight(); ImgY++)
+		FMeshMapTileBuffer TileBuffer(Tile, BakeSampleBufferSize);
+
+		// Calculate interior texels
+		for (FVector2i Texel(0,0); Texel.Y < Tile.GetHeight(); ++Texel.Y)
 		{
-			for (int32 ImgX = 0; ImgX < Tile.GetWidth(); ImgX++)
+			for (Texel.X = 0; Texel.X < Tile.GetWidth(); ++Texel.X)
 			{
 				if (CancelF())
 				{
 					return;
 				}
-
-				BakePixel(OccupancyMap, Tile, ImgX, ImgY);
+				BakePixel(TileBuffer, OccupancyMap, Tile, Texel);
 			}
 		}
 
-		for (int64 k = 0; k < OccupancyMap.GutterTexels.Num(); k++)
+		// Convert TileBuffer data to color data.
+		for (FVector2i Texel(0,0); Texel.Y < Tile.GetHeight(); ++Texel.Y)
 		{
-			TPair<int64, int64> GutterTexel = OccupancyMap.GutterTexels[k];
-			for (int32 Idx = 0; Idx < Bakers.Num(); Idx++)
+			for (Texel.X = 0; Texel.X < Tile.GetWidth(); ++Texel.X)
 			{
-				BakeResults[Idx]->CopyPixel(GutterTexel.Value, GutterTexel.Key);
+				if (CancelF())
+				{
+					return;
+				}
+
+				const int64 PixelLinearIdx = Tile.GetIndex(Texel);
+				const float& PixelWeight = TileBuffer.GetPixelWeight(PixelLinearIdx);
+				if (PixelWeight > 0.0)
+				{
+					const int64 ImageLinearIdx = Tile.GetSourceIndex(Texel);
+					float* PixelBuffer = TileBuffer.GetPixel(PixelLinearIdx);
+					const float OneOverPixelWeight = 1.0f / PixelWeight;
+
+					auto WriteToPixel = [this, &PixelBuffer, &ImageLinearIdx](TArray<int32>& BakeIds, float Weight)
+					{
+						for (int32 Idx : BakeIds)
+						{
+							const FMeshMapEvaluator::FEvaluationContext& Context = BakeContexts[Idx];
+							const int32 NumData = Context.DataLayout.Num();
+							const int32 ResultOffset = BakeOffsets[Idx];
+							for (int32 DataIdx = 0; DataIdx < NumData; ++DataIdx)
+							{
+								const int32 ResultIdx = ResultOffset + DataIdx;
+								const int32 Offset = BakeSampleOffsets[ResultIdx];
+								float* BufferPtr = &PixelBuffer[Offset];
+								const FMeshMapEvaluator::EComponents Stride = Context.DataLayout[DataIdx];
+								FVector4f& Pixel = BakeResults[ResultIdx]->GetPixel(ImageLinearIdx);
+								Pixel = FloatToPixel(BufferPtr, Stride, Weight);
+							}
+						}
+					};
+
+					WriteToPixel(BakeAccumulateLists[(int32)FMeshMapEvaluator::EAccumulateMode::Add], OneOverPixelWeight);
+					WriteToPixel(BakeAccumulateLists[(int32)FMeshMapEvaluator::EAccumulateMode::Overwrite], 1.0f);
+				}
 			}
 		}
-	}, !bParallel ? EParallelForFlags::ForceSingleThread : EParallelForFlags::None);
-
-	for (int32 Idx = 0; Idx < Bakers.Num(); Idx++)
-	{
-		Bakers[Idx]->PostEvaluate(*this, *BakeResults[Idx]);
-	}
-}
-
-void FMeshMapBaker::BakePixel(FImageOccupancyMap& OccupancyMap, const FImageDimensions& Tile, int32 ImgX, int32 ImgY)
-{
-	FVector2i ImageCoords = Tile.GetSourceCoords(ImgX, ImgY);
 
-	// If this texel contains any interior samples, we will accumulate
-	// the result, reset the clear color to black.
-	if (OccupancyMap.TexelNumSamples(Tile.GetIndex(ImgX, ImgY)) > 0)
-	{
-		for (int32 Idx = 0; Idx < Bakers.Num(); Idx++)
+		for (int64 GutterIdx = 0; GutterIdx < OccupancyMap.GutterTexels.Num(); ++GutterIdx)
 		{
-			if (Bakers[Idx]->SupportsMultisampling())
+			int64 GutterPixelTo;
+			int64 GutterPixelFrom;
+			Tie(GutterPixelTo, GutterPixelFrom) = OccupancyMap.GutterTexels[GutterIdx];
+			for (int32 Idx = 0; Idx < NumResults; Idx++)
 			{
-				BakeResults[Idx]->SetPixel(ImageCoords, FVector4f::Zero());
+				BakeResults[Idx]->CopyPixel(GutterPixelFrom, GutterPixelTo);
 			}
 		}
-	}
-	else
+	}, !bParallel ? EParallelForFlags::ForceSingleThread : EParallelForFlags::None);
+}
+
+void FMeshMapBaker::BakePixel(FMeshMapTileBuffer& TileBuffer, FImageOccupancyMap& OccupancyMap, const FImageDimensions& Tile, const FVector2i& TileCoords)
+{
+	if (OccupancyMap.TexelNumSamples(Tile.GetIndex(TileCoords)) == 0)
 	{
 		return;
 	}
 
-	for (int32 SampleIdx = 0; SampleIdx < OccupancyMap.Multisampler.Num(); SampleIdx++)
+	const int64 TilePixelLinearIdx = Tile.GetIndex(TileCoords);
+	float* PixelBuffer = TileBuffer.GetPixel(TilePixelLinearIdx);
+	float& PixelWeight = TileBuffer.GetPixelWeight(TilePixelLinearIdx);
+
+	const FVector2i ImageCoords = Tile.GetSourceCoords(TileCoords.X, TileCoords.Y);
+	const int32 NumSamples = OccupancyMap.Multisampler.Num();
+	const int32 NumEvaluatorAdd = BakeAccumulateLists[(int32)FMeshMapEvaluator::EAccumulateMode::Add].Num();
+	for (int32 SampleIdx = 0; SampleIdx < NumSamples; ++SampleIdx)
 	{
-		int64 LinearIdx = Tile.GetIndex(ImgX, ImgY) * OccupancyMap.Multisampler.Num() + SampleIdx;
-		double OneOverNumSamples = 1.0 / (double)OccupancyMap.Multisampler.Num();
-		if (OccupancyMap.IsInterior(LinearIdx) == false)
+		const int64 LinearIdx = TilePixelLinearIdx * OccupancyMap.Multisampler.Num() + SampleIdx;
+		if (!OccupancyMap.IsInterior(LinearIdx))
 		{
-			for (int32 Idx = 0; Idx < Bakers.Num(); Idx++)
+			PixelWeight += NumEvaluatorAdd > 0;
+			for (int32 Idx : BakeAccumulateLists[(int32) FMeshMapEvaluator::EAccumulateMode::Add] )
 			{
-				if (Bakers[Idx]->SupportsMultisampling())
+				const int32 NumData = BakeContexts[Idx].DataLayout.Num();
+				for (int32 DataIdx = 0; DataIdx < NumData; ++DataIdx)
 				{
-					FVector4f Result = BakeResults[Idx]->GetPixel(ImageCoords);
-					Result += Bakers[Idx]->DefaultSample() * OneOverNumSamples;
-					BakeResults[Idx]->SetPixel(ImageCoords, Result);
+					const int32 ResultIdx = BakeOffsets[Idx] + DataIdx;
+					const int32 Offset = BakeSampleOffsets[ResultIdx];
+					const int32 Stride = (int32) BakeContexts[Idx].DataLayout[DataIdx];
+					for (int32 BufIdx = 0; BufIdx < Stride; ++BufIdx)
+					{
+						PixelBuffer[Offset + BufIdx] += BakeDefaults[Offset + BufIdx];
+					}
 				}
 			}
 		}
 		else
 		{
-			FVector2d UVPosition = (FVector2d)OccupancyMap.TexelQueryUV[LinearIdx];
-			int32 UVTriangleID = OccupancyMap.TexelQueryTriangle[LinearIdx];
+			const FVector2d UVPosition = (FVector2d)OccupancyMap.TexelQueryUV[LinearIdx];
+			const int32 UVTriangleID = OccupancyMap.TexelQueryTriangle[LinearIdx];
 
-			FMeshImageBaker::FCorrespondenceSample Sample;
+			FMeshMapEvaluator::FCorrespondenceSample Sample;
 			DetailMeshSampler.SampleUV(UVTriangleID, UVPosition, Sample);
 
-			for (int32 Idx = 0; Idx < Bakers.Num(); Idx++)
+			BakeSample(TileBuffer, Sample, Tile, TileCoords, ImageCoords, 1.0f);
+		}
+	}
+}
+
+void FMeshMapBaker::BakeSample(FMeshMapTileBuffer& TileBuffer, const FMeshMapEvaluator::FCorrespondenceSample& Sample, const FImageDimensions& Tile,
+	const FVector2i& TileCoords, const FVector2i& ImageCoords, const float& SampleWeight)
+{
+	const int64 TilePixelLinearIdx = Tile.GetIndex(TileCoords);
+	float* PixelBuffer = TileBuffer.GetPixel(TilePixelLinearIdx);
+	float& PixelWeight = TileBuffer.GetPixelWeight(TilePixelLinearIdx);
+	PixelWeight += SampleWeight;
+
+	float* Buffer = static_cast<float*>(FMemory_Alloca(sizeof(float) * BakeSampleBufferSize));
+	float* BufferPtr = Buffer;
+	const int32 NumBakers = Bakers.Num();
+	for (int32 Idx = 0; Idx < NumBakers; ++Idx)
+	{
+		BakeContexts[Idx].Evaluate(BufferPtr, Sample, BakeContexts[Idx].EvalData);
+	}
+	checkSlow((BufferPtr - Buffer) == BakeSampleBufferSize);
+
+	auto AddFn = [PixelBuffer, Buffer](const int32 Idx) -> void
+	{
+		PixelBuffer[Idx] += Buffer[Idx];
+	};
+
+	auto OverwriteFn = [PixelBuffer, Buffer](const int32 Idx) -> void
+	{
+		PixelBuffer[Idx] = Buffer[Idx];
+	};
+
+	auto AccumulateFn = [this](const TArray<int32>& BakeIds, auto&& Op)
+	{
+		for (int32 Idx : BakeIds)
+		{
+			const FMeshMapEvaluator::FEvaluationContext& Context = BakeContexts[Idx];
+			const int32 NumData = Context.DataLayout.Num();
+			const int32 ResultOffset = BakeOffsets[Idx];
+			for (int32 DataIdx = 0; DataIdx < NumData; ++DataIdx)
 			{
-				FVector4f SampleResult = Bakers[Idx]->EvaluateSample(*this, Sample);
-				if (!Bakers[Idx]->SupportsMultisampling())
-				{
-					// TODO/FIXME: Can be more correct & optimal here by not multisampling at all.
-					// Currently the samples are precomputed in the occupancy map.
-					BakeResults[Idx]->SetPixel(ImageCoords, SampleResult);
-				}
-				else
+				const int32 ResultIdx = ResultOffset + DataIdx;
+				const int32 Offset = BakeSampleOffsets[ResultIdx];
+				const int32 Stride = (int32) Context.DataLayout[DataIdx];
+				for (int32 BufIdx = Offset; BufIdx < Offset + Stride; ++BufIdx)
 				{
-					FVector4f Result = BakeResults[Idx]->GetPixel(ImageCoords);
-					Result += SampleResult * OneOverNumSamples;
-					BakeResults[Idx]->SetPixel(ImageCoords, Result);
+					Op(BufIdx);
 				}
 			}
 		}
+	};
+
+	AccumulateFn(BakeAccumulateLists[(int32)FMeshMapEvaluator::EAccumulateMode::Add], AddFn);
+	AccumulateFn(BakeAccumulateLists[(int32)FMeshMapEvaluator::EAccumulateMode::Overwrite], OverwriteFn);
+}
+
+FVector4f FMeshMapBaker::FloatToPixel(float*& Buffer, FMeshMapEvaluator::EComponents Stride, float Weight)
+{
+	// TODO: Decompose this method into non-branching parts.
+	FVector4f Result = FVector4f(0.0, 0.0, 0.0, 1.0f);
+	switch (Stride)
+	{
+	case FMeshMapEvaluator::EComponents::Float1:
+	{
+		float X = Buffer[0] * Weight;
+		Result = FVector4f(X, X, X, 1.0f);
+		break;
+	}
+	case FMeshMapEvaluator::EComponents::Float2:
+	{
+		FVector2f Vec(Buffer[0], Buffer[1]);
+		Vec *= Weight;
+		Result = FVector4f(Vec.X, Vec.Y, 0.0f, 1.0f);
+		break;
 	}
+	case FMeshMapEvaluator::EComponents::Float3:
+	{
+		FVector3f Vec(Buffer[0], Buffer[1], Buffer[2]);
+		Vec *= Weight;
+		Result = FVector4f(Vec.X, Vec.Y, Vec.Z, 1.0f);
+		break;
+	}
+	case FMeshMapEvaluator::EComponents::Float4:
+	{
+		Result = FVector4f(Buffer[0], Buffer[1], Buffer[2], Buffer[3]);
+		Result *= Weight;
+		break;
+	}
+	default:
+		break;
+	}
+	Buffer += static_cast<int32>(Stride);
+	return Result;
 }
 
-int32 FMeshMapBaker::AddBaker(TSharedPtr<FMeshImageBaker> Sampler)
+int32 FMeshMapBaker::AddBaker(TSharedPtr<FMeshMapEvaluator> Sampler)
 {
 	return Bakers.Add(Sampler);
 }
@@ -322,6 +521,13 @@ void FMeshMapBaker::Reset()
 	BakeResults.Empty();
 }
 
+const TArrayView<TUniquePtr<TImageBuilder<FVector4f>>> FMeshMapBaker::GetBakeResults(int32 BakerIdx)
+{
+	int32 ResultIdx = BakeOffsets[BakerIdx];
+	int32 NumResults = BakeOffsets[BakerIdx + 1] - ResultIdx;
+	return TArrayView<TUniquePtr<TImageBuilder<FVector4f>>>(&BakeResults[ResultIdx], NumResults);
+}
+
 void FMeshMapBaker::SetDetailMesh(const FDynamicMesh3* Mesh, const FDynamicMeshAABBTree3* Spatial)
 {
 	DetailMesh = Mesh;
diff --git a/Engine/Plugins/Experimental/GeometryProcessing/Source/DynamicMesh/Private/Sampling/MeshNormalMapBaker.cpp b/Engine/Plugins/Experimental/GeometryProcessing/Source/DynamicMesh/Private/Sampling/MeshNormalMapBaker.cpp
--- a/Engine/Plugins/Experimental/GeometryProcessing/Source/DynamicMesh/Private/Sampling/MeshNormalMapBaker.cpp
+++ b/Engine/Plugins/Experimental/GeometryProcessing/Source/DynamicMesh/Private/Sampling/MeshNormalMapBaker.cpp
@@ -1,7 +1,6 @@
 // Copyright Epic Games, Inc. All Rights Reserved.
 
 #include "Sampling/MeshNormalMapBaker.h"
-#include "Sampling/MeshMapBaker.h"
 
 #include "ExplicitUseGeometryMathTypes.h"		// using UE::Geometry::(math types)
 using namespace UE::Geometry;
@@ -10,78 +9,62 @@ void FMeshNormalMapBaker::Bake()
 {
 	const FMeshImageBakingCache* BakeCache = GetCache();
 	check(BakeCache);
-	DetailMesh = BakeCache->GetDetailMesh();
-	DetailNormalOverlay = BakeCache->GetDetailNormals();
+	const FDynamicMesh3* DetailMesh = BakeCache->GetDetailMesh();
+	const FDynamicMeshNormalOverlay* DetailNormalOverlay = BakeCache->GetDetailNormals();
 	check(DetailNormalOverlay);
 
 	// make sure we have per-triangle tangents computed so that GetInterpolatedTriangleTangent() below works
 	int32 MaxTriID = BakeCache->GetBakeTargetMesh()->MaxTriangleID();
 	check( BaseMeshTangents->GetTangents().Num() >= 3*BakeCache->GetBakeTargetMesh()->MaxTriangleID());
 
+	auto NormalSampleFunction = [&](const FMeshImageBakingCache::FCorrespondenceSample& SampleData)
+	{
+		int32 DetailTriID = SampleData.DetailTriID;
+		if (DetailMesh->IsTriangle(DetailTriID))
+		{
+			// get tangents on base mesh
+			FVector3d BaseTangentX, BaseTangentY;
+			BaseMeshTangents->GetInterpolatedTriangleTangent(
+				SampleData.BaseSample.TriangleIndex, 
+				SampleData.BaseSample.BaryCoords, 
+				BaseTangentX, BaseTangentY);
+
+			// sample normal on detail mesh
+			FVector3d DetailNormal;
+			DetailNormalOverlay->GetTriBaryInterpolate<double>(DetailTriID, &SampleData.DetailBaryCoords.X, &DetailNormal.X);
+			Normalize(DetailNormal);
+
+			// compute normal in tangent space
+			double dx = DetailNormal.Dot(BaseTangentX);
+			double dy = DetailNormal.Dot(BaseTangentY);
+			double dz = DetailNormal.Dot(SampleData.BaseNormal);
+
+			return FVector3f((float)dx, (float)dy, (float)dz);
+		}
+		return FVector3f::UnitZ();
+	};
+
+
 	NormalsBuilder = MakeUnique<TImageBuilder<FVector3f>>();
 	NormalsBuilder->SetDimensions(BakeCache->GetDimensions());
 	FVector3f DefaultMapNormal = (DefaultNormal + FVector3f::One()) * 0.5f;
 	NormalsBuilder->Clear(DefaultMapNormal);
 
+
 	BakeCache->EvaluateSamples([&](const FVector2i& Coords, const FMeshImageBakingCache::FCorrespondenceSample& Sample)
 	{
-		FVector3f RelativeDetailNormal = SampleFunction<FMeshImageBakingCache::FCorrespondenceSample>(Sample);
+		FVector3f RelativeDetailNormal = NormalSampleFunction(Sample);
 		FVector3f MapNormal = (RelativeDetailNormal + FVector3f::One()) * 0.5f;
 		NormalsBuilder->SetPixel(Coords, MapNormal);
 	});
 
 	const FImageOccupancyMap& Occupancy = *BakeCache->GetOccupancyMap();
+
 	for (int64 k = 0; k < Occupancy.GutterTexels.Num(); k++)
 	{
 		TPair<int64, int64> GutterTexel = Occupancy.GutterTexels[k];
 		NormalsBuilder->CopyPixel(GutterTexel.Value, GutterTexel.Key);
 	}
-}
 
-void FMeshNormalMapBaker::PreEvaluate(const FMeshMapBaker& Baker)
-{
-	DetailMesh = Baker.GetDetailMesh();
-	DetailNormalOverlay = Baker.GetDetailMeshNormals();
-	check(DetailNormalOverlay);
-	BaseMeshTangents = Baker.GetTargetMeshTangents().Get();
-
-	// make sure we have per-triangle tangents computed so that GetInterpolatedTriangleTangent() below works
-	int32 MaxTriID = Baker.GetTargetMesh()->MaxTriangleID();
-	check(BaseMeshTangents->GetTangents().Num() >= 3 * Baker.GetTargetMesh()->MaxTriangleID());
-}
 
-FVector4f FMeshNormalMapBaker::EvaluateSample(const FMeshMapBaker& Baker, const FCorrespondenceSample& Sample)
-{
-	FVector3f SampleResult = SampleFunction<FCorrespondenceSample>(Sample);
-	// Map normal space [-1,1] to floating point color space [0,1]
-	FVector3f Result = (SampleResult + FVector3f::One()) * 0.5f;
-	return FVector4f(Result.X, Result.Y, Result.Z, 1.0f);
-}
-
-template<class SampleType>
-FVector3f FMeshNormalMapBaker::SampleFunction(const SampleType& SampleData)
-{
-	int32 DetailTriID = SampleData.DetailTriID;
-	if (DetailMesh->IsTriangle(DetailTriID))
-	{
-		// get tangents on base mesh
-		FVector3d BaseTangentX, BaseTangentY;
-		BaseMeshTangents->GetInterpolatedTriangleTangent(
-			SampleData.BaseSample.TriangleIndex,
-			SampleData.BaseSample.BaryCoords,
-			BaseTangentX, BaseTangentY);
-
-		// sample normal on detail mesh
-		FVector3d DetailNormal;
-		DetailNormalOverlay->GetTriBaryInterpolate<double>(DetailTriID, &SampleData.DetailBaryCoords.X, &DetailNormal.X);
-		Normalize(DetailNormal);
-
-		// compute normal in tangent space
-		double dx = DetailNormal.Dot(BaseTangentX);
-		double dy = DetailNormal.Dot(BaseTangentY);
-		double dz = DetailNormal.Dot(SampleData.BaseNormal);
-
-		return FVector3f((float)dx, (float)dy, (float)dz);
-	}
-	return FVector3f::UnitZ();
-}
+}
\ No newline at end of file
diff --git a/Engine/Plugins/Experimental/GeometryProcessing/Source/DynamicMesh/Private/Sampling/MeshNormalMapEvaluator.cpp b/Engine/Plugins/Experimental/GeometryProcessing/Source/DynamicMesh/Private/Sampling/MeshNormalMapEvaluator.cpp
new file mode 100644
--- /dev/null
+++ b/Engine/Plugins/Experimental/GeometryProcessing/Source/DynamicMesh/Private/Sampling/MeshNormalMapEvaluator.cpp
@@ -0,0 +1,67 @@
+// Copyright Epic Games, Inc. All Rights Reserved.
+
+#include "Sampling/MeshNormalMapEvaluator.h"
+#include "Sampling/MeshMapBaker.h"
+
+#include "ExplicitUseGeometryMathTypes.h"		// using UE::Geometry::(math types)
+using namespace UE::Geometry;
+
+
+void FMeshNormalMapEvaluator::Setup(const FMeshMapBaker& Baker, FEvaluationContext& Context)
+{
+	Context.Evaluate = &EvaluateSample;
+	Context.EvaluateDefault = &EvaluateDefault;
+	Context.EvalData = this;
+	Context.AccumulateMode = EAccumulateMode::Add;
+	Context.DataLayout = { EComponents::Float3 };
+
+	// Cache data from the baker
+	DetailMesh = Baker.GetDetailMesh();
+	DetailNormalOverlay = Baker.GetDetailMeshNormals();
+	checkSlow(DetailNormalOverlay);
+	BaseMeshTangents = Baker.GetTargetMeshTangents().Get();
+}
+
+void FMeshNormalMapEvaluator::EvaluateSample(float*& Out, const FCorrespondenceSample& Sample, void* EvalData)
+{
+	FMeshNormalMapEvaluator* Eval = static_cast<FMeshNormalMapEvaluator*>(EvalData);
+	FVector3f SampleResult = Eval->SampleFunction(Sample);
+	// Map normal space [-1,1] to floating point color space [0,1]
+	FVector3f Result = (SampleResult + FVector3f::One()) * 0.5f;
+	WriteToBuffer(Out, Result);
+}
+
+void FMeshNormalMapEvaluator::EvaluateDefault(float*& Out, void* EvalData)
+{
+	FMeshNormalMapEvaluator* Eval = static_cast<FMeshNormalMapEvaluator*>(EvalData);
+	// Map normal space [-1,1] to floating point color space [0,1]
+	FVector3f Normal = (Eval->DefaultNormal + FVector3f::One()) * 0.5f;
+	WriteToBuffer(Out, Normal);
+}
+
+FVector3f FMeshNormalMapEvaluator::SampleFunction(const FCorrespondenceSample& SampleData)
+{
+	int32 DetailTriID = SampleData.DetailTriID;
+	if (DetailMesh->IsTriangle(DetailTriID))
+	{
+		// get tangents on base mesh
+		FVector3d BaseTangentX, BaseTangentY;
+		BaseMeshTangents->GetInterpolatedTriangleTangent(
+			SampleData.BaseSample.TriangleIndex,
+			SampleData.BaseSample.BaryCoords,
+			BaseTangentX, BaseTangentY);
+
+		// sample normal on detail mesh
+		FVector3d DetailNormal;
+		DetailNormalOverlay->GetTriBaryInterpolate<double>(DetailTriID, &SampleData.DetailBaryCoords.X, &DetailNormal.X);
+		Normalize(DetailNormal);
+
+		// compute normal in tangent space
+		double dx = DetailNormal.Dot(BaseTangentX);
+		double dy = DetailNormal.Dot(BaseTangentY);
+		double dz = DetailNormal.Dot(SampleData.BaseNormal);
+
+		return FVector3f((float)dx, (float)dy, (float)dz);
+	}
+	return FVector3f::UnitZ();
+}
diff --git a/Engine/Plugins/Experimental/GeometryProcessing/Source/DynamicMesh/Private/Sampling/MeshOcclusionMapBaker.cpp b/Engine/Plugins/Experimental/GeometryProcessing/Source/DynamicMesh/Private/Sampling/MeshOcclusionMapBaker.cpp
--- a/Engine/Plugins/Experimental/GeometryProcessing/Source/DynamicMesh/Private/Sampling/MeshOcclusionMapBaker.cpp
+++ b/Engine/Plugins/Experimental/GeometryProcessing/Source/DynamicMesh/Private/Sampling/MeshOcclusionMapBaker.cpp
@@ -1,10 +1,11 @@
 // Copyright Epic Games, Inc. All Rights Reserved.
 
 #include "Sampling/MeshOcclusionMapBaker.h"
-#include "Sampling/MeshMapBaker.h"
 #include "Sampling/SphericalFibonacci.h"
 #include "Sampling/Gaussians.h"
 
+#include "Math/RandomStream.h"
+
 #include "ExplicitUseGeometryMathTypes.h"		// using UE::Geometry::(math types)
 using namespace UE::Geometry;
 
@@ -17,15 +18,15 @@ void FMeshOcclusionMapBaker::Bake()
 
 	const FMeshImageBakingCache* BakeCache = GetCache();
 	check(BakeCache);
-	DetailMesh = BakeCache->GetDetailMesh();
-	DetailSpatial = BakeCache->GetDetailSpatial();
-	DetailNormalOverlay = BakeCache->GetDetailNormals();
+	const FDynamicMesh3* DetailMesh = BakeCache->GetDetailMesh();
+	const FDynamicMeshAABBTree3* DetailSpatial = BakeCache->GetDetailSpatial();
+	const FDynamicMeshNormalOverlay* DetailNormalOverlay = BakeCache->GetDetailNormals();
 	check(DetailNormalOverlay);
 
-	BiasDotThreshold = FMathd::Cos(FMathd::Clamp(90.0 - BiasAngleDeg, 0.0, 90.0) * FMathd::DegToRad);
+	double BiasDotThreshold = FMathd::Cos(FMathd::Clamp(90.0 - BiasAngleDeg, 0.0, 90.0) * FMathd::DegToRad);
 
 	// Precompute occlusion ray directions
-	RayDirections.Empty(NumOcclusionRays);
+	TArray<FVector3d> RayDirections;
 	THemisphericalFibonacci<double>::EDistribution Dist = THemisphericalFibonacci<double>::EDistribution::Cosine;
 	switch (Distribution)
 	{
@@ -58,6 +59,131 @@ void FMeshOcclusionMapBaker::Bake()
 		RayDir = Rotation * RayDir;
 	}
 
+	// Randomized rotation of occlusion rays about the normal
+	FRandomStream RotationGen(31337);
+	FCriticalSection RotationLock;
+	auto GetRandomRotation = [&RotationGen, &RotationLock]() {
+		RotationLock.Lock();
+		double Angle = RotationGen.GetFraction() * FMathd::TwoPi;
+		RotationLock.Unlock();
+		return Angle;
+	};
+
+	auto OcclusionSampleFunction = [&](const FMeshImageBakingCache::FCorrespondenceSample& SampleData, double& OutOcclusion, FVector3d& OutNormal)
+	{
+		// Fallback normal if invalid Tri or fully occluded
+		FVector3d DefaultNormal = FVector3d::UnitZ();
+		switch (NormalSpace)
+		{
+		case ESpace::Tangent:
+			break;
+		case ESpace::Object:
+			DefaultNormal = SampleData.BaseNormal;
+			break;
+		}
+
+		int32 DetailTriID = SampleData.DetailTriID;
+		if (DetailMesh->IsTriangle(DetailTriID))
+		{
+			FIndex3i DetailTri = DetailMesh->GetTriangle(DetailTriID);
+			//FVector3d DetailTriNormal = DetailMesh.GetTriNormal(DetailTriID);
+			FVector3d DetailTriNormal;
+			DetailNormalOverlay->GetTriBaryInterpolate<double>(DetailTriID, &SampleData.DetailBaryCoords.X, &DetailTriNormal.X);
+			Normalize(DetailTriNormal);
+
+			FVector3d BaseTangentX, BaseTangentY;
+			if (WantBentNormal() && NormalSpace == ESpace::Tangent)
+			{
+				check(BaseMeshTangents);
+				BaseMeshTangents->GetInterpolatedTriangleTangent(
+					SampleData.BaseSample.TriangleIndex,
+					SampleData.BaseSample.BaryCoords,
+					BaseTangentX, BaseTangentY);
+			}
+
+			FVector3d DetailBaryCoords = SampleData.DetailBaryCoords;
+			FVector3d DetailPos = DetailMesh->GetTriBaryPoint(DetailTriID, DetailBaryCoords.X, DetailBaryCoords.Y, DetailBaryCoords.Z);
+			DetailPos += 10.0f * FMathf::ZeroTolerance * DetailTriNormal;
+			FFrame3d SurfaceFrame(DetailPos, DetailTriNormal);
+
+			double RotationAngle = GetRandomRotation();
+			SurfaceFrame.Rotate(FQuaterniond(SurfaceFrame.Z(), RotationAngle, false));
+
+			IMeshSpatial::FQueryOptions QueryOptions;
+			QueryOptions.MaxDistance = MaxDistance;
+
+			double AccumOcclusion = 0;
+			FVector3d AccumNormal(FVector3d::Zero());
+			double TotalPointWeight = 0;
+			double TotalNormalWeight = 0;
+			for (FVector3d SphereDir : RayDirections)
+			{
+				FRay3d OcclusionRay(DetailPos, SurfaceFrame.FromFrameVector(SphereDir));
+				ensure(OcclusionRay.Direction.Dot(DetailTriNormal) > 0);
+
+				bool bHit = DetailSpatial->TestAnyHitTriangle(OcclusionRay, QueryOptions);
+
+				if (WantAmbientOcclusion())
+				{
+					// Have weight of point fall off as it becomes more coplanar with face. 
+					// This reduces faceting artifacts that we would otherwise see because geometry does not vary smoothly
+					double PointWeight = 1.0;
+					double BiasDot = OcclusionRay.Direction.Dot(DetailTriNormal);
+					if (BiasDot < BiasDotThreshold)
+					{
+						PointWeight = FMathd::Lerp(0.0, 1.0, FMathd::Clamp(BiasDot / BiasDotThreshold, 0.0, 1.0));
+						PointWeight *= PointWeight;
+					}
+					TotalPointWeight += PointWeight;
+
+					if (bHit)
+					{
+						AccumOcclusion += PointWeight;
+					}
+				}
+
+				if (WantBentNormal())
+				{
+					FVector3d BentNormal = OcclusionRay.Direction;
+					switch (NormalSpace)
+					{
+					case ESpace::Tangent:
+					{
+						// compute normal in tangent space
+						double dx = BentNormal.Dot(BaseTangentX);
+						double dy = BentNormal.Dot(BaseTangentY);
+						double dz = BentNormal.Dot(SampleData.BaseNormal);
+						BentNormal = FVector3d(dx, dy, dz);;
+						break;
+					}
+					case ESpace::Object:
+						break;
+					}
+					TotalNormalWeight += 1.0;
+
+					if (!bHit)
+					{
+						AccumNormal += BentNormal;
+					}
+				}
+			}
+
+			if (WantAmbientOcclusion())
+			{
+				AccumOcclusion = (TotalPointWeight > 0.0001) ? (AccumOcclusion / TotalPointWeight) : 0.0;
+			}
+			if (WantBentNormal())
+			{
+				AccumNormal = (TotalNormalWeight > 0.0 && AccumNormal.Length() > 0.0) ? Normalized(AccumNormal / TotalNormalWeight) : DefaultNormal;
+			}
+			OutOcclusion = AccumOcclusion;
+			OutNormal = AccumNormal;
+			return;
+		}
+		OutOcclusion = 0.0;
+		OutNormal = DefaultNormal;
+	};
+
 	if (WantAmbientOcclusion())
 	{
 		OcclusionBuilder = MakeUnique<TImageBuilder<FVector3f>>();
@@ -84,8 +210,8 @@ void FMeshOcclusionMapBaker::Bake()
 	BakeCache->EvaluateSamples([&](const FVector2i& Coords, const FMeshImageBakingCache::FCorrespondenceSample& Sample)
 	{
 		double Occlusion = 0.0;
-		FVector3d BentNormal = FVector3f::UnitZ();
-		SampleFunction<FMeshImageBakingCache::FCorrespondenceSample>(Sample, Occlusion, BentNormal);
+		FVector3d BentNormal;
+		OcclusionSampleFunction(Sample, Occlusion, BentNormal);
 		if (WantAmbientOcclusion())
 		{
 			FVector3f OcclusionColor = FMathd::Clamp(1.0f - (float)Occlusion, 0.0f, 1.0f) * FVector3f::One();
@@ -126,189 +252,7 @@ void FMeshOcclusionMapBaker::Bake()
 			BlurKernel2d.IntRadius,
 			AOBlurBuffer);
 	}
-}
-
-void FMeshOcclusionMapBaker::PreEvaluate(const FMeshMapBaker& Baker)
-{
-	// TODO: Support evaluation all map types in single pass.
-	check(OcclusionType != EOcclusionMapType::None && OcclusionType != EOcclusionMapType::All);
-
-	DetailMesh = Baker.GetDetailMesh();
-	DetailSpatial = Baker.GetDetailMeshSpatial();
-	DetailNormalOverlay = Baker.GetDetailMeshNormals();
-	check(DetailNormalOverlay);
-
-	if (WantBentNormal() && NormalSpace == ESpace::Tangent)
-	{
-		BaseMeshTangents = Baker.GetTargetMeshTangents().Get();
-	}
-
-	BiasDotThreshold = FMathd::Cos(FMathd::Clamp(90.0 - BiasAngleDeg, 0.0, 90.0) * FMathd::DegToRad);
-
-	// Precompute occlusion ray directions
-	RayDirections.Empty(NumOcclusionRays);
-	THemisphericalFibonacci<double>::EDistribution Dist = THemisphericalFibonacci<double>::EDistribution::Cosine;
-	switch (Distribution)
-	{
-	case EDistribution::Uniform:
-		Dist = THemisphericalFibonacci<double>::EDistribution::Uniform;
-		break;
-	case EDistribution::Cosine:
-		Dist = THemisphericalFibonacci<double>::EDistribution::Cosine;
-		break;
-	}
-	THemisphericalFibonacci<double> Points(NumOcclusionRays, Dist);
-	for (int32 k = 0; k < Points.Num(); ++k)
-	{
-		FVector3d P = Points[k];
-		if (P.Z > 0)
-		{
-			RayDirections.Add(Normalized(P));
-		}
-	}
-
-	// Map occlusion ray hemisphere to conical area (SpreadAngle/2)
-	double ConicalAngle = FMathd::Clamp(SpreadAngle * 0.5, 0.0001, 90.0);
-	for (int32 k = 0; k < RayDirections.Num(); ++k)
-	{
-		FVector3d& RayDir = RayDirections[k];
-		double RayAngle = AngleD(RayDir, FVector3d::UnitZ());
-		FVector3d RayCross = RayDir.Cross(FVector3d::UnitZ());
-		double RotationAngle = RayAngle - FMathd::Lerp(0.0, ConicalAngle, RayAngle / 90.0);
-		FQuaterniond Rotation(RayCross, RotationAngle, true);
-		RayDir = Rotation * RayDir;
-	}
-}
-
-FVector4f FMeshOcclusionMapBaker::EvaluateSample(const FMeshMapBaker& Baker, const FCorrespondenceSample& Sample)
-{
-	double Occlusion = 0.0;
-	FVector3d BentNormal = FVector3f::UnitZ();
-	SampleFunction<FCorrespondenceSample>(Sample, Occlusion, BentNormal);
-	if (WantAmbientOcclusion())
-	{
-		FVector3f OcclusionColor = FMathd::Clamp(1.0f - (float)Occlusion, 0.0f, 1.0f) * FVector3f::One();
-		return FVector4f(OcclusionColor.X, OcclusionColor.Y, OcclusionColor.Z, 1.0f);
-	}
-	else //if (WantBentNormal())
-	{
-		// Map normal space [-1,1] to floating point color space [0,1]
-		FVector3f NormalColor = (FVector3f(BentNormal.X, BentNormal.Y, BentNormal.Z) + FVector3f::One()) * 0.5;
-		return FVector4f(NormalColor.X, NormalColor.Y, NormalColor.Z, 1.0f);
-	}
-}
-
-template<class SampleType>
-void FMeshOcclusionMapBaker::SampleFunction(const SampleType& SampleData, double& OutOcclusion, FVector3d& OutNormal)
-{
-	// Fallback normal if invalid Tri or fully occluded
-	FVector3d DefaultNormal = FVector3d::UnitZ();
-	switch (NormalSpace)
-	{
-	case ESpace::Tangent:
-		break;
-	case ESpace::Object:
-		DefaultNormal = SampleData.BaseNormal;
-		break;
-	}
-
-	int32 DetailTriID = SampleData.DetailTriID;
-	if (!DetailMesh->IsTriangle(DetailTriID))
-	{
-		OutOcclusion = 0.0;
-		OutNormal = DefaultNormal;
-		return;
-	}
-
-	FIndex3i DetailTri = DetailMesh->GetTriangle(DetailTriID);
-	//FVector3d DetailTriNormal = DetailMesh.GetTriNormal(DetailTriID);
-	FVector3d DetailTriNormal;
-	DetailNormalOverlay->GetTriBaryInterpolate<double>(DetailTriID, &SampleData.DetailBaryCoords.X, &DetailTriNormal.X);
-	Normalize(DetailTriNormal);
-
-	FVector3d BaseTangentX, BaseTangentY;
-	if (WantBentNormal() && NormalSpace == ESpace::Tangent)
-	{
-		check(BaseMeshTangents);
-		BaseMeshTangents->GetInterpolatedTriangleTangent(
-			SampleData.BaseSample.TriangleIndex,
-			SampleData.BaseSample.BaryCoords,
-			BaseTangentX, BaseTangentY);
-	}
-
-	FVector3d DetailBaryCoords = SampleData.DetailBaryCoords;
-	FVector3d DetailPos = DetailMesh->GetTriBaryPoint(DetailTriID, DetailBaryCoords.X, DetailBaryCoords.Y, DetailBaryCoords.Z);
-	DetailPos += 10.0f * FMathf::ZeroTolerance * DetailTriNormal;
-	FFrame3d SurfaceFrame(DetailPos, DetailTriNormal);
-
-	double RotationAngle = GetRandomRotation();
-	SurfaceFrame.Rotate(FQuaterniond(SurfaceFrame.Z(), RotationAngle, false));
-
-	IMeshSpatial::FQueryOptions QueryOptions;
-	QueryOptions.MaxDistance = MaxDistance;
 
-	double AccumOcclusion = 0;
-	FVector3d AccumNormal(FVector3d::Zero());
-	double TotalPointWeight = 0;
-	for (FVector3d SphereDir : RayDirections)
-	{
-		FRay3d OcclusionRay(DetailPos, SurfaceFrame.FromFrameVector(SphereDir));
-		ensure(OcclusionRay.Direction.Dot(DetailTriNormal) > 0);
-
-		bool bHit = DetailSpatial->TestAnyHitTriangle(OcclusionRay, QueryOptions);
-
-		if (WantAmbientOcclusion())
-		{
-			// Have weight of point fall off as it becomes more coplanar with face. 
-			// This reduces faceting artifacts that we would otherwise see because geometry does not vary smoothly
-			double PointWeight = 1.0;
-			double BiasDot = OcclusionRay.Direction.Dot(DetailTriNormal);
-			if (BiasDot < BiasDotThreshold)
-			{
-				PointWeight = FMathd::Lerp(0.0, 1.0, FMathd::Clamp(BiasDot / BiasDotThreshold, 0.0, 1.0));
-				PointWeight *= PointWeight;
-			}
-			TotalPointWeight += PointWeight;
 
-			if (bHit)
-			{
-				AccumOcclusion += PointWeight;
-			}
-		}
 
-		if (WantBentNormal() && !bHit)
-		{
-			FVector3d BentNormal = OcclusionRay.Direction;
-			if (NormalSpace == ESpace::Tangent)
-			{
-				// compute normal in tangent space
-				double dx = BentNormal.Dot(BaseTangentX);
-				double dy = BentNormal.Dot(BaseTangentY);
-				double dz = BentNormal.Dot(SampleData.BaseNormal);
-				BentNormal = FVector3d(dx, dy, dz);;
-			}
-			AccumNormal += BentNormal;
-		}
-	}
-
-	if (WantAmbientOcclusion())
-	{
-		AccumOcclusion = (TotalPointWeight > 0.0001) ? (AccumOcclusion / TotalPointWeight) : 0.0;
-	}
-	if (WantBentNormal())
-	{
-		AccumNormal = (AccumNormal.Length() > 0.0) ? Normalized(AccumNormal) : DefaultNormal;
-	}
-	OutOcclusion = AccumOcclusion;
-	OutNormal = AccumNormal;
 }
-
-double FMeshOcclusionMapBaker::GetRandomRotation()
-{
-	// Randomized rotation of occlusion rays about the normal
-	RotationLock.Lock();
-	double Angle = RotationGen.GetFraction() * FMathd::TwoPi;
-	RotationLock.Unlock();
-	return Angle;
-}
-
diff --git a/Engine/Plugins/Experimental/GeometryProcessing/Source/DynamicMesh/Private/Sampling/MeshOcclusionMapEvaluator.cpp b/Engine/Plugins/Experimental/GeometryProcessing/Source/DynamicMesh/Private/Sampling/MeshOcclusionMapEvaluator.cpp
new file mode 100644
--- /dev/null
+++ b/Engine/Plugins/Experimental/GeometryProcessing/Source/DynamicMesh/Private/Sampling/MeshOcclusionMapEvaluator.cpp
@@ -0,0 +1,257 @@
+// Copyright Epic Games, Inc. All Rights Reserved.
+
+#include "Sampling/MeshOcclusionMapEvaluator.h"
+#include "Sampling/MeshMapBaker.h"
+#include "Sampling/SphericalFibonacci.h"
+
+#include "ExplicitUseGeometryMathTypes.h"		// using UE::Geometry::(math types)
+using namespace UE::Geometry;
+
+/**
+ * Occlusion type tests
+ * Using defines for if constexpr compile time branches
+ */
+#define WANT_AMBIENT_OCCLUSION(OccType) \
+	(OccType & EMeshOcclusionMapType::AmbientOcclusion) == EMeshOcclusionMapType::AmbientOcclusion
+#define WANT_BENT_NORMAL(OccType) \
+	(OccType & EMeshOcclusionMapType::BentNormal) == EMeshOcclusionMapType::BentNormal
+#define WANT_ALL(OccType) \
+	(OccType & EMeshOcclusionMapType::All) == EMeshOcclusionMapType::All
+
+const FVector3d FMeshOcclusionMapEvaluator::DefaultTangentNormal = FVector3d::UnitZ();
+const FVector3d FMeshOcclusionMapEvaluator::DefaultObjectNormal = FVector3d::Zero();
+
+void FMeshOcclusionMapEvaluator::Setup(const FMeshMapBaker& Baker, FEvaluationContext& Context)
+{
+	if (WANT_ALL(OcclusionType))
+	{
+		if (NormalSpace == ESpace::Tangent)
+		{
+			Context.Evaluate = &EvaluateSample<EMeshOcclusionMapType::All, ESpace::Tangent>;
+			Context.EvaluateDefault = &EvaluateDefault<EMeshOcclusionMapType::All, ESpace::Tangent>;
+		}
+		else // NormalSpace == ESpace::Object
+		{
+			Context.Evaluate = &EvaluateSample<EMeshOcclusionMapType::All, ESpace::Object>;
+			Context.EvaluateDefault = &EvaluateDefault<EMeshOcclusionMapType::All, ESpace::Object>;
+		}
+		Context.EvalData = this;
+		Context.AccumulateMode = EAccumulateMode::Add;
+		Context.DataLayout = { EComponents::Float1, EComponents::Float3 };
+	}
+	else if (WANT_AMBIENT_OCCLUSION(OcclusionType))
+	{
+		Context.Evaluate = &EvaluateSample<EMeshOcclusionMapType::AmbientOcclusion, ESpace::Tangent>;
+		Context.EvaluateDefault = &EvaluateDefault<EMeshOcclusionMapType::AmbientOcclusion, ESpace::Tangent>;
+		Context.EvalData = this;
+		Context.AccumulateMode = EAccumulateMode::Add;
+		Context.DataLayout = { EComponents::Float1 };
+	}
+	else if (WANT_BENT_NORMAL(OcclusionType))
+	{
+		if (NormalSpace == ESpace::Tangent)
+		{
+			Context.Evaluate = &EvaluateSample<EMeshOcclusionMapType::BentNormal, ESpace::Tangent>;
+			Context.EvaluateDefault = &EvaluateDefault<EMeshOcclusionMapType::BentNormal, ESpace::Tangent>;
+		}
+		else // NormalSpace == ESpace::Object
+		{
+			Context.Evaluate = &EvaluateSample<EMeshOcclusionMapType::BentNormal, ESpace::Object>;
+			Context.EvaluateDefault = &EvaluateDefault<EMeshOcclusionMapType::BentNormal, ESpace::Object>;
+		}
+		Context.EvalData = this;
+		Context.AccumulateMode = EAccumulateMode::Add;
+		Context.DataLayout = { EComponents::Float3 };
+	}
+	else
+	{
+		// TODO: Support error case in the baker to skip over invalid eval configs?
+		checkSlow(false);
+	}
+
+	// Cache data from the baker
+	DetailMesh = Baker.GetDetailMesh();
+	DetailSpatial = Baker.GetDetailMeshSpatial();
+	DetailNormalOverlay = Baker.GetDetailMeshNormals();
+	check(DetailNormalOverlay);
+
+	if (WANT_BENT_NORMAL(OcclusionType) && NormalSpace == ESpace::Tangent)
+	{
+		BaseMeshTangents = Baker.GetTargetMeshTangents().Get();
+	}
+
+	BiasDotThreshold = FMathd::Cos(FMathd::Clamp(90.0 - BiasAngleDeg, 0.0, 90.0) * FMathd::DegToRad);
+
+	// Precompute occlusion ray directions
+	RayDirections.Empty(NumOcclusionRays);
+	THemisphericalFibonacci<double>::EDistribution Dist = THemisphericalFibonacci<double>::EDistribution::Cosine;
+	switch (Distribution)
+	{
+	case EDistribution::Uniform:
+		Dist = THemisphericalFibonacci<double>::EDistribution::Uniform;
+		break;
+	case EDistribution::Cosine:
+		Dist = THemisphericalFibonacci<double>::EDistribution::Cosine;
+		break;
+	}
+	THemisphericalFibonacci<double> Points(NumOcclusionRays, Dist);
+	for (int32 k = 0; k < Points.Num(); ++k)
+	{
+		FVector3d P = Points[k];
+		if (P.Z > 0)
+		{
+			RayDirections.Add(Normalized(P));
+		}
+	}
+
+	// Map occlusion ray hemisphere to conical area (SpreadAngle/2)
+	double ConicalAngle = FMathd::Clamp(SpreadAngle * 0.5, 0.0001, 90.0);
+	for (int32 k = 0; k < RayDirections.Num(); ++k)
+	{
+		FVector3d& RayDir = RayDirections[k];
+		double RayAngle = AngleD(RayDir, FVector3d::UnitZ());
+		FVector3d RayCross = RayDir.Cross(FVector3d::UnitZ());
+		double RotationAngle = RayAngle - FMathd::Lerp(0.0, ConicalAngle, RayAngle / 90.0);
+		FQuaterniond Rotation(RayCross, RotationAngle, true);
+		RayDir = Rotation * RayDir;
+	}
+}
+
+template <EMeshOcclusionMapType ComputeType, FMeshOcclusionMapEvaluator::ESpace ComputeSpace>
+void FMeshOcclusionMapEvaluator::EvaluateSample(float*& Out, const FCorrespondenceSample& Sample, void* EvalData)
+{
+	FMeshOcclusionMapEvaluator* Eval = static_cast<FMeshOcclusionMapEvaluator*>(EvalData);
+	const FVector3d& DefaultNormal = ComputeSpace == ESpace::Tangent ? DefaultTangentNormal : Sample.BaseNormal;
+
+	double Occlusion = 0.0;
+	FVector3f BentNormal = FVector3f::UnitZ();
+	Tie(Occlusion, BentNormal) = Eval->SampleFunction<ComputeType, ComputeSpace>(Sample, DefaultNormal);
+
+	if constexpr (WANT_AMBIENT_OCCLUSION(ComputeType))
+	{
+		WriteToBuffer(Out, FMathd::Clamp(1.0f - (float)Occlusion, 0.0f, 1.0f));
+	}
+
+	if constexpr (WANT_BENT_NORMAL(ComputeType))
+	{
+		// Map normal space[-1, 1] to floating point color space[0, 1]
+		FVector3f NormalColor = (FVector3f(BentNormal.X, BentNormal.Y, BentNormal.Z) + FVector3f::One()) * 0.5;
+		WriteToBuffer(Out, NormalColor);
+	}
+}
+
+template <EMeshOcclusionMapType ComputeType, FMeshOcclusionMapEvaluator::ESpace ComputeSpace>
+void FMeshOcclusionMapEvaluator::EvaluateDefault(float*& Out, void* EvalData)
+{
+	if constexpr (WANT_AMBIENT_OCCLUSION(ComputeType))
+	{
+		WriteToBuffer(Out, 1.0f);
+	}
+
+	if constexpr (WANT_BENT_NORMAL(ComputeType))
+	{
+		const FVector3d& DefaultNormal = ComputeSpace == ESpace::Tangent ? DefaultTangentNormal : DefaultObjectNormal;
+
+		// Map normal space [-1,1] to floating point color space [0,1]
+		FVector3f Normal = (DefaultObjectNormal + FVector3f::One()) * 0.5f;
+		WriteToBuffer(Out, Normal);
+	}
+}
+
+template <EMeshOcclusionMapType ComputeType, FMeshOcclusionMapEvaluator::ESpace ComputeSpace>
+FMeshOcclusionMapEvaluator::FOcclusionTuple FMeshOcclusionMapEvaluator::SampleFunction(const FCorrespondenceSample& SampleData, const FVector3d& DefaultNormal)
+{
+	int32 DetailTriID = SampleData.DetailTriID;
+	if (!DetailMesh->IsTriangle(DetailTriID))
+	{
+		return FOcclusionTuple(0.0, DefaultNormal);
+	}
+
+	FIndex3i DetailTri = DetailMesh->GetTriangle(DetailTriID);
+	//FVector3d DetailTriNormal = DetailMesh.GetTriNormal(DetailTriID);
+	FVector3d DetailTriNormal;
+	DetailNormalOverlay->GetTriBaryInterpolate<double>(DetailTriID, &SampleData.DetailBaryCoords.X, &DetailTriNormal.X);
+	Normalize(DetailTriNormal);
+
+	FVector3d BaseTangentX, BaseTangentY;
+	if constexpr (WANT_BENT_NORMAL(ComputeType) && ComputeSpace == ESpace::Tangent)
+	{
+		check(BaseMeshTangents);
+		BaseMeshTangents->GetInterpolatedTriangleTangent(
+			SampleData.BaseSample.TriangleIndex,
+			SampleData.BaseSample.BaryCoords,
+			BaseTangentX, BaseTangentY);
+	}
+
+	FVector3d DetailBaryCoords = SampleData.DetailBaryCoords;
+	FVector3d DetailPos = DetailMesh->GetTriBaryPoint(DetailTriID, DetailBaryCoords.X, DetailBaryCoords.Y, DetailBaryCoords.Z);
+	DetailPos += 10.0f * FMathf::ZeroTolerance * DetailTriNormal;
+	FFrame3d SurfaceFrame(DetailPos, DetailTriNormal);
+
+	// TODO: Consider pregenerating random rotations.
+	FRandomStream RotationGen;
+	RotationGen.GenerateNewSeed();
+	double RotationAngle = RotationGen.GetFraction() * FMathd::TwoPi;
+	SurfaceFrame.Rotate(FQuaterniond(SurfaceFrame.Z(), RotationAngle, false));
+
+	IMeshSpatial::FQueryOptions QueryOptions;
+	QueryOptions.MaxDistance = MaxDistance;
+
+	double AccumOcclusion = 0;
+	FVector3d AccumNormal(FVector3d::Zero());
+	double TotalPointWeight = 0;
+	for (FVector3d SphereDir : RayDirections)
+	{
+		FRay3d OcclusionRay(DetailPos, SurfaceFrame.FromFrameVector(SphereDir));
+		ensure(OcclusionRay.Direction.Dot(DetailTriNormal) > 0);
+
+		bool bHit = DetailSpatial->TestAnyHitTriangle(OcclusionRay, QueryOptions);
+
+		if constexpr (WANT_AMBIENT_OCCLUSION(ComputeType))
+		{
+			// Have weight of point fall off as it becomes more coplanar with face. 
+			// This reduces faceting artifacts that we would otherwise see because geometry does not vary smoothly
+			double PointWeight = 1.0;
+			double BiasDot = OcclusionRay.Direction.Dot(DetailTriNormal);
+			if (BiasDot < BiasDotThreshold)
+			{
+				PointWeight = FMathd::Lerp(0.0, 1.0, FMathd::Clamp(BiasDot / BiasDotThreshold, 0.0, 1.0));
+				PointWeight *= PointWeight;
+			}
+			TotalPointWeight += PointWeight;
+
+			if (bHit)
+			{
+				AccumOcclusion += PointWeight;
+			}
+		}
+
+		if constexpr (WANT_BENT_NORMAL(ComputeType))
+		{
+			if (!bHit)
+			{
+				FVector3d BentNormal = OcclusionRay.Direction;
+				if (ComputeSpace == ESpace::Tangent)
+				{
+					// compute normal in tangent space
+					double dx = BentNormal.Dot(BaseTangentX);
+					double dy = BentNormal.Dot(BaseTangentY);
+					double dz = BentNormal.Dot(SampleData.BaseNormal);
+					BentNormal = FVector3d(dx, dy, dz);;
+				}
+				AccumNormal += BentNormal;
+			}
+		}
+	}
+
+	if constexpr (WANT_AMBIENT_OCCLUSION(ComputeType))
+	{
+		AccumOcclusion = (TotalPointWeight > 0.0001) ? (AccumOcclusion / TotalPointWeight) : 0.0;
+	}
+	if constexpr (WANT_BENT_NORMAL(ComputeType))
+	{
+		AccumNormal = (AccumNormal.Length() > 0.0) ? Normalized(AccumNormal) : DefaultNormal;
+	}
+	return FOcclusionTuple(AccumOcclusion, AccumNormal);
+}
+
diff --git a/Engine/Plugins/Experimental/GeometryProcessing/Source/DynamicMesh/Private/Sampling/MeshPropertyMapBaker.cpp b/Engine/Plugins/Experimental/GeometryProcessing/Source/DynamicMesh/Private/Sampling/MeshPropertyMapBaker.cpp
--- a/Engine/Plugins/Experimental/GeometryProcessing/Source/DynamicMesh/Private/Sampling/MeshPropertyMapBaker.cpp
+++ b/Engine/Plugins/Experimental/GeometryProcessing/Source/DynamicMesh/Private/Sampling/MeshPropertyMapBaker.cpp
@@ -1,7 +1,6 @@
 // Copyright Epic Games, Inc. All Rights Reserved.
 
 #include "Sampling/MeshPropertyMapBaker.h"
-#include "Sampling/MeshMapBaker.h"
 #include "Util/ColorConstants.h"
 
 #include "ExplicitUseGeometryMathTypes.h"		// using UE::Geometry::(math types)
@@ -32,12 +31,12 @@ void FMeshPropertyMapBaker::Bake()
 {
 	const FMeshImageBakingCache* BakeCache = GetCache();
 	check(BakeCache);
-	DetailMesh = BakeCache->GetDetailMesh();
-	DetailNormalOverlay = BakeCache->GetDetailNormals();
+	const FDynamicMesh3* DetailMesh = BakeCache->GetDetailMesh();
+	const FDynamicMeshNormalOverlay* DetailNormalOverlay = BakeCache->GetDetailNormals();
 	check(DetailNormalOverlay);
-	DetailUVOverlay = DetailMesh->Attributes()->PrimaryUV();
+	const FDynamicMeshUVOverlay* DetailUVOverlay = DetailMesh->Attributes()->PrimaryUV();
 
-	Bounds = DetailMesh->GetBounds();
+	FAxisAlignedBox3d Bounds = DetailMesh->GetBounds();
 	for (int32 j = 0; j < 3; ++j)
 	{
 		if (Bounds.Diagonal()[j] < FMathf::ZeroTolerance)
@@ -48,7 +47,7 @@ void FMeshPropertyMapBaker::Bake()
 		}
 	}
 
-	DefaultValue = FVector3f::Zero();
+	FVector3f DefaultValue(0, 0, 0);
 	switch (this->Property)
 	{
 		case EMeshPropertyBakeType::Position:
@@ -64,12 +63,75 @@ void FMeshPropertyMapBaker::Bake()
 			break;
 	}
 
+	auto PropertySampleFunction = [&](const FMeshImageBakingCache::FCorrespondenceSample& SampleData)
+	{
+		FVector3f Color = DefaultValue;
+		int32 DetailTriID = SampleData.DetailTriID;
+		if (DetailMesh->IsTriangle(DetailTriID))
+		{
+			switch (this->Property)
+			{
+				case EMeshPropertyBakeType::Position:
+					{
+						FVector3d Position = DetailMesh->GetTriBaryPoint(DetailTriID, SampleData.DetailBaryCoords[0], SampleData.DetailBaryCoords[1], SampleData.DetailBaryCoords[2]);
+						Color = PositionToColor(Position, Bounds);
+					}
+					break;
+				default:
+				case EMeshPropertyBakeType::FacetNormal:
+					{
+						FVector3d FacetNormal = DetailMesh->GetTriNormal(DetailTriID);
+						Color = NormalToColor(FacetNormal);
+					}
+					break;
+				break;
+				case EMeshPropertyBakeType::Normal:
+					{
+						if (DetailNormalOverlay->IsSetTriangle(DetailTriID))
+						{
+							FVector3d DetailNormal;
+							DetailNormalOverlay->GetTriBaryInterpolate<double>(DetailTriID, &SampleData.DetailBaryCoords.X, &DetailNormal.X);
+							Normalize(DetailNormal);
+							Color = NormalToColor(DetailNormal);
+						}
+					}
+					break;
+				case EMeshPropertyBakeType::UVPosition:
+					{
+						if (DetailUVOverlay && DetailUVOverlay->IsSetTriangle(DetailTriID))
+						{
+							FVector2d DetailUV;
+							DetailUVOverlay->GetTriBaryInterpolate<double>(DetailTriID, &SampleData.DetailBaryCoords.X, &DetailUV.X);
+							Color = UVToColor(DetailUV);
+						}
+					}
+					break;
+				case EMeshPropertyBakeType::MaterialID:
+					{
+						if (DetailMesh->Attributes() && DetailMesh->Attributes()->HasMaterialID())
+						{
+							const FDynamicMeshMaterialAttribute* DetailMaterialIDAttrib = DetailMesh->Attributes()->GetMaterialID();
+							const int32 MatID = DetailMaterialIDAttrib->GetValue(DetailTriID);
+							Color = LinearColors::SelectColor<FVector3f>(MatID);
+						}
+						else
+						{
+							Color = FVector3f(LinearColors::LightPink3f());
+						}
+					}
+					break;
+			}
+		}
+		return Color;
+	};
+
+
 	ResultBuilder = MakeUnique<TImageBuilder<FVector3f>>();
 	ResultBuilder->SetDimensions(BakeCache->GetDimensions());
 
 	BakeCache->EvaluateSamples([&](const FVector2i& Coords, const FMeshImageBakingCache::FCorrespondenceSample& Sample)
 	{
-		FVector3f Color = SampleFunction<FMeshImageBakingCache::FCorrespondenceSample>(Sample);
+		FVector3f Color = PropertySampleFunction(Sample);
 		ResultBuilder->SetPixel(Coords, Color);
 	});
 
@@ -82,108 +144,4 @@ void FMeshPropertyMapBaker::Bake()
 	}
 
 
-}
-
-void FMeshPropertyMapBaker::PreEvaluate(const FMeshMapBaker& Baker)
-{
-	DetailMesh = Baker.GetDetailMesh();
-	DetailNormalOverlay = Baker.GetDetailMeshNormals();
-	check(DetailNormalOverlay);
-	DetailUVOverlay = Baker.GetDetailMeshUVs();
-
-	Bounds = DetailMesh->GetBounds();
-	for (int32 j = 0; j < 3; ++j)
-	{
-		if (Bounds.Diagonal()[j] < FMathf::ZeroTolerance)
-		{
-			Bounds.Min[j] = Bounds.Center()[j] - FMathf::ZeroTolerance;
-			Bounds.Max[j] = Bounds.Center()[j] + FMathf::ZeroTolerance;
-
-		}
-	}
-
-	DefaultValue = FVector3f::Zero();
-	switch (this->Property)
-	{
-	case EMeshPropertyBakeType::Position:
-		DefaultValue = PositionToColor(Bounds.Center(), Bounds);
-		break;
-	default:
-	case EMeshPropertyBakeType::FacetNormal:
-	case EMeshPropertyBakeType::Normal:
-		DefaultValue = NormalToColor(FVector3d::UnitZ());
-		break;
-	case EMeshPropertyBakeType::UVPosition:
-		DefaultValue = UVToColor(FVector2d::Zero());
-		break;
-	}
-}
-
-FVector4f FMeshPropertyMapBaker::EvaluateSample(const FMeshMapBaker& Baker, const FCorrespondenceSample& Sample)
-{
-	FVector3f SampleResult = SampleFunction<FCorrespondenceSample>(Sample);
-	return FVector4f(SampleResult.X, SampleResult.Y, SampleResult.Z, 1.0f);
-}
-
-template<class SampleType>
-FVector3f FMeshPropertyMapBaker::SampleFunction(const SampleType& SampleData)
-{
-	FVector3f Color = DefaultValue;
-	int32 DetailTriID = SampleData.DetailTriID;
-	if (DetailMesh->IsTriangle(DetailTriID))
-	{
-		switch (this->Property)
-		{
-		case EMeshPropertyBakeType::Position:
-		{
-			FVector3d Position = DetailMesh->GetTriBaryPoint(DetailTriID, SampleData.DetailBaryCoords[0], SampleData.DetailBaryCoords[1], SampleData.DetailBaryCoords[2]);
-			Color = PositionToColor(Position, Bounds);
-		}
-		break;
-		default:
-		case EMeshPropertyBakeType::FacetNormal:
-		{
-			FVector3d FacetNormal = DetailMesh->GetTriNormal(DetailTriID);
-			Color = NormalToColor(FacetNormal);
-		}
-		break;
-		break;
-		case EMeshPropertyBakeType::Normal:
-		{
-			if (DetailNormalOverlay->IsSetTriangle(DetailTriID))
-			{
-				FVector3d DetailNormal;
-				DetailNormalOverlay->GetTriBaryInterpolate<double>(DetailTriID, &SampleData.DetailBaryCoords.X, &DetailNormal.X);
-				Normalize(DetailNormal);
-				Color = NormalToColor(DetailNormal);
-			}
-		}
-		break;
-		case EMeshPropertyBakeType::UVPosition:
-		{
-			if (DetailUVOverlay && DetailUVOverlay->IsSetTriangle(DetailTriID))
-			{
-				FVector2d DetailUV;
-				DetailUVOverlay->GetTriBaryInterpolate<double>(DetailTriID, &SampleData.DetailBaryCoords.X, &DetailUV.X);
-				Color = UVToColor(DetailUV);
-			}
-		}
-		break;
-		case EMeshPropertyBakeType::MaterialID:
-		{
-			if (DetailMesh->Attributes() && DetailMesh->Attributes()->HasMaterialID())
-			{
-				const FDynamicMeshMaterialAttribute* DetailMaterialIDAttrib = DetailMesh->Attributes()->GetMaterialID();
-				const int32 MatID = DetailMaterialIDAttrib->GetValue(DetailTriID);
-				Color = LinearColors::SelectColor<FVector3f>(MatID);
-			}
-			else
-			{
-				Color = FVector3f(LinearColors::LightPink3f());
-			}
-		}
-		break;
-		}
-	}
-	return Color;
-}
+}
\ No newline at end of file
diff --git a/Engine/Plugins/Experimental/GeometryProcessing/Source/DynamicMesh/Private/Sampling/MeshPropertyMapEvaluator.cpp b/Engine/Plugins/Experimental/GeometryProcessing/Source/DynamicMesh/Private/Sampling/MeshPropertyMapEvaluator.cpp
new file mode 100644
--- /dev/null
+++ b/Engine/Plugins/Experimental/GeometryProcessing/Source/DynamicMesh/Private/Sampling/MeshPropertyMapEvaluator.cpp
@@ -0,0 +1,121 @@
+// Copyright Epic Games, Inc. All Rights Reserved.
+
+#include "Sampling/MeshPropertyMapEvaluator.h"
+#include "Sampling/MeshMapBaker.h"
+#include "Util/ColorConstants.h"
+
+#include "ExplicitUseGeometryMathTypes.h"		// using UE::Geometry::(math types)
+using namespace UE::Geometry;
+
+void FMeshPropertyMapEvaluator::Setup(const FMeshMapBaker& Baker, FEvaluationContext& Context)
+{
+	Context.Evaluate = &EvaluateSample;
+	Context.EvaluateDefault = &EvaluateDefault;
+	Context.EvalData = this;
+	Context.AccumulateMode = EAccumulateMode::Add;
+	Context.DataLayout = { EComponents::Float3 };
+
+	// Cache data from the baker
+	DetailMesh = Baker.GetDetailMesh();
+	DetailNormalOverlay = Baker.GetDetailMeshNormals();
+	check(DetailNormalOverlay);
+	DetailUVOverlay = Baker.GetDetailMeshUVs();
+
+	Bounds = DetailMesh->GetBounds();
+	for (int32 j = 0; j < 3; ++j)
+	{
+		if (Bounds.Diagonal()[j] < FMathf::ZeroTolerance)
+		{
+			Bounds.Min[j] = Bounds.Center()[j] - FMathf::ZeroTolerance;
+			Bounds.Max[j] = Bounds.Center()[j] + FMathf::ZeroTolerance;
+
+		}
+	}
+
+	DefaultValue = FVector3f::Zero();
+	switch (this->Property)
+	{
+	case EMeshPropertyMapType::Position:
+		DefaultValue = PositionToColor(Bounds.Center(), Bounds);
+		break;
+	case EMeshPropertyMapType::FacetNormal:
+	case EMeshPropertyMapType::Normal:
+		DefaultValue = NormalToColor(FVector3d::UnitZ());
+		break;
+	case EMeshPropertyMapType::UVPosition:
+		DefaultValue = UVToColor(FVector2d::Zero());
+		break;
+	case EMeshPropertyMapType::MaterialID:
+		DefaultValue = FVector3f(LinearColors::LightPink3f());
+		break;
+	}
+}
+
+void FMeshPropertyMapEvaluator::EvaluateSample(float*& Out, const FCorrespondenceSample& Sample, void* EvalData)
+{
+	FMeshPropertyMapEvaluator* Eval = static_cast<FMeshPropertyMapEvaluator*>(EvalData);
+	FVector3f SampleResult = Eval->SampleFunction(Sample);
+	WriteToBuffer(Out, SampleResult);
+}
+
+void FMeshPropertyMapEvaluator::EvaluateDefault(float*& Out, void* EvalData)
+{
+	WriteToBuffer(Out, FVector3f::Zero());
+}
+
+FVector3f FMeshPropertyMapEvaluator::SampleFunction(const FCorrespondenceSample& SampleData)
+{
+	FVector3f Color = DefaultValue;
+	int32 DetailTriID = SampleData.DetailTriID;
+	if (DetailMesh->IsTriangle(DetailTriID))
+	{
+		switch (this->Property)
+		{
+		case EMeshPropertyMapType::Position:
+		{
+			FVector3d Position = DetailMesh->GetTriBaryPoint(DetailTriID, SampleData.DetailBaryCoords[0], SampleData.DetailBaryCoords[1], SampleData.DetailBaryCoords[2]);
+			Color = PositionToColor(Position, Bounds);
+		}
+		break;
+		case EMeshPropertyMapType::FacetNormal:
+		{
+			FVector3d FacetNormal = DetailMesh->GetTriNormal(DetailTriID);
+			Color = NormalToColor(FacetNormal);
+		}
+		break;
+		case EMeshPropertyMapType::Normal:
+		{
+			if (DetailNormalOverlay->IsSetTriangle(DetailTriID))
+			{
+				FVector3d DetailNormal;
+				DetailNormalOverlay->GetTriBaryInterpolate<double>(DetailTriID, &SampleData.DetailBaryCoords.X, &DetailNormal.X);
+				Normalize(DetailNormal);
+				Color = NormalToColor(DetailNormal);
+			}
+		}
+		break;
+		case EMeshPropertyMapType::UVPosition:
+		{
+			if (DetailUVOverlay && DetailUVOverlay->IsSetTriangle(DetailTriID))
+			{
+				FVector2d DetailUV;
+				DetailUVOverlay->GetTriBaryInterpolate<double>(DetailTriID, &SampleData.DetailBaryCoords.X, &DetailUV.X);
+				Color = UVToColor(DetailUV);
+			}
+		}
+		break;
+		case EMeshPropertyMapType::MaterialID:
+		{
+			if (DetailMesh->Attributes() && DetailMesh->Attributes()->HasMaterialID())
+			{
+				const FDynamicMeshMaterialAttribute* DetailMaterialIDAttrib = DetailMesh->Attributes()->GetMaterialID();
+				const int32 MatID = DetailMaterialIDAttrib->GetValue(DetailTriID);
+				Color = LinearColors::SelectColor<FVector3f>(MatID);
+			}
+		}
+		break;
+		}
+	}
+	return Color;
+}
+
diff --git a/Engine/Plugins/Experimental/GeometryProcessing/Source/DynamicMesh/Private/Sampling/MeshResampleImageBaker.cpp b/Engine/Plugins/Experimental/GeometryProcessing/Source/DynamicMesh/Private/Sampling/MeshResampleImageBaker.cpp
--- a/Engine/Plugins/Experimental/GeometryProcessing/Source/DynamicMesh/Private/Sampling/MeshResampleImageBaker.cpp
+++ b/Engine/Plugins/Experimental/GeometryProcessing/Source/DynamicMesh/Private/Sampling/MeshResampleImageBaker.cpp
@@ -1,7 +1,6 @@
 // Copyright Epic Games, Inc. All Rights Reserved.
 
 #include "Sampling/MeshResampleImageBaker.h"
-#include "Sampling/MeshMapBaker.h"
 
 #include "ExplicitUseGeometryMathTypes.h"		// using UE::Geometry::(math types)
 using namespace UE::Geometry;
@@ -10,17 +9,33 @@ void FMeshResampleImageBaker::Bake()
 {
 	const FMeshImageBakingCache* BakeCache = GetCache();
 	check(BakeCache);
-	DetailMesh = BakeCache->GetDetailMesh();
+	const FDynamicMesh3* DetailMesh = BakeCache->GetDetailMesh();
 
 	check(DetailUVOverlay);
 
+	FVector4f DefaultValue(0, 0, 0, 1.0);
+
+	auto PropertySampleFunction = [&](const FMeshImageBakingCache::FCorrespondenceSample& SampleData)
+	{
+		FVector4f Color = DefaultValue;
+		int32 DetailTriID = SampleData.DetailTriID;
+		if (DetailMesh->IsTriangle(SampleData.DetailTriID) && DetailUVOverlay)
+		{
+			FVector2d DetailUV;
+			DetailUVOverlay->GetTriBaryInterpolate<double>(DetailTriID, &SampleData.DetailBaryCoords.X, &DetailUV.X);
+
+			Color = SampleFunction(DetailUV);
+		}
+		return Color;
+	};
+
 	ResultBuilder = MakeUnique<TImageBuilder<FVector4f>>();
 	ResultBuilder->SetDimensions(BakeCache->GetDimensions());
 	ResultBuilder->Clear(DefaultColor);
 
 	BakeCache->EvaluateSamples([&](const FVector2i& Coords, const FMeshImageBakingCache::FCorrespondenceSample& Sample)
 	{
-		FVector4f Color = PropertySampleFunction<FMeshImageBakingCache::FCorrespondenceSample>(Sample);
+		FVector4f Color = PropertySampleFunction(Sample);
 		ResultBuilder->SetPixel(Coords, Color);
 	});
 
@@ -34,31 +49,6 @@ void FMeshResampleImageBaker::Bake()
 
 }
 
-void FMeshResampleImageBaker::PreEvaluate(const FMeshMapBaker& Baker)
-{
-	DetailMesh = Baker.GetDetailMesh();
-}
-
-FVector4f FMeshResampleImageBaker::EvaluateSample(const FMeshMapBaker& Baker, const FCorrespondenceSample& Sample)
-{
-	return PropertySampleFunction<FCorrespondenceSample>(Sample);
-}
-
-template <class SampleType>
-FVector4f FMeshResampleImageBaker::PropertySampleFunction(const SampleType& SampleData)
-{
-	FVector4f Color = DefaultColor;
-	int32 DetailTriID = SampleData.DetailTriID;
-	if (DetailMesh->IsTriangle(SampleData.DetailTriID) && DetailUVOverlay)
-	{
-		FVector2d DetailUV;
-		DetailUVOverlay->GetTriBaryInterpolate<double>(DetailTriID, &SampleData.DetailBaryCoords.X, &DetailUV.X);
-
-		Color = SampleFunction(DetailUV);
-	}
-	return Color;
-}
-
 
 
 
@@ -78,12 +68,12 @@ void FMeshMultiResampleImageBaker::BakeMaterial(int32 MaterialID)
 	{
 		return;
 	}
-	DetailMesh = BakeCache->GetDetailMesh();
+	const FDynamicMesh3* DetailMesh = BakeCache->GetDetailMesh();
 	if (!ensure(DetailMesh))
 	{
 		return;
 	}
-	DetailMaterialIDAttrib = DetailMesh->Attributes()->GetMaterialID();
+	const FDynamicMeshMaterialAttribute* DetailMaterialIDAttrib = DetailMesh->Attributes()->GetMaterialID();
 	if (!ensure(DetailMaterialIDAttrib))
 	{
 		return;
@@ -93,7 +83,7 @@ void FMeshMultiResampleImageBaker::BakeMaterial(int32 MaterialID)
 		return;
 	}
 
-	BakeCache->EvaluateSamples([this, MaterialID](const FVector2i& Coords, const FMeshImageBakingCache::FCorrespondenceSample& Sample)
+	BakeCache->EvaluateSamples([this, DetailMesh, DetailMaterialIDAttrib, MaterialID](const FVector2i& Coords, const FMeshImageBakingCache::FCorrespondenceSample& Sample)
 	{
 		int32 DetailTriID = Sample.DetailTriID;
 
@@ -139,33 +129,3 @@ void FMeshMultiResampleImageBaker::Bake()
 	}
 
 }
-
-void FMeshMultiResampleImageBaker::PreEvaluate(const FMeshMapBaker& Baker)
-{
-	DetailMesh = Baker.GetDetailMesh();
-	DetailMaterialIDAttrib = ensure(DetailMesh) ? DetailMesh->Attributes()->GetMaterialID() : nullptr;
-	bValidDetailMesh = ensure(DetailMaterialIDAttrib) && ensure(DetailUVOverlay);
-}
-
-FVector4f FMeshMultiResampleImageBaker::EvaluateSample(const FMeshMapBaker& Baker, const FCorrespondenceSample& Sample)
-{
-	if (!bValidDetailMesh)
-	{
-		return DefaultSample();
-	}
-
-	FVector4f Color = DefaultSample();
-	int32 DetailTriID = Sample.DetailTriID;
-	if (DetailMesh->IsTriangle(DetailTriID))
-	{
-		int32 MaterialID = DetailMaterialIDAttrib->GetValue(DetailTriID);
-		TSharedPtr<UE::Geometry::TImageBuilder<FVector4f>, ESPMode::ThreadSafe> TextureImage = MultiTextures.FindRef(MaterialID);
-		if (TextureImage)
-		{
-			FVector2d DetailUV;
-			DetailUVOverlay->GetTriBaryInterpolate<double>(DetailTriID, &Sample.DetailBaryCoords.X, &DetailUV.X);
-			Color = TextureImage->BilinearSampleUV<float>(DetailUV, FVector4f(0, 0, 0, 1));
-		}
-	}
-	return Color;
-}
diff --git a/Engine/Plugins/Experimental/GeometryProcessing/Source/DynamicMesh/Private/Sampling/MeshResampleImageEvaluator.cpp b/Engine/Plugins/Experimental/GeometryProcessing/Source/DynamicMesh/Private/Sampling/MeshResampleImageEvaluator.cpp
new file mode 100644
--- /dev/null
+++ b/Engine/Plugins/Experimental/GeometryProcessing/Source/DynamicMesh/Private/Sampling/MeshResampleImageEvaluator.cpp
@@ -0,0 +1,97 @@
+// Copyright Epic Games, Inc. All Rights Reserved.
+
+#include "Sampling/MeshResampleImageEvaluator.h"
+#include "Sampling/MeshMapBaker.h"
+
+#include "ExplicitUseGeometryMathTypes.h"		// using UE::Geometry::(math types)
+using namespace UE::Geometry;
+
+//
+// FMeshResampleImage Evaluator
+//
+
+void FMeshResampleImageEvaluator::Setup(const FMeshMapBaker& Baker, FEvaluationContext& Context)
+{
+	Context.Evaluate = &EvaluateSample;
+	Context.EvaluateDefault = &EvaluateDefault;
+	Context.EvalData = this;
+	Context.AccumulateMode = EAccumulateMode::Add;
+	Context.DataLayout = { EComponents::Float4 };
+
+	// Cache data from the baker
+	DetailMesh = Baker.GetDetailMesh();
+}
+
+void FMeshResampleImageEvaluator::EvaluateSample(float*& Out, const FCorrespondenceSample& Sample, void* EvalData)
+{
+	FMeshResampleImageEvaluator* Eval = static_cast<FMeshResampleImageEvaluator*>(EvalData);
+	FVector4f SampleResult = Eval->ImageSampleFunction(Sample);
+	WriteToBuffer(Out, SampleResult);
+}
+
+void FMeshResampleImageEvaluator::EvaluateDefault(float*& Out, void* EvalData)
+{
+	WriteToBuffer(Out, FVector4f(0.0f, 0.0f, 0.0f, 1.0f));
+}
+
+FVector4f FMeshResampleImageEvaluator::ImageSampleFunction(const FCorrespondenceSample& SampleData)
+{
+	FVector4f Color = DefaultColor;
+	int32 DetailTriID = SampleData.DetailTriID;
+	if (DetailMesh->IsTriangle(SampleData.DetailTriID) && DetailUVOverlay)
+	{
+		FVector2d DetailUV;
+		DetailUVOverlay->GetTriBaryInterpolate<double>(DetailTriID, &SampleData.DetailBaryCoords.X, &DetailUV.X);
+		Color = SampleFunction(DetailUV);
+	}
+	return Color;
+}
+
+//
+// FMeshMultiResampleImage Evaluator
+//
+
+void FMeshMultiResampleImageEvaluator::Setup(const FMeshMapBaker& Baker, FEvaluationContext& Context)
+{
+	Context.Evaluate = &EvaluateSampleMulti;
+	Context.EvaluateDefault = &EvaluateDefault;
+	Context.EvalData = this;
+	Context.AccumulateMode = EAccumulateMode::Add;
+	Context.DataLayout = { EComponents::Float4 };
+
+	// Cache data from baker
+	DetailMesh = Baker.GetDetailMesh();
+	DetailMaterialIDAttrib = ensure(DetailMesh) ? DetailMesh->Attributes()->GetMaterialID() : nullptr;
+	bValidDetailMesh = ensure(DetailMaterialIDAttrib) && ensure(DetailUVOverlay);
+}
+
+void FMeshMultiResampleImageEvaluator::EvaluateSampleMulti(float*& Out, const FCorrespondenceSample& Sample, void* EvalData)
+{
+	FMeshMultiResampleImageEvaluator* Eval = static_cast<FMeshMultiResampleImageEvaluator*>(EvalData);
+	FVector4f SampleResult = Eval->ImageSampleFunction(Sample);
+	WriteToBuffer(Out, SampleResult);
+}
+
+FVector4f FMeshMultiResampleImageEvaluator::ImageSampleFunction(const FCorrespondenceSample& Sample)
+{
+	FVector4f Color(0.0f, 0.0f, 0.0f, 1.0f);
+	if (!bValidDetailMesh)
+	{
+		return Color;
+	}
+
+	int32 DetailTriID = Sample.DetailTriID;
+	if (DetailMesh->IsTriangle(DetailTriID))
+	{
+		// TODO: Decompose the Map into a sparse array lookup.
+		int32 MaterialID = DetailMaterialIDAttrib->GetValue(DetailTriID);
+		TSharedPtr<UE::Geometry::TImageBuilder<FVector4f>, ESPMode::ThreadSafe> TextureImage = MultiTextures.FindRef(MaterialID);
+		if (TextureImage)
+		{
+			FVector2d DetailUV;
+			DetailUVOverlay->GetTriBaryInterpolate<double>(DetailTriID, &Sample.DetailBaryCoords.X, &DetailUV.X);
+			Color = TextureImage->BilinearSampleUV<float>(DetailUV, FVector4f(0, 0, 0, 1));
+		}
+	}
+	return Color;
+}
diff --git a/Engine/Plugins/Experimental/GeometryProcessing/Source/DynamicMesh/Public/Sampling/MeshCurvatureMapBaker.h b/Engine/Plugins/Experimental/GeometryProcessing/Source/DynamicMesh/Public/Sampling/MeshCurvatureMapBaker.h
--- a/Engine/Plugins/Experimental/GeometryProcessing/Source/DynamicMesh/Public/Sampling/MeshCurvatureMapBaker.h
+++ b/Engine/Plugins/Experimental/GeometryProcessing/Source/DynamicMesh/Public/Sampling/MeshCurvatureMapBaker.h
@@ -90,37 +90,8 @@ protected:
 
 	void Bake_Multi();
 
-	double GetCurvature(int32 vid);
 
 	void GetColorMapRange(FVector3f& NegativeColor, FVector3f& ZeroColor, FVector3f& PositiveColor);
-
-public:
-	//
-	// New Baker interface
-	//
-
-	/** Invoked at start of bake to initialize baker. */
-	virtual void PreEvaluate(const FMeshMapBaker& Baker) override;
-
-	/** Evaluate the sample at this mesh correspondence point. */
-	virtual FVector4f EvaluateSample(const FMeshMapBaker& Baker, const FCorrespondenceSample& Sample) override;
-
-	/** @return the default sample value for the baker. */
-	virtual FVector4f DefaultSample() const override { return FVector4f(0.0f, 0.0f, 0.0f, 1.0f); }
-
-protected:
-	// Cached data
-	const FDynamicMesh3* DetailMesh = nullptr;
-	double MinPreClamp = -TNumericLimits<double>::Max();
-	double MaxPreClamp = TNumericLimits<double>::Max();
-	FInterval1d ClampRange;
-	FVector3f NegativeColor;
-	FVector3f ZeroColor;
-	FVector3f PositiveColor;
-
-private:
-	template <class SampleType>
-	double SampleFunction(const SampleType& Sample);
 };
 
 
diff --git a/Engine/Plugins/Experimental/GeometryProcessing/Source/DynamicMesh/Public/Sampling/MeshCurvatureMapEvaluator.h b/Engine/Plugins/Experimental/GeometryProcessing/Source/DynamicMesh/Public/Sampling/MeshCurvatureMapEvaluator.h
new file mode 100644
--- /dev/null
+++ b/Engine/Plugins/Experimental/GeometryProcessing/Source/DynamicMesh/Public/Sampling/MeshCurvatureMapEvaluator.h
@@ -0,0 +1,83 @@
+// Copyright Epic Games, Inc. All Rights Reserved.
+
+#pragma once
+
+#include "Sampling/MeshMapEvaluator.h"
+
+namespace UE
+{
+namespace Geometry
+{
+
+class FMeshVertexCurvatureCache;
+
+class DYNAMICMESH_API FMeshCurvatureMapEvaluator : public FMeshMapEvaluator
+{
+public:
+	enum class ECurvatureType
+	{
+		Mean = 0,
+		Gaussian = 1,
+		MaxPrincipal = 2,
+		MinPrincipal = 3
+	};
+	ECurvatureType UseCurvatureType = ECurvatureType::Mean;
+
+	enum class EColorMode
+	{
+		BlackGrayWhite = 0,
+		RedGreenBlue = 1,
+		RedBlue = 2
+	};
+	EColorMode UseColorMode = EColorMode::RedGreenBlue;
+
+	enum class EClampMode
+	{
+		FullRange = 0,
+		Positive = 1,
+		Negative = 2
+	};
+	EClampMode UseClampMode = EClampMode::FullRange;
+
+	double RangeScale = 1.0;
+	double MinRangeScale = 0.0;
+
+	// Allows override of the max curvature; if false, range is set based on [-(avg+stddev), avg+stddev]
+	bool bOverrideCurvatureRange = false;
+	double OverrideRangeMax = 0.1;
+
+	// Required input data, can be provided, will be computed otherwise
+	TSharedPtr<FMeshVertexCurvatureCache> Curvatures;
+
+public:
+	/** Invoked at start of bake to initialize baker. */
+	virtual void Setup(const FMeshMapBaker& Baker, FEvaluationContext& Context) override;
+
+	static void EvaluateSample(float*& Out, const FCorrespondenceSample& Sample, void* EvalData);
+
+	static void EvaluateDefault(float*& Out, void* EvalData);
+
+	/** Populate Curvatures member if valid data has not been provided */
+	void CacheDetailCurvatures(const FDynamicMesh3* DetailMesh);
+
+protected:
+	// Cached data
+	const FDynamicMesh3* DetailMesh = nullptr;
+	double MinPreClamp = -TNumericLimits<double>::Max();
+	double MaxPreClamp = TNumericLimits<double>::Max();
+	FInterval1d ClampRange;
+	FVector3f NegativeColor;
+	FVector3f ZeroColor;
+	FVector3f PositiveColor;
+
+protected:
+	double GetCurvature(int32 vid);
+	void GetColorMapRange(FVector3f& NegativeColor, FVector3f& ZeroColor, FVector3f& PositiveColor);
+
+private:
+	double SampleFunction(const FCorrespondenceSample& Sample);
+};
+
+} // end namespace UE::Geometry
+} // end namespace UE
+
diff --git a/Engine/Plugins/Experimental/GeometryProcessing/Source/DynamicMesh/Public/Sampling/MeshImageBaker.h b/Engine/Plugins/Experimental/GeometryProcessing/Source/DynamicMesh/Public/Sampling/MeshImageBaker.h
--- a/Engine/Plugins/Experimental/GeometryProcessing/Source/DynamicMesh/Public/Sampling/MeshImageBaker.h
+++ b/Engine/Plugins/Experimental/GeometryProcessing/Source/DynamicMesh/Public/Sampling/MeshImageBaker.h
@@ -3,23 +3,17 @@
 #pragma once
 
 #include "Sampling/MeshImageBakingCache.h"
-#include "Image/ImageBuilder.h"
 
 namespace UE
 {
 namespace Geometry
 {
 
-class FMeshMapBaker;
-
 class DYNAMICMESH_API FMeshImageBaker
 {
 public:
 	virtual ~FMeshImageBaker() {}
 
-	//
-	// Old Baker interface
-	//
 	void SetCache(const FMeshImageBakingCache* CacheIn)
 	{
 		ImageBakeCache = CacheIn;
@@ -32,38 +26,14 @@ public:
 
 	virtual void Bake()
 	{
-	}
-
-protected:
-	const FMeshImageBakingCache* ImageBakeCache;
-
-public:
-	//
-	// New Baker interface
-	//
-	struct FCorrespondenceSample
-	{
-		FMeshUVSampleInfo BaseSample;
-		FVector3d BaseNormal;
-
-		int32 DetailTriID;
-		FVector3d DetailBaryCoords;
-	};
 
-	/** Invoked at start of bake to initialize baker. */
-	virtual void PreEvaluate(const FMeshMapBaker& Baker) {}
+	}
 
-	/** Evaluate the sample at this mesh correspondence point. */
-	virtual FVector4f EvaluateSample(const FMeshMapBaker& Baker, const FCorrespondenceSample& Sample) { return DefaultSample(); }
 
-	/** Invoked at the end of the bake to finalize the bake result. */
-	virtual void PostEvaluate(const FMeshMapBaker& Baker, TImageBuilder<FVector4f>& Result) {}
+protected:
 
-	/** @return the default sample value for the baker. */
-	virtual FVector4f DefaultSample() const { return FVector4f::Zero(); }
+	const FMeshImageBakingCache* ImageBakeCache;
 
-	/** @return true if this baker supports multisampling. */
-	virtual bool SupportsMultisampling() const { return true; }
 };
 
 } // end namespace UE::Geometry
diff --git a/Engine/Plugins/Experimental/GeometryProcessing/Source/DynamicMesh/Public/Sampling/MeshMapBaker.h b/Engine/Plugins/Experimental/GeometryProcessing/Source/DynamicMesh/Public/Sampling/MeshMapBaker.h
--- a/Engine/Plugins/Experimental/GeometryProcessing/Source/DynamicMesh/Public/Sampling/MeshMapBaker.h
+++ b/Engine/Plugins/Experimental/GeometryProcessing/Source/DynamicMesh/Public/Sampling/MeshMapBaker.h
@@ -6,7 +6,7 @@
 #include "DynamicMeshAttributeSet.h"
 #include "DynamicMeshAABBTree3.h"
 #include "MeshTangents.h"
-#include "Sampling/MeshImageBaker.h"
+#include "Sampling/MeshMapEvaluator.h"
 #include "Sampling/MeshSurfaceSampler.h"
 #include "Spatial/DenseGrid2.h"
 #include "Image/ImageBuilder.h"
@@ -17,6 +17,41 @@ namespace UE
 namespace Geometry
 {
 
+class FImageOccupancyMap;
+
+class FMeshMapTileBuffer
+{
+public:
+	FMeshMapTileBuffer(const FImageDimensions& TileIn, const int32 PixelSizeIn)
+		: Tile(TileIn)
+		, PixelSize(PixelSizeIn + 1) // + 1 for accumulated pixel weight.
+	{
+		Buffer = static_cast<float*>(FMemory::MallocZeroed(sizeof(float) * PixelSize * Tile.Num()));
+	}
+
+	~FMeshMapTileBuffer()
+	{
+		FMemory::Free(Buffer);
+	}
+
+	float& GetPixelWeight(int64 LinearIdx)
+	{
+		checkSlow(LinearIdx >= 0 && LinearIdx < Tile.Num());
+		return Buffer[LinearIdx * PixelSize];
+	}
+
+	float* GetPixel(int64 LinearIdx)
+	{
+		checkSlow(LinearIdx >= 0 && LinearIdx < Tile.Num());
+		return &Buffer[LinearIdx * PixelSize + 1];
+	}
+
+private:
+	const FImageDimensions Tile;
+	const int32 PixelSize;
+	float* Buffer;
+};
+
 class DYNAMICMESH_API FMeshMapBaker
 {
 public:
@@ -29,13 +64,13 @@ public:
 	void Bake();
 
 	/** Add a baker to be processed. */
-	int32 AddBaker(TSharedPtr<FMeshImageBaker> Target);
+	int32 AddBaker(TSharedPtr<FMeshMapEvaluator, ESPMode::ThreadSafe> Target);
 
 	/** Reset the list of bakers. */
 	void Reset();
 
 	/** @return the bake result image for a given baker index. */
-	const TSharedPtr<TImageBuilder<FVector4f>, ESPMode::ThreadSafe>& GetBakeResult(int32 Index) const { return BakeResults[Index]; }
+	const TArrayView<TUniquePtr<TImageBuilder<FVector4f>>> GetBakeResults(int32 BakerIdx);
 
 	/** if this function returns true, we should abort calculation */
 	TFunction<bool(void)> CancelF = []() { return false; };
@@ -92,10 +127,31 @@ public:
 	ECorrespondenceStrategy GetCorrespondenceStrategy() const { return CorrespondenceStrategy; }
 
 protected:
-	void BakePixel(FImageOccupancyMap& OccupancyMap, const FImageDimensions& Tile, int32 ImgX, int32 ImgY);
+	/** Evaluate samples for this tile pixel. */
+	void BakePixel(FMeshMapTileBuffer& TileBuffer, FImageOccupancyMap& OccupancyMap, const FImageDimensions& Tile, const FVector2i& TileCoords);
+
+	/** Evaluate this sample. */
+	void BakeSample(FMeshMapTileBuffer& TileBuffer, const FMeshMapEvaluator::FCorrespondenceSample& Sample, const FImageDimensions& Tile,
+		const FVector2i& TileCoords, const FVector2i& ImageCoords, const float& SampleWeight);
+
+	/** Initialize evaluation contexts and precompute data for bake evaluation. */
+	void InitBake();
+
+	/** Initialize bake sample default floats and colors. */
+	void InitBakeDefaults();
+
+	/**
+	 * Convert float buffer value to float4 color. This function
+	 * will advance the buffer pointer by the stride value.
+	 * 
+	 * @param Buffer [out] input buffer data.
+	 * @param Stride the data layout of the buffer pointer.
+	 * @param Weight the weight to apply to the buffer data.
+	 */
+	FVector4f FloatToPixel(float*& Buffer, FMeshMapEvaluator::EComponents Stride, float Weight);
 
 protected:
-	bool bParallel = true;
+	const bool bParallel = true;
 
 	const FDynamicMesh3* DetailMesh = nullptr;
 	const FDynamicMeshAABBTree3* DetailSpatial = nullptr;
@@ -103,20 +159,45 @@ protected:
 	TSharedPtr<FMeshTangentsd, ESPMode::ThreadSafe> TargetMeshTangents;
 
 	FDynamicMesh3 FlatMesh;
-	TMeshSurfaceUVSampler<FMeshImageBaker::FCorrespondenceSample> DetailMeshSampler;
+	TMeshSurfaceUVSampler<FMeshMapEvaluator::FCorrespondenceSample> DetailMeshSampler;
 
 	FImageDimensions Dimensions = FImageDimensions(128, 128);
 	int32 UVLayer = 0;
 	double Thickness = 3.0;
 	ECorrespondenceStrategy CorrespondenceStrategy = ECorrespondenceStrategy::RaycastStandard;
 	int32 GutterSize = 4;
+
+	/** The square dimensions for multisampling each pixel. */
 	int32 Multisampling = 1;
 
+	/** The square dimensions for tiled processing of the output image(s). */
 	const int32 TileSize = 32;
 
-	TArray<TSharedPtr<FMeshImageBaker, ESPMode::ThreadSafe>> Bakers;
-	TArray<TSharedPtr<TImageBuilder<FVector4f>, ESPMode::ThreadSafe>> BakeResults;
+	/** The total size of the temporary float buffer for BakeSample. */
+	int32 BakeSampleBufferSize = 0;
+
+	/** The list of evaluators to process. */
+	TArray<TSharedPtr<FMeshMapEvaluator, ESPMode::ThreadSafe>> Bakers;
+
+	/** Evaluation contexts for each mesh evaluator. */
+	TArray<FMeshMapEvaluator::FEvaluationContext> BakeContexts;
+
+	/** Lists of Bake indices for each accumulation mode. */
+	TArray<TArray<int32>> BakeAccumulateLists;
+
+	/** Array of default values/colors per BakeResult. */
+	TArray<float> BakeDefaults;
+	TArray<FVector4f> BakeDefaultColors;
+	
+	/** Offsets per Baker into the BakeResults array.*/
+	TArray<int32> BakeOffsets;
+
+	/** Offsets per BakeResult into the BakeSample buffer.*/
+	TArray<int32> BakeSampleOffsets;
+
+	/** Array of bake result images. */
+	TArray<TUniquePtr<TImageBuilder<FVector4f>>> BakeResults;
 };
 
 } // end namespace UE::Geometry
-} // end namespace UE
\ No newline at end of file
+} // end namespace UE
diff --git a/Engine/Plugins/Experimental/GeometryProcessing/Source/DynamicMesh/Public/Sampling/MeshMapEvaluator.h b/Engine/Plugins/Experimental/GeometryProcessing/Source/DynamicMesh/Public/Sampling/MeshMapEvaluator.h
new file mode 100644
--- /dev/null
+++ b/Engine/Plugins/Experimental/GeometryProcessing/Source/DynamicMesh/Public/Sampling/MeshMapEvaluator.h
@@ -0,0 +1,169 @@
+// Copyright Epic Games, Inc. All Rights Reserved.
+
+#pragma once
+
+#include "Sampling/MeshSurfaceSampler.h"
+
+namespace UE
+{
+namespace Geometry
+{
+
+class FMeshMapBaker;
+
+class FMeshMapEvaluator
+{
+public:
+	/** The number of float components of an evaluator result/target. */
+	enum class EComponents
+	{
+		Float1 = 1,
+		Float2 = 2,
+		Float3 = 3,
+		Float4 = 4
+	};
+
+	/** Accumulate mode informs the baker how samples should be accumulated into the buffer. */
+	enum class EAccumulateMode
+	{
+		/** Accumulate sample results (valid & default) into the buffer. */
+		Add,
+		/** Overwrite the previous sample result with the new value into the buffer. */
+		Overwrite,
+		/** Last mode. Used to count modes. Ensure this is always last. */
+		Last
+	};
+
+	/** Detail to target mesh correspondence sample. */
+	struct FCorrespondenceSample
+	{
+		FMeshUVSampleInfo BaseSample;
+		FVector3d BaseNormal;
+
+		int32 DetailTriID;
+		FVector3d DetailBaryCoords;
+	};
+
+	/**
+	 * Evaluator context data to indicate to the baker:
+	 * 1. The evaluation functions.
+	 * 2. The accumulation method for data samples.
+	 * 3. The number of output results/targets for this evaluator.
+	 * 3. The data layout of each output result/target for this evaluator.
+	 */
+	struct FEvaluationContext
+	{
+		/**
+		 * Function pointer to evaluate a mesh correspondence sample into a
+		 * float buffer. This function should:
+		 * 
+		 * - assume that the size of the buffer is correct.
+		 * - advance the the buffer pointer by the DataLayout stride.
+		 * 
+		 * @param Out the buffer to populate.
+		 * @param Sample the detail/target mesh correspondence sample data.
+		 * @param EvalData custom data pointer provided by the evaluation context.
+		 */
+		using EvaluateFn = void(*)(float*& Out, const FCorrespondenceSample& Sample, void* EvalData);
+		EvaluateFn Evaluate = nullptr;
+
+		/**
+		 * Function pointer to evaluate the default sample into a
+		 * float buffer. The default sample is used:
+		 * 
+		 * - to clear the output image on initialization.
+		 * - for samples for which there is no mesh image occupancy nor correspondence.
+		 * 
+		 * This function should:
+		 * 
+		 * - assume that the size of the buffer is correct.
+		 * - advance the the buffer pointer by the DataLayout stride.
+		 * 
+		 * @param Out the buffer to populate.
+		 * @param Sample the detail/target mesh correspondence sample data.
+		 * @param EvalData custom data pointer provided by the evaluation context.
+		 */
+		using EvaluateDefaultFn = void(*)(float*& Out, void* EvalData);
+		EvaluateDefaultFn EvaluateDefault = nullptr;
+
+		/** Define custom data to be passed to the evaluation function. */
+		void* EvalData = nullptr;
+
+		/** Define how the results should be accumulatd by the Baker. */
+		EAccumulateMode AccumulateMode = EAccumulateMode::Add;
+
+		/**
+		 * The data layout of this evaluator including:
+		 * 1. The number of outputs.
+		 * 2. The number of floats (stride) per output. [1,4]
+		 */
+		TArray<EComponents> DataLayout;
+	};
+
+	virtual ~FMeshMapEvaluator() = default;
+
+	/**
+	 * Invoked at start of bake to setup this evaluator and return
+	 * an evaluation context to the baker. The evaluation context
+	 * indicates to the baker how the evaluator should be invoked.
+	 * See FEvaluationContext for more details.
+	 * 
+	 * @param Baker the baker that owns this evaluator.
+	 * @param Context [out] the evaluation context.
+	 */
+	virtual void Setup(const FMeshMapBaker& Baker, FEvaluationContext& Context) = 0;
+
+protected:
+	/**
+	 * Write float data to a float buffer and increment.
+	 *
+	 * @param Out the target float buffer.
+	 * @param Data the data to write into the buffer.
+	 */
+	static void WriteToBuffer(float*& Out, const float Data)
+	{
+		*Out++ = Data;
+	}
+
+	/**
+	 * Write FVector2f data to a float buffer and increment.
+	 *
+	 * @param Out the target float buffer.
+	 * @param Data the data to write into the buffer.
+	 */
+	static void WriteToBuffer(float*& Out, const FVector2f& Data)
+	{
+		*Out++ = Data.X;
+		*Out++ = Data.Y;
+	}
+
+	/**
+	 * Write FVector3f data to a float buffer and increment.
+	 *
+	 * @param Out the target float buffer.
+	 * @param Data the data to write into the buffer.
+	 */
+	static void WriteToBuffer(float*& Out, const FVector3f& Data)
+	{
+		*Out++ = Data.X;
+		*Out++ = Data.Y;
+		*Out++ = Data.Z;
+	}
+
+	/**
+	 * Write FVector4f data to a float buffer and increment.
+	 * 
+	 * @param Out the target float buffer.
+	 * @param Data the data to write into the buffer.
+	 */
+	static void WriteToBuffer(float*& Out, const FVector4f& Data)
+	{
+		*Out++ = Data.X;
+		*Out++ = Data.Y;
+		*Out++ = Data.Z;
+		*Out++ = Data.W;
+	}
+};
+
+} // end namespace UE::Geometry
+} // end namespace UE
diff --git a/Engine/Plugins/Experimental/GeometryProcessing/Source/DynamicMesh/Public/Sampling/MeshNormalMapBaker.h b/Engine/Plugins/Experimental/GeometryProcessing/Source/DynamicMesh/Public/Sampling/MeshNormalMapBaker.h
--- a/Engine/Plugins/Experimental/GeometryProcessing/Source/DynamicMesh/Public/Sampling/MeshNormalMapBaker.h
+++ b/Engine/Plugins/Experimental/GeometryProcessing/Source/DynamicMesh/Public/Sampling/MeshNormalMapBaker.h
@@ -24,7 +24,7 @@ public:
 
 	const TMeshTangents<double>* BaseMeshTangents;
 
-	FVector3f DefaultNormal = FVector3f::UnitZ();
+	FVector3f DefaultNormal = FVector3f(0, 0, 1);
 
 	//
 	// Compute functions
@@ -43,33 +43,6 @@ public:
 protected:
 	TUniquePtr<TImageBuilder<FVector3f>> NormalsBuilder;
 
-public:
-	//
-	// New Baker interface
-	//
-
-	/** Invoked at start of bake to initialize baker. */
-	virtual void PreEvaluate(const FMeshMapBaker& Baker) override;
-
-	/** Evaluate the sample at this mesh correspondence point. */
-	virtual FVector4f EvaluateSample(const FMeshMapBaker& Baker, const FCorrespondenceSample& Sample) override;
-
-	/** @return the default sample value for the baker. */
-	virtual FVector4f DefaultSample() const override
-	{
-		// Map normal space [-1,1] to floating point color space [0,1]
-		FVector3f Normal = (DefaultNormal + FVector3f::One()) * 0.5f;
-		return FVector4f(Normal.X, Normal.Y, Normal.Z, 1.0f);
-	}
-
-protected:
-	// Cached data
-	const FDynamicMesh3* DetailMesh = nullptr;
-	const FDynamicMeshNormalOverlay* DetailNormalOverlay = nullptr;
-
-private:
-	template <class SampleType>
-	FVector3f SampleFunction(const SampleType& Sample);
 };
 
 
diff --git a/Engine/Plugins/Experimental/GeometryProcessing/Source/DynamicMesh/Public/Sampling/MeshNormalMapEvaluator.h b/Engine/Plugins/Experimental/GeometryProcessing/Source/DynamicMesh/Public/Sampling/MeshNormalMapEvaluator.h
new file mode 100644
--- /dev/null
+++ b/Engine/Plugins/Experimental/GeometryProcessing/Source/DynamicMesh/Public/Sampling/MeshNormalMapEvaluator.h
@@ -0,0 +1,37 @@
+// Copyright Epic Games, Inc. All Rights Reserved.
+
+#pragma once
+
+#include "Sampling/MeshMapEvaluator.h"
+#include "MeshTangents.h"
+
+namespace UE
+{
+namespace Geometry
+{
+
+class DYNAMICMESH_API FMeshNormalMapEvaluator : public FMeshMapEvaluator
+{
+public:
+	/** Invoked at start of bake to initialize baker. */
+	virtual void Setup(const FMeshMapBaker& Baker, FEvaluationContext& Context) override;
+
+	static void EvaluateSample(float*& Out, const FCorrespondenceSample& Sample, void* EvalData);
+
+	static void EvaluateDefault(float*& Out, void* EvalData);
+
+protected:
+	// Cached data
+	const FDynamicMesh3* DetailMesh = nullptr;
+	const FDynamicMeshNormalOverlay* DetailNormalOverlay = nullptr;
+	const TMeshTangents<double>* BaseMeshTangents;
+
+	FVector3f DefaultNormal = FVector3f::UnitZ();
+
+private:
+	FVector3f SampleFunction(const FCorrespondenceSample& Sample);
+};
+
+} // end namespace UE::Geometry
+} // end namespace UE
+
diff --git a/Engine/Plugins/Experimental/GeometryProcessing/Source/DynamicMesh/Public/Sampling/MeshOcclusionMapBaker.h b/Engine/Plugins/Experimental/GeometryProcessing/Source/DynamicMesh/Public/Sampling/MeshOcclusionMapBaker.h
--- a/Engine/Plugins/Experimental/GeometryProcessing/Source/DynamicMesh/Public/Sampling/MeshOcclusionMapBaker.h
+++ b/Engine/Plugins/Experimental/GeometryProcessing/Source/DynamicMesh/Public/Sampling/MeshOcclusionMapBaker.h
@@ -7,7 +7,6 @@
 #include "Sampling/MeshImageBakingCache.h"
 #include "Image/ImageBuilder.h"
 #include "MeshTangents.h"
-#include "Math/RandomStream.h"
 
 namespace UE
 {
@@ -125,57 +124,6 @@ public:
 protected:
 	TUniquePtr<TImageBuilder<FVector3f>> OcclusionBuilder;
 	TUniquePtr<TImageBuilder<FVector3f>> NormalBuilder;
-
-public:
-	//
-	// New Baker interface
-	//
-
-	/** Invoked at start of bake to initialize baker. */
-	virtual void PreEvaluate(const FMeshMapBaker& Baker) override;
-
-	/** Evaluate the sample at this mesh correspondence point. */
-	virtual FVector4f EvaluateSample(const FMeshMapBaker& Baker, const FCorrespondenceSample& Sample) override;
-
-	/** @return the default sample value for the baker. */
-	virtual FVector4f DefaultSample() const override
-	{
-		if (WantAmbientOcclusion())
-		{
-			return FVector4f::One();
-		}
-		else //if (WantBentNormal())
-		{
-			FVector3f DefaultNormal = FVector3f::UnitZ();
-			switch (NormalSpace)
-			{
-			case ESpace::Tangent:
-				break;
-			case ESpace::Object:
-				DefaultNormal = FVector3f::Zero();
-				break;
-			}
-			// Map normal space [-1,1] to floating point color space [0,1]
-			DefaultNormal = (DefaultNormal + FVector3f::One()) * 0.5f;
-			return FVector4f(DefaultNormal.X, DefaultNormal.Y, DefaultNormal.Z, 1.0f);
-		}
-	}
-
-protected:
-	const FDynamicMesh3* DetailMesh = nullptr;
-	const FDynamicMeshAABBTree3* DetailSpatial = nullptr;
-	const FDynamicMeshNormalOverlay* DetailNormalOverlay = nullptr;
-	double BiasDotThreshold = 0.25;
-	TArray<FVector3d> RayDirections;
-
-	FRandomStream RotationGen = FRandomStream(31337);
-	FCriticalSection RotationLock;
-
-private:
-	template<class SampleType>
-	void SampleFunction(const SampleType& Sample, double& Occlusion, FVector3d& Normal);
-
-	double GetRandomRotation();
 };
 
 } // end namespace UE::Geometry
diff --git a/Engine/Plugins/Experimental/GeometryProcessing/Source/DynamicMesh/Public/Sampling/MeshOcclusionMapEvaluator.h b/Engine/Plugins/Experimental/GeometryProcessing/Source/DynamicMesh/Public/Sampling/MeshOcclusionMapEvaluator.h
new file mode 100644
--- /dev/null
+++ b/Engine/Plugins/Experimental/GeometryProcessing/Source/DynamicMesh/Public/Sampling/MeshOcclusionMapEvaluator.h
@@ -0,0 +1,82 @@
+// Copyright Epic Games, Inc. All Rights Reserved.
+
+#pragma once
+
+#include "Sampling/MeshMapEvaluator.h"
+#include "Templates/Tuple.h"
+#include "Misc/EnumClassFlags.h"
+#include "Math/RandomStream.h"
+#include "MeshTangents.h"
+
+namespace UE
+{
+namespace Geometry
+{
+
+enum class EMeshOcclusionMapType : uint8
+{
+	None = 0,
+	AmbientOcclusion = 1,
+	BentNormal = 2,
+	All = 3
+};
+ENUM_CLASS_FLAGS(EMeshOcclusionMapType);
+
+class DYNAMICMESH_API FMeshOcclusionMapEvaluator : public FMeshMapEvaluator
+{
+public:
+	enum class EDistribution
+	{
+		Uniform,
+		Cosine
+	};
+
+	enum class ESpace
+	{
+		Tangent,
+		Object
+	};
+
+	EMeshOcclusionMapType OcclusionType = EMeshOcclusionMapType::All;
+	int32 NumOcclusionRays = 32;
+	double MaxDistance = TNumericLimits<double>::Max();
+	double SpreadAngle = 180.0;
+	EDistribution Distribution = EDistribution::Cosine;
+
+	// Ambient Occlusion
+	double BiasAngleDeg = 15.0;
+
+	// Bent Normal
+	ESpace NormalSpace = ESpace::Tangent;
+
+	using FOcclusionTuple = TTuple<float, FVector3f>;
+
+public:
+	/** Invoked at start of bake to initialize baker. */
+	virtual void Setup(const FMeshMapBaker & Baker, FEvaluationContext & Context) override;
+
+	template <EMeshOcclusionMapType OcclusionType, ESpace NormalSpace>
+	static void EvaluateSample(float*& Out, const FCorrespondenceSample& Sample, void* EvalData);
+
+	template <EMeshOcclusionMapType OcclusionType, ESpace NormalSpace>
+	static void EvaluateDefault(float*& Out, void* EvalData);
+
+protected:
+	// Cached data
+	const FDynamicMesh3* DetailMesh = nullptr;
+	const FDynamicMeshAABBTree3* DetailSpatial = nullptr;
+	const FDynamicMeshNormalOverlay* DetailNormalOverlay = nullptr;
+	const TMeshTangents<double>* BaseMeshTangents = nullptr;
+	double BiasDotThreshold = 0.25;
+	TArray<FVector3d> RayDirections;
+
+	static const FVector3d DefaultTangentNormal;
+	static const FVector3d DefaultObjectNormal;
+
+private:
+	template <EMeshOcclusionMapType OcclusionType, ESpace NormalSpace>
+	FOcclusionTuple SampleFunction(const FCorrespondenceSample& SampleData, const FVector3d& DefaultNormal);
+};
+
+} // end namespace UE::Geometry
+} // end namespace UE
diff --git a/Engine/Plugins/Experimental/GeometryProcessing/Source/DynamicMesh/Public/Sampling/MeshPropertyMapBaker.h b/Engine/Plugins/Experimental/GeometryProcessing/Source/DynamicMesh/Public/Sampling/MeshPropertyMapBaker.h
--- a/Engine/Plugins/Experimental/GeometryProcessing/Source/DynamicMesh/Public/Sampling/MeshPropertyMapBaker.h
+++ b/Engine/Plugins/Experimental/GeometryProcessing/Source/DynamicMesh/Public/Sampling/MeshPropertyMapBaker.h
@@ -49,35 +49,6 @@ public:
 
 protected:
 	TUniquePtr<TImageBuilder<FVector3f>> ResultBuilder;
-
-public:
-	//
-	// New Baker interface
-	//
-
-	/** Invoked at start of bake to initialize baker. */
-	virtual void PreEvaluate(const FMeshMapBaker& Baker) override;
-
-	/** Evaluate the sample at this mesh correspondence point. */
-	virtual FVector4f EvaluateSample(const FMeshMapBaker& Baker, const FCorrespondenceSample& Sample) override;
-
-	/** @return the default sample value for the baker. */
-	virtual FVector4f DefaultSample() const override { return FVector4f(0.0f, 0.0f, 0.0f, 1.0f); }
-
-	/** @return true if this baker supports multisampling. */
-	virtual bool SupportsMultisampling() const override { return Property != EMeshPropertyBakeType::MaterialID; }
-
-protected:
-	// Cached data
-	const FDynamicMesh3* DetailMesh = nullptr;
-	const FDynamicMeshNormalOverlay* DetailNormalOverlay = nullptr;
-	const FDynamicMeshUVOverlay* DetailUVOverlay = nullptr;
-	FAxisAlignedBox3d Bounds;
-	FVector3f DefaultValue = FVector3f::Zero();
-
-private:
-	template <class SampleType>
-	FVector3f SampleFunction(const SampleType& Sample);
 };
 
 } // end namespace UE::Geometry
diff --git a/Engine/Plugins/Experimental/GeometryProcessing/Source/DynamicMesh/Public/Sampling/MeshPropertyMapEvaluator.h b/Engine/Plugins/Experimental/GeometryProcessing/Source/DynamicMesh/Public/Sampling/MeshPropertyMapEvaluator.h
new file mode 100644
--- /dev/null
+++ b/Engine/Plugins/Experimental/GeometryProcessing/Source/DynamicMesh/Public/Sampling/MeshPropertyMapEvaluator.h
@@ -0,0 +1,68 @@
+// Copyright Epic Games, Inc. All Rights Reserved.
+
+#pragma once
+
+#include "Sampling/MeshMapEvaluator.h"
+
+namespace UE
+{
+namespace Geometry
+{
+
+enum class EMeshPropertyMapType
+{
+	Position = 1,
+	Normal = 2,
+	FacetNormal = 3,
+	UVPosition = 4,
+	MaterialID = 5
+};
+
+class DYNAMICMESH_API FMeshPropertyMapEvaluator : public FMeshMapEvaluator
+{
+public:
+	EMeshPropertyMapType Property = EMeshPropertyMapType::Normal;
+
+public:
+	/** Invoked at start of bake to initialize baker. */
+	virtual void Setup(const FMeshMapBaker& Baker, FEvaluationContext& Context) override;
+
+	static void EvaluateSample(float*& Out, const FCorrespondenceSample& Sample, void* EvalData);
+
+	static void EvaluateDefault(float*& Out, void* EvalData);
+
+protected:
+	// Cached data
+	const FDynamicMesh3* DetailMesh = nullptr;
+	const FDynamicMeshNormalOverlay* DetailNormalOverlay = nullptr;
+	const FDynamicMeshUVOverlay* DetailUVOverlay = nullptr;
+	FAxisAlignedBox3d Bounds;
+	FVector3f DefaultValue = FVector3f::Zero();
+
+private:
+	FVector3f NormalToColor(const FVector3d Normal)
+	{
+		return (FVector3f)((Normal + FVector3d::One()) * 0.5);
+	}
+
+	FVector3f UVToColor(const FVector2d UV)
+	{
+		double X = FMathd::Clamp(UV.X, 0.0, 1.0);
+		double Y = FMathd::Clamp(UV.Y, 0.0, 1.0);
+		return (FVector3f)FVector3d(X, Y, 0);
+	}
+
+	FVector3f PositionToColor(const FVector3d Position, const FAxisAlignedBox3d SafeBounds)
+	{
+		double X = (Position.X - SafeBounds.Min.X) / SafeBounds.Width();
+		double Y = (Position.Y - SafeBounds.Min.Y) / SafeBounds.Height();
+		double Z = (Position.Z - SafeBounds.Min.Z) / SafeBounds.Depth();
+		return (FVector3f)FVector3d(X, Y, Z);
+	}
+
+	FVector3f SampleFunction(const FCorrespondenceSample& Sample);
+};
+
+} // end namespace UE::Geometry
+} // end namespace UE
+
diff --git a/Engine/Plugins/Experimental/GeometryProcessing/Source/DynamicMesh/Public/Sampling/MeshResampleImageBaker.h b/Engine/Plugins/Experimental/GeometryProcessing/Source/DynamicMesh/Public/Sampling/MeshResampleImageBaker.h
--- a/Engine/Plugins/Experimental/GeometryProcessing/Source/DynamicMesh/Public/Sampling/MeshResampleImageBaker.h
+++ b/Engine/Plugins/Experimental/GeometryProcessing/Source/DynamicMesh/Public/Sampling/MeshResampleImageBaker.h
@@ -42,26 +42,6 @@ public:
 
 protected:
 	TUniquePtr<TImageBuilder<FVector4f>> ResultBuilder;
-
-public:
-	//
-	// New Baker interface
-	//
-
-	/** Invoked at start of bake to initialize sampler. */
-	virtual void PreEvaluate(const FMeshMapBaker& Baker) override;
-
-	virtual FVector4f EvaluateSample(const FMeshMapBaker& Baker, const FCorrespondenceSample& Sample) override;
-
-	virtual FVector4f DefaultSample() const override { return FVector4f(0.0f, 0.0f, 0.0f, 1.0f); }
-
-protected:
-	// Cached data
-	const FDynamicMesh3* DetailMesh = nullptr;
-
-private:
-	template <class SampleType>
-	FVector4f PropertySampleFunction(const SampleType& Sample);
 };
 
 
@@ -78,24 +58,6 @@ protected:
 
 	void InitResult();
 	void BakeMaterial(int32 MaterialID);
-
-public:
-	//
-	// New Baker interface
-	//
-
-	/** Invoked at start of bake to initialize baker. */
-	virtual void PreEvaluate(const FMeshMapBaker& Baker) override;
-
-	/** Evaluate the sample at this mesh correspondence point. */
-	virtual FVector4f EvaluateSample(const FMeshMapBaker& Baker, const FCorrespondenceSample& Sample) override;
-
-	/** @return the default sample value for the baker. */
-	virtual FVector4f DefaultSample() const override { return FVector4f(0.0f, 0.0f, 0.0f, 1.0f); }
-
-protected:
-	const FDynamicMeshMaterialAttribute* DetailMaterialIDAttrib = nullptr;
-	bool bValidDetailMesh = false;
 };
 
 } // end namespace UE::Geometry
diff --git a/Engine/Plugins/Experimental/GeometryProcessing/Source/DynamicMesh/Public/Sampling/MeshResampleImageEvaluator.h b/Engine/Plugins/Experimental/GeometryProcessing/Source/DynamicMesh/Public/Sampling/MeshResampleImageEvaluator.h
new file mode 100644
--- /dev/null
+++ b/Engine/Plugins/Experimental/GeometryProcessing/Source/DynamicMesh/Public/Sampling/MeshResampleImageEvaluator.h
@@ -0,0 +1,62 @@
+// Copyright Epic Games, Inc. All Rights Reserved.
+
+#pragma once
+
+#include "Sampling/MeshMapEvaluator.h"
+#include "Image/ImageBuilder.h"
+
+namespace UE
+{
+namespace Geometry
+{
+
+/** Resample image */
+class DYNAMICMESH_API FMeshResampleImageEvaluator : public FMeshMapEvaluator
+{
+public:
+	TFunction<FVector4f(FVector2d)> SampleFunction = [](FVector2d Position) { return FVector4f::Zero(); };
+	const FDynamicMeshUVOverlay* DetailUVOverlay = nullptr;
+
+	FVector4f DefaultColor = FVector4f(0, 0, 0, 1);
+
+public:
+	/** Invoked at start of bake to initialize baker. */
+	virtual void Setup(const FMeshMapBaker& Baker, FEvaluationContext& Context) override;
+
+	static void EvaluateSample(float*& Out, const FCorrespondenceSample& Sample, void* EvalData);
+
+	static void EvaluateDefault(float*& Out, void* EvalData);
+
+protected:
+	// Cached data
+	const FDynamicMesh3* DetailMesh = nullptr;
+
+private:
+	FVector4f ImageSampleFunction(const FCorrespondenceSample& Sample);
+};
+
+/** Resample multi texture by material ID */
+class DYNAMICMESH_API FMeshMultiResampleImageEvaluator : public FMeshResampleImageEvaluator
+{
+public:
+	TMap<int32, TSharedPtr<UE::Geometry::TImageBuilder<FVector4f>, ESPMode::ThreadSafe>> MultiTextures;
+
+public:
+	virtual ~FMeshMultiResampleImageEvaluator() = default;
+
+	/** Invoked at start of bake to initialize baker. */
+	virtual void Setup(const FMeshMapBaker & Baker, FEvaluationContext & Context) override;
+
+	static void EvaluateSampleMulti(float*& Out, const FCorrespondenceSample& Sample, void* EvalData);
+
+protected:
+	const FDynamicMeshMaterialAttribute* DetailMaterialIDAttrib = nullptr;
+	bool bValidDetailMesh = false;
+
+private:
+	FVector4f ImageSampleFunction(const FCorrespondenceSample& Sample);
+};
+
+} // end namespace UE::Geometry
+} // end namespace UE
+
diff --git a/Engine/Plugins/Experimental/MeshModelingToolset/Source/MeshModelingTools/Private/BakeMeshAttributeMapsTool.cpp b/Engine/Plugins/Experimental/MeshModelingToolset/Source/MeshModelingTools/Private/BakeMeshAttributeMapsTool.cpp
--- a/Engine/Plugins/Experimental/MeshModelingToolset/Source/MeshModelingTools/Private/BakeMeshAttributeMapsTool.cpp
+++ b/Engine/Plugins/Experimental/MeshModelingToolset/Source/MeshModelingTools/Private/BakeMeshAttributeMapsTool.cpp
@@ -8,12 +8,11 @@
 #include "DynamicMeshAttributeSet.h"
 #include "MeshTransforms.h"
 #include "MeshDescriptionToDynamicMesh.h"
-#include "Sampling/MeshImageBakingCache.h"
-#include "Sampling/MeshNormalMapBaker.h"
-#include "Sampling/MeshOcclusionMapBaker.h"
-#include "Sampling/MeshCurvatureMapBaker.h"
-#include "Sampling/MeshPropertyMapBaker.h"
-#include "Sampling/MeshResampleImageBaker.h"
+#include "Sampling/MeshNormalMapEvaluator.h"
+#include "Sampling/MeshOcclusionMapEvaluator.h"
+#include "Sampling/MeshCurvatureMapEvaluator.h"
+#include "Sampling/MeshPropertyMapEvaluator.h"
+#include "Sampling/MeshResampleImageEvaluator.h"
 #include "Util/IndexUtil.h"
 
 #include "SimpleDynamicMeshComponent.h"
@@ -139,8 +138,8 @@ public:
 			return;
 		}
 
-		TSharedPtr<FMeshNormalMapBaker, ESPMode::ThreadSafe> NormalBaker = MakeShared<FMeshNormalMapBaker, ESPMode::ThreadSafe>();
-		Baker->AddBaker(NormalBaker);
+		TSharedPtr<FMeshNormalMapEvaluator, ESPMode::ThreadSafe> NormalEval = MakeShared<FMeshNormalMapEvaluator, ESPMode::ThreadSafe>();
+		Baker->AddBaker(NormalEval);
 		Baker->SetTargetMeshTangents(BaseMeshTangents);
 		Baker->Bake();
 		SetResult(MoveTemp(Baker));
@@ -171,45 +170,35 @@ public:
 			return;
 		}
 
-		auto InitializeBaker = [this](TSharedPtr<FMeshOcclusionMapBaker>& BakerIn) {
-			BakerIn->NumOcclusionRays = Settings.OcclusionRays;
-			BakerIn->MaxDistance = Settings.MaxDistance;
-			BakerIn->SpreadAngle = Settings.SpreadAngle;
-			BakerIn->BlurRadius = Settings.BlurRadius;
-			BakerIn->BiasAngleDeg = Settings.BiasAngle;
+		TSharedPtr<FMeshOcclusionMapEvaluator> OcclusionEval = MakeShared<FMeshOcclusionMapEvaluator>();
+		OcclusionEval->OcclusionType = EMeshOcclusionMapType::All;
+		OcclusionEval->NumOcclusionRays = Settings.OcclusionRays;
+		OcclusionEval->MaxDistance = Settings.MaxDistance;
+		OcclusionEval->SpreadAngle = Settings.SpreadAngle;
+		OcclusionEval->BiasAngleDeg = Settings.BiasAngle;
 
-			switch (Settings.Distribution)
-			{
-			case EOcclusionMapDistribution::Cosine:
-				BakerIn->Distribution = FMeshOcclusionMapBaker::EDistribution::Cosine;
-				break;
-			case EOcclusionMapDistribution::Uniform:
-				BakerIn->Distribution = FMeshOcclusionMapBaker::EDistribution::Uniform;
-				break;
-			}
-
-			switch (Settings.NormalSpace)
-			{
-			case ENormalMapSpace::Tangent:
-				BakerIn->NormalSpace = FMeshOcclusionMapBaker::ESpace::Tangent;
-				break;
-			case ENormalMapSpace::Object:
-				BakerIn->NormalSpace = FMeshOcclusionMapBaker::ESpace::Object;
-				break;
-			}
-		};
-
-		TSharedPtr<FMeshOcclusionMapBaker> OcclusionBaker = MakeShared<FMeshOcclusionMapBaker>();
-		InitializeBaker(OcclusionBaker);
-		OcclusionBaker->OcclusionType = EOcclusionMapType::AmbientOcclusion;
+		switch (Settings.Distribution)
+		{
+		case EOcclusionMapDistribution::Cosine:
+			OcclusionEval->Distribution = FMeshOcclusionMapEvaluator::EDistribution::Cosine;
+			break;
+		case EOcclusionMapDistribution::Uniform:
+			OcclusionEval->Distribution = FMeshOcclusionMapEvaluator::EDistribution::Uniform;
+			break;
+		}
 
-		TSharedPtr<FMeshOcclusionMapBaker> BentNormalBaker = MakeShared<FMeshOcclusionMapBaker>();
-		InitializeBaker(BentNormalBaker);
-		BentNormalBaker->OcclusionType = EOcclusionMapType::BentNormal;
+		switch (Settings.NormalSpace)
+		{
+		case ENormalMapSpace::Tangent:
+			OcclusionEval->NormalSpace = FMeshOcclusionMapEvaluator::ESpace::Tangent;
+			break;
+		case ENormalMapSpace::Object:
+			OcclusionEval->NormalSpace = FMeshOcclusionMapEvaluator::ESpace::Object;
+			break;
+		}
 
 		Baker->SetTargetMeshTangents(BaseMeshTangents);
-		Baker->AddBaker(OcclusionBaker);
-		Baker->AddBaker(BentNormalBaker);
+		Baker->AddBaker(OcclusionEval);
 		Baker->Bake();
 		SetResult(MoveTemp(Baker));
 	}
@@ -238,13 +227,12 @@ public:
 			return;
 		}
 
-		TSharedPtr<FMeshCurvatureMapBaker> CurvatureBaker = MakeShared<FMeshCurvatureMapBaker>();
+		TSharedPtr<FMeshCurvatureMapEvaluator, ESPMode::ThreadSafe> CurvatureBaker = MakeShared<FMeshCurvatureMapEvaluator, ESPMode::ThreadSafe>();
 		CurvatureBaker->RangeScale = FMathd::Clamp(Settings.RangeMultiplier, 0.0001, 1000.0);
 		CurvatureBaker->MinRangeScale = FMathd::Clamp(Settings.MinRangeMultiplier, 0.0, 1.0);
-		CurvatureBaker->UseCurvatureType = (FMeshCurvatureMapBaker::ECurvatureType)Settings.CurvatureType;
-		CurvatureBaker->UseColorMode = (FMeshCurvatureMapBaker::EColorMode)Settings.ColorMode;
-		CurvatureBaker->UseClampMode = (FMeshCurvatureMapBaker::EClampMode)Settings.ClampMode;
-		CurvatureBaker->BlurRadius = Settings.BlurRadius;
+		CurvatureBaker->UseCurvatureType = (FMeshCurvatureMapEvaluator::ECurvatureType)Settings.CurvatureType;
+		CurvatureBaker->UseColorMode = (FMeshCurvatureMapEvaluator::EColorMode)Settings.ColorMode;
+		CurvatureBaker->UseClampMode = (FMeshCurvatureMapEvaluator::EClampMode)Settings.ClampMode;
 		Baker->AddBaker(CurvatureBaker);
 		Baker->Bake();
 		SetResult(MoveTemp(Baker));
@@ -274,8 +262,8 @@ public:
 			return;
 		}
 
-		TSharedPtr<FMeshPropertyMapBaker> PropertyBaker = MakeShared<FMeshPropertyMapBaker>();
-		PropertyBaker->Property = (EMeshPropertyBakeType)Settings.PropertyTypeIndex;
+		TSharedPtr<FMeshPropertyMapEvaluator, ESPMode::ThreadSafe> PropertyBaker = MakeShared<FMeshPropertyMapEvaluator, ESPMode::ThreadSafe>();
+		PropertyBaker->Property = (EMeshPropertyMapType) Settings.PropertyTypeIndex;
 		Baker->AddBaker(PropertyBaker);
 		Baker->Bake();
 		SetResult(MoveTemp(Baker));
@@ -307,7 +295,7 @@ public:
 			return;
 		}
 
-		TSharedPtr<FMeshResampleImageBaker> ResampleBaker = MakeShared<FMeshResampleImageBaker>();
+		TSharedPtr<FMeshResampleImageEvaluator, ESPMode::ThreadSafe> ResampleBaker = MakeShared<FMeshResampleImageEvaluator, ESPMode::ThreadSafe>();
 		ResampleBaker->DetailUVOverlay = UVOverlay;
 		ResampleBaker->SampleFunction = [this](FVector2d UVCoord) {
 			return TextureImage->BilinearSampleUV<float>(UVCoord, FVector4f(0, 0, 0, 1));
@@ -344,7 +332,7 @@ public:
 			return;
 		}
 
-		TSharedPtr<FMeshMultiResampleImageBaker> TextureBaker = MakeShared<FMeshMultiResampleImageBaker>();
+		TSharedPtr<FMeshMultiResampleImageEvaluator, ESPMode::ThreadSafe> TextureBaker = MakeShared<FMeshMultiResampleImageEvaluator, ESPMode::ThreadSafe>();
 		TextureBaker->DetailUVOverlay = UVOverlay;
 		TextureBaker->MultiTextures = MaterialToTextureImageMap;
 		Baker->AddBaker(TextureBaker);
@@ -1028,42 +1016,42 @@ UBakeMeshAttributeMapsTool::EOpState UBakeMeshAttributeMapsTool::UpdateResult_Cu
 	{
 	default:
 	case EBakedCurvatureTypeMode::MeanAverage:
-		CurvatureMapSettings.CurvatureType = (int32)FMeshCurvatureMapBaker::ECurvatureType::Mean;
+		CurvatureMapSettings.CurvatureType = (int32)FMeshCurvatureMapEvaluator::ECurvatureType::Mean;
 		break;
 	case EBakedCurvatureTypeMode::Gaussian:
-		CurvatureMapSettings.CurvatureType = (int32)FMeshCurvatureMapBaker::ECurvatureType::Gaussian;
+		CurvatureMapSettings.CurvatureType = (int32)FMeshCurvatureMapEvaluator::ECurvatureType::Gaussian;
 		break;
 	case EBakedCurvatureTypeMode::Max:
-		CurvatureMapSettings.CurvatureType = (int32)FMeshCurvatureMapBaker::ECurvatureType::MaxPrincipal;
+		CurvatureMapSettings.CurvatureType = (int32)FMeshCurvatureMapEvaluator::ECurvatureType::MaxPrincipal;
 		break;
 	case EBakedCurvatureTypeMode::Min:
-		CurvatureMapSettings.CurvatureType = (int32)FMeshCurvatureMapBaker::ECurvatureType::MinPrincipal;
+		CurvatureMapSettings.CurvatureType = (int32)FMeshCurvatureMapEvaluator::ECurvatureType::MinPrincipal;
 		break;
 	}
 	switch (CurvatureMapProps->ColorMode)
 	{
 	default:
 	case EBakedCurvatureColorMode::Grayscale:
-		CurvatureMapSettings.ColorMode = (int32)FMeshCurvatureMapBaker::EColorMode::BlackGrayWhite;
+		CurvatureMapSettings.ColorMode = (int32)FMeshCurvatureMapEvaluator::EColorMode::BlackGrayWhite;
 		break;
 	case EBakedCurvatureColorMode::RedBlue:
-		CurvatureMapSettings.ColorMode = (int32)FMeshCurvatureMapBaker::EColorMode::RedBlue;
+		CurvatureMapSettings.ColorMode = (int32)FMeshCurvatureMapEvaluator::EColorMode::RedBlue;
 		break;
 	case EBakedCurvatureColorMode::RedGreenBlue:
-		CurvatureMapSettings.ColorMode = (int32)FMeshCurvatureMapBaker::EColorMode::RedGreenBlue;
+		CurvatureMapSettings.ColorMode = (int32)FMeshCurvatureMapEvaluator::EColorMode::RedGreenBlue;
 		break;
 	}
 	switch (CurvatureMapProps->Clamping)
 	{
 	default:
 	case EBakedCurvatureClampMode::None:
-		CurvatureMapSettings.ClampMode = (int32)FMeshCurvatureMapBaker::EClampMode::FullRange;
+		CurvatureMapSettings.ClampMode = (int32)FMeshCurvatureMapEvaluator::EClampMode::FullRange;
 		break;
 	case EBakedCurvatureClampMode::Positive:
-		CurvatureMapSettings.ClampMode = (int32)FMeshCurvatureMapBaker::EClampMode::Positive;
+		CurvatureMapSettings.ClampMode = (int32)FMeshCurvatureMapEvaluator::EClampMode::Positive;
 		break;
 	case EBakedCurvatureClampMode::Negative:
-		CurvatureMapSettings.ClampMode = (int32)FMeshCurvatureMapBaker::EClampMode::Negative;
+		CurvatureMapSettings.ClampMode = (int32)FMeshCurvatureMapEvaluator::EClampMode::Negative;
 		break;
 	}
 	CurvatureMapSettings.BlurRadius = (CurvatureMapProps->bGaussianBlur) ? CurvatureMapProps->BlurRadius : 0.0;
@@ -1090,21 +1078,20 @@ UBakeMeshAttributeMapsTool::EOpState UBakeMeshAttributeMapsTool::UpdateResult_Me
 	switch (Settings->MapType)
 	{
 		case EBakeMapType::NormalImage:
-			MeshPropertyMapSettings.PropertyTypeIndex = (int32)EMeshPropertyBakeType::Normal;
+			MeshPropertyMapSettings.PropertyTypeIndex = (int32)EMeshPropertyMapType::Normal;
 			break;
 		case EBakeMapType::FaceNormalImage:
-			MeshPropertyMapSettings.PropertyTypeIndex = (int32)EMeshPropertyBakeType::FacetNormal;
+			MeshPropertyMapSettings.PropertyTypeIndex = (int32)EMeshPropertyMapType::FacetNormal;
 			break;
 		case EBakeMapType::PositionImage:
-			MeshPropertyMapSettings.PropertyTypeIndex = (int32)EMeshPropertyBakeType::Position;
+			MeshPropertyMapSettings.PropertyTypeIndex = (int32)EMeshPropertyMapType::Position;
 			break;
 		case EBakeMapType::MaterialID:
-			MeshPropertyMapSettings.PropertyTypeIndex = (int32)EMeshPropertyBakeType::MaterialID;
+			MeshPropertyMapSettings.PropertyTypeIndex = (int32)EMeshPropertyMapType::MaterialID;
 			break;
 		default:
 			check(false);		// should not be possible!
 	}
-	//MeshPropertyMapSettings.BlurRadius = (CurvatureMapProps->bGaussianBlur) ? CurvatureMapProps->BlurRadius : 0.0;
 
 	if (!(CachedMeshPropertyMapSettings == MeshPropertyMapSettings))
 	{
@@ -1430,7 +1417,7 @@ void UBakeMeshAttributeMapsTool::OnMapsUpdated(const TUniquePtr<FMeshMapBaker>&
 	{
 		FTexture2DBuilder TextureBuilder;
 		TextureBuilder.Initialize(FTexture2DBuilder::ETextureType::NormalMap, CachedNormalMapSettings.Dimensions);
-		TextureBuilder.Copy(*NewResult->GetBakeResult(0));
+		TextureBuilder.Copy(*NewResult->GetBakeResults(0)[0]);
 		TextureBuilder.Commit(false);
 		CachedNormalMap = TextureBuilder.GetTexture2D();
 		break;
@@ -1439,13 +1426,13 @@ void UBakeMeshAttributeMapsTool::OnMapsUpdated(const TUniquePtr<FMeshMapBaker>&
 	{
 		FTexture2DBuilder TextureBuilder;
 		TextureBuilder.Initialize(FTexture2DBuilder::ETextureType::AmbientOcclusion, CachedOcclusionMapSettings.Dimensions);
-		TextureBuilder.Copy(*NewResult->GetBakeResult(0));
+		TextureBuilder.Copy(*NewResult->GetBakeResults(0)[0]);
 		TextureBuilder.Commit(false);
 		CachedOcclusionMap = TextureBuilder.GetTexture2D();
 
 		FTexture2DBuilder TextureNormalBuilder;
 		TextureNormalBuilder.Initialize(FTexture2DBuilder::ETextureType::NormalMap, CachedOcclusionMapSettings.Dimensions);
-		TextureNormalBuilder.Copy(*NewResult->GetBakeResult(1));
+		TextureNormalBuilder.Copy(*NewResult->GetBakeResults(0)[1]);
 		TextureNormalBuilder.Commit(false);
 		CachedBentNormalMap = TextureNormalBuilder.GetTexture2D();
 		break;
@@ -1454,7 +1441,7 @@ void UBakeMeshAttributeMapsTool::OnMapsUpdated(const TUniquePtr<FMeshMapBaker>&
 	{
 		FTexture2DBuilder TextureBuilder;
 		TextureBuilder.Initialize(FTexture2DBuilder::ETextureType::Color, CachedCurvatureMapSettings.Dimensions);
-		TextureBuilder.Copy(*NewResult->GetBakeResult(0));
+		TextureBuilder.Copy(*NewResult->GetBakeResults(0)[0]);
 		TextureBuilder.Commit(false);
 		CachedCurvatureMap = TextureBuilder.GetTexture2D();
 		break;
@@ -1466,7 +1453,7 @@ void UBakeMeshAttributeMapsTool::OnMapsUpdated(const TUniquePtr<FMeshMapBaker>&
 	{
 		FTexture2DBuilder TextureBuilder;
 		TextureBuilder.Initialize(FTexture2DBuilder::ETextureType::Color, CachedMeshPropertyMapSettings.Dimensions);
-		TextureBuilder.Copy(*NewResult->GetBakeResult(0));
+		TextureBuilder.Copy(*NewResult->GetBakeResults(0)[0]);
 		TextureBuilder.Commit(false);
 		CachedMeshPropertyMap = TextureBuilder.GetTexture2D();
 		break;
@@ -1475,7 +1462,7 @@ void UBakeMeshAttributeMapsTool::OnMapsUpdated(const TUniquePtr<FMeshMapBaker>&
 	{
 		FTexture2DBuilder TextureBuilder;
 		TextureBuilder.Initialize(FTexture2DBuilder::ETextureType::Color, CachedTexture2DImageSettings.Dimensions);
-		TextureBuilder.Copy(*NewResult->GetBakeResult(0), true);
+		TextureBuilder.Copy(*NewResult->GetBakeResults(0)[0], true);
 		TextureBuilder.Commit(false);
 		CachedTexture2DImageMap = TextureBuilder.GetTexture2D();
 		break;
@@ -1484,7 +1471,7 @@ void UBakeMeshAttributeMapsTool::OnMapsUpdated(const TUniquePtr<FMeshMapBaker>&
 	{
 		FTexture2DBuilder TextureBuilder;
 		TextureBuilder.Initialize(FTexture2DBuilder::ETextureType::Color, CachedTexture2DImageSettings.Dimensions);
-		TextureBuilder.Copy(*NewResult->GetBakeResult(0), true);
+		TextureBuilder.Copy(*NewResult->GetBakeResults(0)[0], true);
 		TextureBuilder.Commit(false);
 		CachedTexture2DImageMap = TextureBuilder.GetTexture2D();
 		break;
diff --git a/Engine/Source/Runtime/GeometryCore/Public/Image/ImageBuilder.h b/Engine/Source/Runtime/GeometryCore/Public/Image/ImageBuilder.h
--- a/Engine/Source/Runtime/GeometryCore/Public/Image/ImageBuilder.h
+++ b/Engine/Source/Runtime/GeometryCore/Public/Image/ImageBuilder.h
@@ -60,6 +60,32 @@ public:
 			ImageCoords.X < Dimensions.GetWidth() && ImageCoords.Y < Dimensions.GetHeight();
 	}
 
+	/**
+	 * Get the Pixel at the given X/Y coordinates
+	 */
+	PixelType& GetPixel(int32 X, int32 Y)
+	{
+		int64 LinearIndex = Dimensions.GetIndex(FVector2i(X, Y));
+		return Image[LinearIndex];
+	}
+
+	/**
+	 * Get the Pixel at the given X/Y coordinates
+	 */
+	PixelType& GetPixel(const FVector2i& ImageCoords)
+	{
+		int64 LinearIndex = Dimensions.GetIndex(ImageCoords);
+		return Image[LinearIndex];
+	}
+
+	/**
+	 * Get the Pixel at the given linear index
+	 */
+	PixelType& GetPixel(int64 LinearIndex)
+	{
+		return Image[LinearIndex];
+	}
+
 	/**
 	 * Get the Pixel at the given X/Y coordinates
 	 */
