GeometryProcessing: add various functions to FDynamicMeshUVEditor. New SetTriangleUVsFromExpMap() variant can specify a specific target frame, supports normal blending and smoothing internally. Ported box and cylinder projection from UVProjectionOp to SetTriangleUVsFromBoxProjection() and SetTriangleUVsFromCylinderProjection(). Cylinder projection currently has been modified to create two islands for cylinder, instead of one (should eventually make this optional)
#rnx

diff --git a/Engine/Plugins/Experimental/GeometryProcessing/Source/DynamicMesh/Private/Parameterization/DynamicMeshUVEditor.cpp b/Engine/Plugins/Experimental/GeometryProcessing/Source/DynamicMesh/Private/Parameterization/DynamicMeshUVEditor.cpp
--- a/Engine/Plugins/Experimental/GeometryProcessing/Source/DynamicMesh/Private/Parameterization/DynamicMeshUVEditor.cpp
+++ b/Engine/Plugins/Experimental/GeometryProcessing/Source/DynamicMesh/Private/Parameterization/DynamicMeshUVEditor.cpp
@@ -7,11 +7,17 @@
 #include "DynamicMesh/MeshNormals.h"
 #include "MeshBoundaryLoops.h"
 #include "DynamicMesh/MeshIndexUtil.h"
+#include "DynamicMesh/MeshTransforms.h"
+#include "DynamicMesh/DynamicMeshAABBTree3.h"
+#include "MeshQueries.h"
+#include "MeshWeights.h"
 
 #include "Parameterization/MeshDijkstra.h"
 #include "Parameterization/MeshLocalParam.h"
 #include "Solvers/MeshParameterizationSolvers.h"
 
+#include "Async/ParallelFor.h"
+
 #include "ExplicitUseGeometryMathTypes.h"		// using UE::Geometry::(math types)
 using namespace UE::Geometry;
 
@@ -48,6 +54,29 @@ void FDynamicMeshUVEditor::CreateUVLayer(int32 LayerIndex)
 }
 
 
+void FDynamicMeshUVEditor::ResetUVs()
+{
+	if (ensure(UVOverlay))
+	{
+		UVOverlay->ClearElements();
+	}
+}
+
+
+
+void FDynamicMeshUVEditor::TransformUVElements(const TArray<int32>& ElementIDs, TFunctionRef<FVector2f(const FVector2f&)> TransformFunc)
+{
+	for (int32 elemid : ElementIDs)
+	{
+		if (UVOverlay->IsElement(elemid))
+		{
+			FVector2f UV = UVOverlay->GetElement(elemid);
+			UV = TransformFunc(UV);
+			UVOverlay->SetElement(elemid, UV);
+		}
+	}
+}
+
 
 template<typename EnumerableType>
 void InternalSetPerTriangleUVs(EnumerableType TriangleIDs, const FDynamicMesh3* Mesh, FDynamicMeshUVOverlay* UVOverlay, double ScaleFactor, FUVEditResult* Result)
@@ -95,12 +124,25 @@ void FDynamicMeshUVEditor::SetPerTriangleUVs(double ScaleFactor, FUVEditResult*
 	InternalSetPerTriangleUVs(Mesh->TriangleIndicesItr(), Mesh, UVOverlay, ScaleFactor, Result);
 }
 
-
 void FDynamicMeshUVEditor::SetTriangleUVsFromProjection(const TArray<int32>& Triangles, const FFrame3d& ProjectionFrame, FUVEditResult* Result)
+{
+	SetTriangleUVsFromPlanarProjection(Triangles, [](const FVector3d& P) { return P; },  
+		ProjectionFrame, FVector2d::One(), Result);
+}
+
+void FDynamicMeshUVEditor::SetTriangleUVsFromPlanarProjection(
+	const TArray<int32>& Triangles, 
+	TFunctionRef<FVector3d(const FVector3d&)> PointTransform,
+	const FFrame3d& ProjectionFrame, 
+	const FVector2d& Dimensions, 
+	FUVEditResult* Result)
 {
 	if (ensure(UVOverlay) == false) return;
 	if (!Triangles.Num()) return;
 
+	double ScaleX = (FMathd::Abs(Dimensions.X) > FMathf::ZeroTolerance) ? (1.0 / Dimensions.X) : 1.0;
+	double ScaleY = (FMathd::Abs(Dimensions.Y) > FMathf::ZeroTolerance) ? (1.0 / Dimensions.Y) : 1.0;
+
 	TMap<int32, int32> BaseToOverlayVIDMap;
 	TArray<int32> NewUVIndices;
 
@@ -113,7 +155,11 @@ void FDynamicMeshUVEditor::SetTriangleUVsFromProjection(const TArray<int32>& Tri
 			const int32* FoundElementID = BaseToOverlayVIDMap.Find(BaseTri[j]);
 			if (FoundElementID == nullptr)
 			{
-				FVector2f UV = (FVector2f)ProjectionFrame.ToPlaneUV(Mesh->GetVertex(BaseTri[j]), 2);
+				FVector3d Pos = Mesh->GetVertex(BaseTri[j]);
+				FVector3d TransformPos = PointTransform(Pos);
+				FVector2f UV = (FVector2f)ProjectionFrame.ToPlaneUV(TransformPos, 2);
+				UV.X *= ScaleX;
+				UV.Y *= ScaleY;
 				ElemTri[j] = UVOverlay->AppendElement(UV);
 				NewUVIndices.Add(ElemTri[j]);
 				BaseToOverlayVIDMap.Add(BaseTri[j], ElemTri[j]);
@@ -133,19 +179,15 @@ void FDynamicMeshUVEditor::SetTriangleUVsFromProjection(const TArray<int32>& Tri
 }
 
 
-
-bool FDynamicMeshUVEditor::SetTriangleUVsFromExpMap(const TArray<int32>& Triangles, FUVEditResult* Result)
+bool FDynamicMeshUVEditor::EstimateGeodesicCenterFrameVertex(const FDynamicMesh3& Mesh, FFrame3d& FrameOut, int32& VertexIDOut, bool bAlignToUnitAxes)
 {
-	if (ensure(UVOverlay) == false) return false;
-	if (Triangles.Num() == 0) return false;
-
-	FDynamicSubmesh3 SubmeshCalc(Mesh, Triangles, (int)EMeshComponents::None, false);
-	FDynamicMesh3& Submesh = SubmeshCalc.GetSubmesh();
-	FMeshNormals::QuickComputeVertexNormals(Submesh);
+	VertexIDOut = *Mesh.VertexIndicesItr().begin();
+	FVector3d Normal = FMeshNormals::ComputeVertexNormal(Mesh, VertexIDOut);
 
-	FMeshBoundaryLoops LoopsCalc(&Submesh, true);
-	if (LoopsCalc.GetLoopCount() == 0)
+	FMeshBoundaryLoops LoopsCalc(&Mesh, true);
+	if (ensure(LoopsCalc.GetLoopCount() > 0) == false )
 	{
+		FrameOut = Mesh.GetVertexFrame(VertexIDOut, false, &Normal);
 		return false;
 	}
 	const FEdgeLoop& Loop = LoopsCalc[LoopsCalc.GetMaxVerticesLoopIndex()];
@@ -155,22 +197,64 @@ bool FDynamicMeshUVEditor::SetTriangleUVsFromExpMap(const TArray<int32>& Triangl
 		SeedPoints.Add(FVector2d(vid, 0.0));
 	}
 
-	TMeshDijkstra<FDynamicMesh3> Dijkstra(&Submesh);
+	TMeshDijkstra<FDynamicMesh3> Dijkstra(&Mesh);
 	Dijkstra.ComputeToMaxDistance(SeedPoints, TNumericLimits<float>::Max());
 	int32 MaxDistVID = Dijkstra.GetMaxGraphDistancePointID();
-	if (!Submesh.IsVertex(MaxDistVID))
+	if ( ensure(Mesh.IsVertex(MaxDistVID)) == false )
 	{
+		FrameOut = Mesh.GetVertexFrame(VertexIDOut, false, &Normal);
 		return false;
 	}
+	VertexIDOut = MaxDistVID;
+	Normal = FMeshNormals::ComputeVertexNormal(Mesh, MaxDistVID);
+	FrameOut = Mesh.GetVertexFrame(MaxDistVID, false, &Normal);
 
-
-	FFrame3d SeedFrame(Submesh.GetVertex(MaxDistVID), (FVector3d)Submesh.GetVertexNormal(MaxDistVID));
 	// try to generate consistent frame alignment...
-	SeedFrame.ConstrainedAlignPerpAxes(0, 1, 2, FVector3d::UnitX(), FVector3d::UnitY(), 0.95);
+	if (bAlignToUnitAxes)
+	{
+		FrameOut.ConstrainedAlignPerpAxes(0, 1, 2, FVector3d::UnitX(), FVector3d::UnitY(), 0.95);
+	}
+
+	return true;
+}
+
+
+bool FDynamicMeshUVEditor::EstimateGeodesicCenterFrameVertex(const FDynamicMesh3& Mesh, const TArray<int32>& Triangles, FFrame3d& FrameOut, int32& VertexIDOut, bool bAlignToUnitAxes)
+{
+	FDynamicSubmesh3 SubmeshCalc(&Mesh, Triangles, (int)EMeshComponents::None, false);
+	FDynamicMesh3& Submesh = SubmeshCalc.GetSubmesh();
+	FFrame3d SeedFrame;
+	int32 FrameVertexID;
+	bool bFrameOK = EstimateGeodesicCenterFrameVertex(Submesh, SeedFrame, FrameVertexID, true);
+	if (!bFrameOK)
+	{
+		return false;
+	}
+	FrameVertexID = SubmeshCalc.MapVertexToBaseMesh(FrameVertexID);
+	return true;
+}
+
+
+bool FDynamicMeshUVEditor::SetTriangleUVsFromExpMap(const TArray<int32>& Triangles, FUVEditResult* Result)
+{
+	if (ensure(UVOverlay) == false) return false;
+	if (Triangles.Num() == 0) return false;
+
+	FDynamicSubmesh3 SubmeshCalc(Mesh, Triangles, (int)EMeshComponents::None, false);
+	FDynamicMesh3& Submesh = SubmeshCalc.GetSubmesh();
+	FMeshNormals::QuickComputeVertexNormals(Submesh);
+
+	FFrame3d SeedFrame;
+	int32 FrameVertexID;
+	bool bFrameOK = EstimateGeodesicCenterFrameVertex(Submesh, SeedFrame, FrameVertexID, true);
+	if (!bFrameOK)
+	{
+		return false;
+	}
 
 	TMeshLocalParam<FDynamicMesh3> Param(&Submesh);
 	Param.ParamMode = ELocalParamTypes::ExponentialMapUpwindAvg;
-	Param.ComputeToMaxDistance(MaxDistVID, SeedFrame, TNumericLimits<float>::Max());
+	Param.ComputeToMaxDistance(FrameVertexID, SeedFrame, TNumericLimits<float>::Max());
 
 	TArray<int32> VtxElementIDs;
 	TArray<int32> NewElementIDs;
@@ -210,6 +294,124 @@ bool FDynamicMeshUVEditor::SetTriangleUVsFromExpMap(const TArray<int32>& Triangl
 
 
 
+bool FDynamicMeshUVEditor::SetTriangleUVsFromExpMap(
+	const TArray<int32>& Triangles,
+	TFunctionRef<FVector3d(const FVector3d&)> PointTransform,
+	const FFrame3d& ProjectionFrame,
+	const FVector2d& Dimensions,
+	int32 NormalSmoothingRounds,
+	double NormalSmoothingAlpha,
+	double FrameNormalBlendWeight,
+	FUVEditResult* Result)
+{
+	if (ensure(UVOverlay) == false) return false;
+	if (Triangles.Num() == 0) return false;
+
+	double ScaleX = (FMathd::Abs(Dimensions.X) > FMathf::ZeroTolerance) ? (1.0 / Dimensions.X) : 1.0;
+	double ScaleY = (FMathd::Abs(Dimensions.Y) > FMathf::ZeroTolerance) ? (1.0 / Dimensions.Y) : 1.0;
+
+	FDynamicSubmesh3 SubmeshCalc(Mesh, Triangles, (int)EMeshComponents::None, false);
+	FDynamicMesh3& Submesh = SubmeshCalc.GetSubmesh();
+	MeshTransforms::ApplyTransform(Submesh, PointTransform, [&](const FVector3f& V) { return V;});
+	FMeshNormals::QuickComputeVertexNormals(Submesh);
+
+	NormalSmoothingRounds = FMath::Clamp(NormalSmoothingRounds, 0, 500);
+	NormalSmoothingAlpha = FMathd::Clamp(NormalSmoothingAlpha, 0.0, 1.0);
+	if (NormalSmoothingRounds > 0 && NormalSmoothingAlpha > 0)
+	{
+		int32 NumV = Submesh.MaxVertexID();
+		TArray<FVector3d> SmoothedNormals;
+		SmoothedNormals.SetNum(NumV);
+		for (int32 ri = 0; ri < NormalSmoothingRounds; ++ri)
+		{
+			SmoothedNormals.Init(FVector3d::Zero(), NumV);
+			for (int32 vid : Submesh.VertexIndicesItr())
+			{
+				FVector3d SmoothedNormal = FVector3d::Zero();
+				UE::Geometry::FMeshWeights::CotanWeightsBlendSafe(Submesh, vid, [&](int32 nbrvid, double Weight)
+				{
+					SmoothedNormal += Weight * (FVector3d)Submesh.GetVertexNormal(nbrvid);
+				});
+				SmoothedNormal.Normalize();
+				SmoothedNormals[vid] = Lerp((FVector3d)Submesh.GetVertexNormal(vid), SmoothedNormal, NormalSmoothingAlpha);
+			}
+			for (int32 vid : Submesh.VertexIndicesItr())
+			{
+				Submesh.SetVertexNormal(vid, (FVector3f)SmoothedNormals[vid]);
+			}
+		}
+	}
+
+	FDynamicMeshAABBTree3 Spatial(&Submesh, true);
+	double NearDistSqr;
+	int32 SeedTID = Spatial.FindNearestTriangle(ProjectionFrame.Origin, NearDistSqr);
+	FDistPoint3Triangle3d Query = TMeshQueries<FDynamicMesh3>::TriangleDistance(Submesh, SeedTID, ProjectionFrame.Origin);
+	FIndex3i TriVerts = Submesh.GetTriangle(SeedTID);
+
+	FFrame3d ParamSeedFrame = ProjectionFrame;
+	ParamSeedFrame.Origin = Query.ClosestTrianglePoint;
+	// correct for inverted frame
+	if (ParamSeedFrame.Z().Dot(Submesh.GetTriNormal(SeedTID)) < 0)
+	{
+		ParamSeedFrame.Rotate(FQuaterniond(ParamSeedFrame.X(), 180.0, true));
+	}
+
+	// apply normal blending
+	FrameNormalBlendWeight = FMathd::Clamp(FrameNormalBlendWeight, 0, 1);
+	if (FrameNormalBlendWeight > 0)
+	{
+		FVector3d FrameZ = ParamSeedFrame.Z();
+		for (int32 vid : Submesh.VertexIndicesItr())
+		{
+			FVector3d N = (FVector3d)Submesh.GetVertexNormal(vid);
+			N = Lerp(N, FrameZ, FrameNormalBlendWeight);
+			Submesh.SetVertexNormal(vid, (FVector3f)N);
+		}
+	}
+
+	TMeshLocalParam<FDynamicMesh3> Param(&Submesh);
+	Param.ParamMode = ELocalParamTypes::ExponentialMapUpwindAvg;
+	Param.ComputeToMaxDistance(ParamSeedFrame, TriVerts, TNumericLimits<float>::Max());
+
+	TArray<int32> VtxElementIDs;
+	TArray<int32> NewElementIDs;
+	VtxElementIDs.Init(IndexConstants::InvalidID, Submesh.MaxVertexID());
+	for (int32 vid : Submesh.VertexIndicesItr())
+	{
+		if (Param.HasUV(vid))
+		{
+			FVector2f UV = (FVector2f)Param.GetUV(vid);
+			UV.X *= ScaleX;
+			UV.Y *= ScaleY;
+			VtxElementIDs[vid] = UVOverlay->AppendElement(UV);
+			NewElementIDs.Add(VtxElementIDs[vid]);
+		}
+	}
+
+	int32 NumFailed = 0;
+	for (int32 tid : Submesh.TriangleIndicesItr())
+	{
+		FIndex3i SubTri = Submesh.GetTriangle(tid);
+		FIndex3i UVTri = { VtxElementIDs[SubTri.A], VtxElementIDs[SubTri.B], VtxElementIDs[SubTri.C] };
+		if (UVTri.A == IndexConstants::InvalidID || UVTri.B == IndexConstants::InvalidID || UVTri.C == IndexConstants::InvalidID)
+		{
+			NumFailed++;
+			continue;
+		}
+
+		int32 BaseTID = SubmeshCalc.MapTriangleToBaseMesh(tid);
+		UVOverlay->SetTriangle(BaseTID, UVTri);
+	}
+
+	if (Result != nullptr)
+	{
+		Result->NewUVElements = MoveTemp(NewElementIDs);
+	}
+
+	return (NumFailed == 0);
+}
+
+
 
 
 bool FDynamicMeshUVEditor::SetTriangleUVsFromFreeBoundaryConformal(const TArray<int32>& Triangles, FUVEditResult* Result)
@@ -552,3 +754,239 @@ bool FDynamicMeshUVEditor::CreateSeamAlongVertexPath(const TArray<int32>& Vertex
 
 	return true;
 }
+
+
+
+
+void FDynamicMeshUVEditor::SetTriangleUVsFromBoxProjection(
+	const TArray<int32>& Triangles, 
+	TFunctionRef<FVector3d(const FVector3d&)> PointTransform, 
+	const FFrame3d& BoxFrame, 
+	const FVector3d& BoxDimensions, 
+	FUVEditResult* Result)
+{
+	if (ensure(UVOverlay) == false) return;
+	int32 NumTriangles = Triangles.Num();
+	if (!NumTriangles) return;
+
+	const int Minor1s[3] = { 1, 0, 0 };
+	const int Minor2s[3] = { 2, 2, 1 };
+	const int Minor1Flip[3] = { -1, 1, 1 };
+	const int Minor2Flip[3] = { -1, -1, 1 };
+
+	auto GetTriNormal = [this, &PointTransform](int32 tid) -> FVector3d
+	{
+		FVector3d A, B, C;
+		Mesh->GetTriVertices(tid, A, B, C);
+		return VectorUtil::Normal(PointTransform(A), PointTransform(B), PointTransform(C));
+	};
+
+	double ScaleX = (FMathd::Abs(BoxDimensions.X) > FMathf::ZeroTolerance) ? (1.0 / BoxDimensions.X) : 1.0;
+	double ScaleY = (FMathd::Abs(BoxDimensions.Y) > FMathf::ZeroTolerance) ? (1.0 / BoxDimensions.Y) : 1.0;
+	double ScaleZ = (FMathd::Abs(BoxDimensions.Z) > FMathf::ZeroTolerance) ? (1.0 / BoxDimensions.Z) : 1.0;
+	FVector3d Scale(ScaleX, ScaleY, ScaleZ);
+
+	TArray<FVector3d> TriNormals;
+	TArray<FIndex2i> TriangleBoxPlaneAssignments;
+	TriNormals.SetNum(NumTriangles);
+	TriangleBoxPlaneAssignments.SetNum(NumTriangles);
+	ParallelFor(NumTriangles, [&](int32 i)
+	{
+		int32 tid = Triangles[i];
+		TriNormals[i] = GetTriNormal(tid);
+		FVector3d N = BoxFrame.ToFrameVector(TriNormals[i]);
+		N *= Scale;
+		FVector3d NAbs(FMathd::Abs(N.X), FMathd::Abs(N.Y), FMathd::Abs(N.Z));
+		int MajorAxis = NAbs[0] > NAbs[1] ? (NAbs[0] > NAbs[2] ? 0 : 2) : (NAbs[1] > NAbs[2] ? 1 : 2);
+		double MajorAxisSign = FMathd::Sign(N[MajorAxis]);
+		int Bucket = (MajorAxisSign > 0) ? (MajorAxis+3) : MajorAxis;
+		TriangleBoxPlaneAssignments[i] = FIndex2i(MajorAxis, Bucket);
+	});
+
+	auto ProjAxis = [](const FVector3d& P, int Axis1, int Axis2, float Axis1Scale, float Axis2Scale)
+	{
+		return FVector2f(float(P[Axis1]) * Axis1Scale, float(P[Axis2]) * Axis2Scale);
+	};
+
+	TMap<FIndex2i, int32> BaseToOverlayVIDMap;
+	TArray<int32> NewUVIndices;
+
+	for ( int32 i = 0; i < NumTriangles; ++i )
+	{
+		int32 tid = Triangles[i];
+		FIndex3i BaseTri = Mesh->GetTriangle(tid);
+		FIndex2i TriBoxInfo = TriangleBoxPlaneAssignments[i];
+		FVector3d N = BoxFrame.ToFrameVector(TriNormals[i]);
+
+		int MajorAxis = TriBoxInfo.A;
+		int Bucket = TriBoxInfo.B;
+		double MajorAxisSign = FMathd::Sign(N[MajorAxis]);
+		int Minor1 = Minor1s[MajorAxis];
+		int Minor2 = Minor2s[MajorAxis];
+
+		FIndex3i ElemTri;
+		for (int32 j = 0; j < 3; ++j)
+		{
+			FIndex2i ElementKey(BaseTri[j], Bucket);
+			const int32* FoundElementID = BaseToOverlayVIDMap.Find(ElementKey);
+			if (FoundElementID == nullptr)
+			{
+				FVector3d Pos = Mesh->GetVertex(BaseTri[j]);
+				FVector3d TransformPos = PointTransform(Pos);
+				FVector3d BoxPos = BoxFrame.ToFramePoint(TransformPos);
+				BoxPos *= Scale;
+
+				FVector2f UV = ProjAxis(BoxPos, Minor1, Minor2, MajorAxisSign * Minor1Flip[MajorAxis], Minor2Flip[MajorAxis]);
+
+				ElemTri[j] = UVOverlay->AppendElement(UV);
+				NewUVIndices.Add(ElemTri[j]);
+				BaseToOverlayVIDMap.Add(ElementKey, ElemTri[j]);
+			}
+			else
+			{
+				ElemTri[j] = *FoundElementID;
+			}
+		}
+		UVOverlay->SetTriangle(tid, ElemTri);
+	}
+
+	if (Result != nullptr)
+	{
+		Result->NewUVElements = MoveTemp(NewUVIndices);
+	}
+}
+
+
+
+
+
+
+void FDynamicMeshUVEditor::SetTriangleUVsFromCylinderProjection(
+	const TArray<int32>& Triangles, 
+	TFunctionRef<FVector3d(const FVector3d&)> PointTransform, 
+	const FFrame3d& BoxFrame, 
+	const FVector3d& BoxDimensions, 
+	FUVEditResult* Result)
+{
+	if (ensure(UVOverlay) == false) return;
+	int32 NumTriangles = Triangles.Num();
+	if (!NumTriangles) return;
+
+	const int Minor1s[3] = { 1, 0, 0 };
+	const int Minor2s[3] = { 2, 2, 1 };
+	const int Minor1Flip[3] = { -1, 1, 1 };
+	const int Minor2Flip[3] = { -1, -1, 1 };
+
+	auto GetTriNormalCentroid = [this, &PointTransform](int32 tid) -> TPair<FVector3d,FVector3d>
+	{
+		FVector3d A, B, C;
+		Mesh->GetTriVertices(tid, A, B, C);
+		A = PointTransform(A); B = PointTransform(B); C = PointTransform(C);
+		return TPair<FVector3d, FVector3d>{ VectorUtil::Normal(A, B, C), (A + B + C) / 3.0};
+	};
+
+	double ScaleX = (FMathd::Abs(BoxDimensions.X) > FMathf::ZeroTolerance) ? (1.0 / BoxDimensions.X) : 1.0;
+	double ScaleY = (FMathd::Abs(BoxDimensions.Y) > FMathf::ZeroTolerance) ? (1.0 / BoxDimensions.Y) : 1.0;
+	double ScaleZ = (FMathd::Abs(BoxDimensions.Z) > FMathf::ZeroTolerance) ? (1.0 / BoxDimensions.Z) : 1.0;
+	FVector3d Scale(ScaleX, ScaleY, ScaleZ);
+
+	double w = FMathd::Sqrt(BoxDimensions.X*BoxDimensions.X + BoxDimensions.Y*BoxDimensions.Y);
+	double h = BoxDimensions.Z;
+	double Angle = FMathd::Atan2(h, w);
+	double DotThresholdRejectFromPlane = FMathd::Cos(Angle);
+
+	TArray<FVector3d> TriNormals;
+	TArray<FIndex2i> TriangleCylinderAssignments;
+	TriNormals.SetNum(NumTriangles);
+	TriangleCylinderAssignments.SetNum(NumTriangles);
+	ParallelFor(NumTriangles, [&](int32 i)
+	{
+		int32 tid = Triangles[i];
+		TPair<FVector3d, FVector3d> NormalCentroid = GetTriNormalCentroid(tid);
+		TriNormals[i] = NormalCentroid.Key;
+		FVector3d N = BoxFrame.ToFrameVector(TriNormals[i]);
+		N = Normalized(N * Scale);
+
+		if (FMathd::Abs(N.Z) > DotThresholdRejectFromPlane)
+		{
+			int MajorAxis = 2;
+			double MajorAxisSign = FMathd::Sign(N[MajorAxis]);
+			// project to +/- Z
+			int Bucket = MajorAxisSign > 0 ? 1 : 0;
+
+			TriangleCylinderAssignments[i] = FIndex2i(MajorAxis, Bucket);
+		}
+		else
+		{
+			FVector3d Centroid = BoxFrame.ToFramePoint(NormalCentroid.Value);
+			double CentroidAngle = FMathd::Atan2(Centroid.Y, Centroid.X);
+			int Bucket = (CentroidAngle < 0) ? 3 : 4;
+			TriangleCylinderAssignments[i] = FIndex2i(-1, Bucket);
+		}
+	});
+
+	auto ProjAxis = [](const FVector3d& P, int Axis1, int Axis2, float Axis1Scale, float Axis2Scale)
+	{
+		return FVector2f(float(P[Axis1]) * Axis1Scale, float(P[Axis2]) * Axis2Scale);
+	};
+
+	TMap<FIndex2i, int32> BaseToOverlayVIDMap;
+	TArray<int32> NewUVIndices;
+
+	for ( int32 i = 0; i < NumTriangles; ++i )
+	{
+		int32 tid = Triangles[i];
+		FIndex3i BaseTri = Mesh->GetTriangle(tid);
+		FIndex2i TriBoxInfo = TriangleCylinderAssignments[i];
+		FVector3d N = BoxFrame.ToFrameVector(TriNormals[i]);
+
+		int MajorAxis = TriBoxInfo.A;
+		int Bucket = TriBoxInfo.B;
+
+		FIndex3i ElemTri;
+		for (int32 j = 0; j < 3; ++j)
+		{
+			FIndex2i ElementKey(BaseTri[j], Bucket);
+			const int32* FoundElementID = BaseToOverlayVIDMap.Find(ElementKey);
+			if (FoundElementID == nullptr)
+			{
+				FVector3d TransPos = PointTransform(Mesh->GetVertex(BaseTri[j]));
+				FVector3d BoxPos = Scale * BoxFrame.ToFramePoint(TransPos);
+
+				FVector2f UV = FVector2f::Zero();
+				if (Bucket <= 2)
+				{
+					UV = ProjAxis(BoxPos, 0, 1, FMathd::Sign(N[MajorAxis]) * Minor1Flip[MajorAxis], Minor2Flip[MajorAxis]);
+				}
+				else
+				{
+					double VAngle = FMathd::Atan2Positive(BoxPos.Y, BoxPos.X);
+					if (Bucket == 4 && VAngle < FMathd::HalfPi)
+					{
+						VAngle += FMathd::TwoPi;
+					}
+					else if (Bucket == 3 && VAngle > FMathd::HalfPi)
+					{
+						VAngle -= FMathd::TwoPi;
+					}
+					UV = FVector2f(-(float(VAngle) * FMathf::InvPi - 1.0f), -float(BoxPos.Z));
+				}
+
+				ElemTri[j] = UVOverlay->AppendElement(UV);
+				NewUVIndices.Add(ElemTri[j]);
+				BaseToOverlayVIDMap.Add(ElementKey, ElemTri[j]);
+			}
+			else
+			{
+				ElemTri[j] = *FoundElementID;
+			}
+		}
+
+		UVOverlay->SetTriangle(tid, ElemTri);
+	}
+
+	if (Result != nullptr)
+	{
+		Result->NewUVElements = MoveTemp(NewUVIndices);
+	}
+}
\ No newline at end of file
diff --git a/Engine/Plugins/Experimental/GeometryProcessing/Source/DynamicMesh/Public/Parameterization/DynamicMeshUVEditor.h b/Engine/Plugins/Experimental/GeometryProcessing/Source/DynamicMesh/Public/Parameterization/DynamicMeshUVEditor.h
--- a/Engine/Plugins/Experimental/GeometryProcessing/Source/DynamicMesh/Public/Parameterization/DynamicMeshUVEditor.h
+++ b/Engine/Plugins/Experimental/GeometryProcessing/Source/DynamicMesh/Public/Parameterization/DynamicMeshUVEditor.h
@@ -47,6 +47,12 @@ public:
 	 */
 	void CreateUVLayer(int32 UVLayerIndex);
 
+
+	/**
+	 * Clear UVs on active layer
+	 */
+	void ResetUVs();
+
 	/**
 	 * Create new UV island for each Triangle, by planar projection onto plane of Triangle. No transforms/etc are applied.
 	 */
@@ -57,12 +63,24 @@ public:
 	 */
 	void SetPerTriangleUVs(double ScaleFactor = 1.0, FUVEditResult* Result = nullptr);
 
+	/**
+	 * Apply arbitrary transform to UV elements
+	 */
+	void TransformUVElements(const TArray<int32>& ElementIDs, TFunctionRef<FVector2f(const FVector2f&)> TransformFunc);
+
 
 	/**
 	 * Create new UV island for given Triangles, and set UVs by planar projection to ProjectionFrame. No transforms/etc are applied.
 	 */
 	void SetTriangleUVsFromProjection(const TArray<int32>& Triangles, const FFrame3d& ProjectionFrame, FUVEditResult* Result = nullptr);
 
+	/**
+	 * Create new UV island for given Triangles, and set UVs by planar projection to ProjectionFrame. 
+	 * PointTransform is applied to points before projectiong onto ProjectionFrame X/Y axes
+	 * Projected U/V coordinates are divided by Dimensions.X/Y
+	 */
+	void SetTriangleUVsFromPlanarProjection(const TArray<int32>& Triangles, TFunctionRef<FVector3d(const FVector3d&)> PointTransform, const FFrame3d& ProjectionFrame, const FVector2d& Dimensions, FUVEditResult* Result = nullptr);
+
 	/**
 	 * Create new UV island for given Triangles, and set UVs for that island using Discrete Exponential Map.
 	 * ExpMap center-point is calculated by finding maximum (Dijkstra-approximated) geodesic distance from border of island.
@@ -70,6 +88,23 @@ public:
 	 */
 	bool SetTriangleUVsFromExpMap(const TArray<int32>& Triangles, FUVEditResult* Result = nullptr);
 
+
+	/**
+	 * Create new UV island for given Triangles, and set UVs for that island using Discrete Exponential Map.
+	 * ExpMap center-point is calculated by finding maximum (Dijkstra-approximated) geodesic distance from border of island.
+	 * @warning computes a single ExpMap, so input triangle set must be connected, however this is not verified internally
+	 */
+	bool SetTriangleUVsFromExpMap(
+		const TArray<int32>& Triangles, 
+		TFunctionRef<FVector3d(const FVector3d&)> PointTransform,
+		const FFrame3d& ProjectionFrame,
+		const FVector2d& Dimensions,
+		int32 NormalSmoothingRounds = 0,
+		double NormalSmoothingAlpha = 0.5,
+		double FrameNormalBlendWeight = 0,
+		FUVEditResult* Result = nullptr);
+
+
 	/**
 	 * Create new UV island for given Triangles, and set UVs for that island using Discrete Natural Conformal Map (equivalent to Least-Squares Conformal Map)
 	 * @warning computes a single parameterization, so input triangle set must be connected, however this is not verified internally
@@ -92,6 +127,42 @@ public:
 	 * @return true on success
 	 */
 	bool CreateSeamAlongVertexPath(const TArray<int32>& VertexPath, FUVEditResult* Result = nullptr);
+
+
+	/**
+	 * Set UVs by box projection. Triangles will be grouped to "best" box face
+	 * PointTransform is applied to points before projectiong onto ProjectionFrame X/Y axes
+	 * Projected U/V coordinates are divided by Dimensions.X/Y
+	 */
+	void SetTriangleUVsFromBoxProjection(const TArray<int32>& Triangles, TFunctionRef<FVector3d(const FVector3d&)> PointTransform, const FFrame3d& BoxFrame, const FVector3d& BoxDimensions, FUVEditResult* Result = nullptr);
+
+
+	/**
+	 * Set UVs by box projection. Triangles will be grouped to "best" box face
+	 * PointTransform is applied to points before projectiong onto ProjectionFrame X/Y axes
+	 * Projected U/V coordinates are divided by Dimensions.X/Y
+	 */
+	void SetTriangleUVsFromCylinderProjection(const TArray<int32>& Triangles, TFunctionRef<FVector3d(const FVector3d&)> PointTransform, const FFrame3d& CylFrame, const FVector3d& CylDimensions, FUVEditResult* Result = nullptr);
+
+
+
+
+public:
+
+	//
+	// Utility functions that could maybe go elsewhere but are mainly useful in parameterization context
+	//
+
+	/**
+	 * Compute geodesic center of given Mesh
+	 */
+	static bool EstimateGeodesicCenterFrameVertex(const FDynamicMesh3& Mesh, FFrame3d& FrameOut, int32 &VertexIDOut, bool bAlignToUnitAxes = true);
+
+	/**
+	 * Compute geodesic center of given Mesh triangles (assumes they are connected)
+	 */
+	static bool EstimateGeodesicCenterFrameVertex(const FDynamicMesh3& Mesh, const TArray<int32>& Triangles, FFrame3d& FrameOut, int32& VertexIDOut, bool bAlignToUnitAxes = true);
+
 };
 
 
