ModelingTools: Fix post op selection for TriEdit: Extrude, Offset, Cut Faces & Duplicate.

#rnx

diff --git a/Engine/Plugins/Experimental/GeometryProcessing/Source/DynamicMesh/Private/Operations/MeshPlaneCut.cpp b/Engine/Plugins/Experimental/GeometryProcessing/Source/DynamicMesh/Private/Operations/MeshPlaneCut.cpp
--- a/Engine/Plugins/Experimental/GeometryProcessing/Source/DynamicMesh/Private/Operations/MeshPlaneCut.cpp
+++ b/Engine/Plugins/Experimental/GeometryProcessing/Source/DynamicMesh/Private/Operations/MeshPlaneCut.cpp
@@ -20,7 +20,13 @@ using namespace UE::Geometry;
 
 void FMeshPlaneCut::SplitCrossingEdges(TArray<double>& Signs, TSet<int>& ZeroEdges, TSet<int>& OnCutEdges, bool bDeleteTrisOnPlane)
 {
-	Signs.Reset(); ZeroEdges.Reset(); OnCutEdges.Reset();
+	TSet<int> OnSplitEdges;
+	SplitCrossingEdges(Signs, ZeroEdges, OnCutEdges, OnSplitEdges, bDeleteTrisOnPlane);
+}
+
+void FMeshPlaneCut::SplitCrossingEdges(TArray<double>& Signs, TSet<int>& ZeroEdges, TSet<int>& OnCutEdges, TSet<int>& OnSplitEdges, bool bDeleteTrisOnPlane)
+{
+	Signs.Reset(); ZeroEdges.Reset(); OnCutEdges.Reset(); OnSplitEdges.Reset();
 	OnCutVertices.Reset();
 
 	double InvalidDist = -FMathd::MaxReal;
@@ -152,7 +158,8 @@ void FMeshPlaneCut::SplitCrossingEdges(TArray<double>& Signs, TSet<int>& ZeroEdg
 			continue; // edge split really shouldn't fail; skip the edge if it somehow does
 		}
 
-		NewEdges.Add(splitInfo.NewEdges.A);
+		OnSplitEdges.Add(EID);
+		NewEdges.Add(splitInfo.NewEdges.A); OnSplitEdges.Add(splitInfo.NewEdges.A);
 		NewEdges.Add(splitInfo.NewEdges.B); OnCutEdges.Add(splitInfo.NewEdges.B);
 		if (splitInfo.NewEdges.C != FDynamicMesh3::InvalidID)
 		{
@@ -383,8 +390,8 @@ bool FMeshPlaneCut::SplitEdgesOnly(bool bAssignNewGroups)
 {
 	// split edges with current plane
 	TArray<double> Signs;
-	TSet<int32> ZeroEdges, OnCutEdges;
-	SplitCrossingEdges(Signs, ZeroEdges, OnCutEdges, false);
+	TSet<int32> ZeroEdges, OnCutEdges, OnSplitEdges;
+	SplitCrossingEdges(Signs, ZeroEdges, OnCutEdges, OnSplitEdges, false);
 
 	if (bAssignNewGroups == false)
 	{
@@ -453,6 +460,46 @@ bool FMeshPlaneCut::SplitEdgesOnly(bool bAssignNewGroups)
 		Result.Triangles = MoveTemp(Component.Indices);
 	}
 
+	// Compute the set of triangle IDs in the cut mesh that represent
+	// the original/seed triangle selection along the cut. This assumes
+	// that the edge filter function contains edges that originate
+	// from source triangles.
+	ResultSeedTriangles.Reset();
+	for (int tid : CutEdgeTriangles)
+	{
+		// TODO: We currently assume that all cut edges are on the interior of
+		// our seed triangles, thus all tris adjacent to the cut edge are seed
+		// triangles. This assumption fails in the following edge case.
+		//
+		//    o---o---o
+		//    |xx/ \xx|
+		//    |x/   \x| <---> Cut plane   
+		//    |/     \|
+		//    o-------o    x = Seed tris
+		//
+		// CutEdges filters the OnCutEdges list by checking if both ends of the edge
+		// are OnCutVertices. In this scenario, the edge introduced across the non
+		// seed triangle on the bottom is included.
+		ResultSeedTriangles.Add(tid);
+
+		// Walk the edges of the CutEdgeTriangles, skipping seed, split & cut edges,
+		// to identify extra interior edges. Both triangles along that interior
+		// edge are also seed triangles.
+		FIndex3i TriEdges = Mesh->GetTriEdges(tid);
+		for (int j = 0; j < 3; j++)
+		{
+			int eid = TriEdges[j];
+			FIndex2i ev = Mesh->GetEdgeV(eid);
+			bool bIsSeedEdge = (EdgeFilterFunc && EdgeFilterFunc(eid));
+			if (bIsSeedEdge || CutEdges.Contains(eid) || OnSplitEdges.Contains(eid))
+			{
+				continue;
+			}
+			FIndex2i EdgeTris = Mesh->GetEdgeT(eid);
+			ResultSeedTriangles.Add(EdgeTris.A);
+			ResultSeedTriangles.Add(EdgeTris.B);
+		}
+	}
 	return true;
 }
 
diff --git a/Engine/Plugins/Experimental/GeometryProcessing/Source/DynamicMesh/Public/Operations/MeshPlaneCut.h b/Engine/Plugins/Experimental/GeometryProcessing/Source/DynamicMesh/Public/Operations/MeshPlaneCut.h
--- a/Engine/Plugins/Experimental/GeometryProcessing/Source/DynamicMesh/Public/Operations/MeshPlaneCut.h
+++ b/Engine/Plugins/Experimental/GeometryProcessing/Source/DynamicMesh/Public/Operations/MeshPlaneCut.h
@@ -91,6 +91,8 @@ public:
 	};
 	/** List of output cut regions (eg that have separate GroupIDs). Currently only calculated by SplitEdgesOnly() path */
 	TArray<FCutResultRegion> ResultRegions;
+	/** List of output cut triangles representing the seed triangles along the cut. Currently only calculated by SplitEdgesOnly() path */
+	TArray<int32> ResultSeedTriangles;
 
 public:
 
@@ -158,6 +160,7 @@ protected:
 
 	void CollapseDegenerateEdges(const TSet<int>& OnCutEdges, const TSet<int>& ZeroEdges);
 	void SplitCrossingEdges(TArray<double>& Signs, TSet<int>& ZeroEdges, TSet<int>& OnCutEdges, bool bDeleteTrisOnPlane = true);
+	void SplitCrossingEdges(TArray<double>& Signs, TSet<int>& ZeroEdges, TSet<int>& OnCutEdges, TSet<int>& OnSplitEdges, bool bDeleteTrisOnPlane = true);
 	bool ExtractBoundaryLoops(const TSet<int>& OnCutEdges, const TSet<int>& ZeroEdges, FMeshPlaneCut::FOpenBoundary& Boundary);
 
 	// set of vertices lying on plane after calling SplitCrossingEdges
diff --git a/Engine/Plugins/Experimental/MeshModelingToolset/Source/MeshModelingTools/Private/EditMeshPolygonsTool.cpp b/Engine/Plugins/Experimental/MeshModelingToolset/Source/MeshModelingTools/Private/EditMeshPolygonsTool.cpp
--- a/Engine/Plugins/Experimental/MeshModelingToolset/Source/MeshModelingTools/Private/EditMeshPolygonsTool.cpp
+++ b/Engine/Plugins/Experimental/MeshModelingToolset/Source/MeshModelingTools/Private/EditMeshPolygonsTool.cpp
@@ -1301,11 +1301,18 @@ void UEditMeshPolygonsTool::ApplyExtrude(bool bIsOffset)
 
 	// construct new selection
 	FGroupTopologySelection NewSelection;
-	for (const FOffsetMeshRegion::FOffsetInfo& Info : Extruder.OffsetRegions)
+	if (!bTriangleMode)
 	{
-		for (int32 gid : Info.OffsetGroups)
+		for (const FOffsetMeshRegion::FOffsetInfo& Info : Extruder.OffsetRegions)
 		{
-			NewSelection.SelectedGroupIDs.Add(gid);
+			NewSelection.SelectedGroupIDs.Append(Info.OffsetGroups);
+		}
+	}
+	else
+	{
+		for (const FOffsetMeshRegion::FOffsetInfo& Info : Extruder.OffsetRegions)
+		{
+			NewSelection.SelectedGroupIDs.Append(Info.InitialTriangles);
 		}
 	}
 
@@ -1509,9 +1516,20 @@ void UEditMeshPolygonsTool::ApplyCutFaces()
 	Cut.EdgeFilterFunc = [&](int EdgeID) { return Edges.IsSelected(EdgeID); };
 	if (Cut.SplitEdgesOnly(true))
 	{
-		for (FMeshPlaneCut::FCutResultRegion& Region : Cut.ResultRegions)
+		if (!bTriangleMode)
+		{
+			for (FMeshPlaneCut::FCutResultRegion& Region : Cut.ResultRegions)
+			{
+				OutputSelection.SelectedGroupIDs.Add(Region.GroupID);
+			}
+		}
+		else
 		{
-			OutputSelection.SelectedGroupIDs.Add(Region.GroupID);
+			// Retain the selection along the cut. ResultSeedTriangles does not
+			// contain selected tris that are not cut, so re-add the original selected
+			// tris.
+			OutputSelection.SelectedGroupIDs.Append(Cut.ResultSeedTriangles);
+			OutputSelection.SelectedGroupIDs.Append(ActiveTriangleSelection);
 		}
 	}
 
@@ -1923,10 +1941,7 @@ void UEditMeshPolygonsTool::ApplyDuplicate()
 	Editor.DuplicateTriangles(AllTriangles, Mappings, EditResult);
 
 	FGroupTopologySelection NewSelection;
-	for ( int NewGroupID : EditResult.NewGroups )
-	{ 
-		NewSelection.SelectedGroupIDs.Add(NewGroupID);
-	}
+	NewSelection.SelectedGroupIDs.Append(bTriangleMode ? EditResult.NewTriangles : EditResult.NewGroups);
 
 	TUniquePtr<FMeshChange> MeshChange = MakeUnique<FMeshChange>(ChangeTracker.EndChange());
 	CompleteMeshEditChange(LOCTEXT("PolyMeshDisconnectChange", "Disconnect"), MoveTemp(MeshChange), NewSelection);
