GeometryProcessing: Fix crash when trying to access UV overlay with parent vertex index of -1.

diff --git a/Engine/Plugins/Experimental/GeometryProcessing/Source/DynamicMesh/Private/Parameterization/MeshUVTransforms.cpp b/Engine/Plugins/Experimental/GeometryProcessing/Source/DynamicMesh/Private/Parameterization/MeshUVTransforms.cpp
--- a/Engine/Plugins/Experimental/GeometryProcessing/Source/DynamicMesh/Private/Parameterization/MeshUVTransforms.cpp
+++ b/Engine/Plugins/Experimental/GeometryProcessing/Source/DynamicMesh/Private/Parameterization/MeshUVTransforms.cpp
@@ -142,6 +142,10 @@ void UE::MeshUVTransforms::MakeSeamsDisjoint(FDynamicMeshUVOverlay* UVOverlay)
 	for (int32 elemid : UVOverlay->ElementIndicesItr())
 	{
 		int32 ParentVtxID = UVOverlay->GetParentVertex(elemid);
+		if (ParentVtxID == -1)
+		{
+			continue;
+		}
 		int32 SeamListIndex = SeamVerticesMap[ParentVtxID];
 		if (SeamListIndex == -1)
 		{
