fixes for potential crash / invalid operations on radial fracture of pillar
 - make 'embed surface path' find the intended start vertex for its path when it's known
 - guard against path being empty before accessing element 0
 - guard against invalid edge collapse in mesh boolean that was triggered by this example

#rnx

diff --git a/Engine/Plugins/Experimental/GeometryProcessing/Source/DynamicMesh/Private/Operations/EmbedSurfacePath.cpp b/Engine/Plugins/Experimental/GeometryProcessing/Source/DynamicMesh/Private/Operations/EmbedSurfacePath.cpp
--- a/Engine/Plugins/Experimental/GeometryProcessing/Source/DynamicMesh/Private/Operations/EmbedSurfacePath.cpp
+++ b/Engine/Plugins/Experimental/GeometryProcessing/Source/DynamicMesh/Private/Operations/EmbedSurfacePath.cpp
@@ -133,7 +133,7 @@ struct FIndexDistance
 	}
 };
 
-bool WalkMeshPlanar(const FDynamicMesh3* Mesh, int StartTri, FVector3d StartPt, int EndTri, int EndVertID, FVector3d EndPt, FVector3d WalkPlaneNormal, TFunction<FVector3d(const FDynamicMesh3*, int)> VertexToPosnFn,
+bool WalkMeshPlanar(const FDynamicMesh3* Mesh, int StartTri, int StartVID, FVector3d StartPt, int EndTri, int EndVertID, FVector3d EndPt, FVector3d WalkPlaneNormal, TFunction<FVector3d(const FDynamicMesh3*, int)> VertexToPosnFn,
 	bool bAllowBackwardsSearch, double AcceptEndPtOutsideDist, double PtOnPlaneThresholdSq, TArray<TPair<FMeshSurfacePoint, int>>& WalkedPath, double BackwardsTolerance)
 {
 	auto SetTriVertPositions = [&VertexToPosnFn, &Mesh](FIndex3i TriVertIDs, FTriangle3d& Tri)
@@ -175,10 +175,26 @@ bool WalkMeshPlanar(const FDynamicMesh3* Mesh, int StartTri, FVector3d StartPt,
 	FIndex3i StartTriVertIDs = Mesh->GetTriangle(StartTri);
 	SetTriVertPositions(StartTriVertIDs, CurrentTri);
 	FDistPoint3Triangle3d CurrentTriDist(StartPt, CurrentTri); // heavy duty way to get barycentric coordinates and check if on triangle; should be robust to degenerate triangles unlike VectorUtil's barycentric coordinate function
-	// TODO: use TrianglePosToSurfacePoint to snap to edge or vertex as needed (OR do this as a post-process and delete the point if doing so leads to a duplicate point!)
-	CurrentTriDist.ComputeResult();
-	// TODO: replace barycoords result with edge or vertex surface point data if within distance threshold of vertex or edge!
-	ComputedPointsAndSources.Emplace(FMeshSurfacePoint(StartTri, CurrentTriDist.TriangleBaryCoords), FWalkIndices(StartPt, -1, StartTri));
+	int StartVIDIndex = -1;
+	if (StartVID != -1)
+	{
+		StartVIDIndex = StartTriVertIDs.IndexOf(StartVID);
+	}
+	if (StartVIDIndex == -1)
+	{
+		// TODO: use TrianglePosToSurfacePoint to snap to edge or vertex as needed (OR do this as a post-process and delete the point if doing so leads to a duplicate point!)
+		CurrentTriDist.ComputeResult();
+		// TODO: replace barycoords result with edge or vertex surface point data if within distance threshold of vertex or edge!
+		ComputedPointsAndSources.Emplace(FMeshSurfacePoint(StartTri, CurrentTriDist.TriangleBaryCoords), FWalkIndices(StartPt, -1, StartTri));
+	}
+	else
+	{
+		// if a valid StartVID was given, assume that's our closest point
+		CurrentTriDist.TriangleBaryCoords = FVector3d::Zero();
+		CurrentTriDist.TriangleBaryCoords[StartVIDIndex] = 1.0;
+		CurrentTriDist.ClosestTrianglePoint = StartPt;
+		ComputedPointsAndSources.Emplace(FMeshSurfacePoint(StartTri, CurrentTriDist.TriangleBaryCoords), FWalkIndices(StartPt, -1, StartTri));
+	}
 
 	FVector3d ForwardsDirection = EndPt - StartPt;
 
@@ -560,7 +576,7 @@ bool FMeshSurfacePath::IsConnected() const
 
 
 bool FMeshSurfacePath::AddViaPlanarWalk(
-	int StartTri, FVector3d StartPt, int EndTri, int EndVertID, 
+	int StartTri, int StartVID, FVector3d StartPt, int EndTri, int EndVertID, 
 	FVector3d EndPt, FVector3d WalkPlaneNormal, TFunction<FVector3d(const FDynamicMesh3*, int)> VertexToPosnFn,
 	bool bAllowBackwardsSearch, double AcceptEndPtOutsideDist, double PtOnPlaneThresholdSq, double BackwardsTolerance)
 {
@@ -571,7 +587,7 @@ bool FMeshSurfacePath::AddViaPlanarWalk(
 			return MeshArg->GetVertex(VertexID);
 		};
 	}
-	return WalkMeshPlanar(Mesh, StartTri, StartPt, EndTri, EndVertID, EndPt, WalkPlaneNormal, VertexToPosnFn,
+	return WalkMeshPlanar(Mesh, StartTri, StartVID, StartPt, EndTri, EndVertID, EndPt, WalkPlaneNormal, VertexToPosnFn,
 		bAllowBackwardsSearch, AcceptEndPtOutsideDist, PtOnPlaneThresholdSq, Path, BackwardsTolerance);
 }
 
@@ -842,7 +858,7 @@ bool UE::Geometry::EmbedProjectedPaths(FDynamicMesh3* Mesh, const TArrayView<con
 				{
 					return false;
 				}
-				bool bWalkSuccess = SurfacePath.AddViaPlanarWalk(CurrentSeedTriID, StartPos, -1, LastVert, FVector3d(Path2D[IdxB].X, Path2D[IdxB].Y, 0), WalkNormal, ProjectToFrame, false, FMathf::ZeroTolerance, PtSnapVertexOrEdgeThresholdSq);
+				bool bWalkSuccess = SurfacePath.AddViaPlanarWalk(CurrentSeedTriID, -1, StartPos, -1, LastVert, FVector3d(Path2D[IdxB].X, Path2D[IdxB].Y, 0), WalkNormal, ProjectToFrame, false, FMathf::ZeroTolerance, PtSnapVertexOrEdgeThresholdSq);
 				if (!bWalkSuccess)
 				{
 					return false;
diff --git a/Engine/Plugins/Experimental/GeometryProcessing/Source/DynamicMesh/Private/Operations/MeshBoolean.cpp b/Engine/Plugins/Experimental/GeometryProcessing/Source/DynamicMesh/Private/Operations/MeshBoolean.cpp
--- a/Engine/Plugins/Experimental/GeometryProcessing/Source/DynamicMesh/Private/Operations/MeshBoolean.cpp
+++ b/Engine/Plugins/Experimental/GeometryProcessing/Source/DynamicMesh/Private/Operations/MeshBoolean.cpp
@@ -1084,6 +1084,49 @@ void FMeshBoolean::SimplifyAlongNewEdges(int NumMeshesToProcess, FDynamicMesh3*
 					continue;
 				}
 
+				// do some pre-collapse sanity checks on the matched edge (if present) to see if it will fail to collapse
+				bool bAttemptCollapse = true;
+				if (bHasMatches)
+				{
+					int OtherRemoveV = Matches[RemoveVIdx];
+					int OtherKeepV = Matches[KeepVIdx];
+
+					int a = OtherRemoveV, b = OtherKeepV;
+					int eab = CutMesh[1]->FindEdge(OtherRemoveV, OtherKeepV);
+
+					const FDynamicMesh3::FEdge EdgeAB = CutMesh[1]->GetEdge(eab);
+					int t0 = EdgeAB.Tri[0];
+					if (t0 == FDynamicMesh3::InvalidID)
+					{
+						bAttemptCollapse = false;
+					}
+					else
+					{
+						FIndex3i T0tv = CutMesh[1]->GetTriangle(t0);
+						int c = IndexUtil::FindTriOtherVtx(a, b, T0tv);
+						checkSlow(EdgeAB.Tri[1] == FDynamicMesh3::InvalidID);
+						// We cannot collapse if edge lists of a and b share vertices other
+						//  than c and d  (because then we will make a triangle [x b b].
+						//  Brute-force search logic adapted from FDynamicMesh3::CollapseEdge implementation.
+						//  (simplified because we know this is a boundary edge)
+						CutMesh[1]->EnumerateVertexVertices(a, [&](int VID)
+						{
+							if (!bAttemptCollapse || VID == c || VID == b)
+							{
+								return;
+							}
+							CutMesh[1]->EnumerateVertexVertices(b, [&](int VID2)
+							{
+								bAttemptCollapse &= (VID != VID2);
+							});
+						});
+					}
+				}
+				if (!bAttemptCollapse)
+				{
+					break; // don't try starting from other vertex if the match edge couldn't be collapsed
+				}
+
 				FDynamicMesh3::FEdgeCollapseInfo CollapseInfo;
 				int RemoveV = Edge.Vert[RemoveVIdx];
 				int KeepV = Edge.Vert[KeepVIdx];
diff --git a/Engine/Plugins/Experimental/GeometryProcessing/Source/DynamicMesh/Private/Operations/MeshMeshCut.cpp b/Engine/Plugins/Experimental/GeometryProcessing/Source/DynamicMesh/Private/Operations/MeshMeshCut.cpp
--- a/Engine/Plugins/Experimental/GeometryProcessing/Source/DynamicMesh/Private/Operations/MeshMeshCut.cpp
+++ b/Engine/Plugins/Experimental/GeometryProcessing/Source/DynamicMesh/Private/Operations/MeshMeshCut.cpp
@@ -345,7 +345,7 @@ namespace MeshCut
 				FVector3d WalkPlaneNormal = BaseFaceNormals[Seg.BaseTID].Cross(PtB.Pos - PtA.Pos);
 				if (ensure(Normalize(WalkPlaneNormal) > 0))
 				{
-					bool bWalkSuccess = SurfacePath.AddViaPlanarWalk(StartTID,
+					bool bWalkSuccess = SurfacePath.AddViaPlanarWalk(StartTID, PtA.ElemID,
 						Mesh->GetVertex(PtA.ElemID), -1, PtB.ElemID,
 						Mesh->GetVertex(PtB.ElemID), WalkPlaneNormal, nullptr /*TODO: transform fn goes here?*/, false, FMathd::ZeroTolerance, SnapToleranceSq, .001);
 					if (!bWalkSuccess)
@@ -357,7 +357,7 @@ namespace MeshCut
 						EmbeddedPath.Reset();
 						if (SurfacePath.EmbedSimplePath(false, EmbeddedPath, false, SnapToleranceSq))
 						{
-							ensure(EmbeddedPath[0] == PtA.ElemID);
+							ensure(EmbeddedPath.Num() > 0 && EmbeddedPath[0] == PtA.ElemID);
 							if (VertexChains)
 							{
 								if (SegmentToChain)
diff --git a/Engine/Plugins/Experimental/GeometryProcessing/Source/DynamicMesh/Public/Operations/EmbedSurfacePath.h b/Engine/Plugins/Experimental/GeometryProcessing/Source/DynamicMesh/Public/Operations/EmbedSurfacePath.h
--- a/Engine/Plugins/Experimental/GeometryProcessing/Source/DynamicMesh/Public/Operations/EmbedSurfacePath.h
+++ b/Engine/Plugins/Experimental/GeometryProcessing/Source/DynamicMesh/Public/Operations/EmbedSurfacePath.h
@@ -97,7 +97,7 @@ public:
 		bIsClosed = false;
 	}
 
-	bool AddViaPlanarWalk(int StartTri, FVector3d StartPt, int EndTri, int EndVertID, FVector3d EndPt,
+	bool AddViaPlanarWalk(int StartTri, int StartVID, FVector3d StartPt, int EndTri, int EndVertID, FVector3d EndPt,
 		FVector3d WalkPlaneNormal, TFunction<FVector3d(const FDynamicMesh3*, int)> VertexToPosnFn = nullptr,
 		bool bAllowBackwardsSearch = true, double AcceptEndPtOutsideDist = FMathd::ZeroTolerance, 
 		double PtOnPlaneThresholdSq = FMathf::ZeroTolerance*100, double BackwardsTolerance = FMathd::ZeroTolerance*10);
