add a few missing tests for valid triangle id when doing parallelfor to max ID

#rnx

diff --git a/Engine/Plugins/Experimental/GeometryProcessing/Source/DynamicMesh/Private/NormalFlowRemesher.cpp b/Engine/Plugins/Experimental/GeometryProcessing/Source/DynamicMesh/Private/NormalFlowRemesher.cpp
--- a/Engine/Plugins/Experimental/GeometryProcessing/Source/DynamicMesh/Private/NormalFlowRemesher.cpp
+++ b/Engine/Plugins/Experimental/GeometryProcessing/Source/DynamicMesh/Private/NormalFlowRemesher.cpp
@@ -173,6 +173,10 @@ void FNormalFlowRemesher::TrackedFaceProjectionPass(double& MaxDistanceMoved, bo
 	// First compute all rotated triangles and store their corner positions and weights
 	ParallelFor(Mesh->MaxTriangleID(), [this, NormalProjTarget, bIsTuningIteration, SmoothAreaDistance, UseSmoothFunc, MaxStepDistance](int32 TriangleIndex)
 	{
+		if (!Mesh->IsTriangle(TriangleIndex))
+		{
+			return;
+		}
 		FVector3d TriangleNormal, Centroid;
 		double Area;
 		Mesh->GetTriInfo(TriangleIndex, TriangleNormal, Area, Centroid);
diff --git a/Engine/Source/Runtime/Experimental/GeometryFramework/Private/Components/DynamicMeshSceneProxy.h b/Engine/Source/Runtime/Experimental/GeometryFramework/Private/Components/DynamicMeshSceneProxy.h
--- a/Engine/Source/Runtime/Experimental/GeometryFramework/Private/Components/DynamicMeshSceneProxy.h
+++ b/Engine/Source/Runtime/Experimental/GeometryFramework/Private/Components/DynamicMeshSceneProxy.h
@@ -246,6 +246,10 @@ public:
 		}
 		ParallelFor(Mesh->MaxTriangleID(), [&](int tid)
 		{
+			if (!Mesh->IsTriangle(tid))
+			{
+				return;
+			}
 			int MatIdx;
 			MaterialID->GetValue(tid, &MatIdx);
 			if (MatIdx >= 0 && MatIdx < NumMaterials)
diff --git a/Engine/Source/Runtime/GeometryCore/Public/Implicit/CachingMeshSDF.h b/Engine/Source/Runtime/GeometryCore/Public/Implicit/CachingMeshSDF.h
--- a/Engine/Source/Runtime/GeometryCore/Public/Implicit/CachingMeshSDF.h
+++ b/Engine/Source/Runtime/GeometryCore/Public/Implicit/CachingMeshSDF.h
@@ -318,7 +318,7 @@ private:
 			{
 				cancelled = true;
 			}
-			if (cancelled)
+			if (cancelled || !Mesh->IsTriangle(TID))
 			{
 				return;
 			}
