make mesh booleans handle degenerate sliver triangles: fix their intersection not being detected + the mesh cut not properly cutting through them

#rnx

diff --git a/Engine/Plugins/Experimental/GeometryProcessing/Source/DynamicMesh/Private/Operations/EmbedSurfacePath.cpp b/Engine/Plugins/Experimental/GeometryProcessing/Source/DynamicMesh/Private/Operations/EmbedSurfacePath.cpp
--- a/Engine/Plugins/Experimental/GeometryProcessing/Source/DynamicMesh/Private/Operations/EmbedSurfacePath.cpp
+++ b/Engine/Plugins/Experimental/GeometryProcessing/Source/DynamicMesh/Private/Operations/EmbedSurfacePath.cpp
@@ -422,30 +422,44 @@ bool WalkMeshPlanar(const FDynamicMesh3* Mesh, int StartTri, int StartVID, FVect
 	// note we don't even check that the barycoords match in the edge case -- conceptually the path can only cross the edge at one point, so it 'should' be fine to treat them as close enough, and having two points on the same edge would mess up the simple embedding code
 	if (WalkedPath.Num() && WalkedPath[0].Key.PointType == ESurfacePointType::Triangle)
 	{
-		RefineSurfacePtFromTriangleToSubElement(Mesh, WalkedPath[0].Key.Pos(Mesh), WalkedPath[0].Key, PtOnPlaneThresholdSq);
+		FMeshSurfacePoint& SurfacePt = WalkedPath[0].Key;
+		
+		// special case code to handle if we started on an exact vertex ID, to ensure refinement gets the same vertex ID
+		// (note we do this after-the-fact fix rather than refining up-front to keep the traversal logic simpler)
+		if (StartVIDIndex > -1 && SurfacePt.BaryCoord[StartVIDIndex] == 1.0) // if we had an exact start vertex
+		{
+			FIndex3i TriVertIDs = Mesh->GetTriangle(SurfacePt.ElementID);
+			SurfacePt.ElementID = TriVertIDs[StartVIDIndex];
+			SurfacePt.PointType = ESurfacePointType::Vertex;
+		}
+		else
+		{
+			RefineSurfacePtFromTriangleToSubElement(Mesh, SurfacePt.Pos(Mesh), SurfacePt, PtOnPlaneThresholdSq);
+		}
 		if (WalkedPath.Num() > 1 &&
-			WalkedPath[0].Key.PointType != ESurfacePointType::Triangle &&
-			WalkedPath[0].Key.PointType == WalkedPath[1].Key.PointType &&
-			WalkedPath[0].Key.ElementID == WalkedPath[1].Key.ElementID)
+			SurfacePt.PointType != ESurfacePointType::Triangle &&
+			SurfacePt.PointType == WalkedPath[1].Key.PointType &&
+			SurfacePt.ElementID == WalkedPath[1].Key.ElementID)
 		{
-			if (WalkedPath.Last().Key.PointType == ESurfacePointType::Edge) // copy closer barycoord
+			if (SurfacePt.PointType == ESurfacePointType::Edge) // copy closer barycoord
 			{
-				WalkedPath[1].Key.BaryCoord = WalkedPath[0].Key.BaryCoord;
+				WalkedPath[1].Key.BaryCoord = SurfacePt.BaryCoord;
 			}
 			WalkedPath.RemoveAt(0);
 		}
 	}
 	if (WalkedPath.Num() && WalkedPath.Last().Key.PointType == ESurfacePointType::Triangle)
 	{
-		RefineSurfacePtFromTriangleToSubElement(Mesh, WalkedPath.Last().Key.Pos(Mesh), WalkedPath.Last().Key, PtOnPlaneThresholdSq);
+		FMeshSurfacePoint& SurfacePt = WalkedPath.Last().Key;
+		RefineSurfacePtFromTriangleToSubElement(Mesh, SurfacePt.Pos(Mesh), SurfacePt, PtOnPlaneThresholdSq);
 		if (WalkedPath.Num() > 1 &&
-			WalkedPath.Last().Key.PointType != ESurfacePointType::Triangle &&
-			WalkedPath.Last().Key.PointType == WalkedPath.Last(1).Key.PointType &&
-			WalkedPath.Last().Key.ElementID == WalkedPath.Last(1).Key.ElementID)
+			SurfacePt.PointType != ESurfacePointType::Triangle &&
+			SurfacePt.PointType == WalkedPath.Last(1).Key.PointType &&
+			SurfacePt.ElementID == WalkedPath.Last(1).Key.ElementID)
 		{
-			if (WalkedPath.Last().Key.PointType == ESurfacePointType::Edge) // copy closer barycoord
+			if (SurfacePt.PointType == ESurfacePointType::Edge) // copy closer barycoord
 			{
-				WalkedPath.Last(1).Key.BaryCoord = WalkedPath.Last().Key.BaryCoord;
+				WalkedPath.Last(1).Key.BaryCoord = SurfacePt.BaryCoord;
 			}
 			WalkedPath.Pop();
 		}
diff --git a/Engine/Plugins/Experimental/GeometryProcessing/Source/DynamicMesh/Private/Operations/MeshMeshCut.cpp b/Engine/Plugins/Experimental/GeometryProcessing/Source/DynamicMesh/Private/Operations/MeshMeshCut.cpp
--- a/Engine/Plugins/Experimental/GeometryProcessing/Source/DynamicMesh/Private/Operations/MeshMeshCut.cpp
+++ b/Engine/Plugins/Experimental/GeometryProcessing/Source/DynamicMesh/Private/Operations/MeshMeshCut.cpp
@@ -47,6 +47,22 @@ namespace MeshCut
 			Init(WorkingMesh);
 		}
 
+		// get the direction of the first non-degenerate edge (or a default vector if all edges were degenerate)
+		static FVector3d GetDegenTriangleEdgeDirection(const FDynamicMesh3* Mesh, int TID, FVector3d DefaultDir = FVector3d::ZAxisVector)
+		{ 
+			FVector3d V[3];
+			Mesh->GetTriVertices(TID, V[0], V[1], V[2]);
+			for (int Prev = 2, Idx = 0; Idx < 3; Prev = Idx++)
+			{
+				FVector3d E = V[Idx] - V[Prev];
+				if (E.Normalize())
+				{
+					return E;
+				}
+			}
+			return DefaultDir;
+		}
+
 		void Init(FDynamicMesh3* WorkingMesh)
 		{
 			BaseFaceNormals.SetNumUninitialized(WorkingMesh->MaxTriangleID());
@@ -86,7 +102,7 @@ namespace MeshCut
 		TArray<FSegmentToElements> Segments;
 
 
-		void AddSegments(const MeshIntersection::FIntersectionsQueryResult& Intersections, TMap<FVector3d, int>& SegVtxMap, int WhichSide)
+		void AddSegments(const MeshIntersection::FIntersectionsQueryResult& Intersections, int WhichSide)
 		{
 			int SegStart = Segments.Num();
 			Segments.SetNum(SegStart + Intersections.Segments.Num());
@@ -97,27 +113,21 @@ namespace MeshCut
 				const MeshIntersection::FSegmentIntersection& Seg = Intersections.Segments[SegIdx];
 				FSegmentToElements& SegToEls = Segments[SegStart+SegIdx];
 				SegToEls.BaseTID = Seg.TriangleID[WhichSide];
+
+				FTriangle3d Tri;
+				Mesh->GetTriVertices(SegToEls.BaseTID, Tri.V[0], Tri.V[1], Tri.V[2]);
+				FIndex3i TriVIDs = Mesh->GetTriangle(SegToEls.BaseTID);
+				int PrevOnEdgeIdx = -1;
+				FVector3d PrevOnEdgePos;
 				for (int SegPtIdx = 0; SegPtIdx < 2; SegPtIdx++)
 				{
-					// re-use intersection vertices for exact position matches
-					int* ExistingSegVertexIdx = SegVtxMap.Find(Seg.Point[SegPtIdx]);
-					if (ExistingSegVertexIdx)
-					{
-						SegToEls.PtOnMeshIdx[SegPtIdx] = *ExistingSegVertexIdx;
-					}
-
 					int NewPtIdx = IntersectionVerts.Num();
 					FPtOnMesh& PtOnMesh = IntersectionVerts.Emplace_GetRef();
 					PtOnMesh.Pos = Seg.Point[SegPtIdx];
 					SegToEls.PtOnMeshIdx[SegPtIdx] = NewPtIdx;
-					SegVtxMap.Add(Seg.Point[SegPtIdx], NewPtIdx);
 
 					// decide whether the point is on a vertex, edge or triangle
 
-					FTriangle3d Tri;
-					Mesh->GetTriVertices(SegToEls.BaseTID, Tri.V[0], Tri.V[1], Tri.V[2]);
-					FIndex3i TriVIDs = Mesh->GetTriangle(SegToEls.BaseTID);
-
 					int OnVertexIdx = OnVertex(Tri, PtOnMesh.Pos);
 					if (OnVertexIdx > -1)
 					{
@@ -126,16 +136,28 @@ namespace MeshCut
 						continue;
 					}
 
-					// check  for an edge match
+					// check for an edge match
 					int OnEdgeIdx = OnEdge(Tri, PtOnMesh.Pos);
 					if (OnEdgeIdx > -1)
 					{
+						// if segment is degenerate and stuck to one edge, see if it could cross
+						if (PrevOnEdgeIdx == OnEdgeIdx && FVector3d::DistSquared(PrevOnEdgePos, PtOnMesh.Pos) < SnapToleranceSq)
+						{
+							int OnEdgeReplaceIdx = OnEdgeWithSkip(Tri, PtOnMesh.Pos, OnEdgeIdx);
+							if (OnEdgeReplaceIdx > -1)
+							{
+								OnEdgeIdx = OnEdgeReplaceIdx;
+							}
+						}
 						PtOnMesh.Type = EVertexType::Edge;
 						PtOnMesh.ElemID = Mesh->GetTriEdge(SegToEls.BaseTID, OnEdgeIdx);
 
 						check(PtOnMesh.ElemID > -1);
 						EdgeVertices.Add(PtOnMesh.ElemID, NewPtIdx);
 
+						PrevOnEdgeIdx = OnEdgeIdx;
+						PrevOnEdgePos = PtOnMesh.Pos;
+
 						continue;
 					}
 
@@ -343,40 +365,51 @@ namespace MeshCut
 				FMeshSurfacePath SurfacePath(Mesh);
 				int StartTID = Mesh->GetVtxSingleTriangle(PtA.ElemID); // TODO: would be faster to have a PlanarWalk call that takes a start vertex ID !
 				FVector3d WalkPlaneNormal = BaseFaceNormals[Seg.BaseTID].Cross(PtB.Pos - PtA.Pos);
-				if (ensure(Normalize(WalkPlaneNormal) > 0))
+				if (Normalize(WalkPlaneNormal) == 0)
 				{
-					bool bWalkSuccess = SurfacePath.AddViaPlanarWalk(StartTID, PtA.ElemID,
-						Mesh->GetVertex(PtA.ElemID), -1, PtB.ElemID,
-						Mesh->GetVertex(PtB.ElemID), WalkPlaneNormal, nullptr /*TODO: transform fn goes here?*/, false, FMathd::ZeroTolerance, SnapToleranceSq, .001);
-					if (!bWalkSuccess)
+					if (FVector3d::DistSquared(PtA.Pos, PtB.Pos) > SnapToleranceSq)
 					{
-						bSuccess = false;
+						// path points are separated, expect degeneracy to come from colinear triangle vertices
+						// this implies vertices are spread along the original edges, which would already be connected
+						// so we shouldn't need to do any additional work to connect things
+						continue;
 					}
-					else
+					// Path points are not separated; we may need to connect across a (likely degenerate) triangle
+					// Use a walk normal that can separate the triangle vertices (even if the triangle's collapsed to a line segment)
+					WalkPlaneNormal = GetDegenTriangleEdgeDirection(Mesh, StartTID);
+					if (!ensure(Normalize(WalkPlaneNormal) > 0))
+					{
+						// there was no non-degenerate edge; triangle is a point, nothing to walk here
+						continue;
+					}
+				}
+				bool bWalkSuccess = SurfacePath.AddViaPlanarWalk(StartTID, PtA.ElemID,
+					Mesh->GetVertex(PtA.ElemID), -1, PtB.ElemID,
+					Mesh->GetVertex(PtB.ElemID), WalkPlaneNormal, nullptr /*TODO: transform fn goes here?*/, false, FMathd::ZeroTolerance, SnapToleranceSq, .001);
+				if (!bWalkSuccess)
+				{
+					bSuccess = false;
+				}
+				else
+				{
+					EmbeddedPath.Reset();
+					if (SurfacePath.EmbedSimplePath(false, EmbeddedPath, false, SnapToleranceSq))
 					{
-						EmbeddedPath.Reset();
-						if (SurfacePath.EmbedSimplePath(false, EmbeddedPath, false, SnapToleranceSq))
+						ensure(EmbeddedPath.Num() > 0 && EmbeddedPath[0] == PtA.ElemID);
+						if (VertexChains)
 						{
-							ensure(EmbeddedPath.Num() > 0 && EmbeddedPath[0] == PtA.ElemID);
-							if (VertexChains)
+							if (SegmentToChain)
 							{
-								if (SegmentToChain)
-								{
-									(*SegmentToChain)[SegIdx] = VertexChains->Num();
-								}
-								VertexChains->Add(EmbeddedPath.Num());
-								VertexChains->Append(EmbeddedPath);
+								(*SegmentToChain)[SegIdx] = VertexChains->Num();
 							}
-						}
-						else
-						{
-							bSuccess = false;
+							VertexChains->Add(EmbeddedPath.Num());
+							VertexChains->Append(EmbeddedPath);
 						}
 					}
-				}
-				else
-				{
-					bSuccess = false;
+					else
+					{
+						bSuccess = false;
+					}
 				}
 			}
 			
@@ -470,6 +503,27 @@ namespace MeshCut
 			return BestIdx;
 		}
 
+		int OnEdgeWithSkip(const FTriangle3d& Tri, const FVector3d& V, int SkipIdx)
+		{
+			double BestDSq = SnapToleranceSq;
+			int BestIdx = -1;
+			for (int Idx = 0; Idx < 3; Idx++)
+			{
+				if (Idx == SkipIdx)
+				{
+					continue;
+				}
+				FSegment3d Seg{ Tri.V[Idx], Tri.V[(Idx + 1) % 3] };
+				double DSq = Seg.DistanceSquared(V);
+				if (DSq < BestDSq)
+				{
+					BestDSq = DSq;
+					BestIdx = Idx;
+				}
+			}
+			return BestIdx;
+		}
+
 		int ClosestEdge(FIndex2i EIDs, const FVector3d& Pos)
 		{
 			double BestIdx = -1;
@@ -528,9 +582,8 @@ bool FMeshSelfCut::Cut(const MeshIntersection::FIntersectionsQueryResult& Inters
 	ResetOutputs();
 
 	MeshCut::FCutWorkingInfo WorkingInfo(Mesh, SnapTolerance);
-	TMap<FVector3d, int> SegVtxMap;
-	WorkingInfo.AddSegments(Intersections, SegVtxMap, 0);
-	WorkingInfo.AddSegments(Intersections, SegVtxMap, 1);
+	WorkingInfo.AddSegments(Intersections, 0);
+	WorkingInfo.AddSegments(Intersections, 1);
 	WorkingInfo.InsertFaceVertices();
 	WorkingInfo.InsertEdgeVertices();
 	return WorkingInfo.ConnectEdges(bTrackInsertedVertices ? &VertexChains : nullptr);
@@ -549,8 +602,7 @@ bool FMeshMeshCut::Cut(const MeshIntersection::FIntersectionsQueryResult& Inters
 	for (int MeshIdx = 0; MeshIdx < MeshesToProcess; MeshIdx++)
 	{
 		MeshCut::FCutWorkingInfo WorkingInfo(Mesh[MeshIdx], SnapTolerance);
-		TMap<FVector3d, int> SegVtxMap;
-		WorkingInfo.AddSegments(Intersections, SegVtxMap, MeshIdx); // add intersection segments
+		WorkingInfo.AddSegments(Intersections, MeshIdx); // add intersection segments
 		WorkingInfo.InsertFaceVertices(); // insert vertices for intersection segments w/ endpoints on faces
 		WorkingInfo.InsertEdgeVertices(); // insert vertices for intersection segments w/ endpoints on edges
 
diff --git a/Engine/Source/Runtime/GeometryCore/Public/Intersection/IntrSegment2Triangle2.h b/Engine/Source/Runtime/GeometryCore/Public/Intersection/IntrSegment2Triangle2.h
--- a/Engine/Source/Runtime/GeometryCore/Public/Intersection/IntrSegment2Triangle2.h
+++ b/Engine/Source/Runtime/GeometryCore/Public/Intersection/IntrSegment2Triangle2.h
@@ -18,7 +18,8 @@ namespace Geometry
 {
 
 /**
- * Compute intersection between 2D segment and 2D triangle
+ * Compute intersection between 2D segment and 2D triangle.
+ * Note that if Segment.Extent is zero, this will test if the Segment Center is in the triangle.
  */
 template <typename Real>
 class TIntrSegment2Triangle2
@@ -85,14 +86,82 @@ public:
 			return (Result == EIntersectionResult::Intersects);
 		}
 
-		// if either line direction is not a normalized vector, 
-		//   results are garbage, so fail query
-		if (IsNormalized(Segment.Direction) == false)
+		// if segment extent is zero, segment is just a point, so just check the center vs the tri
+		// (Note this is also the case where Segment.Direction is the zero vector)
+		if (Segment.Extent == (Real)0)
 		{
-			Type = EIntersectionType::Empty;
-			Result = EIntersectionResult::InvalidQuery;
-			return false;
+			int pos = 0, neg = 0;
+			for (int TriPrev = 2, TriIdx = 0; TriIdx < 3; TriPrev = TriIdx++)
+			{
+				FVector2<Real> ToPt = Segment.Center - Triangle.V[TriIdx];
+				FVector2<Real> EdgePerp = PerpCW(Triangle.V[TriIdx] - Triangle.V[TriPrev]);
+				Real EdgeLen = Normalize(EdgePerp);
+				if (EdgeLen == 0) // triangle is just a line segment; try edgeperp = one of the other edges
+				{
+					FVector2<Real> OtherV = Triangle.V[(TriIdx + 1) % 3];
+					EdgePerp = OtherV - Triangle.V[TriIdx];
+					EdgeLen = Normalize(EdgePerp);
+					if (EdgeLen == 0) // triangle is just a point; go by distance between vertex and center
+					{
+						if (Triangle.V[0].DistanceSquared(Segment.Center) <= Tolerance*Tolerance)
+						{
+							pos = neg = 0; // closer than tolerance; accept collision
+						}
+						else
+						{
+							pos = neg = 1;
+						}
+						break;
+					}
+					else // valid edge, test from the other side of the degenerate tri as well
+					{
+						FVector2<Real> ToPtFromOther = Segment.Center - OtherV;
+						FVector2<Real> BackwardsEdgePerp = -EdgePerp;
+						Real OtherSideSign = BackwardsEdgePerp.Dot(ToPtFromOther);
+						if (OtherSideSign < -Tolerance)
+						{
+							neg++;
+						}
+						else if (OtherSideSign > Tolerance)
+						{
+							pos++;
+						}
+					}
+				}
+				
+				Real SideSign = EdgePerp.Dot(ToPt);
+				if (SideSign < -Tolerance)
+				{
+					neg++;
+				}
+				else if (SideSign > Tolerance)
+				{
+					pos++;
+				}
+			}
+			if (pos == 0 || neg == 0)
+			{
+				// note: intersection is a single point but conceptually a segment intersection
+				//  -- the whole (degenerate) segment intersects the triangle
+				Type = EIntersectionType::Segment;
+				Result = EIntersectionResult::Intersects;
+				Quantity = 2;
+				Point0 = Segment.Center;
+				Param0 = 0;
+				Point1 = Segment.Center;
+				Param1 = 0;
+				return true;
+			}
+			else
+			{
+				Type = EIntersectionType::Empty;
+				Result = EIntersectionResult::NoIntersection;
+				return false;
+			}
 		}
+		
+		// FSegment2d Direction should always be normalized unless Extent is 0 (handled above)
+		checkSlow(IsNormalized(Segment.Direction));
 
 		FVector3<Real> dist;
 		FVector3i sign;
diff --git a/Engine/Source/Runtime/GeometryCore/Public/Intersection/IntrTriangle3Triangle3.h b/Engine/Source/Runtime/GeometryCore/Public/Intersection/IntrTriangle3Triangle3.h
--- a/Engine/Source/Runtime/GeometryCore/Public/Intersection/IntrTriangle3Triangle3.h
+++ b/Engine/Source/Runtime/GeometryCore/Public/Intersection/IntrTriangle3Triangle3.h
@@ -153,6 +153,35 @@ public:
 		// Get the plane of Triangle0.
 		TPlane3<Real> Plane0(Triangle0.V[0], Triangle0.V[1], Triangle0.V[2]);
 
+		if (Plane0.Normal == FVector3<Real>::Zero())
+		{
+			// This function was written assuming triangle 0 has a valid normal;
+			//  if it's degenerate, try swapping the triangles
+			TPlane3<Real> Plane1(Triangle1.V[0], Triangle1.V[1], Triangle1.V[2]);
+			if (Plane1.Normal != FVector3<Real>::Zero())
+			{
+				// Find the intersection with the triangles swapped
+				TIntrTriangle3Triangle3<Real> SwappedTrisIntr(Triangle1, Triangle0);
+				SwappedTrisIntr.Tolerance = Tolerance;
+				SwappedTrisIntr.bReportCoplanarIntersection = bReportCoplanarIntersection;
+				bool bRes = SwappedTrisIntr.Find();
+				// copy results back
+				Result = SwappedTrisIntr.Result;
+				Type = SwappedTrisIntr.Type;
+				Quantity = SwappedTrisIntr.Quantity;
+				for (int PtIdx = 0; PtIdx < 6; PtIdx++)
+				{
+					Points[PtIdx] = SwappedTrisIntr.Points[PtIdx];
+				}
+				return bRes;
+			}
+			else
+			{
+				// both triangles degenerate; we don't handle this
+				return false;
+			}
+		}
+
 		// Compute the signed distances of Triangle1 vertices to Plane0.  Use
 		// an epsilon-thick plane test.
 		int pos1, neg1, zero1;
@@ -544,11 +573,11 @@ public:
 			for (i = 0; i < 3; ++i)
 			{
 				projTri.V[i] = GetYZ(triangle.V[i]);
-				projEnd0.X = end0.Y;
-				projEnd0.Y = end0.Z;
-				projEnd1.X = end1.Y;
-				projEnd1.Y = end1.Z;
 			}
+			projEnd0.X = end0.Y;
+			projEnd0.Y = end0.Z;
+			projEnd1.X = end1.Y;
+			projEnd1.Y = end1.Z;
 		}
 		else if (maxNormal == 1)
 		{
@@ -556,11 +585,11 @@ public:
 			for (i = 0; i < 3; ++i)
 			{
 				projTri.V[i] = GetXZ(triangle.V[i]);
-				projEnd0.X = end0.X;
-				projEnd0.Y = end0.Z;
-				projEnd1.X = end1.X;
-				projEnd1.Y = end1.Z;
 			}
+			projEnd0.X = end0.X;
+			projEnd0.Y = end0.Z;
+			projEnd1.X = end1.X;
+			projEnd1.Y = end1.Z;
 		}
 		else
 		{
@@ -568,11 +597,11 @@ public:
 			for (i = 0; i < 3; ++i)
 			{
 				projTri.V[i] = GetXY(triangle.V[i]);
-				projEnd0.X = end0.X;
-				projEnd0.Y = end0.Y;
-				projEnd1.X = end1.X;
-				projEnd1.Y = end1.Y;
 			}
+			projEnd0.X = end0.X;
+			projEnd0.Y = end0.Y;
+			projEnd1.X = end1.X;
+			projEnd1.Y = end1.Y;
 		}
 
 		TSegment2<Real> projSeg(projEnd0, projEnd1);
