Allow construction of a Laplacian mesh deformer on a mesh type that has only edges and vertices (no triangles, polygons, volumes, etc.) A uniform weighted Laplacian must be used in this case.

Several internal constructors/functions are changed to operate on generic (template) mesh types. The public-facing ConstructUniformConstrainedMeshDeformer() in ConstrainedMeshDeformer.h uses an abstract mesh class instead, to avoid exposing those internal headers.

#rnx

diff --git a/Engine/Plugins/Experimental/GeometryProcessing/Source/DynamicMesh/Private/Solvers/ConstrainedMeshDeformer.cpp b/Engine/Plugins/Experimental/GeometryProcessing/Source/DynamicMesh/Private/Solvers/ConstrainedMeshDeformer.cpp
--- a/Engine/Plugins/Experimental/GeometryProcessing/Source/DynamicMesh/Private/Solvers/ConstrainedMeshDeformer.cpp
+++ b/Engine/Plugins/Experimental/GeometryProcessing/Source/DynamicMesh/Private/Solvers/ConstrainedMeshDeformer.cpp
@@ -10,9 +10,14 @@ TUniquePtr<UE::Solvers::IConstrainedMeshSolver> UE::MeshDeformation::ConstructCo
 	return Deformer;
 }
 
+TUniquePtr<UE::Solvers::IConstrainedMeshSolver> UE::MeshDeformation::ConstructUniformConstrainedMeshDeformer(const UE::MeshDeformation::FDynamicGraph3d& Graph)
+{
+	TUniquePtr<UE::Solvers::IConstrainedMeshSolver> Deformer(new FConstrainedMeshDeformer(Graph));
+	return Deformer;
+}
 
 TUniquePtr<UE::Solvers::IConstrainedLaplacianMeshSolver> UE::MeshDeformation::ConstructSoftMeshDeformer(const FDynamicMesh3& DynamicMesh)
 {
 	TUniquePtr<UE::Solvers::IConstrainedLaplacianMeshSolver> Deformer(new FSoftMeshDeformer(DynamicMesh));
 	return Deformer;
-}
\ No newline at end of file
+}
diff --git a/Engine/Plugins/Experimental/GeometryProcessing/Source/DynamicMesh/Private/Solvers/Internal/ConstrainedMeshDeformationSolver.cpp b/Engine/Plugins/Experimental/GeometryProcessing/Source/DynamicMesh/Private/Solvers/Internal/ConstrainedMeshDeformationSolver.cpp
--- a/Engine/Plugins/Experimental/GeometryProcessing/Source/DynamicMesh/Private/Solvers/Internal/ConstrainedMeshDeformationSolver.cpp
+++ b/Engine/Plugins/Experimental/GeometryProcessing/Source/DynamicMesh/Private/Solvers/Internal/ConstrainedMeshDeformationSolver.cpp
@@ -66,53 +66,11 @@ FConstrainedMeshDeformationSolver::~FConstrainedMeshDeformationSolver()
 FConstrainedMeshDeformationSolver::FConstrainedMeshDeformationSolver(const FDynamicMesh3& DynamicMesh, const ELaplacianWeightScheme Scheme, const EMatrixSolverType MatrixSolverType)
 	: VertexCount(DynamicMesh.VertexCount())
 {
-
 	FSparseMatrixD LaplacianInternal;
 	FSparseMatrixD LaplacianBoundary;
 	ConstructLaplacian(Scheme, DynamicMesh, VtxLinearization, LaplacianInternal, LaplacianBoundary);
 
-	// Copy the original boundary vertex locations
-	const int32 BoundaryVertexCount = VtxLinearization.NumBoundaryVerts();
-
-	// Number of vertices in the interior of the mesh
-	InternalVertexCount = VertexCount - BoundaryVertexCount;
-
-	// Copy the original boundary vertex locations
-	{
-		const auto ToVertId = VtxLinearization.ToId();
-		BoundaryPositions.SetZero(BoundaryVertexCount);
-		for (int32 i = 0; i < BoundaryVertexCount; ++i)
-		{
-			const int32 VtxId = ToVertId[i + InternalVertexCount];
-			const FVector3d Pos = DynamicMesh.GetVertex(VtxId);
-			BoundaryPositions.SetXYZ(i, Pos);
-		}
-	}
-
-	checkSlow(LaplacianInternal.rows() == LaplacianInternal.cols());
-
-	TUniquePtr<FSparseMatrixD> LTLPtr(new FSparseMatrixD(LaplacianInternal.rows(), LaplacianInternal.cols()));
-	FSparseMatrixD& LTLMatrix = *(LTLPtr);
-
-	bool bIsLaplacianSymmetric = (Scheme == ELaplacianWeightScheme::Valence || Scheme == ELaplacianWeightScheme::Uniform);
-
-	if (bIsLaplacianSymmetric)
-	{
-		// Laplacian is symmetric, i.e. equal to its transpose
-		LTLMatrix = LaplacianInternal * LaplacianInternal;
-		BoundaryOperator = -1. * LaplacianInternal * LaplacianBoundary;
-
-		ConstrainedSolver.Reset(new FConstrainedSolver(LTLPtr, MatrixSolverType));
-	}
-	else
-	{
-		// the laplacian 
-		LTLMatrix = LaplacianInternal.transpose() * LaplacianInternal;
-		BoundaryOperator = -1. * LaplacianInternal.transpose() * LaplacianBoundary;
-
-		ConstrainedSolver.Reset(new FConstrainedSolver(LTLPtr, MatrixSolverType));
-	}
-
+	Init(DynamicMesh, Scheme, MatrixSolverType, LaplacianInternal, LaplacianBoundary);
 }
 
 void FConstrainedMeshDeformationSolver::AddConstraint(const int32 VtxId, const double Weight, const FVector3d& Pos, const bool bPostFix)
@@ -288,19 +246,6 @@ bool FConstrainedMeshDeformationSolver::CopyBoundaryPositions(TArray<FVector3d>&
 }
 
 
-void FConstrainedMeshDeformationSolver::ExtractInteriorVertexPositions(const FDynamicMesh3& DynamicMesh, FSOAPositions& VertexPositions) const
-{
-	VertexPositions.SetZero(InternalVertexCount);
-
-	const auto& ToVtxId = VtxLinearization.ToId();
-
-	for (int32 i = 0; i < InternalVertexCount; ++i)
-	{
-		const int32 VtxId = ToVtxId[i];
-		const FVector3d& Pos = DynamicMesh.GetVertex(VtxId);
-		VertexPositions.SetXYZ(i, Pos);
-	}
-}
 
 void FConstrainedMeshDeformationSolver::UpdateWithPostFixConstraints(FSOAPositions& PositionVector) const
 {
diff --git a/Engine/Plugins/Experimental/GeometryProcessing/Source/DynamicMesh/Private/Solvers/Internal/ConstrainedMeshDeformationSolver.h b/Engine/Plugins/Experimental/GeometryProcessing/Source/DynamicMesh/Private/Solvers/Internal/ConstrainedMeshDeformationSolver.h
--- a/Engine/Plugins/Experimental/GeometryProcessing/Source/DynamicMesh/Private/Solvers/Internal/ConstrainedMeshDeformationSolver.h
+++ b/Engine/Plugins/Experimental/GeometryProcessing/Source/DynamicMesh/Private/Solvers/Internal/ConstrainedMeshDeformationSolver.h
@@ -4,10 +4,12 @@
 
 #include "CoreMinimal.h"
 #include "DynamicMesh/DynamicMesh3.h"
+#include "LaplacianOperators.h"
 #include "Solvers/MatrixInterfaces.h"
 #include "Solvers/MeshLaplacian.h"
 #include "Solvers/MeshLinearization.h"
 #include "Solvers/ConstrainedMeshSolver.h"
+#include "ConstrainedPoissonSolver.h"
 
 #include "MatrixBase.h"
 
@@ -17,7 +19,6 @@ namespace UE
 namespace Geometry
 {
 
-class FConstrainedSolver;
 
 /**
  * FConstrainedMeshDeformationSolver is an implmentation of IConstrainedMeshSolver that solves
@@ -33,10 +34,13 @@ class DYNAMICMESH_API FConstrainedMeshDeformationSolver : public UE::Solvers::IC
 public:
 	typedef UE::Solvers::FPositionConstraint  FConstraintPosition;
 
-
 	FConstrainedMeshDeformationSolver(const FDynamicMesh3& DynamicMesh, const ELaplacianWeightScheme Scheme, const EMatrixSolverType MatrixSolverType);
 	virtual ~FConstrainedMeshDeformationSolver();
 
+	// Create the solver using a generic mesh type. Assume uniform weight scheme.
+	template<typename MeshType>
+	FConstrainedMeshDeformationSolver(const MeshType& Mesh, const EMatrixSolverType MatrixSolverType);
+
 	// Add constraint associated with given vertex id.  Boundary vertices will be ignored
 	void AddConstraint(const int32 VtxId, const double Weight, const FVector3d& Pos, const bool bPostFix) override;
 
@@ -74,8 +78,68 @@ public:
 
 protected:
 
+	// Call this from constructors after building matrices
+	template<typename MeshType>
+	void Init(const MeshType& Mesh, const ELaplacianWeightScheme Scheme, const EMatrixSolverType MatrixSolverType, const FSparseMatrixD& LaplacianInternal, const FSparseMatrixD& LaplacianBoundary)
+	{
+		// Copy the original boundary vertex locations
+		const int32 BoundaryVertexCount = VtxLinearization.NumBoundaryVerts();
+
+		// Number of vertices in the interior of the mesh
+		InternalVertexCount = VertexCount - BoundaryVertexCount;
+
+		// Copy the original boundary vertex locations
+		{
+			const auto ToVertId = VtxLinearization.ToId();
+			BoundaryPositions.SetZero(BoundaryVertexCount);
+			for (int32 i = 0; i < BoundaryVertexCount; ++i)
+			{
+				const int32 VtxId = ToVertId[i + InternalVertexCount];
+				const FVector3d Pos = Mesh.GetVertex(VtxId);
+				BoundaryPositions.SetXYZ(i, Pos);
+			}
+		}
+
+		checkSlow(LaplacianInternal.rows() == LaplacianInternal.cols());
+
+		TUniquePtr<FSparseMatrixD> LTLPtr(new FSparseMatrixD(LaplacianInternal.rows(), LaplacianInternal.cols()));
+		FSparseMatrixD& LTLMatrix = *(LTLPtr);
+
+		bool bIsLaplacianSymmetric = (Scheme == ELaplacianWeightScheme::Valence || Scheme == ELaplacianWeightScheme::Uniform);
+
+		if (bIsLaplacianSymmetric)
+		{
+			// Laplacian is symmetric, i.e. equal to its transpose
+			LTLMatrix = LaplacianInternal * LaplacianInternal;
+			BoundaryOperator = -1. * LaplacianInternal * LaplacianBoundary;
+
+			ConstrainedSolver.Reset(new FConstrainedSolver(LTLPtr, MatrixSolverType));
+		}
+		else
+		{
+			// the laplacian 
+			LTLMatrix = LaplacianInternal.transpose() * LaplacianInternal;
+			BoundaryOperator = -1. * LaplacianInternal.transpose() * LaplacianBoundary;
+
+			ConstrainedSolver.Reset(new FConstrainedSolver(LTLPtr, MatrixSolverType));
+		}
+	}
+	
+	template<typename MeshType>
+	void ExtractInteriorVertexPositions(const MeshType& Mesh, FSOAPositions& VertexPositions) const
+	{
+		VertexPositions.SetZero(InternalVertexCount);
+
+		const auto& ToVtxId = VtxLinearization.ToId();
+
+		for (int32 i = 0; i < InternalVertexCount; ++i)
+		{
+			const int32 VtxId = ToVtxId[i];
+			const FVector3d& Pos = Mesh.GetVertex(VtxId);
+			VertexPositions.SetXYZ(i, Pos);
+		}
+	}
 
-	void ExtractInteriorVertexPositions(const FDynamicMesh3& DynamicMesh, FSOAPositions& Positions) const;
 
 	//void UpdateMeshWithConstraints();
 	// Respect any bPostFix constraints by moving those vertices to position defined by said constraint.
@@ -122,9 +186,16 @@ protected:
 
 
 
+template<typename MeshType>
+FConstrainedMeshDeformationSolver::FConstrainedMeshDeformationSolver(const MeshType& Mesh, const EMatrixSolverType MatrixSolverType)
+	: VertexCount(Mesh.VertexCount())
+{
+	FSparseMatrixD LaplacianInternal;
+	FSparseMatrixD LaplacianBoundary;
+	ConstructUniformLaplacian(Mesh, VtxLinearization, LaplacianInternal, LaplacianBoundary);
 
-
-
+	Init(Mesh, ELaplacianWeightScheme::Uniform, MatrixSolverType, LaplacianInternal, LaplacianBoundary);
+}
 
 /**
  * FSoftMeshDeformationSolver is an implmentation of IConstrainedLaplacianMeshSolver that solves
diff --git a/Engine/Plugins/Experimental/GeometryProcessing/Source/DynamicMesh/Private/Solvers/Internal/ConstrainedMeshDeformers.h b/Engine/Plugins/Experimental/GeometryProcessing/Source/DynamicMesh/Private/Solvers/Internal/ConstrainedMeshDeformers.h
--- a/Engine/Plugins/Experimental/GeometryProcessing/Source/DynamicMesh/Private/Solvers/Internal/ConstrainedMeshDeformers.h
+++ b/Engine/Plugins/Experimental/GeometryProcessing/Source/DynamicMesh/Private/Solvers/Internal/ConstrainedMeshDeformers.h
@@ -25,6 +25,10 @@ public:
 	FConstrainedMeshDeformer(const FDynamicMesh3& DynamicMesh, const ELaplacianWeightScheme LaplacianType);
 	~FConstrainedMeshDeformer() override {}
 
+	// Generic mesh type, assume uniform weight scheme
+	template<typename MeshT>
+	FConstrainedMeshDeformer(const MeshT& Mesh);
+	
 	bool Deform(TArray<FVector3d>& PositionBuffer) override;
 
 private:
@@ -33,6 +37,33 @@ private:
 	FSOAPositions OriginalInteriorPositions;
 };
 
+// Generic mesh type, assume uniform weight scheme
+template<typename MeshT>
+FConstrainedMeshDeformer::FConstrainedMeshDeformer(const MeshT& Mesh) :
+	FConstrainedMeshDeformationSolver(Mesh, EMatrixSolverType::LU),
+	LaplacianVectors(FConstrainedMeshDeformationSolver::InternalVertexCount)
+{
+	// The current vertex positions 
+	// Note: the OriginalInteriorPositions are being stored as member data 
+	// for use if the solver is iterative.
+	// FSOAPositions OriginalInteriorPositions; 
+	ExtractInteriorVertexPositions(Mesh, OriginalInteriorPositions);
+
+	// The biharmonic part of the constrained solver
+	//   Biharmonic := Laplacian^{T} * Laplacian
+
+	const auto& Biharmonic = ConstrainedSolver->Biharmonic();
+
+	// Compute the Laplacian Vectors
+	//    := Biharmonic * VertexPostion
+	// In the case of the cotangent laplacian this can be identified as the mean curvature * normal.
+	checkSlow(LaplacianVectors.Num() == OriginalInteriorPositions.Num());
+
+	for (int32 i = 0; i < 3; ++i)
+	{
+		LaplacianVectors.Array(i) = Biharmonic * OriginalInteriorPositions.Array(i);
+	}
+}
 
 
 /**
diff --git a/Engine/Plugins/Experimental/GeometryProcessing/Source/DynamicMesh/Private/Solvers/Internal/LaplacianOperators.cpp b/Engine/Plugins/Experimental/GeometryProcessing/Source/DynamicMesh/Private/Solvers/Internal/LaplacianOperators.cpp
--- a/Engine/Plugins/Experimental/GeometryProcessing/Source/DynamicMesh/Private/Solvers/Internal/LaplacianOperators.cpp
+++ b/Engine/Plugins/Experimental/GeometryProcessing/Source/DynamicMesh/Private/Solvers/Internal/LaplacianOperators.cpp
@@ -3,7 +3,7 @@
 #include "LaplacianOperators.h"
 
 #include "Solvers/MatrixInterfaces.h"
-#include "Solvers/LaplacianMatrixAssembly.h"
+
 
 #include <cmath> // double version of sqrt
 #include <vector> // used by eigen to initialize sparse matrix
@@ -29,22 +29,6 @@ PRAGMA_DEFAULT_VISIBILITY_END
 using namespace UE::Geometry;
 
 
-void ConstructUniformLaplacian(const FDynamicMesh3& DynamicMesh, FVertexLinearization& VertexMap, FSparseMatrixD& LaplacianInterior, FSparseMatrixD& LaplacianBoundary)
-{
-	// Sync the mapping between the mesh vertex ids and their offsets in a nominal linear array.
-	VertexMap.Reset(DynamicMesh);
-	const int32 NumVerts = VertexMap.NumVerts();
-	const int32 NumBoundaryVerts = VertexMap.NumBoundaryVerts();
-	const int32 NumInteriorVerts = NumVerts - NumBoundaryVerts;
-
-	FEigenSparseMatrixAssembler Interior(NumInteriorVerts, NumInteriorVerts);
-	FEigenSparseMatrixAssembler Boundary(NumInteriorVerts, NumBoundaryVerts);
-	UE::MeshDeformation::ConstructUniformLaplacian<double>(DynamicMesh, VertexMap, Interior, Boundary);
-	Interior.ExtractResult(LaplacianInterior);
-	Boundary.ExtractResult(LaplacianBoundary);
-}
-
-
 void ConstructUmbrellaLaplacian(const FDynamicMesh3& DynamicMesh, FVertexLinearization& VertexMap, FSparseMatrixD& LaplacianInterior, FSparseMatrixD& LaplacianBoundary)
 {
 	// Sync the mapping between the mesh vertex ids and their offsets in a nominal linear array.
@@ -226,7 +210,7 @@ void ConstructLaplacian(const ELaplacianWeightScheme Scheme, const FDynamicMesh3
 //
 
 
-static void ExtractBoundaryVerts(const FVertexLinearization& VertexMap, TArray<int32>& BoundaryVerts)
+void ExtractBoundaryVerts(const FVertexLinearization& VertexMap, TArray<int32>& BoundaryVerts)
 {
 	int32 NumBoundaryVerts = VertexMap.NumBoundaryVerts();
 	int32 NumInternalVerts = VertexMap.NumVerts() - NumBoundaryVerts;
@@ -240,23 +224,6 @@ static void ExtractBoundaryVerts(const FVertexLinearization& VertexMap, TArray<i
 	}
 }
 
-TUniquePtr<FSparseMatrixD> ConstructUniformLaplacian(const FDynamicMesh3& DynamicMesh, FVertexLinearization& VertexMap, TArray<int32>* BoundaryVerts)
-{
-	TUniquePtr<FSparseMatrixD> LaplacianMatrix(new FSparseMatrixD);
-	FSparseMatrixD BoundaryMatrix;
-
-	ConstructUniformLaplacian(DynamicMesh, VertexMap, *LaplacianMatrix, BoundaryMatrix);
-
-	if (BoundaryVerts)
-	{
-		ExtractBoundaryVerts(VertexMap, *BoundaryVerts);
-	}
-
-	return LaplacianMatrix;
-}
-
-
-
 TUniquePtr<FSparseMatrixD> ConstructUmbrellaLaplacian(const FDynamicMesh3& DynamicMesh, FVertexLinearization& VertexMap, TArray<int32>* BoundaryVerts)
 {
 	TUniquePtr<FSparseMatrixD> LaplacianMatrix(new FSparseMatrixD);
diff --git a/Engine/Plugins/Experimental/GeometryProcessing/Source/DynamicMesh/Private/Solvers/Internal/LaplacianOperators.h b/Engine/Plugins/Experimental/GeometryProcessing/Source/DynamicMesh/Private/Solvers/Internal/LaplacianOperators.h
--- a/Engine/Plugins/Experimental/GeometryProcessing/Source/DynamicMesh/Private/Solvers/Internal/LaplacianOperators.h
+++ b/Engine/Plugins/Experimental/GeometryProcessing/Source/DynamicMesh/Private/Solvers/Internal/LaplacianOperators.h
@@ -6,7 +6,7 @@
 #include "DynamicMesh/DynamicMesh3.h"
 #include "Solvers/MeshLaplacian.h"
 #include "Solvers/MeshLinearization.h"
-
+#include "Solvers/LaplacianMatrixAssembly.h"
 #include "FSparseMatrixD.h"
 
 // this is a private header
@@ -55,10 +55,6 @@ TUniquePtr<FSparseMatrixD> ConstructLaplacian(const ELaplacianWeightScheme Schem
 	FVertexLinearization& VertexMap,
 	TArray<int32>* BoundaryVerts = NULL);
 
-TUniquePtr<FSparseMatrixD> ConstructUniformLaplacian(const FDynamicMesh3& DynamicMesh,
-	FVertexLinearization& VertexMap,
-	TArray<int32>* BoundaryVerts = NULL);
-
 TUniquePtr<FSparseMatrixD> ConstructUmbrellaLaplacian(const FDynamicMesh3& DynamicMesh,
 	FVertexLinearization& VertexMap,
 	TArray<int32>* BoundaryVerts = NULL);
@@ -80,3 +76,39 @@ TUniquePtr<FSparseMatrixD> ConstructCotangentLaplacian(const FDynamicMesh3& Dyna
 TUniquePtr<FSparseMatrixD> ConstructMeanValueWeightLaplacian(const FDynamicMesh3& DynamicMesh,
 	FVertexLinearization& VertexMap,
 	TArray<int32>* BoundaryVerts = NULL);
+
+void ExtractBoundaryVerts(const FVertexLinearization& VertexMap, TArray<int32>& BoundaryVerts);
+
+
+template<typename MeshType>
+void ConstructUniformLaplacian(const MeshType& Mesh, FVertexLinearization& VertexMap, FSparseMatrixD& LaplacianInterior, FSparseMatrixD& LaplacianBoundary)
+{
+	// Sync the mapping between the mesh vertex ids and their offsets in a nominal linear array.
+	VertexMap.Reset(Mesh);
+	const int32 NumVerts = VertexMap.NumVerts();
+	const int32 NumBoundaryVerts = VertexMap.NumBoundaryVerts();
+	const int32 NumInteriorVerts = NumVerts - NumBoundaryVerts;
+
+	FEigenSparseMatrixAssembler Interior(NumInteriorVerts, NumInteriorVerts);
+	FEigenSparseMatrixAssembler Boundary(NumInteriorVerts, NumBoundaryVerts);
+	UE::MeshDeformation::ConstructUniformLaplacian<double>(Mesh, VertexMap, Interior, Boundary);
+	Interior.ExtractResult(LaplacianInterior);
+	Boundary.ExtractResult(LaplacianBoundary);
+}
+
+template<typename MeshType>
+TUniquePtr<FSparseMatrixD> ConstructUniformLaplacian(const MeshType& Mesh, FVertexLinearization& VertexMap, TArray<int32>* BoundaryVerts)
+{
+	TUniquePtr<FSparseMatrixD> LaplacianMatrix(new FSparseMatrixD);
+	FSparseMatrixD BoundaryMatrix;
+
+	ConstructUniformLaplacian(Mesh, VertexMap, *LaplacianMatrix, BoundaryMatrix);
+
+	if (BoundaryVerts)
+	{
+		ExtractBoundaryVerts(VertexMap, *BoundaryVerts);
+	}
+
+	return LaplacianMatrix;
+}
+
diff --git a/Engine/Plugins/Experimental/GeometryProcessing/Source/DynamicMesh/Private/Solvers/PrecomputedMeshWeightData.cpp b/Engine/Plugins/Experimental/GeometryProcessing/Source/DynamicMesh/Private/Solvers/PrecomputedMeshWeightData.cpp
--- a/Engine/Plugins/Experimental/GeometryProcessing/Source/DynamicMesh/Private/Solvers/PrecomputedMeshWeightData.cpp
+++ b/Engine/Plugins/Experimental/GeometryProcessing/Source/DynamicMesh/Private/Solvers/PrecomputedMeshWeightData.cpp
@@ -5,35 +5,6 @@
 #include "ExplicitUseGeometryMathTypes.h"		// using UE::Geometry::(math types)
 using namespace UE::Geometry;
 
-// Utility to compute the number of elements in the sparse laplacian matrix
-int32 UE::MeshDeformation::ComputeNumMatrixElements(const FDynamicMesh3& DynamicMesh, const TArray<int32>& ToVtxId)
-{
-	const int32 NumVerts = ToVtxId.Num();
-	TArray<int32> OneRingSize;
-	{
-		OneRingSize.SetNumUninitialized(NumVerts);
-
-		for (int32 i = 0; i < NumVerts; ++i)
-		{
-			const int32 VertId = ToVtxId[i];
-			OneRingSize[i] = DynamicMesh.GetVtxEdgeCount(VertId);
-		}
-	}
-
-	// Compute the total number of entries in the sparse matrix
-	int32 NumMatrixEntries = 0;
-	{
-		for (int32 i = 0; i < NumVerts; ++i)
-		{
-			NumMatrixEntries += 1 + OneRingSize[i]; // myself plus my neighbors
-		}
-
-	}
-
-	return NumMatrixEntries;
-}
-
-
 
 UE::MeshDeformation::CotanTriangleData::CotanTriangleData(const CotanTriangleData& Other)
 {
diff --git a/Engine/Plugins/Experimental/GeometryProcessing/Source/DynamicMesh/Public/Solvers/ConstrainedMeshDeformer.h b/Engine/Plugins/Experimental/GeometryProcessing/Source/DynamicMesh/Public/Solvers/ConstrainedMeshDeformer.h
--- a/Engine/Plugins/Experimental/GeometryProcessing/Source/DynamicMesh/Public/Solvers/ConstrainedMeshDeformer.h
+++ b/Engine/Plugins/Experimental/GeometryProcessing/Source/DynamicMesh/Public/Solvers/ConstrainedMeshDeformer.h
@@ -7,6 +7,7 @@
 #include "Solvers/ConstrainedMeshSolver.h"
 #include "Solvers/MeshLaplacian.h"
 #include "DynamicMesh/DynamicMesh3.h"
+#include "Curve/DynamicGraph3.h"
 
 namespace UE
 {
@@ -58,6 +59,35 @@ namespace UE
 		TUniquePtr<UE::Solvers::IConstrainedMeshSolver> DYNAMICMESH_API ConstructConstrainedMeshDeformer(const ELaplacianWeightScheme WeightScheme, const FDynamicMesh3& DynamicMesh);
 
 
+		/**
+		 * Interface for a mesh with only edges and vertices
+		 */
+		//class IVertexEdgeGraph
+		//{
+		//public:
+		//	virtual const FVector3d& GetVertex(int) const = 0;
+		//	virtual int VertexCount() const = 0;
+		//	virtual int MaxVertexID() const = 0;
+		//	virtual FRefCountVector::IndexEnumerable VertexIndicesItr() const = 0;
+		//	virtual FSmallListSet::MappedValueEnumerable VtxVerticesItr(int VID) const = 0;
+		//	virtual int GetVtxEdgeCount(int VID) const = 0;
+
+		//	// Required for FVertexLinearization
+		//	virtual bool IsBoundaryVertex(int VID) const = 0;
+		//};
+
+		/**
+		 * Construct a Mesh Deformer object for the given generic mesh that uses Biharmonic Laplacian Mesh Deformation 
+		 * to solve for the deformed vertex positions.
+		 *
+		 * Similar to ConstructConstrainedMeshDeformer() however a uniform-weighted Laplacian is always used.
+		 * This allows for the solution of deformation problems on a more generic mesh type (triangles or faces are
+		 * not required.)
+		 * 
+		 */
+		TUniquePtr<UE::Solvers::IConstrainedMeshSolver> DYNAMICMESH_API ConstructUniformConstrainedMeshDeformer(const FDynamicGraph3d& Graph);
+
+
 		/**
 		 * Construct a Mesh Deformer object for the given mesh that uses Biharmonic Laplacian Mesh Deformation to solve
 		 * for the deformed vertex positions.
diff --git a/Engine/Plugins/Experimental/GeometryProcessing/Source/DynamicMesh/Public/Solvers/LaplacianMatrixAssembly.h b/Engine/Plugins/Experimental/GeometryProcessing/Source/DynamicMesh/Public/Solvers/LaplacianMatrixAssembly.h
--- a/Engine/Plugins/Experimental/GeometryProcessing/Source/DynamicMesh/Public/Solvers/LaplacianMatrixAssembly.h
+++ b/Engine/Plugins/Experimental/GeometryProcessing/Source/DynamicMesh/Public/Solvers/LaplacianMatrixAssembly.h
@@ -39,10 +39,11 @@ namespace UE
 		*
 		*   LaplacianInterior * Vector_InteriorVerts + LaplacianBoundary * Vector_BoundaryVerts = Full Laplacian applied to interior vertices.
 		*/
-		template<typename RealType>
-		void ConstructUniformLaplacian(const FDynamicMesh3& DynamicMesh, const FVertexLinearization& VertexMap,
-			UE::Solvers::TSparseMatrixAssembler<RealType>& LaplacianInterior, 
-			UE::Solvers::TSparseMatrixAssembler<RealType>& LaplacianBoundary);
+		template<typename RealType, typename MeshType>
+		void ConstructUniformLaplacian(const MeshType& Mesh,
+									   const FVertexLinearization& VertexMap,
+									   UE::Solvers::TSparseMatrixAssembler<RealType>& LaplacianInterior,
+									   UE::Solvers::TSparseMatrixAssembler<RealType>& LaplacianBoundary);
 
 		/**
 		* Construct a sparse matrix representation of an umbrella weighted Laplacian.
@@ -211,10 +212,11 @@ namespace UE
 
 
 
-template<typename RealType>
-void UE::MeshDeformation::ConstructUniformLaplacian(const FDynamicMesh3& DynamicMesh, const FVertexLinearization& VertexMap,
-	UE::Solvers::TSparseMatrixAssembler<RealType>& LaplacianInterior, 
-	UE::Solvers::TSparseMatrixAssembler<RealType>& LaplacianBoundary)
+template<typename RealType, typename MeshType>
+void UE::MeshDeformation::ConstructUniformLaplacian(const MeshType& Mesh, 
+													const FVertexLinearization& VertexMap,
+													UE::Solvers::TSparseMatrixAssembler<RealType>& LaplacianInterior,
+													UE::Solvers::TSparseMatrixAssembler<RealType>& LaplacianBoundary)
 {
 	//check(VertexMap_is_good)
 	const TArray<int32>& ToMeshV = VertexMap.ToId();
@@ -224,7 +226,7 @@ void UE::MeshDeformation::ConstructUniformLaplacian(const FDynamicMesh3& Dynamic
 	const int32 NumInteriorVerts = NumVerts - NumBoundaryVerts;
 
 	// pre-allocate space when possible
-	int32 NumMatrixEntries = ComputeNumMatrixElements(DynamicMesh, ToMeshV);
+	int32 NumMatrixEntries = ComputeNumMatrixElements(Mesh, ToMeshV);
 	LaplacianInterior.ReserveEntriesFunc(NumMatrixEntries);
 
 	// Construct Laplacian Matrix: loop over verts constructing the corresponding matrix row.
@@ -236,7 +238,7 @@ void UE::MeshDeformation::ConstructUniformLaplacian(const FDynamicMesh3& Dynamic
 
 		checkSlow(!DynamicMesh.IsBoundaryVertex(VertId));  // we should only be looping over the internal verts
 
-		for (int32 NeighborVertId : DynamicMesh.VtxVerticesItr(VertId))
+		for (int32 NeighborVertId : Mesh.VtxVerticesItr(VertId))
 		{
 			const int32 j = ToIndex[NeighborVertId];
 			RealType NeighborWeight = RealType(1);
diff --git a/Engine/Plugins/Experimental/GeometryProcessing/Source/DynamicMesh/Public/Solvers/MeshLinearization.h b/Engine/Plugins/Experimental/GeometryProcessing/Source/DynamicMesh/Public/Solvers/MeshLinearization.h
--- a/Engine/Plugins/Experimental/GeometryProcessing/Source/DynamicMesh/Public/Solvers/MeshLinearization.h
+++ b/Engine/Plugins/Experimental/GeometryProcessing/Source/DynamicMesh/Public/Solvers/MeshLinearization.h
@@ -23,16 +23,19 @@ class FVertexLinearization : public FElementLinearization
 {
 public:
 	FVertexLinearization() = default;
-	FVertexLinearization(const FDynamicMesh3& DynamicMesh)
+
+	template<typename MeshType>
+	FVertexLinearization(const MeshType& Mesh)
 	{
-		Reset(DynamicMesh);
+		Reset(Mesh);
 	}
 
-	void Reset(const FDynamicMesh3& DynamicMesh)
+	template<typename MeshType>
+	void Reset(const MeshType& Mesh)
 	{
 		Empty();
-		FElementLinearization::Populate(DynamicMesh.MaxVertexID(), DynamicMesh.VertexCount(), DynamicMesh.VertexIndicesItr());
-		RemapBoundaryVerts(DynamicMesh);
+		FElementLinearization::Populate(Mesh.MaxVertexID(), Mesh.VertexCount(), Mesh.VertexIndicesItr());
+		RemapBoundaryVerts(Mesh);
 	}
 
 	int32 NumVerts() const { return FElementLinearization::NumIds(); }
@@ -42,7 +45,8 @@ public:
 private:
 
 	// Moves the boundary verts to the end of the arrays and records the number of boundary verts
-	void RemapBoundaryVerts(const FDynamicMesh3& DynamicMesh)
+	template<typename MeshType>
+	void RemapBoundaryVerts(const MeshType& DynamicMesh)
 	{
 		int32 VertCount = NumVerts();
 
diff --git a/Engine/Plugins/Experimental/GeometryProcessing/Source/DynamicMesh/Public/Solvers/PrecomputedMeshWeightData.h b/Engine/Plugins/Experimental/GeometryProcessing/Source/DynamicMesh/Public/Solvers/PrecomputedMeshWeightData.h
--- a/Engine/Plugins/Experimental/GeometryProcessing/Source/DynamicMesh/Public/Solvers/PrecomputedMeshWeightData.h
+++ b/Engine/Plugins/Experimental/GeometryProcessing/Source/DynamicMesh/Public/Solvers/PrecomputedMeshWeightData.h
@@ -12,7 +12,34 @@ namespace UE
 		using namespace UE::Geometry;
 
 		// Utility to compute the number of elements in the sparse laplacian matrix
-		DYNAMICMESH_API int32 ComputeNumMatrixElements(const FDynamicMesh3& DynamicMesh, const TArray<int32>& ToVtxId);
+		template<typename MeshT>
+		DYNAMICMESH_API int32 ComputeNumMatrixElements(const MeshT& DynamicMesh, const TArray<int32>& ToVtxId)
+		{
+			const int32 NumVerts = ToVtxId.Num();
+			TArray<int32> OneRingSize;
+			{
+				OneRingSize.SetNumUninitialized(NumVerts);
+
+				for (int32 i = 0; i < NumVerts; ++i)
+				{
+					const int32 VertId = ToVtxId[i];
+					OneRingSize[i] = DynamicMesh.GetVtxEdgeCount(VertId);
+				}
+			}
+
+			// Compute the total number of entries in the sparse matrix
+			int32 NumMatrixEntries = 0;
+			{
+				for (int32 i = 0; i < NumVerts; ++i)
+				{
+					NumMatrixEntries += 1 + OneRingSize[i]; // myself plus my neighbors
+				}
+
+			}
+
+			return NumMatrixEntries;
+		}
+
 
 		/**
 		* The per-triangle data used in constructing the cotangent weighted laplacian.
diff --git a/Engine/Source/Runtime/GeometryCore/Public/Curve/DynamicGraph3.h b/Engine/Source/Runtime/GeometryCore/Public/Curve/DynamicGraph3.h
new file mode 100644
--- /dev/null
+++ b/Engine/Source/Runtime/GeometryCore/Public/Curve/DynamicGraph3.h
@@ -0,0 +1,134 @@
+// Copyright Epic Games, Inc. All Rights Reserved.
+
+#pragma once
+
+#include "CoreTypes.h"
+
+#include "BoxTypes.h"
+#include "DynamicGraph.h"
+#include "SegmentTypes.h"
+#include "Util/DynamicVector.h"
+#include "Util/IndexUtil.h"
+#include "Util/IteratorUtil.h"
+#include "Util/RefCountVector.h"
+#include "Util/SmallListSet.h"
+#include "VectorTypes.h"
+#include "VectorUtil.h"
+
+namespace UE
+{
+namespace Geometry
+{
+
+template <typename T>
+class FDynamicGraph3 : public FDynamicGraph
+{
+	TDynamicVectorN<T, 3> Vertices;
+
+public:
+	static FVector3<T> InvalidVertex()
+	{
+		return FVector3<T>(TNumericLimits<T>::Max(), 0, 0);
+	}
+
+	FVector3<T> GetVertex(int VID) const
+	{
+		return vertices_refcount.IsValid(VID) ? Vertices.AsVector3(VID) : InvalidVertex();
+	}
+
+	void SetVertex(int VID, FVector3<T> VNewPos)
+	{
+		check(VectorUtil::IsFinite(VNewPos)); // this will really catch a lot of bugs...
+		if (vertices_refcount.IsValid(VID))
+		{
+			Vertices.SetVector3(VID, VNewPos);
+			updateTimeStamp(true);
+		}
+	}
+
+	using FDynamicGraph::GetEdgeV;
+	bool GetEdgeV(int EID, FVector3<T>& A, FVector3<T>& B) const
+	{
+		if (edges_refcount.IsValid(EID))
+		{
+			A = Vertices.AsVector3(edges[EID].A);
+			B = Vertices.AsVector3(edges[EID].B);
+			return true;
+		}
+		return false;
+	}
+
+	TSegment3<T> GetEdgeSegment(int EID) const
+	{
+		checkfSlow(edges_refcount.IsValid(EID), TEXT("FDynamicGraph2.GetEdgeSegment: invalid segment with id %d"), EID);
+		const FEdge& e = edges[EID];
+		return TSegment3<T>(
+			Vertices.AsVector3(e.A),
+			Vertices.AsVector3(e.B));
+	}
+
+	FVector3<T> GetEdgeCenter(int EID) const
+	{
+		checkfSlow(edges_refcount.IsValid(EID), TEXT("FDynamicGraph3.GetEdgeCenter: invalid segment with id %d"), EID);
+		const FEdge& e = edges[EID];
+		return 0.5 * (Vertices.AsVector3(e.A) + Vertices.AsVector3(e.B));
+	}
+
+	int AppendVertex(FVector3<T> V)
+	{
+		int vid = append_vertex_internal();
+		Vertices.InsertAt({{V.X, V.Y, V.Z}}, vid);
+		return vid;
+	}
+
+	FRefCountVector::IndexEnumerable VertexIndicesItr() const
+	{
+		return vertices_refcount.Indices();
+	}
+
+	/** Enumerate positions of all Vertices in graph */
+	value_iteration<FVector3<T>> VerticesItr() const
+	{
+		return vertices_refcount.MappedIndices<FVector3<T>>(
+			[=](int vid) {
+				return Vertices.template AsVector3<T>(vid);
+			});
+	}
+
+
+	// compute vertex bounding box
+	FAxisAlignedBox2d GetBounds() const
+	{
+		TAxisAlignedBox2<T> AABB;
+		for (const FVector3<T>& V : Vertices())
+		{
+			AABB.Contain(V);
+		}
+		return AABB;
+	}
+
+
+protected:
+	// internal used in SplitEdge
+	virtual int append_new_split_vertex(int A, int B) override
+	{
+		FVector3<T> vNew = 0.5 * (GetVertex(A) + GetVertex(B));
+		int f = AppendVertex(vNew);
+		return f;
+	}
+
+	virtual void subclass_validity_checks(TFunction<void(bool)> CheckOrFailF) const override
+	{
+		for (int VID : VertexIndices())
+		{
+			FVector3<T> V = GetVertex(VID);
+			CheckOrFailF(VectorUtil::IsFinite(V));
+		}
+	}
+};
+
+typedef FDynamicGraph3<double> FDynamicGraph3d;
+
+
+} // end namespace UE::Geometry
+} // end namespace UE
\ No newline at end of file
