Fix Mac build

#rnx

diff --git a/Engine/Plugins/Experimental/GeometryProcessing/Source/DynamicMesh/Public/Solvers/LaplacianMatrixAssembly.h b/Engine/Plugins/Experimental/GeometryProcessing/Source/DynamicMesh/Public/Solvers/LaplacianMatrixAssembly.h
--- a/Engine/Plugins/Experimental/GeometryProcessing/Source/DynamicMesh/Public/Solvers/LaplacianMatrixAssembly.h
+++ b/Engine/Plugins/Experimental/GeometryProcessing/Source/DynamicMesh/Public/Solvers/LaplacianMatrixAssembly.h
@@ -236,7 +236,7 @@ void UE::MeshDeformation::ConstructUniformLaplacian(const MeshType& Mesh,
 		const int32 VertId = ToMeshV[i];
 		RealType CenterWeight = RealType(0); // equal and opposite the sum of the neighbor weights
 
-		checkSlow(!DynamicMesh.IsBoundaryVertex(VertId));  // we should only be looping over the internal verts
+		checkSlow(!Mesh.IsBoundaryVertex(VertId));  // we should only be looping over the internal verts
 
 		for (int32 NeighborVertId : Mesh.VtxVerticesItr(VertId))
 		{
