MeshModelingTools: Make FTriangleGroupTopology call the correct RebuildTopology(), and add a convenience != operator for group topology selections.

#rnx

diff --git a/Engine/Plugins/Experimental/GeometryProcessing/Source/DynamicMesh/Private/GroupTopology.cpp b/Engine/Plugins/Experimental/GeometryProcessing/Source/DynamicMesh/Private/GroupTopology.cpp
--- a/Engine/Plugins/Experimental/GeometryProcessing/Source/DynamicMesh/Private/GroupTopology.cpp
+++ b/Engine/Plugins/Experimental/GeometryProcessing/Source/DynamicMesh/Private/GroupTopology.cpp
@@ -810,8 +810,15 @@ void FGroupTopology::GetAllVertexGroups(int32 VertexID, TArray<int32>& GroupsOut
 
 
 FTriangleGroupTopology::FTriangleGroupTopology(const FDynamicMesh3* Mesh, bool bAutoBuild) 
-	: FGroupTopology(Mesh, bAutoBuild)
+	: FGroupTopology(Mesh, false)
 {
+	// Note that we can't rely on the call to RebuildTopology in the base class constructor
+	// because virtual function call resolution is not active in constructors (so it would
+	// call the base class RebuildTopology)
+	if (bAutoBuild)
+	{
+		RebuildTopology();
+	}
 }
 
 
diff --git a/Engine/Plugins/Experimental/GeometryProcessing/Source/DynamicMesh/Public/GroupTopology.h b/Engine/Plugins/Experimental/GeometryProcessing/Source/DynamicMesh/Public/GroupTopology.h
--- a/Engine/Plugins/Experimental/GeometryProcessing/Source/DynamicMesh/Public/GroupTopology.h
+++ b/Engine/Plugins/Experimental/GeometryProcessing/Source/DynamicMesh/Public/GroupTopology.h
@@ -108,6 +108,11 @@ struct DYNAMICMESH_API FGroupTopologySelection
 			&& Contains(Selection);
 	}
 
+	bool operator!=(const FGroupTopologySelection& Selection) const
+	{
+		return !(*this == Selection);
+	}
+
 	inline bool IsSelectedTriangle(const FDynamicMesh3* Mesh, const FGroupTopology* Topology, int32 TriangleID) const;
 
 
@@ -433,6 +438,7 @@ FIndex2i FGroupTopology::MakeEdgeID(int Group1, int Group2) const
 	return (Group1 < Group2) ? FIndex2i(Group1, Group2) : FIndex2i(Group2, Group1);
 }
 
+// Defined here instead of the class declaration because it requires FGroupTopology to be defined.
 bool FGroupTopologySelection::IsSelectedTriangle(const FDynamicMesh3* Mesh, const FGroupTopology* Topology, int32 TriangleID) const
 {
 	int GroupID = Topology->GetGroupID(TriangleID);
diff --git a/Engine/Plugins/Experimental/MeshModelingToolset/Source/MeshModelingTools/Public/HoleFillTool.h b/Engine/Plugins/Experimental/MeshModelingToolset/Source/MeshModelingTools/Public/HoleFillTool.h
--- a/Engine/Plugins/Experimental/MeshModelingToolset/Source/MeshModelingTools/Public/HoleFillTool.h
+++ b/Engine/Plugins/Experimental/MeshModelingToolset/Source/MeshModelingTools/Public/HoleFillTool.h
@@ -288,8 +288,16 @@ protected:
 	{
 	public:
 		FBasicTopology(const UE::Geometry::FDynamicMesh3* Mesh, bool bAutoBuild) :
-			FGroupTopology(Mesh, bAutoBuild)
-		{}
+			FGroupTopology(Mesh, false)
+		{
+			if (bAutoBuild)
+			{
+				// Virtual func resolution doesn't work in constructors. Though we're not currently
+				// overriding RebuildTopology, let's do the proper thing in case we get copy-pasted
+				// somewhere where we do.
+				RebuildTopology();
+			}
+		}
 
 		int GetGroupID(int TriangleID) const override
 		{
diff --git a/Engine/Plugins/Experimental/MeshModelingToolset/Source/MeshModelingTools/Public/MeshBoundaryToolBase.h b/Engine/Plugins/Experimental/MeshModelingToolset/Source/MeshModelingTools/Public/MeshBoundaryToolBase.h
--- a/Engine/Plugins/Experimental/MeshModelingToolset/Source/MeshModelingTools/Public/MeshBoundaryToolBase.h
+++ b/Engine/Plugins/Experimental/MeshModelingToolset/Source/MeshModelingTools/Public/MeshBoundaryToolBase.h
@@ -54,8 +54,16 @@ protected:
 	{
 	public:
 		FBasicTopology(const FDynamicMesh3* Mesh, bool bAutoBuild) :
-			FGroupTopology(Mesh, bAutoBuild)
-		{}
+			FGroupTopology(Mesh, false)
+		{
+			if (bAutoBuild)
+			{
+				// Virtual func resolution doesn't work in constructors. Though we're not currently
+				// overriding RebuildTopology, let's do the proper thing in case we get copy-pasted
+				// somewhere where we do.
+				RebuildTopology();
+			}
+		}
 
 		int GetGroupID(int TriangleID) const override
 		{
