GeometryProcessing updates
- add FMeshNormals::SmoothVertexNormals(), does rounds of vertex normal smoothing. Moved from FDynamicMeshUVEditor::SetTriangleUVsFromExpMap().
- add tests to ensure edge iexists in DynamicMeshOverlay::IsSeamEdge() and similar functions, as parent GetEdge() will fail otherwise
- add FMeshRegionGraph, utility class for representing connections between mesh triangle patches, supports merging patches
- FDynamicMeshUVEditor::SetTriangleUVsFromBoxProjection can now optionally use FMeshRegionGraph to cluster small UV-islands into larger ones. Update call sites in UVProjectionOp
#rnx

diff --git a/Engine/Plugins/Experimental/GeometryProcessing/Source/DynamicMesh/Private/Parameterization/DynamicMeshUVEditor.cpp b/Engine/Plugins/Experimental/GeometryProcessing/Source/DynamicMesh/Private/Parameterization/DynamicMeshUVEditor.cpp
--- a/Engine/Plugins/Experimental/GeometryProcessing/Source/DynamicMesh/Private/Parameterization/DynamicMeshUVEditor.cpp
+++ b/Engine/Plugins/Experimental/GeometryProcessing/Source/DynamicMesh/Private/Parameterization/DynamicMeshUVEditor.cpp
@@ -16,6 +16,7 @@
 #include "Parameterization/MeshLocalParam.h"
 #include "Parameterization/MeshUVPacking.h"
 #include "Solvers/MeshParameterizationSolvers.h"
+#include "Parameterization/MeshRegionGraph.h"
 
 #include "Async/ParallelFor.h"
 
@@ -246,7 +247,7 @@ bool FDynamicMeshUVEditor::EstimateGeodesicCenterFrameVertex(const FDynamicMesh3
 }
 
 
-bool FDynamicMeshUVEditor::SetTriangleUVsFromExpMap(const TArray<int32>& Triangles, FUVEditResult* Result)
+bool FDynamicMeshUVEditor::SetTriangleUVsFromExpMap(const TArray<int32>& Triangles, const FExpMapOptions& Options, FUVEditResult* Result)
 {
 	if (ensure(UVOverlay) == false) return false;
 	if (Triangles.Num() == 0) return false;
@@ -257,6 +258,12 @@ bool FDynamicMeshUVEditor::SetTriangleUVsFromExpMap(const TArray<int32>& Triangl
 	FDynamicMesh3& Submesh = SubmeshCalc.GetSubmesh();
 	FMeshNormals::QuickComputeVertexNormals(Submesh);
 
+	if (Options.NormalSmoothingRounds > 0)
+	{
+		FMeshNormals::SmoothVertexNormals(Submesh, Options.NormalSmoothingRounds, Options.NormalSmoothingAlpha);
+	}
+
+
 	FFrame3d SeedFrame;
 	int32 FrameVertexID;
 	bool bFrameOK = EstimateGeodesicCenterFrameVertex(Submesh, SeedFrame, FrameVertexID, true);
@@ -330,32 +337,7 @@ bool FDynamicMeshUVEditor::SetTriangleUVsFromExpMap(
 	MeshTransforms::ApplyTransform(Submesh, PointTransform, [&](const FVector3f& V) { return V;});
 	FMeshNormals::QuickComputeVertexNormals(Submesh);
 
-	NormalSmoothingRounds = FMath::Clamp(NormalSmoothingRounds, 0, 500);
-	NormalSmoothingAlpha = FMathd::Clamp(NormalSmoothingAlpha, 0.0, 1.0);
-	if (NormalSmoothingRounds > 0 && NormalSmoothingAlpha > 0)
-	{
-		int32 NumV = Submesh.MaxVertexID();
-		TArray<FVector3d> SmoothedNormals;
-		SmoothedNormals.SetNum(NumV);
-		for (int32 ri = 0; ri < NormalSmoothingRounds; ++ri)
-		{
-			SmoothedNormals.Init(FVector3d::Zero(), NumV);
-			for (int32 vid : Submesh.VertexIndicesItr())
-			{
-				FVector3d SmoothedNormal = FVector3d::Zero();
-				UE::Geometry::FMeshWeights::CotanWeightsBlendSafe(Submesh, vid, [&](int32 nbrvid, double Weight)
-				{
-					SmoothedNormal += Weight * (FVector3d)Submesh.GetVertexNormal(nbrvid);
-				});
-				SmoothedNormal.Normalize();
-				SmoothedNormals[vid] = Lerp((FVector3d)Submesh.GetVertexNormal(vid), SmoothedNormal, NormalSmoothingAlpha);
-			}
-			for (int32 vid : Submesh.VertexIndicesItr())
-			{
-				Submesh.SetVertexNormal(vid, (FVector3f)SmoothedNormals[vid]);
-			}
-		}
-	}
+	FMeshNormals::SmoothVertexNormals(Submesh, NormalSmoothingRounds, NormalSmoothingAlpha);
 
 	FDynamicMeshAABBTree3 Spatial(&Submesh, true);
 	double NearDistSqr;
@@ -780,6 +762,7 @@ void FDynamicMeshUVEditor::SetTriangleUVsFromBoxProjection(
 	TFunctionRef<FVector3d(const FVector3d&)> PointTransform, 
 	const FFrame3d& BoxFrame, 
 	const FVector3d& BoxDimensions, 
+	int32 MinIslandTriCount,
 	FUVEditResult* Result)
 {
 	if (ensure(UVOverlay) == false) return;
@@ -805,23 +788,62 @@ void FDynamicMeshUVEditor::SetTriangleUVsFromBoxProjection(
 	double ScaleZ = (FMathd::Abs(BoxDimensions.Z) > FMathf::ZeroTolerance) ? (1.0 / BoxDimensions.Z) : 1.0;
 	FVector3d Scale(ScaleX, ScaleY, ScaleZ);
 
+	// compute assignments to the available planes based on face normals
 	TArray<FVector3d> TriNormals;
 	TArray<FIndex2i> TriangleBoxPlaneAssignments;
 	TriNormals.SetNum(NumTriangles);
 	TriangleBoxPlaneAssignments.SetNum(NumTriangles);
+	TArray<int32> IndexMap;
+	IndexMap.SetNum(Mesh->MaxTriangleID());
 	ParallelFor(NumTriangles, [&](int32 i)
 	{
 		int32 tid = Triangles[i];
 		TriNormals[i] = GetTriNormal(tid);
-		FVector3d N = BoxFrame.ToFrameVector(TriNormals[i]);
-		N *= Scale;
-		FVector3d NAbs(FMathd::Abs(N.X), FMathd::Abs(N.Y), FMathd::Abs(N.Z));
+		FVector3d ScaledNormal = BoxFrame.ToFrameVector(TriNormals[i]);
+		ScaledNormal *= Scale;
+		FVector3d NAbs(FMathd::Abs(ScaledNormal.X), FMathd::Abs(ScaledNormal.Y), FMathd::Abs(ScaledNormal.Z));
 		int MajorAxis = NAbs[0] > NAbs[1] ? (NAbs[0] > NAbs[2] ? 0 : 2) : (NAbs[1] > NAbs[2] ? 1 : 2);
-		double MajorAxisSign = FMathd::Sign(N[MajorAxis]);
+		double MajorAxisSign = FMathd::Sign(ScaledNormal[MajorAxis]);
 		int Bucket = (MajorAxisSign > 0) ? (MajorAxis+3) : MajorAxis;
 		TriangleBoxPlaneAssignments[i] = FIndex2i(MajorAxis, Bucket);
+		IndexMap[tid] = i;
 	});
 
+
+	// Optimize face assignments. Small regions are grouped with larger neighbour regions.
+	if (MinIslandTriCount > 1)
+	{
+		FMeshConnectedComponents Components(Mesh);
+		Components.FindConnectedTriangles(Triangles, [&](int32 t1, int32 t2) { return TriangleBoxPlaneAssignments[IndexMap[t1]] == TriangleBoxPlaneAssignments[IndexMap[t2]]; });
+		FMeshRegionGraph RegionGraph;
+		RegionGraph.BuildFromComponents(*Mesh, Components, [&](int32 ComponentIdx) { int32 tid = Components[ComponentIdx].Indices[0]; return TriangleBoxPlaneAssignments[IndexMap[tid]].A; });
+		// todo: similarity measure should probably take normals into account
+		bool bMerged = RegionGraph.MergeSmallRegions(MinIslandTriCount-1, 
+			[&](int32 A, int32 B) { return RegionGraph.GetRegionTriCount(A) > RegionGraph.GetRegionTriCount(B); });
+		bool bSwapped = RegionGraph.OptimizeBorders();
+		if (bMerged || bSwapped)
+		{
+			int32 N = RegionGraph.MaxRegionIndex();
+			for (int32 k = 0; k < N; ++k)
+			{
+				if (RegionGraph.IsRegion(k))
+				{
+					int32 MajorAxis = RegionGraph.GetExternalID(k);
+					const TArray<int32>& Tris = RegionGraph.GetRegionTris(k);
+					for (int32 tid : Tris)
+					{
+						int32 i = IndexMap[tid];
+						FVector3d ScaledNormal = BoxFrame.ToFrameVector(TriNormals[i]) * Scale;
+						double MajorAxisSign = FMathd::Sign(ScaledNormal[MajorAxis]);
+						int Bucket = (MajorAxisSign > 0) ? (MajorAxis + 3) : MajorAxis;
+						TriangleBoxPlaneAssignments[i] = FIndex2i(MajorAxis, Bucket);
+					}
+				}
+			}
+		}
+	}
+
+
 	auto ProjAxis = [](const FVector3d& P, int Axis1, int Axis2, float Axis1Scale, float Axis2Scale)
 	{
 		return FVector2f(float(P[Axis1]) * Axis1Scale, float(P[Axis2]) * Axis2Scale);
diff --git a/Engine/Plugins/Experimental/GeometryProcessing/Source/DynamicMesh/Private/Parameterization/MeshRegionGraph.cpp b/Engine/Plugins/Experimental/GeometryProcessing/Source/DynamicMesh/Private/Parameterization/MeshRegionGraph.cpp
new file mode 100644
--- /dev/null
+++ b/Engine/Plugins/Experimental/GeometryProcessing/Source/DynamicMesh/Private/Parameterization/MeshRegionGraph.cpp
@@ -0,0 +1,328 @@
+// Copyright Epic Games, Inc. All Rights Reserved.
+
+#include "Parameterization/MeshRegionGraph.h"
+
+#include "Async/ParallelFor.h"
+
+#include "ExplicitUseGeometryMathTypes.h"		// using UE::Geometry::(math types)
+using namespace UE::Geometry;
+
+
+void FMeshRegionGraph::BuildFromComponents(const FDynamicMesh3& Mesh, const FMeshConnectedComponents& Components, TFunctionRef<int32(int32)> ExternalIDFunc)
+{
+	int32 N = Components.Num();
+	Regions.SetNum(N);
+	TriangleToRegionMap.Init(-1, Mesh.MaxTriangleID());
+	TriangleNbrTris.Init(FIndex3i(-1, -1, -1), Mesh.MaxTriangleID());
+
+	for (int32 k = 0; k < N; ++k)
+	{
+		Regions[k].Triangles = Components[k].Indices;
+		Regions[k].ExternalID = ExternalIDFunc(k);
+
+		for (int32 tid : Components[k].Indices)
+		{
+			TriangleToRegionMap[tid] = k;
+			TriangleNbrTris[tid] = Mesh.GetTriNeighbourTris(tid);
+		}
+	}
+
+	for (int32 k = 0; k < N; ++k)
+	{
+		BuildNeigbours(k);
+	}
+}
+
+
+void FMeshRegionGraph::BuildFromTriangleSets(const FDynamicMesh3& Mesh, const TArray<TArray<int32>>& TriangleSets, TFunctionRef<int32(int32)> ExternalIDFunc)
+{
+	int32 N = TriangleSets.Num();
+	Regions.SetNum(N);
+	TriangleToRegionMap.Init(-1, Mesh.MaxTriangleID());
+	TriangleNbrTris.Init(FIndex3i(-1, -1, -1), Mesh.MaxTriangleID());
+
+	for (int32 k = 0; k < N; ++k)
+	{
+		Regions[k].Triangles = TriangleSets[k];
+		Regions[k].ExternalID = ExternalIDFunc(k);
+
+		for (int32 tid : TriangleSets[k])
+		{
+			TriangleToRegionMap[tid] = k;
+			TriangleNbrTris[tid] = Mesh.GetTriNeighbourTris(tid);
+		}
+	}
+
+	for (int32 k = 0; k < N; ++k)
+	{
+		BuildNeigbours(k);
+	}
+}
+
+
+
+TArray<int32> FMeshRegionGraph::GetNeighbours(int32 RegionIdx) const
+{
+	TArray<int32> Nbrs;
+	if (IsRegion(RegionIdx))
+	{
+		for (const FNeighbour& Nbr : Regions[RegionIdx].Neighbours)
+		{
+			Nbrs.Add(Nbr.RegionIndex);
+		}
+	}
+	return Nbrs;
+}
+
+
+bool FMeshRegionGraph::AreRegionsConnected(int32 RegionAIndex, int32 RegionBIndex) const
+{
+	if (IsRegion(RegionAIndex) == false || IsRegion(RegionBIndex) == false) return false;
+
+	for (const FNeighbour& Nbr : Regions[RegionAIndex].Neighbours)
+	{
+		if (Nbr.RegionIndex == RegionBIndex)
+		{
+			return true;
+		}
+	}
+	return false;
+}
+
+
+
+bool FMeshRegionGraph::MergeSmallRegions(int32 SmallThreshold,
+										 TFunctionRef<float(int32 SmallRgnIdx, int32 NbrRgnIdx)> RegionSimilarityFunc)
+{
+	bool bMergedAny = false;
+	int32 UseSmallThreshold = 1;
+	while (UseSmallThreshold < SmallThreshold || MergeSmallRegionsPass(UseSmallThreshold, RegionSimilarityFunc))
+	{
+		bMergedAny = true;
+		UseSmallThreshold = FMath::Min(UseSmallThreshold + 1, SmallThreshold);
+	}
+	return bMergedAny;
+}
+
+
+
+bool FMeshRegionGraph::MergeSmallRegionsPass(int32 SmallThreshold,
+										 TFunctionRef<float(int32 SmallRgnIdx, int32 NbrRgnIdx)> RegionSimilarityFunc)
+{
+	bool bMergedAny = false;
+
+	TArray<int32> SmallRegions;
+	for (int32 k = 0; k < MaxRegionIndex(); ++k)
+	{
+		if (IsRegion(k) && Regions[k].Triangles.Num() <= SmallThreshold)
+		{
+			SmallRegions.Add(k);
+		}
+	}
+	if (SmallRegions.Num() == 0)
+	{
+		return false;
+	}
+
+	SmallRegions.Sort([this](int a, int b) { return Regions[a].Triangles.Num() < Regions[b].Triangles.Num(); });
+
+	struct FMatch 
+	{
+		float Score = 0.0f;
+		int Index = -1;
+	};
+
+	for (int32 j = 0; j < SmallRegions.Num(); ++j)
+	{
+		int32 SmallRegionIdx = SmallRegions[j];
+		if (IsRegion(SmallRegionIdx) == false) continue;
+		const FRegion& SmallRegion = Regions[SmallRegionIdx];
+
+		TArray<FMatch> Matches;
+		for (const FNeighbour& Nbr : SmallRegion.Neighbours)
+		{
+			float Score = RegionSimilarityFunc(SmallRegionIdx, Nbr.RegionIndex);
+			Matches.Add(FMatch{ Score, Nbr.RegionIndex });
+		}
+		if (Matches.Num() == 0)
+		{
+			continue;
+		}
+
+		Matches.Sort([](const FMatch& A, const FMatch& B) { return A.Score > B.Score; });
+
+		if (MergeRegion(SmallRegionIdx, Matches[0].Index))
+		{
+			bMergedAny = true;
+		}
+	}
+
+	return true;
+}
+
+
+
+
+bool FMeshRegionGraph::MergeRegion(int32 RemoveRgnIdx, int32 MergeToRgnIdx)
+{
+	if (AreRegionsConnected(RemoveRgnIdx, MergeToRgnIdx) == false)
+	{
+		ensure(false);
+		return false;
+	}
+
+	FRegion& RemoveRegion = Regions[RemoveRgnIdx];
+	FRegion& AppendToRegion = Regions[MergeToRgnIdx];
+
+	// append triangles to the new region
+	for (int32 tid : RemoveRegion.Triangles)
+	{
+		AppendToRegion.Triangles.Add(tid);
+		TriangleToRegionMap[tid] = MergeToRgnIdx;
+	}
+	RemoveRegion.Triangles.Reset();
+	RemoveRegion.bValid = false;
+
+	// rebuild all affected regions, which should just be all nbrs of the removed region
+	// (which will include the merge-to region)
+	// TODO: this can be more efficient. At minimum can update all nbr sets and then rebuild counts
+	for (FNeighbour& Nbr : RemoveRegion.Neighbours)
+	{
+		BuildNeigbours(Nbr.RegionIndex);
+	}
+
+	return true;
+}
+
+
+
+bool FMeshRegionGraph::OptimizeBorders(int32 MaxRounds)
+{
+	auto GetSwapNbrRegionIndex = [this](int32 tid) -> int32
+	{
+		int32 Index = TriangleToRegionMap[tid];
+		FIndex3i NbrTris = TriangleNbrTris[tid];
+		int32 NbrCount = 0;
+		TArray<int32, TFixedAllocator<3>> UniqueNbrs;
+		for (int32 j = 0; j < 3; ++j)
+		{
+			if (NbrTris[j] >= 0)
+			{
+				int32 NbrRegionIndex = TriangleToRegionMap[NbrTris[j]];
+				if (IsRegion(NbrRegionIndex) && NbrRegionIndex != Index)
+				{
+					NbrCount++;
+					UniqueNbrs.AddUnique(NbrRegionIndex);
+				}
+			}
+		}
+		if (NbrCount >= 2 && UniqueNbrs.Num() == 1)
+		{
+			return UniqueNbrs[0];
+		}
+		return -1;
+	};
+
+	bool bModified = false;
+
+	bool bDone = false;
+	int32 RoundCounter = 0;
+	while (bDone == false && RoundCounter++ < MaxRounds)
+	{
+		TArray<FIndex2i> PotentialSwaps;
+		for (const FRegion& Region : Regions)
+		{
+			for (int32 tid : Region.Triangles)
+			{
+				int32 SwapToNbrIndex = GetSwapNbrRegionIndex(tid);
+				if (IsRegion(SwapToNbrIndex))
+				{
+					PotentialSwaps.Add( FIndex2i(tid, SwapToNbrIndex) );
+				}
+			}
+		}
+		bDone = (PotentialSwaps.Num() == 0);
+
+		for (const FIndex2i& Swap : PotentialSwaps)
+		{
+			int32 tid = Swap.A;
+			int32 Index = TriangleToRegionMap[tid];
+			int32 SwapToNbrIndex = Swap.B;
+
+			int32 CheckNbrIndex = GetSwapNbrRegionIndex(tid);
+			if (CheckNbrIndex == SwapToNbrIndex)
+			{
+				Regions[Index].Triangles.RemoveSwap(tid, false);
+				Regions[SwapToNbrIndex].Triangles.Add(tid);
+				TriangleToRegionMap[tid] = SwapToNbrIndex;
+				bModified = true;
+			}
+		}
+	}
+
+	return bModified;
+}
+
+
+
+
+
+void FMeshRegionGraph::BuildNeigbours(int32 RegionIdx)
+{
+	auto CountNeighbourFunc = [](FRegion& Region, int32 NbrRegionIndex)
+	{
+		for (FNeighbour& Nbr : Region.Neighbours)
+		{
+			if (Nbr.RegionIndex == NbrRegionIndex)
+			{
+				Nbr.Count++;
+				return;
+			}
+		}
+		check(false);
+	};
+
+
+	FRegion& Region = Regions[RegionIdx];
+	Region.Neighbours.Reset();
+
+	TSet<int32> NeighbourSet;
+
+	for (int32 tid : Region.Triangles)
+	{
+		int32 RegionIndex = TriangleToRegionMap[tid];
+		check(RegionIndex == RegionIdx);
+
+		FIndex3i Nbrs = TriangleNbrTris[tid];
+		for (int32 j = 0; j < 3; ++j)
+		{
+			if (Nbrs[j] == FDynamicMesh3::InvalidID)
+			{
+				Region.bIsOnMeshBoundary = true;
+			}
+			else
+			{
+				int32 NbrRegionIndex = TriangleToRegionMap[Nbrs[j]];
+				if (NbrRegionIndex == -1)
+				{
+					Region.bIsOnROIBoundary = true;
+				}
+				else
+				{
+					if (NbrRegionIndex != RegionIndex)
+					{
+						if (NeighbourSet.Contains(NbrRegionIndex) == false)
+						{
+							NeighbourSet.Add(NbrRegionIndex);
+							FNeighbour Nbr;
+							Nbr.RegionIndex = NbrRegionIndex;
+							Region.Neighbours.Add(Nbr);
+						}
+
+						CountNeighbourFunc(Region, NbrRegionIndex);
+					}
+				}
+			}
+		}
+	}
+}
\ No newline at end of file
diff --git a/Engine/Plugins/Experimental/GeometryProcessing/Source/DynamicMesh/Public/Parameterization/DynamicMeshUVEditor.h b/Engine/Plugins/Experimental/GeometryProcessing/Source/DynamicMesh/Public/Parameterization/DynamicMeshUVEditor.h
--- a/Engine/Plugins/Experimental/GeometryProcessing/Source/DynamicMesh/Public/Parameterization/DynamicMeshUVEditor.h
+++ b/Engine/Plugins/Experimental/GeometryProcessing/Source/DynamicMesh/Public/Parameterization/DynamicMeshUVEditor.h
@@ -86,12 +86,30 @@ public:
 	 */
 	void SetTriangleUVsFromPlanarProjection(const TArray<int32>& Triangles, TFunctionRef<FVector3d(const FVector3d&)> PointTransform, const FFrame3d& ProjectionFrame, const FVector2d& Dimensions, FUVEditResult* Result = nullptr);
 
+
+	/**
+	 * FExpMapOptions provides additional control over ExpMap UV generation below
+	 */
+	struct FExpMapOptions
+	{
+		/** Number of rounds of explicit uniform normal smoothing to apply to mesh normals */
+		int32 NormalSmoothingRounds;
+		/** Alpha for smoothing, valid range is 0-1 */
+		double NormalSmoothingAlpha;
+
+		FExpMapOptions()
+		{
+			NormalSmoothingRounds = 0;
+			NormalSmoothingAlpha = 0.25f;
+		}
+	};
+
 	/**
 	 * Create new UV island for given Triangles, and set UVs for that island using Discrete Exponential Map.
 	 * ExpMap center-point is calculated by finding maximum (Dijkstra-approximated) geodesic distance from border of island.
 	 * @warning computes a single ExpMap, so input triangle set must be connected, however this is not verified internally
 	 */
-	bool SetTriangleUVsFromExpMap(const TArray<int32>& Triangles, FUVEditResult* Result = nullptr);
+	bool SetTriangleUVsFromExpMap(const TArray<int32>& Triangles, const FExpMapOptions& Options = FExpMapOptions(), FUVEditResult* Result = nullptr);
 
 
 	/**
@@ -138,8 +156,9 @@ public:
 	 * Set UVs by box projection. Triangles will be grouped to "best" box face
 	 * PointTransform is applied to points before projectiong onto ProjectionFrame X/Y axes
 	 * Projected U/V coordinates are divided by Dimensions.X/Y
+	 * @param MinIslandTriCount Any UV island with fewer triangles than this count will be merged into a neighbouring island
 	 */
-	void SetTriangleUVsFromBoxProjection(const TArray<int32>& Triangles, TFunctionRef<FVector3d(const FVector3d&)> PointTransform, const FFrame3d& BoxFrame, const FVector3d& BoxDimensions, FUVEditResult* Result = nullptr);
+	void SetTriangleUVsFromBoxProjection(const TArray<int32>& Triangles, TFunctionRef<FVector3d(const FVector3d&)> PointTransform, const FFrame3d& BoxFrame, const FVector3d& BoxDimensions, int32 MinIslandTriCount = 2, FUVEditResult* Result = nullptr);
 
 
 	/**
diff --git a/Engine/Plugins/Experimental/GeometryProcessing/Source/DynamicMesh/Public/Parameterization/MeshRegionGraph.h b/Engine/Plugins/Experimental/GeometryProcessing/Source/DynamicMesh/Public/Parameterization/MeshRegionGraph.h
new file mode 100644
--- /dev/null
+++ b/Engine/Plugins/Experimental/GeometryProcessing/Source/DynamicMesh/Public/Parameterization/MeshRegionGraph.h
@@ -0,0 +1,148 @@
+// Copyright Epic Games, Inc. All Rights Reserved.
+
+#pragma once
+
+#include "DynamicMesh/DynamicMesh3.h"
+#include "Selections/MeshConnectedComponents.h"
+
+namespace UE
+{
+namespace Geometry
+{
+
+
+
+/**
+ * FMeshRegionGraph represents neighbourhood relationships between mesh patches. 
+ * This is similar in some ways to a FGroupTopology, however only patch-adjacency
+ * connections are tracked. Connectivity between patches can be queried, and
+ * patches can be merged without rebuilding the graph.
+ * 
+ * Each mesh Patch/Region is assigned an integer ID/RegionIndex.
+ * After removing or merging regions, some RegionIndex values may be invalid, use IsRegion() to check
+ * 
+ * An integer ID can be associated with each Region using the ExternalIDFunc passed to the various Build() functions.
+ */
+class DYNAMICMESH_API FMeshRegionGraph
+{
+public:
+
+	/**
+	 * Build a region graph for a Mesh using the given ConnectedComponents
+	 * @param ExternalIDFunc This allows the client to provide an integer ID that will be associated with each Region - called with the region index in Components
+	 */
+	void BuildFromComponents(const FDynamicMesh3& Mesh, 
+							 const FMeshConnectedComponents& Components,
+							 TFunctionRef<int32(int32)> ExternalIDFunc);
+
+	/**
+	 * Build a region graph for a Mesh using the given triangle sets
+	 * @param ExternalIDFunc This allows the client to provide an integer ID that will be associated with each Region - called with the region index in TriangleSets
+	 */
+	void BuildFromTriangleSets(const FDynamicMesh3& Mesh,
+							 const TArray<TArray<int32>>& TriangleSets,
+							 TFunctionRef<int32(int32)> ExternalIDFunc);
+
+	/** 
+	 * @return Max valid region index. Some region indices may be invalid. 
+	 */
+	int32 MaxRegionIndex() const { return Regions.Num(); }
+
+	/** 
+	 * @return true if the RegionIndex is valid 
+	 */
+	int32 IsRegion(int32 RegionIdx) const { return RegionIdx >= 0 && RegionIdx < Regions.Num() && Regions[RegionIdx].bValid; }
+	
+	/** 
+	 * @return the ExternalID for a Region, defined at Build time 
+	 */
+	int32 GetExternalID(int32 RegionIdx) const { return IsRegion(RegionIdx) ? Regions[RegionIdx].ExternalID : -1; }
+
+	/** 
+	 * @return the triangle count for a Region 
+	 */
+	int32 GetRegionTriCount(int32 RegionIdx) const { return IsRegion(RegionIdx) ? Regions[RegionIdx].Triangles.Num() : 0; }
+
+	/** 
+	 * @return list of triangles in a Region 
+	 */
+	const TArray<int32>& GetRegionTris(int32 RegionIdx) const { return IsRegion(RegionIdx) ? Regions[RegionIdx].Triangles : EmptyTriSet; }
+
+	/** 
+	 * @return list of Neigbour Region Indices for a Region 
+	 */
+	TArray<int32> GetNeighbours(int32 RegionIdx) const;
+
+	/**
+	 * @return true if the two regions are adjacent, ie they share an edge on at least one triangle
+	 */
+	bool AreRegionsConnected(int32 RegionAIndex, int32 RegionBIndex) const;
+
+	/** @return true if the two triangles are in the same region */
+	bool AreTrianglesConnected(int32 TriangleA, int32 TriangleB) const
+	{
+		return TriangleToRegionMap[TriangleA] == TriangleToRegionMap[TriangleB];
+	}
+
+	/**
+	 * Incrementally merge regions with a triangle count below the given SmallThreshold.
+	 * A Small region is merged to it's most-similar neighbour, as defined by RegionSimilarityFunc.
+	 * The merge order is sorted by largest Similarity, ie most-similar regions are merged first.
+	 */
+	bool MergeSmallRegions(int32 SmallThreshold,
+						   TFunctionRef<float(int32 SmallRgnIdx, int32 NbrRgnIdx)> RegionSimilarityFunc );
+
+	/**
+	 * Merge one existing Region into another
+	 * @param RemoveRgnIdx this region will be removed
+	 * @param MergeToRgnIdx this region will be added to
+	 */
+	bool MergeRegion(int32 RemoveRgnIdx, int32 MergeToRgnIdx);
+
+	/**
+	 * Optimize region borders by swapping triangles. The swap criteria is, if 2/3 of the
+	 * neighbours of a given Triangle in region A are both in region B, swap it to region B.
+	 * This swapping process is repeated in rounds, MaxRounds times.
+	 * @return true if any swaps were applied
+	 */
+	bool OptimizeBorders(int32 MaxRounds = 25);
+
+
+public:
+
+	struct FNeighbour
+	{
+		int32 RegionIndex;
+		int32 Count;
+	};
+
+	struct FRegion
+	{
+		bool bValid = true;
+
+		int32 ExternalID;
+
+		TArray<int32> Triangles;
+		TArray<FNeighbour> Neighbours;
+
+		bool bIsOnMeshBoundary = false;
+		bool bIsOnROIBoundary = false;
+	};
+
+	TArray<FRegion> Regions;
+	TArray<int32> TriangleToRegionMap;
+	TArray<FIndex3i> TriangleNbrTris;
+
+	TArray<int32> EmptyTriSet;
+
+protected:
+	void BuildNeigbours(int32 RegionIdx);
+	
+	bool MergeSmallRegionsPass(int32 SmallThreshold,
+						   TFunctionRef<float(int32 SmallRgnIdx, int32 NbrRgnIdx)> RegionSimilarityFunc);
+
+};
+
+
+} // end namespace UE::Geometry
+} // end namespace UE
\ No newline at end of file
diff --git a/Engine/Plugins/Experimental/MeshModelingToolset/Source/ModelingOperators/Private/ParameterizationOps/UVProjectionOp.cpp b/Engine/Plugins/Experimental/MeshModelingToolset/Source/ModelingOperators/Private/ParameterizationOps/UVProjectionOp.cpp
--- a/Engine/Plugins/Experimental/MeshModelingToolset/Source/ModelingOperators/Private/ParameterizationOps/UVProjectionOp.cpp
+++ b/Engine/Plugins/Experimental/MeshModelingToolset/Source/ModelingOperators/Private/ParameterizationOps/UVProjectionOp.cpp
@@ -75,12 +75,12 @@ void FUVProjectionOp::CalculateResult_Box(FProgressCancel* Progress)
 			Triangles.Add(tid);
 		}
 		UVEditor.SetTriangleUVsFromBoxProjection(Triangles, [&](const FVector3d P) { return MeshToProjectionSpace.TransformPosition(P); },
-			ProjectionBox.Frame, 2*ProjectionBox.Extents, &EditResult);
+			ProjectionBox.Frame, 2*ProjectionBox.Extents, MinRegionSize, &EditResult);
 	}
 	else
 	{
 		UVEditor.SetTriangleUVsFromBoxProjection(*TriangleROI, [&](const FVector3d P) { return MeshToProjectionSpace.TransformPosition(P); },
-			ProjectionBox.Frame, 2*ProjectionBox.Extents, &EditResult);
+			ProjectionBox.Frame, 2*ProjectionBox.Extents, MinRegionSize, &EditResult);
 		//FCompactMaps CompactMaps;
 		//UVEditor.GetOverlay()->CompactInPlace(CompactMaps);
 	}
diff --git a/Engine/Plugins/Experimental/MeshModelingToolset/Source/ModelingOperators/Public/ParameterizationOps/UVProjectionOp.h b/Engine/Plugins/Experimental/MeshModelingToolset/Source/ModelingOperators/Public/ParameterizationOps/UVProjectionOp.h
--- a/Engine/Plugins/Experimental/MeshModelingToolset/Source/ModelingOperators/Public/ParameterizationOps/UVProjectionOp.h
+++ b/Engine/Plugins/Experimental/MeshModelingToolset/Source/ModelingOperators/Public/ParameterizationOps/UVProjectionOp.h
@@ -42,6 +42,7 @@ public:
 
 	EUVProjectionMethod ProjectionMethod = EUVProjectionMethod::Box;
 	FOrientedBox3d ProjectionBox;
+	int32 MinRegionSize = 2;
 	float CylinderSplitAngle = 45.0f;
 
 	// ExpMap parameters for modifying normals
@@ -49,8 +50,6 @@ public:
 	int32 SmoothingRounds = 0;
 	double SmoothingAlpha = 0.5;
 
-	float CylinderProjectToTopOrBottomAngleThreshold = 45.0f;
-
 	// position to use as UV origin
 	FVector2f UVOrigin = FVector2f(0.5f,0.5f);
 	// rotation applied to computed UVs
diff --git a/Engine/Source/Runtime/GeometryCore/Private/DynamicMesh/DynamicMeshOverlay.cpp b/Engine/Source/Runtime/GeometryCore/Private/DynamicMesh/DynamicMeshOverlay.cpp
--- a/Engine/Source/Runtime/GeometryCore/Private/DynamicMesh/DynamicMeshOverlay.cpp
+++ b/Engine/Source/Runtime/GeometryCore/Private/DynamicMesh/DynamicMeshOverlay.cpp
@@ -472,6 +472,10 @@ bool TDynamicMeshOverlay<RealType,ElementSize>::IsSeamEdge(int eid, bool* bIsNon
 	{
 		*bIsNonIntersecting = false;
 	}
+	if (ParentMesh->IsEdge(eid) == false)
+	{
+		return false;
+	}
 
 	FIndex2i et = ParentMesh->GetEdgeT(eid);
 	if (et.B == FDynamicMesh3::InvalidID)
@@ -560,6 +564,11 @@ bool TDynamicMeshOverlay<RealType,ElementSize>::IsSeamEdge(int eid, bool* bIsNon
 template<typename RealType, int ElementSize>
 bool TDynamicMeshOverlay<RealType, ElementSize>::IsSeamEndEdge(int eid) const
 {
+	if (ParentMesh->IsEdge(eid) == false)
+	{
+		return false;
+	}
+
 	FIndex2i et = ParentMesh->GetEdgeT(eid);
 	if (et.B == FDynamicMesh3::InvalidID)
 	{
diff --git a/Engine/Source/Runtime/GeometryCore/Private/DynamicMesh/MeshNormals.cpp b/Engine/Source/Runtime/GeometryCore/Private/DynamicMesh/MeshNormals.cpp
--- a/Engine/Source/Runtime/GeometryCore/Private/DynamicMesh/MeshNormals.cpp
+++ b/Engine/Source/Runtime/GeometryCore/Private/DynamicMesh/MeshNormals.cpp
@@ -325,6 +325,46 @@ void FMeshNormals::QuickComputeVertexNormals(FDynamicMesh3& Mesh, bool bInvert)
 }
 
 
+void FMeshNormals::SmoothVertexNormals(FDynamicMesh3& Mesh, int32 SmoothingRounds, double SmoothingAlpha)
+{
+	SmoothingRounds = FMath::Clamp(SmoothingRounds, 0, 500);
+	SmoothingAlpha = FMathd::Clamp(SmoothingAlpha, 0.0, 1.0);
+	if (SmoothingRounds > 0 && SmoothingAlpha > 0)
+	{
+		int32 NumV = Mesh.MaxVertexID();
+		TArray<FVector3d> SmoothedNormals;
+		SmoothedNormals.SetNum(NumV);
+		for (int32 ri = 0; ri < SmoothingRounds; ++ri)
+		{
+			SmoothedNormals.Init(FVector3d::Zero(), NumV);
+
+			// compute
+			ParallelFor(NumV, [&](int32 vid)
+			{
+				if (Mesh.IsVertex(vid))
+				{
+					FVector3d SmoothedNormal = FVector3d::Zero();
+					Mesh.EnumerateVertexVertices(vid, [&](int32 nbrvid)
+					{
+						SmoothedNormal += (FVector3d)Mesh.GetVertexNormal(nbrvid);
+					});
+					SmoothedNormals[vid] = Lerp((FVector3d)Mesh.GetVertexNormal(vid), SmoothedNormal, SmoothingAlpha);
+					Normalize(SmoothedNormals[vid]);
+				}
+			});
+
+			// update
+			for (int32 vid : Mesh.VertexIndicesItr())
+			{
+				Mesh.SetVertexNormal(vid, (FVector3f)SmoothedNormals[vid]);
+			}
+		}
+	}
+}
+
+
+
+
 void FMeshNormals::QuickComputeVertexNormalsForTriangles(FDynamicMesh3& Mesh, const TArray<int32>& Triangles, bool bWeightByArea, bool bWeightByAngle, bool bInvert)
 {
 	if (Mesh.HasVertexNormals() == false)
diff --git a/Engine/Source/Runtime/GeometryCore/Public/DynamicMesh/MeshNormals.h b/Engine/Source/Runtime/GeometryCore/Public/DynamicMesh/MeshNormals.h
--- a/Engine/Source/Runtime/GeometryCore/Public/DynamicMesh/MeshNormals.h
+++ b/Engine/Source/Runtime/GeometryCore/Public/DynamicMesh/MeshNormals.h
@@ -112,6 +112,13 @@ public:
 	 */
 	static void QuickComputeVertexNormals(FDynamicMesh3& Mesh, bool bInvert = false);
 
+
+	/**
+	 * Apply rounds of explicit uniform-weighted normal smoothing to the VertexNormals attribute of the given Mesh.
+	 */
+	static void SmoothVertexNormals(FDynamicMesh3& Mesh, int32 SmoothingRounds, double SmoothingAlpha);
+
+
 	/**
 	 * Compute per-vertex normals for the vertices of a set of triangles of a Mesh
 	 * @param bWeightByArea weight neighbor triangles by area
