ModelingTools: Fix gutter artifacts in BakeTexture when multisampling is enabled.

#rnx

diff --git a/Engine/Plugins/Experimental/GeometryProcessing/Source/DynamicMesh/Private/Sampling/MeshMapBaker.cpp b/Engine/Plugins/Experimental/GeometryProcessing/Source/DynamicMesh/Private/Sampling/MeshMapBaker.cpp
--- a/Engine/Plugins/Experimental/GeometryProcessing/Source/DynamicMesh/Private/Sampling/MeshMapBaker.cpp
+++ b/Engine/Plugins/Experimental/GeometryProcessing/Source/DynamicMesh/Private/Sampling/MeshMapBaker.cpp
@@ -202,34 +202,12 @@ void FMeshMapBaker::Bake()
 					continue;
 				}
 
-				float* PixelBuffer = TileBuffer.GetPixel(TilePixelLinearIdx);
-				float& PixelWeight = TileBuffer.GetPixelWeight(TilePixelLinearIdx);
-
 				const FVector2i ImageCoords = Tile.GetSourceCoords(TileCoords);
 				const int32 NumSamples = OccupancyMap.Multisampler.Num();
-				const int32 NumEvaluatorAdd = BakeAccumulateLists[(int32)FMeshMapEvaluator::EAccumulateMode::Add].Num();
 				for (int32 SampleIdx = 0; SampleIdx < NumSamples; ++SampleIdx)
 				{
 					const int64 LinearIdx = TilePixelLinearIdx * OccupancyMap.Multisampler.Num() + SampleIdx;
-					if (!OccupancyMap.IsInterior(LinearIdx))
-					{
-						PixelWeight += NumEvaluatorAdd > 0;
-						for (int32 Idx : BakeAccumulateLists[(int32) FMeshMapEvaluator::EAccumulateMode::Add] )
-						{
-							const int32 NumData = BakeContexts[Idx].DataLayout.Num();
-							for (int32 DataIdx = 0; DataIdx < NumData; ++DataIdx)
-							{
-								const int32 ResultIdx = BakeOffsets[Idx] + DataIdx;
-								const int32 Offset = BakeSampleOffsets[ResultIdx];
-								const int32 Stride = (int32) BakeContexts[Idx].DataLayout[DataIdx];
-								for (int32 BufIdx = 0; BufIdx < Stride; ++BufIdx)
-								{
-									PixelBuffer[Offset + BufIdx] += BakeDefaults[Offset + BufIdx];
-								}
-							}
-						}
-					}
-					else
+					if (OccupancyMap.IsInterior(LinearIdx))
 					{
 						const FVector2d UVPosition = (FVector2d)OccupancyMap.TexelQueryUV[LinearIdx];
 						const int32 UVTriangleID = OccupancyMap.TexelQueryTriangle[LinearIdx];
