GeometryProcessing: Add FDynamicMeshUVEditor::ScaleUVAreaTo3DArea() and ::AutoOrientUVArea(), FMeshRegionGraph::MoveRegionTris()
#rnx

diff --git a/Engine/Plugins/Runtime/GeometryProcessing/Source/DynamicMesh/Private/Parameterization/DynamicMeshUVEditor.cpp b/Engine/Plugins/Runtime/GeometryProcessing/Source/DynamicMesh/Private/Parameterization/DynamicMeshUVEditor.cpp
--- a/Engine/Plugins/Runtime/GeometryProcessing/Source/DynamicMesh/Private/Parameterization/DynamicMeshUVEditor.cpp
+++ b/Engine/Plugins/Runtime/GeometryProcessing/Source/DynamicMesh/Private/Parameterization/DynamicMeshUVEditor.cpp
@@ -11,6 +11,7 @@
 #include "DynamicMesh/DynamicMeshAABBTree3.h"
 #include "MeshQueries.h"
 #include "MeshWeights.h"
+#include "CompGeom/ConvexHull2.h"
 
 #include "Parameterization/MeshDijkstra.h"
 #include "Parameterization/MeshLocalParam.h"
@@ -41,6 +42,13 @@ FDynamicMeshUVEditor::FDynamicMeshUVEditor(FDynamicMesh3* MeshIn, int32 UVLayerI
 }
 
 
+FDynamicMeshUVEditor::FDynamicMeshUVEditor(FDynamicMesh3* MeshIn, FDynamicMeshUVOverlay* UVOverlayIn)
+{
+	Mesh = MeshIn;
+	UVOverlay = UVOverlayIn;
+	ensure(UVOverlay->GetParentMesh() == Mesh);
+}
+
 
 void FDynamicMeshUVEditor::CreateUVLayer(int32 LayerIndex)
 {
@@ -263,7 +271,6 @@ bool FDynamicMeshUVEditor::SetTriangleUVsFromExpMap(const TArray<int32>& Triangl
 		FMeshNormals::SmoothVertexNormals(Submesh, Options.NormalSmoothingRounds, Options.NormalSmoothingAlpha);
 	}
 
-
 	FFrame3d SeedFrame;
 	int32 FrameVertexID;
 	bool bFrameOK = EstimateGeodesicCenterFrameVertex(Submesh, SeedFrame, FrameVertexID, true);
@@ -313,7 +320,6 @@ bool FDynamicMeshUVEditor::SetTriangleUVsFromExpMap(const TArray<int32>& Triangl
 
 
 
-
 bool FDynamicMeshUVEditor::SetTriangleUVsFromExpMap(
 	const TArray<int32>& Triangles,
 	TFunctionRef<FVector3d(const FVector3d&)> PointTransform,
@@ -1037,6 +1043,165 @@ void FDynamicMeshUVEditor::SetTriangleUVsFromCylinderProjection(
 
 
 
+bool FDynamicMeshUVEditor::ScaleUVAreaTo3DArea(const TArray<int32>& Triangles, bool bRecenterAtOrigin)
+{
+	double Area3D = TMeshQueries<FDynamicMesh3>::GetVolumeArea(*Mesh, Triangles).Y;
+	if (FMathd::Abs(Area3D) < FMathf::Epsilon || FMathd::IsFinite(Area3D) == false)
+	{
+		return false;
+	}
+
+	TSet<int32> Elements;
+
+	double Area2D = 0.0;
+	FAxisAlignedBox2f UVBounds = FAxisAlignedBox2f::Empty();
+	for (int32 tid : Triangles)
+	{
+		if (UVOverlay->IsSetTriangle(tid))
+		{
+			FIndex3i UVTri = UVOverlay->GetTriangle(tid);
+			Elements.Add(UVTri.A); Elements.Add(UVTri.B); Elements.Add(UVTri.C);
+			FVector2f U = UVOverlay->GetElement(UVTri.A);
+			FVector2f V = UVOverlay->GetElement(UVTri.B);
+			FVector2f W = UVOverlay->GetElement(UVTri.C);
+			UVBounds.Contain(U); UVBounds.Contain(V); UVBounds.Contain(W);
+			Area2D += (double)VectorUtil::Area(U, V, W);
+		}
+	}
+	if (Elements.Num() == 0 || FMathd::Abs(Area2D) < FMathf::Epsilon ||  FMathd::IsFinite(Area2D) == false )
+	{
+		return false;
+	}
+
+	double UVScale = FMathd::Sqrt(Area3D) / FMathd::Sqrt(Area2D);
+	if (!FMathd::IsFinite(UVScale))
+	{
+		return false;
+	}
+
+	FVector2f ScaleOrigin = UVBounds.Center();
+	FVector2f Translation = (bRecenterAtOrigin) ? FVector2f::Zero() : ScaleOrigin;
+	for (int32 eid : Elements)
+	{
+		FVector2f UV = UVOverlay->GetElement(eid);
+		UV = (UV - ScaleOrigin) * UVScale + Translation;
+		UVOverlay->SetElement(eid, UV);
+	}
+
+	return true;
+}
+
+
+
+bool FDynamicMeshUVEditor::AutoOrientUVArea(const TArray<int32>& Triangles)
+{
+	TSet<int32> Elements;
+	TArray<FVector2f> UVs;
+	FAxisAlignedBox2f UVBounds = FAxisAlignedBox2f::Empty();
+	for (int32 tid : Triangles)
+	{
+		if (UVOverlay->IsSetTriangle(tid))
+		{
+			FIndex3i UVTri = UVOverlay->GetTriangle(tid);
+			for (int32 j = 0; j < 3; ++j)
+			{
+				bool bIsAlreadyInSet = false;
+				Elements.Add(UVTri[j], &bIsAlreadyInSet);
+				if (bIsAlreadyInSet == false)
+				{
+					FVector2f UV = UVOverlay->GetElement(UVTri[j]);
+					UVs.Add(UV);
+					UVBounds.Contain(UV);
+				}
+			}
+		}
+	}
+	int32 N = UVs.Num();
+	if (N == 0)
+	{
+		return false;
+	}
+
+	// shift to origin so we can skip subtract below
+	FVector2f BoxCenter = UVBounds.Center();
+	for ( int32 k = 0; k < N; ++k )
+	{
+		UVs[k] -= BoxCenter;
+	}
+
+	FConvexHull2f ConvexHull;
+	if (ConvexHull.Solve(UVs) == false)
+	{
+		return false;
+	}
+	TArray<int32> const& HullPointIndices = ConvexHull.GetPolygonIndices();
+	check(HullPointIndices[0] != HullPointIndices.Last());
+
+	TArray<FVector2f> HullPoints;
+	HullPoints.Reserve(HullPointIndices.Num());
+	for (int32 k = 0; k < HullPointIndices.Num(); ++k)
+	{
+		HullPoints.Add(UVs[HullPointIndices[k]]);
+	}
+	int32 NV = HullPoints.Num();
+
+
+	// This is basically a brute-force rotating-calipers implementation. Probably should be moved to a CompGeom class (and implemented more efficiently)
+	double MinBoxArea = UVBounds.Area();
+	FVector2f MinAxisDirection = FVector2f::UnitX();
+	bool bFoundSmallerBox = false;
+	for (int32 j = 0; j <= NV; ++j)
+	{
+		FVector2f A = HullPoints[j%NV], B = HullPoints[(j+1)%NV], C = HullPoints[(j+NV/2)%NV];
+		FVector2f Axis0 = B - A; 
+		float Dimension0 = Normalize(Axis0);
+		FVector2f Axis1 = PerpCW(Axis0);
+		float Dimension1 = (C - A).Dot(Axis1);
+
+		FInterval1f Interval0 = FInterval1f::Empty();
+		FInterval1f Interval1 = FInterval1f::Empty();
+
+		bool bAbort = false;
+		// Use modulo iteration here to try to grow the box more quickly, which (hopefully) hits the early-out faster
+		// (would be worth profiling to see if this is a good idea, due to cache coherency)
+		FModuloIteration Iter(NV);
+		uint32 Index;
+		while (Iter.GetNextIndex(Index))
+		{
+			Interval0.Contain(Axis0.Dot(HullPoints[Index]));
+			Dimension0 = FMathf::Max(Dimension0, Interval0.Length());
+			Interval1.Contain(Axis1.Dot(HullPoints[Index]));
+			Dimension1 = FMathf::Max(Dimension1, Interval1.Length());
+			if (Dimension0 * Dimension1 > MinBoxArea)
+			{
+				bAbort = true;
+				break;
+			}
+		}
+		if (bAbort == false && (Dimension0 * Dimension1 < MinBoxArea) )
+		{
+			MinBoxArea = Dimension0 * Dimension1;
+			MinAxisDirection = Axis0;
+			bFoundSmallerBox = true;
+		}
+	}
+
+	if (bFoundSmallerBox)
+	{
+		float RotationAngle = FMathf::Atan2(MinAxisDirection.Y, MinAxisDirection.X);
+		FMatrix2f RotMatrix = FMatrix2f::RotationRad(-RotationAngle);
+		for (int32 eid : Elements)
+		{
+			FVector2f UV = UVOverlay->GetElement(eid);
+			UV = RotMatrix * (UV - BoxCenter) + BoxCenter;
+			UVOverlay->SetElement(eid, UV);
+		}
+	}
+
+	return true;
+}
+
+
 bool FDynamicMeshUVEditor::QuickPack(int32 TargetTextureResolution, float GutterSize)
 {
 	// always split bowties before packing
diff --git a/Engine/Plugins/Runtime/GeometryProcessing/Source/DynamicMesh/Public/Parameterization/DynamicMeshUVEditor.h b/Engine/Plugins/Runtime/GeometryProcessing/Source/DynamicMesh/Public/Parameterization/DynamicMeshUVEditor.h
--- a/Engine/Plugins/Runtime/GeometryProcessing/Source/DynamicMesh/Public/Parameterization/DynamicMeshUVEditor.h
+++ b/Engine/Plugins/Runtime/GeometryProcessing/Source/DynamicMesh/Public/Parameterization/DynamicMeshUVEditor.h
@@ -37,6 +37,12 @@ public:
 	 */
 	explicit FDynamicMeshUVEditor(FDynamicMesh3* MeshIn, int32 UVLayerIndex, bool bCreateIfMissing);
 
+	/**
+	 * Construct UV Editor for a UV Overlay of the given Mesh.
+	 */
+	explicit FDynamicMeshUVEditor(FDynamicMesh3* MeshIn, FDynamicMeshUVOverlay* UVOverlayIn);
+
+
 	FDynamicMesh3* GetMesh() { return Mesh; }
 	const FDynamicMesh3* GetMesh() const { return Mesh; }
 
@@ -169,6 +175,20 @@ public:
 	void SetTriangleUVsFromCylinderProjection(const TArray<int32>& Triangles, TFunctionRef<FVector3d(const FVector3d&)> PointTransform, const FFrame3d& CylFrame, const FVector3d& CylDimensions, float CylinderSplitAngle, FUVEditResult* Result = nullptr);
 
 
+	/**
+	* Compute the UV-space and 3D area of the given Triangles, and then scale the UV area to be equvalent
+	* to the 3D area. Scaling is around the center of the UV bounding box.
+	* @param bRecenterAtOrigin if true, UVs are translated after scaling such that the UV bounding box center is at (0,0)
+	*/
+	bool ScaleUVAreaTo3DArea(const TArray<int32>& Triangles, bool bRecenterAtOrigin);
+
+	/**
+	* Compute an oriented UV-space bounding box for the given Triangles, and then rotate the UVs such that the
+	* box is aligned with the X axis. 
+	* The UV-space box is found by computing the 2D Convex Hull and then finding the hull edge that results in the minimal box.
+	*/
+	bool AutoOrientUVArea(const TArray<int32>& Triangles);
+
 
 	/**
 	 * Pack UVs into unit rectangle
diff --git a/Engine/Plugins/Runtime/GeometryProcessing/Source/DynamicMesh/Public/Parameterization/MeshRegionGraph.h b/Engine/Plugins/Runtime/GeometryProcessing/Source/DynamicMesh/Public/Parameterization/MeshRegionGraph.h
--- a/Engine/Plugins/Runtime/GeometryProcessing/Source/DynamicMesh/Public/Parameterization/MeshRegionGraph.h
+++ b/Engine/Plugins/Runtime/GeometryProcessing/Source/DynamicMesh/Public/Parameterization/MeshRegionGraph.h
@@ -108,6 +108,14 @@ public:
 	bool OptimizeBorders(int32 MaxRounds = 25);
 
 
+	/**
+	 * @return list of triangles in a Region
+	 * @warning This method leaves the FMeshRegionGraph in an invalid state, and should only be used as a means to avoid memory copies while extracting the triangle sets from a FMeshRegionGraph that will then be discarded
+	 */
+	TArray<int32>&& MoveRegionTris(int32 RegionIdx) { return MoveTemp(Regions[RegionIdx].Triangles); }
+
+
+
 public:
 
 	struct FNeighbour
