GeometryCore: add TNormalHistogram, simple spherical histogram based on TSphericalFibonacci point set
GeometryProcessing: add PatchBasedMeshUVGenerator, partial implementation of new AutoUV generator, currently generates UVs based on an initial patch decomposition computed externally, by incrementally merging those patches if possible.
ModelingTools: extend RecomputeUVsOp to optionally support island merging based on PatchBasedMeshUVGenerator. Add support for Normal smoothing in ExpMap, and UV Island rotation optimization. Expose new options/etc in RecomputeUVsTool, as well as add optional UV-space view.
#rnx

diff --git a/Engine/Plugins/Experimental/MeshModelingToolset/Source/MeshModelingTools/Private/RecomputeUVsTool.cpp b/Engine/Plugins/Experimental/MeshModelingToolset/Source/MeshModelingTools/Private/RecomputeUVsTool.cpp
--- a/Engine/Plugins/Experimental/MeshModelingToolset/Source/MeshModelingTools/Private/RecomputeUVsTool.cpp
+++ b/Engine/Plugins/Experimental/MeshModelingToolset/Source/MeshModelingTools/Private/RecomputeUVsTool.cpp
@@ -39,17 +39,19 @@ void URecomputeUVsTool::Setup()
 	InputMesh = MakeShared<FDynamicMesh3, ESPMode::ThreadSafe>();
 	*InputMesh = UE::ToolTarget::GetDynamicMeshCopy(Target);
 
+	FComponentMaterialSet MaterialSet = UE::ToolTarget::GetMaterialSet(Target);
+	FTransform TargetTransform = (FTransform)UE::ToolTarget::GetLocalToWorldTransform(Target);
+
 	Preview = NewObject<UMeshOpPreviewWithBackgroundCompute>(this);
 	Preview->Setup(this->TargetWorld, this);
 	Preview->PreviewMesh->SetTangentsMode(EDynamicMeshComponentTangentsMode::AutoCalculated);
 	Preview->PreviewMesh->ReplaceMesh(*InputMesh);
-	Preview->ConfigureMaterials(UE::ToolTarget::GetMaterialSet(Target).Materials,
-								ToolSetupUtil::GetDefaultWorkingMaterial(GetToolManager()) );
-	Preview->PreviewMesh->SetTransform((FTransform)UE::ToolTarget::GetLocalToWorldTransform(Target));
+	Preview->ConfigureMaterials(MaterialSet.Materials, ToolSetupUtil::GetDefaultWorkingMaterial(GetToolManager()) );
+	Preview->PreviewMesh->SetTransform(TargetTransform);
 
 	Preview->OnMeshUpdated.AddLambda([this](UMeshOpPreviewWithBackgroundCompute* Op)
 	{
-		MaterialSettings->UpdateMaterials();
+		OnPreviewMeshUpdated();
 	});
 
 	UE::ToolTarget::HideSourceObject(Target);
@@ -85,6 +87,18 @@ void URecomputeUVsTool::Setup()
 	MaterialSettings->UpdateMaterials();
 	Preview->OverrideMaterial = MaterialSettings->GetActiveOverrideMaterial();
 
+	if (bCreateUVLayoutViewOnSetup)
+	{
+		UVLayoutView = NewObject<UUVLayoutPreview>(this);
+		UVLayoutView->CreateInWorld(TargetWorld);
+		UVLayoutView->SetSourceMaterials(MaterialSet);
+		UVLayoutView->SetSourceWorldPosition(TargetTransform, UE::ToolTarget::GetTargetActor(Target)->GetComponentsBoundingBox());
+		UVLayoutView->Settings->bVisible = false;
+		UVLayoutView->Settings->bShowWireframe = false;
+		UVLayoutView->Settings->RestoreProperties(this, TEXT("RecomputeUVsTool"));
+		AddToolPropertySource(UVLayoutView->Settings);
+	}
+
 	Preview->InvalidateResult();    // start compute
 
 	SetToolDisplayName(LOCTEXT("ToolNameLocal", "UV Unwrap"));
@@ -114,6 +128,12 @@ void URecomputeUVsTool::OnPropertyModified(UObject* PropertySet, FProperty* Prop
 
 void URecomputeUVsTool::Shutdown(EToolShutdownType ShutdownType)
 {
+	if (UVLayoutView)
+	{
+		UVLayoutView->Settings->SaveProperties(this, TEXT("RecomputeUVsTool"));
+		UVLayoutView->Disconnect();
+	}
+
 	UVChannelProperties->SaveProperties(this);
 	Settings->SaveProperties(this);
 	PolygroupLayerProperties->RestoreProperties(this, TEXT("RecomputeUVsTool"));
@@ -135,9 +155,23 @@ void URecomputeUVsTool::Shutdown(EToolShutdownType ShutdownType)
 
 }
 
+void URecomputeUVsTool::Render(IToolsContextRenderAPI* RenderAPI)
+{
+	if (UVLayoutView)
+	{
+		UVLayoutView->Render(RenderAPI);
+	}
+}
+
+
 void URecomputeUVsTool::OnTick(float DeltaTime)
 {
 	Preview->Tick(DeltaTime);
+
+	if (UVLayoutView)
+	{
+		UVLayoutView->OnTick(DeltaTime);
+	}
 }
 
 bool URecomputeUVsTool::CanAccept() const
@@ -156,6 +190,15 @@ TUniquePtr<FDynamicMeshOperator> URecomputeUVsTool::MakeNewOperator()
 	RecomputeUVsOp->IslandMode = (UE::Geometry::ERecomputeUVsIslandMode)(int)Settings->IslandMode;
 	RecomputeUVsOp->UnwrapType = (UE::Geometry::ERecomputeUVsUnwrapType)(int)Settings->UnwrapType;
 
+	RecomputeUVsOp->bAutoRotate = (Settings->AutoRotation == ERecomputeUVsToolOrientationMode::MinBoxBounds);
+
+	RecomputeUVsOp->NormalSmoothingRounds = Settings->SmoothingSteps;
+	RecomputeUVsOp->NormalSmoothingAlpha = Settings->SmoothingAlpha;
+
+	RecomputeUVsOp->bMergingOptimization = Settings->bIslandMerging;
+	RecomputeUVsOp->MergingThreshold = Settings->MergingThreshold;
+	RecomputeUVsOp->MaxNormalDeviationDeg = Settings->MaxAngleDeviation;
+
 	RecomputeUVsOp->bPackUVs = Settings->bAutoPack;
 	if (Settings->bAutoPack)
 	{
@@ -210,6 +253,24 @@ void URecomputeUVsTool::UpdateActiveGroupLayer()
 		ensureMsgf(FoundAttrib, TEXT("Selected Attribute Not Found! Falling back to Default group layer."));
 		ActiveGroupSet = MakeShared<UE::Geometry::FPolygroupSet, ESPMode::ThreadSafe>(InputMesh.Get(), FoundAttrib);
 	}
+}
+
+
+void URecomputeUVsTool::OnPreviewMeshUpdated()
+{
+	if (UVLayoutView)
+	{
+		int32 UVChannel = UVChannelProperties ? UVChannelProperties->GetSelectedChannelIndex(true) : 0;
+		Preview->PreviewMesh->ProcessMesh([&](const FDynamicMesh3& NewMesh)
+		{
+			UVLayoutView->UpdateUVMesh(&NewMesh, UVChannel);
+		});
+	}
+
+	if (MaterialSettings)
+	{
+		MaterialSettings->UpdateMaterials();
+	}
 
 }
 
diff --git a/Engine/Plugins/Experimental/MeshModelingToolset/Source/MeshModelingTools/Public/RecomputeUVsTool.h b/Engine/Plugins/Experimental/MeshModelingToolset/Source/MeshModelingTools/Public/RecomputeUVsTool.h
--- a/Engine/Plugins/Experimental/MeshModelingToolset/Source/MeshModelingTools/Public/RecomputeUVsTool.h
+++ b/Engine/Plugins/Experimental/MeshModelingToolset/Source/MeshModelingTools/Public/RecomputeUVsTool.h
@@ -10,6 +10,7 @@
 #include "Properties/MeshUVChannelProperties.h"
 #include "PropertySets/PolygroupLayersProperties.h"
 #include "Polygroups/PolygroupSet.h"
+#include "Drawing/UVLayoutPreview.h"
 
 #include "RecomputeUVsTool.generated.h"
 
@@ -39,9 +40,9 @@ enum class ERecomputeUVsUnwrapType
 {
 	// Values must match UE::Geometry::ERecomputeUVsUnwrapType
 
-	/** */
+	/** ExpMap UV Flattening Algorithm is very fast but has limited ability to control stretching/distortion */
 	ExpMap = 0,
-	/** */
+	/** Conformal UV Flattening algorithm is increasingly-expensive on large islands but better at distortion control */
 	Conformal = 1
 };
 
@@ -52,12 +53,23 @@ enum class ERecomputeUVsIslandMode
 {
 	// Values must match UE::Geometry::ERecomputeUVsIslandMode
 
-	/** */
+	/** Use Active Mesh Polygroups Layer to define initial UV Islands */
 	PolyGroups = 0,
-	/** */
+	/** Use Existing UV Layer to define UV Islands (ie re-solve UV flattening based on existing UVs) */
 	ExistingUVs = 1
 };
 
+
+UENUM()
+enum class ERecomputeUVsToolOrientationMode
+{
+	/**  */
+	None,
+	/**  */
+	MinBoxBounds
+};
+
+
 UENUM()
 enum class ERecomputeUVsToolUVScaleMode
 {
@@ -76,15 +88,43 @@ class MESHMODELINGTOOLS_API URecomputeUVsToolProperties : public UInteractiveToo
 	GENERATED_BODY()
 public:
 
-	UPROPERTY(EditAnywhere, Category = "Compute UVs")
+	/** IslandMode determines which patches of triangles will be used as initial UV Islands */
+	UPROPERTY(EditAnywhere, Category = "UV Generation")
 	ERecomputeUVsIslandMode IslandMode = ERecomputeUVsIslandMode::PolyGroups;
 
-	UPROPERTY(EditAnywhere, Category = "Compute UVs")
+	/** Which type of UV flattening algorithm to use */
+	UPROPERTY(EditAnywhere, Category = "UV Generation", meta = (EditCondition = "bIslandMerging == false"))
 	ERecomputeUVsUnwrapType UnwrapType = ERecomputeUVsUnwrapType::Conformal;
 
+	/** Which type of automatic island orientation to use */
+	UPROPERTY(EditAnywhere, Category = "UV Generation")
+	ERecomputeUVsToolOrientationMode AutoRotation = ERecomputeUVsToolOrientationMode::MinBoxBounds;
+
+	/** If Island Merging is enabled, then the initial UV islands will be merged into larger islands if it does not increase distortion/stretching beyond the limits below */
+	UPROPERTY(EditAnywhere, Category = "UV Generation")
+	bool bIslandMerging = false;
+
+	/** Distortion/Stretching Threshold for island merging - larger values increase the allowable UV stretching */
+	UPROPERTY(EditAnywhere, Category = "UV Generation", meta = (UIMin = "1.0", UIMax = "5.0", ClampMin = "1.0", EditCondition = "bIslandMerging == true"))
+	float MergingThreshold = 1.5f;
+
+	/** UV islands will not be merged if their average face normals deviate by larger than this amount */
+	UPROPERTY(EditAnywhere, Category = "UV Generation", meta = (UIMin = "0.0", UIMax = "90.0", ClampMin = "0.0", ClampMax = "180.0", EditCondition = "bIslandMerging == true"))
+	float MaxAngleDeviation = 45.0f;
+
+	/** Number of smoothing steps to apply in the ExpMap UV Generation method (Smoothing slightly increases distortion but produces more stable results) */
+	UPROPERTY(EditAnywhere, Category = "ExpMap UV Generator", meta = (UIMin = "0", UIMax = "25", ClampMin = "0", ClampMax = "1000", EditCondition = "UnwrapType == ERecomputeUVsUnwrapType::ExpMap || bIslandMerging == true" ))
+	int SmoothingSteps = 5;
+
+	/** Smoothing parameter, larger values result in faster smoothing in each step */
+	UPROPERTY(EditAnywhere, Category = "ExpMap UV Generator", meta = (UIMin = "0", UIMax = "1.0", ClampMin = "0", ClampMax = "1.0", EditCondition = "UnwrapType == ERecomputeUVsUnwrapType::ExpMap || bIslandMerging == true"))
+	float SmoothingAlpha = 0.25f;
+
+	/** If enabled, result UVs are automatically packed into the standard UV 0-1 square */
 	UPROPERTY(EditAnywhere, Category = "UV Layout")
 	bool bAutoPack = true;
 
+	/** Target texture resolution used for UV packing, which determines gutter size */
 	UPROPERTY(EditAnywhere, Category = "UV Layout", meta = (UIMin = "64", UIMax = "2048", ClampMin = "2", ClampMax = "4096", EditCondition = bAutoPack))
 	int TextureResolution = 1024;
 
@@ -113,6 +153,7 @@ public:
 	virtual void Setup() override;
 	virtual void Shutdown(EToolShutdownType ShutdownType) override;
 
+	virtual void Render(IToolsContextRenderAPI* RenderAPI) override;
 	virtual void OnTick(float DeltaTime) override;
 
 	virtual bool HasCancel() const override { return true; }
@@ -137,6 +178,13 @@ protected:
 	UPROPERTY()
 	TObjectPtr<UExistingMeshMaterialProperties> MaterialSettings = nullptr;
 
+	UPROPERTY()
+	bool bCreateUVLayoutViewOnSetup = true;
+
+	UPROPERTY()
+	TObjectPtr<UUVLayoutPreview> UVLayoutView = nullptr;
+
+
 protected:
 	UPROPERTY()
 	TObjectPtr<UMeshOpPreviewWithBackgroundCompute> Preview = nullptr;
@@ -148,4 +196,6 @@ protected:
 	TSharedPtr<UE::Geometry::FPolygroupSet, ESPMode::ThreadSafe> ActiveGroupSet;
 	void OnSelectedGroupLayerChanged();
 	void UpdateActiveGroupLayer();
+
+	void OnPreviewMeshUpdated();
 };
diff --git a/Engine/Plugins/Experimental/MeshModelingToolset/Source/ModelingOperators/Private/ParameterizationOps/RecomputeUVsOp.cpp b/Engine/Plugins/Experimental/MeshModelingToolset/Source/ModelingOperators/Private/ParameterizationOps/RecomputeUVsOp.cpp
--- a/Engine/Plugins/Experimental/MeshModelingToolset/Source/ModelingOperators/Private/ParameterizationOps/RecomputeUVsOp.cpp
+++ b/Engine/Plugins/Experimental/MeshModelingToolset/Source/ModelingOperators/Private/ParameterizationOps/RecomputeUVsOp.cpp
@@ -3,14 +3,26 @@
 #include "ParameterizationOps/RecomputeUVsOp.h"
 
 #include "DynamicMesh/DynamicMeshAttributeSet.h"
+#include "DynamicMesh/MeshNormals.h"
 #include "Selections/MeshConnectedComponents.h"
 #include "Parameterization/MeshLocalParam.h"
 #include "Parameterization/DynamicMeshUVEditor.h"
+#include "Parameterization/PatchBasedMeshUVGenerator.h"
+
+#include "Parameterization/MeshRegionGraph.h"
+#include "DynamicSubmesh3.h"
+#include "Parameterization/MeshDijkstra.h"
+
+#include "Async/ParallelFor.h"
+
 
 #include "ExplicitUseGeometryMathTypes.h"		// using UE::Geometry::(math types)
 using namespace UE::Geometry;
 
 
+#define LOCTEXT_NAMESPACE "RecomputeUVsOp"
+
+
 void FRecomputeUVsOp::NormalizeUVAreas(const FDynamicMesh3& Mesh, FDynamicMeshUVOverlay* Overlay, float GlobalScale)
 {
 	FMeshConnectedComponents UVComponents(&Mesh);
@@ -54,7 +66,7 @@ void FRecomputeUVsOp::NormalizeUVAreas(const FDynamicMesh3& Mesh, FDynamicMeshUV
 		for (int elemid : UVElements)
 		{
 			FVector2d UV = FVector2d(Overlay->GetElement(elemid));
-			UV = (UV - ComponentOrigin) * LinearScale + ComponentOrigin;
+			UV = (UV - ComponentOrigin) * LinearScale;
 			Overlay->SetElement(elemid, FVector2f(UV));
 		}
 	}
@@ -62,10 +74,32 @@ void FRecomputeUVsOp::NormalizeUVAreas(const FDynamicMesh3& Mesh, FDynamicMeshUV
 
 
 void FRecomputeUVsOp::CalculateResult(FProgressCancel* Progress)
+{
+	NewResultInfo = FGeometryResult(EGeometryResultType::InProgress);
+
+	bool bOK = false;
+	if (bMergingOptimization)
+	{
+		bOK = CalculateResult_RegionOptimization(Progress);
+	}
+	else
+	{
+		bOK = CalculateResult_Basic(Progress);
+	}
+
+	NewResultInfo.SetSuccess(bOK, Progress);
+	SetResultInfo(NewResultInfo);
+
+}
+
+
+
+
+bool FRecomputeUVsOp::CalculateResult_Basic(FProgressCancel* Progress)
 {
 	if (!InputMesh.IsValid())
 	{
-		return;
+		return false;
 	}
 
 	ResultMesh = MakeUnique<FDynamicMesh3>(*InputMesh);
@@ -75,12 +109,12 @@ void FRecomputeUVsOp::CalculateResult(FProgressCancel* Progress)
 	FDynamicMeshUVOverlay* UseOverlay = UVEditor.GetOverlay();
 	if (ensure(UseOverlay != nullptr) == false)
 	{
-		return;
+		return false;
 	}
 
 	if (Progress && Progress->Cancelled())
 	{
-		return;
+		return false;
 	}
 
 	// find group-connected-components
@@ -109,40 +143,71 @@ void FRecomputeUVsOp::CalculateResult(FProgressCancel* Progress)
 
 	if (Progress && Progress->Cancelled())
 	{
-		return;
+		return false;
 	}
 
 	// TODO: the solves here could be done in parallel if we pre-allocated the island element IDs
 
 	int32 NumComponents = ConnectedComponents.Num();
+	TArray<bool> bComponentSolved;
+	bComponentSolved.Init(false, NumComponents);
 	int32 SuccessCount = 0;
 	for (int32 k = 0; k < NumComponents; ++k)
 	{
 		const TArray<int32>& ComponentTris = ConnectedComponents[k].Indices;
 
-		bool bComputedUVs = false;
+		bComponentSolved[k] = false;
 		switch (UnwrapType)
 		{
 		case ERecomputeUVsUnwrapType::ExpMap:
-			bComputedUVs = UVEditor.SetTriangleUVsFromExpMap(ComponentTris);
-			break;
+		{
+			FDynamicMeshUVEditor::FExpMapOptions Options;
+			Options.NormalSmoothingRounds = this->NormalSmoothingRounds;
+			Options.NormalSmoothingAlpha = this->NormalSmoothingAlpha;
+			bComponentSolved[k] = UVEditor.SetTriangleUVsFromExpMap(ComponentTris, Options);
+		}
+		break;
 
 		case ERecomputeUVsUnwrapType::ConformalFreeBoundary:
-			bComputedUVs = UVEditor.SetTriangleUVsFromFreeBoundaryConformal(ComponentTris);
+			bComponentSolved[k] = UVEditor.SetTriangleUVsFromFreeBoundaryConformal(ComponentTris);
+			if (bComponentSolved[k])
+			{
+				UVEditor.ScaleUVAreaTo3DArea(ComponentTris, true);
+			}
 			break;
 		}
 
-		if (bComputedUVs)
+		if (bComponentSolved[k])
 		{
 			SuccessCount++;
 		}
 
 		if (Progress && Progress->Cancelled())
 		{
-			return;
+			return false;
 		}
 	}
 
+	if (bAutoRotate)
+	{
+		ParallelFor(NumComponents, [&](int32 k)
+		{
+			if (Progress && Progress->Cancelled())
+			{
+				return;
+			}
+			if (bComponentSolved[k])
+			{
+				const TArray<int32>& ComponentTris = ConnectedComponents[k].Indices;
+				UVEditor.AutoOrientUVArea(ComponentTris);
+			};
+		});
+	}
+	if (Progress && Progress->Cancelled())
+	{
+		return false;
+	}
+
 
 	if (bNormalizeAreas)
 	{
@@ -150,9 +215,157 @@ void FRecomputeUVsOp::CalculateResult(FProgressCancel* Progress)
 		NormalizeUVAreas(*ResultMesh, UseOverlay, AreaScaling);
 	}
 
+	if (Progress && Progress->Cancelled())
+	{
+		return false;
+	}
+
 	if (bPackUVs)
 	{
 		UVEditor.QuickPack(PackingTextureResolution, PackingGutterWidth);
 	}
 
-}
\ No newline at end of file
+	if (Progress && Progress->Cancelled())
+	{
+		return false;
+	}
+
+	return true;
+}
+
+
+
+
+
+
+bool FRecomputeUVsOp::CalculateResult_RegionOptimization(FProgressCancel* Progress)
+{
+	if (!InputMesh.IsValid())
+	{
+		return false;
+	}
+
+	ResultMesh = MakeUnique<FDynamicMesh3>(*InputMesh);
+
+	FDynamicMesh3* BaseMesh = ResultMesh.Get();
+	FDynamicMeshUVEditor UVEditor(BaseMesh, UVLayer, true);
+	FDynamicMeshUVOverlay* UseOverlay = UVEditor.GetOverlay();
+	if (ensure(UseOverlay != nullptr) == false)
+	{
+		return false;
+	}
+
+	if (Progress && Progress->Cancelled())
+	{
+		return false;
+	}
+
+	// initialize UV generation with connected components
+	FMeshConnectedComponents ConnectedComponents(ResultMesh.Get());
+	if (IslandMode == ERecomputeUVsIslandMode::PolyGroups)
+	{
+		if (InputGroups != nullptr)
+		{
+			ConnectedComponents.FindConnectedTriangles([this](int32 CurTri, int32 NbrTri) {
+				return InputGroups->GetTriangleGroup(CurTri) == InputGroups->GetTriangleGroup(NbrTri);
+			});
+		}
+		else
+		{
+			ConnectedComponents.FindConnectedTriangles([this](int32 CurTri, int32 NbrTri) {
+				return ResultMesh->GetTriangleGroup(CurTri) == ResultMesh->GetTriangleGroup(NbrTri);
+			});
+		}
+	}
+	else
+	{
+		ConnectedComponents.FindConnectedTriangles([&](int32 Triangle0, int32 Triangle1) {
+			return UseOverlay->AreTrianglesConnected(Triangle0, Triangle1);
+		});
+	}
+
+	if (Progress && Progress->Cancelled())
+	{
+		return false;
+	}
+
+
+	FPatchBasedMeshUVGenerator UVGenerator;
+	UVGenerator.MergingThreshold = this->MergingThreshold;
+	UVGenerator.CompactnessThreshold = this->CompactnessThreshold;
+	UVGenerator.MaxNormalDeviationDeg = this->MaxNormalDeviationDeg;
+	UVGenerator.NormalSmoothingRounds = this->NormalSmoothingRounds;
+	UVGenerator.NormalSmoothingAlpha = this->NormalSmoothingAlpha;
+
+	TArray<TArray<int32>> UVIslandTriangleSets;
+	bool bIslandsOK = UVGenerator.ComputeIslandsByRegionMerging(*BaseMesh, *UseOverlay, ConnectedComponents, UVIslandTriangleSets, Progress);
+
+	if (bIslandsOK == false)
+	{
+		NewResultInfo.AddError(FGeometryError(0, LOCTEXT("IslandMergingFailed", "Failed to merge UV islands")));
+		return false;
+	}
+
+	if (Progress && Progress->Cancelled())
+	{
+		return false;
+	}
+
+	TArray<bool> bIslandUVsValid;
+	int32 NumSolvesFailed = UVGenerator.ComputeUVsFromTriangleSets(*BaseMesh, *UseOverlay, UVIslandTriangleSets, bIslandUVsValid, Progress);
+	if (NumSolvesFailed > 0)
+	{
+		NewResultInfo.AddWarning(FGeometryWarning(0, LOCTEXT("PartialSolvesFailed", "Failed to compute UVs for some UV islands")));
+	}
+
+	if (Progress && Progress->Cancelled())
+	{
+		return false;
+	}
+
+	if (bAutoRotate)
+	{
+		ParallelFor(UVIslandTriangleSets.Num(), [&](int32 k)
+		{
+			if (Progress && Progress->Cancelled())
+			{
+				return;
+			}
+			if (bIslandUVsValid[k])
+			{
+				UVEditor.AutoOrientUVArea(UVIslandTriangleSets[k]);
+			}
+		});
+	}
+	if (Progress && Progress->Cancelled())
+	{
+		return false;
+	}
+
+	if (bNormalizeAreas)
+	{
+		// todo should be a DynamicUVEditor function?
+		NormalizeUVAreas(*ResultMesh, UseOverlay, AreaScaling);
+	}
+
+	if (Progress && Progress->Cancelled())
+	{
+		return false;
+	}
+
+	if (bPackUVs)
+	{
+		UVEditor.QuickPack(PackingTextureResolution, PackingGutterWidth);
+	}
+
+	if (Progress && Progress->Cancelled())
+	{
+		return false;
+	}
+
+	return true;
+}
+
+
+
+#undef LOCTEXT_NAMESPACE
\ No newline at end of file
diff --git a/Engine/Plugins/Experimental/MeshModelingToolset/Source/ModelingOperators/Public/ParameterizationOps/RecomputeUVsOp.h b/Engine/Plugins/Experimental/MeshModelingToolset/Source/ModelingOperators/Public/ParameterizationOps/RecomputeUVsOp.h
--- a/Engine/Plugins/Experimental/MeshModelingToolset/Source/ModelingOperators/Public/ParameterizationOps/RecomputeUVsOp.h
+++ b/Engine/Plugins/Experimental/MeshModelingToolset/Source/ModelingOperators/Public/ParameterizationOps/RecomputeUVsOp.h
@@ -45,6 +45,9 @@ public:
 	// source groups (optional)
 	TSharedPtr<UE::Geometry::FPolygroupSet, ESPMode::ThreadSafe> InputGroups;
 
+	// orientation control
+	bool bAutoRotate = true;
+
 	// area scaling
 	bool bNormalizeAreas = true;
 	float AreaScaling = 1.0;
@@ -60,6 +63,20 @@ public:
 	ERecomputeUVsUnwrapType UnwrapType = ERecomputeUVsUnwrapType::ExpMap;
 	ERecomputeUVsIslandMode IslandMode = ERecomputeUVsIslandMode::PolyGroups;
 
+	// 
+	// ExpMap Options
+	//
+	int32 NormalSmoothingRounds = 0;
+	double NormalSmoothingAlpha = 0.25;
+
+	//
+	// Patch Merging options
+	//
+	bool bMergingOptimization = false;
+	double MergingThreshold = 1.5;
+	double CompactnessThreshold = 9999999.0;		// effectively disabled as it usually is not a good idea
+	double MaxNormalDeviationDeg = 45.0;
+
 	// set ability on protected transform.
 	void SetTransform(const FTransform3d& XForm)
 	{
@@ -75,7 +92,12 @@ public:
 
 protected:
 
+	FGeometryResult NewResultInfo;
+
 	void NormalizeUVAreas(const FDynamicMesh3& Mesh, FDynamicMeshUVOverlay* Overlay, float GlobalScale = 1.0f);
+
+	virtual bool CalculateResult_Basic(FProgressCancel* Progress);
+	virtual bool CalculateResult_RegionOptimization(FProgressCancel* Progress);
 };
 
 } // end namespace UE::Geometry
diff --git a/Engine/Plugins/Runtime/GeometryProcessing/Source/DynamicMesh/Private/Parameterization/PatchBasedMeshUVGenerator.cpp b/Engine/Plugins/Runtime/GeometryProcessing/Source/DynamicMesh/Private/Parameterization/PatchBasedMeshUVGenerator.cpp
new file mode 100644
--- /dev/null
+++ b/Engine/Plugins/Runtime/GeometryProcessing/Source/DynamicMesh/Private/Parameterization/PatchBasedMeshUVGenerator.cpp
@@ -0,0 +1,478 @@
+// Copyright Epic Games, Inc. All Rights Reserved.
+
+#include "Parameterization/PatchBasedMeshUVGenerator.h"
+
+#include "DynamicMesh/MeshNormals.h"
+#include "DynamicSubmesh3.h"
+#include "Sampling/NormalHistogram.h"
+
+#include "Parameterization/MeshLocalParam.h"
+#include "Parameterization/DynamicMeshUVEditor.h"
+#include "Parameterization/MeshRegionGraph.h"
+#include "Parameterization/MeshDijkstra.h"
+
+#include "Async/ParallelFor.h"
+
+#include "ExplicitUseGeometryMathTypes.h"		// using UE::Geometry::(math types)
+using namespace UE::Geometry;
+
+
+struct FUVAreaMetrics
+{
+	double AvgSqrAreaDistortion = 1.0;
+	FVector3d AreaWeightedNormal = FVector3d::Zero();
+	double NormalSpread = 1.0;
+
+	FAxisAlignedBox2d BoundsUV = FAxisAlignedBox2d::Empty();
+	FAxisAlignedBox3d Bounds3D = FAxisAlignedBox3d::Empty();
+
+	double SurfAreaSum = 0.0;
+	double UVAreaSum = 0.0;
+	double UVPerimeter = 0.0;
+	double Compactness = 1.0;
+
+	static FUVAreaMetrics Invalid() {
+		FUVAreaMetrics Tmp; Tmp.AvgSqrAreaDistortion = TNumericLimits<double>::Max(); return Tmp;
+	}
+};
+
+
+static FUVAreaMetrics Compute_UV_Area_Metrics(const FDynamicMesh3* Mesh, const TArray<int32>& Triangles, int32 NormalSmoothingRounds, double NormalSmoothingAlpha)
+{
+	FDynamicSubmesh3 SubmeshCalc(Mesh, Triangles, (int)EMeshComponents::None, false);
+	FDynamicMesh3& Submesh = SubmeshCalc.GetSubmesh();
+	FMeshNormals::QuickComputeVertexNormals(Submesh);
+	if (NormalSmoothingRounds > 0)
+	{
+		FMeshNormals::SmoothVertexNormals(Submesh, NormalSmoothingRounds, NormalSmoothingAlpha);
+	}
+
+
+	TArray<FVector2d> SeedPoints;
+	TArray<bool> BoundaryVerts;
+	BoundaryVerts.Init(false, Submesh.MaxVertexID());
+	for (int32 eid : Submesh.EdgeIndicesItr())
+	{
+		if (Submesh.IsBoundaryEdge(eid))
+		{
+			FIndex2i BoundaryV = Submesh.GetEdgeV(eid);
+			for (int32 j = 0; j < 2; ++j)
+			{
+				if (BoundaryVerts[BoundaryV[j]] == false)
+				{
+					BoundaryVerts[BoundaryV[j]] = true;
+					SeedPoints.Add(FVector2d(BoundaryV[j], 0.0));
+				}
+			}
+		}
+	}
+
+	FFrame3d SeedFrame;
+	int32 SeedVID = *Submesh.VertexIndicesItr().begin();
+	if (SeedPoints.Num() > 0)
+	{
+		TMeshDijkstra<FDynamicMesh3> Dijkstra(&Submesh);
+		Dijkstra.ComputeToMaxDistance(SeedPoints, TNumericLimits<float>::Max());
+		SeedVID = Dijkstra.GetMaxGraphDistancePointID();
+		if (ensure(Submesh.IsVertex(SeedVID)) == false)
+		{
+			SeedVID = *Submesh.VertexIndicesItr().begin();
+		}
+	}
+
+	FVector3d Normal = FMeshNormals::ComputeVertexNormal(Submesh, SeedVID);
+	SeedFrame = Submesh.GetVertexFrame(SeedVID, false, &Normal);
+
+	TMeshLocalParam<FDynamicMesh3> Param(&Submesh);
+	Param.ParamMode = ELocalParamTypes::ExponentialMapUpwindAvg;
+	Param.ComputeToMaxDistance(SeedVID, SeedFrame, TNumericLimits<float>::Max());
+
+	double SumAreaDistortion = 0.0;
+	double SumDistortionsSqr = 0.0;
+	double Count = 0;
+	double MaxAreaDistortion = 1.0;
+
+	const double ClampAreaDistortionMetric = 100.0;
+
+	FUVAreaMetrics Metrics;
+
+	TNormalHistogram<double> NormalHistogram(128);
+
+	// TODO:
+	//   - compute bounding box/circle, compare to UV area (convexity metric - could also consider in 3D?)
+
+	for (int32 tid : Submesh.TriangleIndicesItr())
+	{
+		FIndex3i Tri = Submesh.GetTriangle(tid);
+
+		FVector3d TriNormal, Centroid; double Area3D;
+		Submesh.GetTriInfo(tid, TriNormal, Area3D, Centroid);
+
+		Metrics.AreaWeightedNormal += Area3D * TriNormal;
+		Metrics.SurfAreaSum += Area3D;
+		Metrics.Bounds3D.Contain(Submesh.GetTriBounds(tid));
+		NormalHistogram.Count(TriNormal, Area3D);
+
+		FVector2d UVTri[3];
+		UVTri[0] = Param.GetUV(Tri.A);
+		UVTri[1] = Param.GetUV(Tri.B);
+		UVTri[2] = Param.GetUV(Tri.C);
+		// TODO: this param seems to emit flipped areas ?!?
+		double SignedArea2D = VectorUtil::SignedArea(UVTri[0], UVTri[1], UVTri[2]);
+		double Area2D = VectorUtil::Area(UVTri[0], UVTri[1], UVTri[2]);
+
+		// TODO: penalize UV flips?
+
+		Metrics.UVAreaSum += Area2D;
+		Metrics.BoundsUV.Contain(UVTri[0]);
+		Metrics.BoundsUV.Contain(UVTri[1]);
+		Metrics.BoundsUV.Contain(UVTri[2]);
+
+		double Distortion = FMath::Max(Area3D, Area2D) / FMath::Min(Area3D, Area2D);
+
+		Distortion = FMath::Min(Distortion, ClampAreaDistortionMetric);
+		if (SignedArea2D > 0)
+		{
+			Distortion = ClampAreaDistortionMetric;
+		}
+
+		SumAreaDistortion += Area3D * Distortion;
+		SumDistortionsSqr += Area3D * Distortion * Distortion;
+		MaxAreaDistortion = FMath::Max(MaxAreaDistortion, Distortion);
+
+		FIndex3i TriEdges = Submesh.GetTriEdges(tid);
+		for (int j = 0; j < 3; ++j)
+		{
+			if (Submesh.IsBoundaryEdge(TriEdges[j]))
+			{
+				Metrics.UVPerimeter += UVTri[j].Distance(UVTri[(j + 1) % 3]);
+			}
+		}
+
+		Count++;
+	}
+
+
+	double AvgAreaDistortion = SumAreaDistortion / Metrics.SurfAreaSum;
+	double AvgSqrAreaDistortion = SumDistortionsSqr / Metrics.SurfAreaSum;
+	double L2AreaDistortion = FMathd::Sqrt(SumDistortionsSqr);
+
+	Metrics.AvgSqrAreaDistortion = AvgSqrAreaDistortion;
+
+	Metrics.AreaWeightedNormal.Normalize();
+	Metrics.NormalSpread = NormalHistogram.WeightedSpreadMetric();
+
+	// things to try:
+	//  - quick OBB fit to try to get rectangular bounds
+	double IdealBoxArea = FMathd::Sqrt(Metrics.UVAreaSum);
+	double RealBoxArea = Metrics.BoundsUV.Area();
+	Metrics.Compactness = FMathd::Max(IdealBoxArea, RealBoxArea) / FMathd::Min(IdealBoxArea, RealBoxArea);
+	Metrics.Compactness = FMathd::Sqrt(Metrics.Compactness);
+
+	//double IdealUVSquareArea = (Metrics.UVPerimeter/4.0) * (Metrics.UVPerimeter/4.0);		// if patch was square, this would be it's area
+	//double SquareCompactness = Metrics.UVAreaSum / IdealUVSquareArea;
+
+	//double UVPerimeterCircleRadius = (Metrics.UVPerimeter / (2.0 * FMathd::Pi));
+	//double UVCircleArea = FMathd::Pi * UVPerimeterCircleRadius * UVPerimeterCircleRadius;
+	//Metrics.Compactness = UVCircleArea / Metrics.UVAreaSum;
+	//ensure(Metrics.Compactness >= 1.0);
+
+	return Metrics;
+}
+
+
+
+
+
+bool FPatchBasedMeshUVGenerator::ComputeIslandsByRegionMerging(
+	FDynamicMesh3& TargetMesh,
+	FDynamicMeshUVOverlay& TargetUVOverlay,
+	const FMeshConnectedComponents& ConnectedComponents,
+	TArray<TArray<int32>>& IslandsOut,
+	FProgressCancel* Progress)
+{
+	if (Progress && Progress->Cancelled())
+	{
+		return false;
+	}
+
+	FMeshRegionGraph RegionGraph;
+	RegionGraph.BuildFromComponents(TargetMesh, ConnectedComponents, [&](int32 k) { return k; } );
+
+	if (Progress && Progress->Cancelled())
+	{
+		return false;
+	}
+
+	TSet<int32> RemainingRegions;
+	for (int32 k = 0; k < RegionGraph.MaxRegionIndex(); ++k)
+	{
+		if (RegionGraph.IsRegion(k))
+		{
+			RemainingRegions.Add(k);
+		}
+	}
+
+	FDynamicMeshUVEditor::FExpMapOptions Options;
+	Options.NormalSmoothingRounds = this->NormalSmoothingRounds;
+	Options.NormalSmoothingAlpha = this->NormalSmoothingAlpha;
+
+	// TODO: use constants instead of relative/multipliers
+	//double MergeThreshold = 4.0;
+	//double NeverMergeThreshold = 25.0;
+	double MergeDistortionThreshold = this->MergingThreshold;
+	double InvalidMaxDistortion = 4.0 * this->MergingThreshold;
+
+	TArray<FUVAreaMetrics> RegionMetrics;
+	RegionMetrics.SetNum(RegionGraph.MaxRegionIndex());
+	ParallelFor(RegionGraph.MaxRegionIndex(), [&](int32 rid)
+	{
+		const TArray<int32>& RegionTris = RegionGraph.GetRegionTris(rid);
+		RegionMetrics[rid] = Compute_UV_Area_Metrics(&TargetMesh, RegionTris, Options.NormalSmoothingRounds, Options.NormalSmoothingAlpha);
+
+		if (Progress && Progress->Cancelled())
+		{
+			return;
+		}
+	});
+	if (Progress && Progress->Cancelled())
+	{
+		return false;
+	}
+
+	TSet<TPair<int, int>> FailedMatches;
+
+	struct FPotentialMatch
+	{
+		int RegionIdx;
+		int NbrRegionIdx;
+		FUVAreaMetrics Metrics;
+	};
+
+
+	struct FPatchJoinParams
+	{
+		double DistortionThreshold;
+		double CompactnessThreshold;
+		double AvgNormalAngleThreshold;
+	};
+
+	TArray<FPatchJoinParams> ParameterIncrements;
+	ParameterIncrements.Add(FPatchJoinParams{ MergeDistortionThreshold, CompactnessThreshold, MaxNormalDeviationDeg/8.0 });
+	ParameterIncrements.Add(FPatchJoinParams{ MergeDistortionThreshold, CompactnessThreshold, MaxNormalDeviationDeg/4.0 });
+	ParameterIncrements.Add(FPatchJoinParams{ MergeDistortionThreshold, CompactnessThreshold, MaxNormalDeviationDeg/2.0 });
+	ParameterIncrements.Add(FPatchJoinParams{ MergeDistortionThreshold, CompactnessThreshold, MaxNormalDeviationDeg });
+
+	int32 CurFilterIndex = 0;
+	int32 PassCount = 0;
+
+	bool bDone = false;
+	while (!bDone)
+	{
+		FPatchJoinParams ActiveJoinParams = ParameterIncrements[CurFilterIndex];
+
+		bDone = true;
+		int32 MergedInPass = 0;
+
+		TArray<FPotentialMatch> Matches;
+
+		// collect up all valid potential matches
+		for (int32 rid : RemainingRegions)
+		{
+			TArray<int32> Nbrs = RegionGraph.GetNeighbours(rid);
+			for (int32 nbrid : Nbrs)
+			{
+				if (rid < nbrid)
+				{
+					if (FailedMatches.Contains(TPair<int, int>(rid, nbrid)) == false)
+					{
+						Matches.Add(FPotentialMatch{ rid, nbrid, FUVAreaMetrics::Invalid() });
+					}
+				}
+			}
+		}
+
+		// compute them all in parallel
+		ParallelFor(Matches.Num(), [&](int32 i)
+		{
+			FPotentialMatch& Match = Matches[i];
+			const TArray<int32>& RegionTris = RegionGraph.GetRegionTris(Match.RegionIdx);
+			const TArray<int32>& NbrTris = RegionGraph.GetRegionTris(Match.NbrRegionIdx);
+			TArray<int32> MergedTriangles;
+			MergedTriangles.Append(RegionTris);
+			MergedTriangles.Append(NbrTris);
+			Match.Metrics = Compute_UV_Area_Metrics(&TargetMesh, MergedTriangles, Options.NormalSmoothingRounds, Options.NormalSmoothingAlpha);
+
+			if (Progress && Progress->Cancelled())
+			{
+				return;
+			}
+		} /*, EParallelForFlags::ForceSingleThread*/ );
+
+		if (Progress && Progress->Cancelled())
+		{
+			return false;
+		}
+
+		auto SortScoreFunc = [&](const FPotentialMatch& Match)
+		{
+			double NormalAngleRad = AngleR(RegionMetrics[Match.RegionIdx].AreaWeightedNormal, RegionMetrics[Match.NbrRegionIdx].AreaWeightedNormal);
+			NormalAngleRad *= 10.0;
+			//NormalAngleRad *= 0.1;
+			double DistortionDelta1 = FMathd::Abs(Match.Metrics.AvgSqrAreaDistortion - RegionMetrics[Match.RegionIdx].AvgSqrAreaDistortion);
+			double DistortionDelta2 = FMathd::Abs(Match.Metrics.AvgSqrAreaDistortion - RegionMetrics[Match.NbrRegionIdx].AvgSqrAreaDistortion);
+
+			// should we consider compactness for sorting?
+
+			//double CompactnessDelta1 = FMathd::Abs(Match.Metrics.Compactness - RegionMetrics[Match.RegionIdx].Compactness);
+			//double CompactnessDelta2 = FMathd::Abs(Match.Metrics.Compactness - RegionMetrics[Match.NbrRegionIdx].Compactness);
+
+			return NormalAngleRad + FMath::Max(DistortionDelta1, DistortionDelta2);
+			//return NormalAngleRad + FMath::Max(DistortionDelta1, DistortionDelta2) + FMath::Max(CompactnessDelta1, CompactnessDelta2);
+		};
+
+		// sort
+		Matches.Sort([&](const FPotentialMatch& A, const FPotentialMatch& B) { return SortScoreFunc(A) < SortScoreFunc(B); });
+
+		auto DistortionMeasureFunc = [](const FPotentialMatch& Match)
+		{
+			// TODO: compactness penalizes squares, does not seem right
+
+			//return (Match.Metrics.AvgSqrAreaDistortion + FMathd::Sqrt(Match.Metrics.Compactness)) * 0.5;
+			return Match.Metrics.AvgSqrAreaDistortion;
+		};
+
+
+		auto NormalsFilterFunc = [&](const FPotentialMatch& Match) -> bool
+		{
+			double NormalAngleDeg = AngleD(RegionMetrics[Match.RegionIdx].AreaWeightedNormal, RegionMetrics[Match.NbrRegionIdx].AreaWeightedNormal);
+			return NormalAngleDeg < ActiveJoinParams.AvgNormalAngleThreshold;
+		};
+		auto CompactnessFilterFunc = [&](const FPotentialMatch& Match) -> bool
+		{
+			return Match.Metrics.Compactness < ActiveJoinParams.CompactnessThreshold;
+		};
+
+
+		TSet<int32> ModifiedRegions;
+		for (int32 i = 0; i < Matches.Num(); ++i)
+		{
+			FPotentialMatch& Match = Matches[i];
+			if (ModifiedRegions.Contains(Match.RegionIdx) || ModifiedRegions.Contains(Match.NbrRegionIdx))
+			{
+				continue;
+			}
+
+			double JoinedDistortionMeasure = DistortionMeasureFunc(Match);
+			bool bNormalsAcceptable = NormalsFilterFunc(Match);
+			bool bCompactnessAcceptable = CompactnessFilterFunc(Match);
+
+			if (bNormalsAcceptable == false || bCompactnessAcceptable == false || JoinedDistortionMeasure > InvalidMaxDistortion)
+			{
+				FailedMatches.Add(TPair<int, int>(Match.RegionIdx, Match.NbrRegionIdx));
+			}
+			else if (JoinedDistortionMeasure < ActiveJoinParams.DistortionThreshold)
+			{
+				bool bMerged = RegionGraph.MergeRegion(Match.NbrRegionIdx, Match.RegionIdx);
+				if (bMerged)
+				{
+					RegionMetrics[Match.RegionIdx] = Match.Metrics;
+					RemainingRegions.Remove(Match.NbrRegionIdx);
+					ModifiedRegions.Add(Match.RegionIdx);
+					ModifiedRegions.Add(Match.NbrRegionIdx);
+					bDone = false;
+					MergedInPass++;
+				}
+				else
+				{
+					FailedMatches.Add(TPair<int, int>(Match.RegionIdx, Match.NbrRegionIdx));
+				}
+			}
+
+			if (Progress && Progress->Cancelled())
+			{
+				return false;
+			}
+		}
+
+		if (Progress && Progress->Cancelled())
+		{
+			return false;
+		}
+
+		//UE_LOG(LogGeometry, Warning, TEXT("ParamStep %d - PASS %d - Merged %d - Remaining %d"), CurFilterIndex, PassCount, MergedInPass, RemainingRegions.Num());
+
+		PassCount++;
+		if (bDone && CurFilterIndex < (ParameterIncrements.Num()-1) )
+		{
+			CurFilterIndex++;
+			FailedMatches.Reset();
+			bDone = false;
+		}
+	}
+
+	// accumulate remaining valid regions
+
+	int32 NumValidRegions = 0;
+	for (int32 rid : RemainingRegions)
+	{
+		if (RegionGraph.IsRegion(rid) && RegionGraph.GetRegionTris(rid).Num() > 0)
+		{
+			NumValidRegions++;
+		}
+	}
+
+	IslandsOut.Reserve(NumValidRegions);
+
+	for (int32 rid : RemainingRegions)
+	{
+		if (RegionGraph.IsRegion(rid))
+		{
+			IslandsOut.Add(RegionGraph.MoveRegionTris(rid));
+		}
+		if (Progress && Progress->Cancelled())
+		{
+			return false;
+		}
+	}
+
+	return true;
+}
+
+
+
+int32 FPatchBasedMeshUVGenerator::ComputeUVsFromTriangleSets(
+	FDynamicMesh3& TargetMesh,
+	FDynamicMeshUVOverlay& TargetUVOverlay,
+	const TArray<TArray<int32>>& TriangleSets,
+	TArray<bool>& bValidUVsComputedOut,
+	FProgressCancel* Progress
+)
+{
+	FDynamicMeshUVEditor::FExpMapOptions Options;
+	Options.NormalSmoothingRounds = this->NormalSmoothingRounds;
+	Options.NormalSmoothingAlpha = this->NormalSmoothingAlpha;
+
+	FDynamicMeshUVEditor UVEditor(&TargetMesh, &TargetUVOverlay);
+
+	int32 NumUVIslands = TriangleSets.Num();
+	bValidUVsComputedOut.Init(false, NumUVIslands);
+
+	int32 NumFailed = 0;
+	for ( int32 k = 0; k < NumUVIslands; ++k)
+	{
+		const TArray<int32>& TriangleSet = TriangleSets[k];
+		bValidUVsComputedOut[k] = UVEditor.SetTriangleUVsFromExpMap(TriangleSet, Options);
+		if (bValidUVsComputedOut[k] == false)
+		{
+			NumFailed++;
+		}
+
+		if (Progress && Progress->Cancelled())
+		{
+			return -1;
+		}
+	}
+
+	return NumFailed;
+}
diff --git a/Engine/Plugins/Runtime/GeometryProcessing/Source/DynamicMesh/Public/Parameterization/PatchBasedMeshUVGenerator.h b/Engine/Plugins/Runtime/GeometryProcessing/Source/DynamicMesh/Public/Parameterization/PatchBasedMeshUVGenerator.h
new file mode 100644
--- /dev/null
+++ b/Engine/Plugins/Runtime/GeometryProcessing/Source/DynamicMesh/Public/Parameterization/PatchBasedMeshUVGenerator.h
@@ -0,0 +1,74 @@
+// Copyright Epic Games, Inc. All Rights Reserved.
+
+#pragma once
+
+#include "DynamicMesh/DynamicMesh3.h"
+#include "DynamicMesh/DynamicMeshAttributeSet.h"
+#include "Selections/MeshConnectedComponents.h"
+#include "Util/ProgressCancel.h"
+
+namespace UE
+{
+namespace Geometry
+{
+
+
+/**
+ * FPatchBasedMeshUVGenerator is an automatic UV generator for a triangle mesh that, generally, works by
+ * first decomposing the mesh into small patches for which a "known good" parameterization can be computed,
+ * and then incrementally merging those patches into larger ones to create final UV islands.
+ * 
+ * The various steps of the generator can be called independently, eg if one already has a patch decomposition,
+ * then the initial steps can be skipped
+ * 
+ */
+class DYNAMICMESH_API FPatchBasedMeshUVGenerator
+{
+public:
+
+	//
+	// Parameters for island merging
+	//
+	double MergingThreshold = 1.5;
+	double CompactnessThreshold = 9999999.0;		// effectively disabled as it usually is not a good idea
+	double MaxNormalDeviationDeg = 45.0;
+
+	//
+	// ExpMap parameters
+	//
+	int32 NormalSmoothingRounds = 0;
+	double NormalSmoothingAlpha = 0.25;
+
+
+	/**
+	 * Incrementally combine existing mesh patches into larger patches that meet the various criteria defined
+	 * by the parameters above
+	 * @param InitialComponents the initial mesh patches, ie triangle lists
+	 * @param IslandsOut the set of output mesh patches / triangle lists
+	 * @return false if catastrophic failure ocurred (ie IslandsOut is invalid) or if Progress was cancelled
+	 */
+	bool ComputeIslandsByRegionMerging(
+		FDynamicMesh3& TargetMesh,
+		FDynamicMeshUVOverlay& TargetUVOverlay,
+		const FMeshConnectedComponents& InitialComponents, 
+		TArray<TArray<int32>>& IslandsOut,
+		FProgressCancel* Progress = nullptr);
+
+	/**
+	 * Calculate UVs for the input mesh patches of TargetMesh and store in TargetUVOVerlay
+	 * @param TriangleSets the set of desired UV islands, defined as triangle lists
+	 * @param bValidUVsComputedOut output list of booleans the same length as TriangleSets, set to false if UV computation failed for that island
+	 * @return number of failed UV islands, 0 if all were successful, and -1 if Progress was cancelled
+	 */
+	int32 ComputeUVsFromTriangleSets(
+		FDynamicMesh3& TargetMesh,
+		FDynamicMeshUVOverlay& TargetUVOverlay,
+		const TArray<TArray<int32>>& TriangleSets,
+		TArray<bool>& bValidUVsComputedOut,
+		FProgressCancel* Progress = nullptr );
+
+};
+
+
+} // end namespace UE::Geometry
+} // end namespace UE
\ No newline at end of file
diff --git a/Engine/Source/Runtime/GeometryCore/Public/Sampling/NormalHistogram.h b/Engine/Source/Runtime/GeometryCore/Public/Sampling/NormalHistogram.h
new file mode 100644
--- /dev/null
+++ b/Engine/Source/Runtime/GeometryCore/Public/Sampling/NormalHistogram.h
@@ -0,0 +1,101 @@
+// Copyright Epic Games, Inc. All Rights Reserved.
+
+#pragma once
+
+#include "MathUtil.h"
+#include "VectorTypes.h"
+#include "Sampling/SphericalFibonacci.h"
+
+namespace UE
+{
+namespace Geometry
+{
+
+
+/**
+ * TNormalHistogram calculates/represents a histogram on a 3D sphere. The histogram bins
+ * are defined by a Spherical Fibonacci distribution, ie each bin is the voronoi cell of
+ * the Fibonacci point. This point set can be calculated for any N.
+ * 
+ * One current limitation of this implementation is that the unit +/- X/Y/Z axes are not bin center points.
+ */
+template<typename RealType>
+class TNormalHistogram
+{
+public:
+	int32 Bins;
+	TSphericalFibonacci<RealType> BinPoints;
+	TArray<RealType> WeightedCounts;
+
+	/**
+	 * Initialize a normal histogram with the given number of bins
+	 */
+	TNormalHistogram(int32 NumBins)
+	{
+		Bins = NumBins;
+		BinPoints = TSphericalFibonacci<RealType>(NumBins);
+		WeightedCounts.Init(0, BinPoints.Num());
+	}
+
+	/**
+	 * Count a Normal in the histogram, ie find the Bin it should be included in and add the Weight to that Bin (The Normal does not strictly need to be normalized)
+	 * @param Weight optional weight for this normal
+	 */
+	void Count(const FVector3<RealType>& Normal, RealType Weight = 1)
+	{
+		int32 BinIndex = BinPoints.FindIndex(Normal);
+		WeightedCounts[BinIndex] += Weight;
+	}
+
+	/**
+	 * @return the normal direction for the Histogram bin that contains the most points / largest mass
+	 */
+	FVector3<RealType> FindMaxNormal() const
+	{
+		int MaxIndex = 0;
+		for (int k = 1; k < Bins; ++k)
+		{
+			if (WeightedCounts[k] > WeightedCounts[MaxIndex])
+			{
+				MaxIndex = k;
+			}
+		}
+		return BinPoints[MaxIndex];
+	}
+
+
+	/**
+	 * Attempt to compute a measure of the range of variation in normals across the occupied bins
+	 */
+	RealType WeightedSpreadMetric() const
+	{
+		int32 NumOccupiedBins = 0;
+		RealType TotalWeightedCount = 0;
+		for (int k = 0; k < Bins; ++k)
+		{
+			if (WeightedCounts[k] > 0)
+			{
+				TotalWeightedCount += WeightedCounts[k];
+				NumOccupiedBins++;
+			}
+		}
+
+		RealType Metric = 0;
+		for (int32 k = 0; k < Bins; ++k)
+		{
+			if (WeightedCounts[k] > 0)
+			{
+				RealType MassFraction = WeightedCounts[k] / TotalWeightedCount;
+				MassFraction /= (RealType)NumOccupiedBins;
+				Metric += MassFraction;
+			}
+		}
+
+		return 1.0 / Metric;
+	}
+};
+
+
+
+} // end namespace UE::Geometry
+} // end namespace UE
