MeshModelingTools: Fix edge loop insertion crash when using retriangulate in some weird topology situations.

#rnx

diff --git a/Engine/Plugins/Runtime/GeometryProcessing/Source/DynamicMesh/Private/Operations/GroupEdgeInserter.cpp b/Engine/Plugins/Runtime/GeometryProcessing/Source/DynamicMesh/Private/Operations/GroupEdgeInserter.cpp
--- a/Engine/Plugins/Runtime/GeometryProcessing/Source/DynamicMesh/Private/Operations/GroupEdgeInserter.cpp
+++ b/Engine/Plugins/Runtime/GeometryProcessing/Source/DynamicMesh/Private/Operations/GroupEdgeInserter.cpp
@@ -33,7 +33,7 @@ bool InsertEdgeLoopEdgesInDirection(const FGroupEdgeInserter::FEdgeLoopInsertion
 	int32 NextGroupID, int32 NextEdgeID, int32 NextCornerID, int32 NextBoundaryID,
 	TSet<int32>& AlteredGroups, int32& NumInserted, FGroupEdgeInserter::FOptionalOutputParams& OptionalOut,
 	FProgressCancel* Progress);
-void InsertNewVertexEndpoints(
+bool InsertNewVertexEndpoints(
 	const FGroupEdgeInserter::FEdgeLoopInsertionParams& Params,
 	int32 GroupEdgeID, int32 StartCornerID,
 	TArray<FGroupEdgeInserter::FGroupEdgeSplitPoint>& EndPointsOut, 
@@ -90,11 +90,12 @@ bool InsertSingleWithRetriangulation(FDynamicMesh3& Mesh, FGroupTopology& Topolo
 	const FGroupEdgeInserter::FGroupEdgeSplitPoint& EndPoint,
 	FGroupEdgeInserter::FOptionalOutputParams& OptionalOut, FProgressCancel* Progress);
 }
-using namespace GroupEdgeInserterLocals;
 
 /** Inserts an edge loop into a mesh, where an edge loop is a sequence of (group) edges across quads. */
 bool FGroupEdgeInserter::InsertEdgeLoops(const FEdgeLoopInsertionParams& Params, FOptionalOutputParams OptionalOut, FProgressCancel* Progress)
 {
+	using namespace GroupEdgeInserterLocals;
+
 	if (Progress && Progress->Cancelled())
 	{
 		return false;
@@ -139,16 +140,15 @@ bool FGroupEdgeInserter::InsertEdgeLoops(const FEdgeLoopInsertionParams& Params,
 	// Although we have code that can insert new edges at edge endpoints, it is cleaner to do splits for all the loops
 	// down an edge ahead of time to make vertex endpoints, in part because if we don't, a split can change the eid
 	// of the next endpoint.
-	InsertNewVertexEndpoints(Params, Params.GroupEdgeID, Params.StartCornerID, StartEndpoints, OptionalOut);
+	bool bSuccess = InsertNewVertexEndpoints(Params, Params.GroupEdgeID, Params.StartCornerID, StartEndpoints, OptionalOut);
 
-	if (StartEndpoints.Num() == 0 || (Progress && Progress->Cancelled()))
+	if (!bSuccess || StartEndpoints.Num() == 0 || (Progress && Progress->Cancelled()))
 	{
 		return false;
 	}
 
 	// Insert edges in both directions. In case of a loop, the second call won't do anything because
 	// AlteredGroups will be updated.
-	bool bSuccess = true;
 	int32 TotalNumInserted = 0;
 	if (bHaveForwardEdge)
 	{
@@ -292,8 +292,11 @@ bool InsertEdgeLoopEdgesInDirection(const FGroupEdgeInserter::FEdgeLoopInsertion
 		}
 
 		// Otherwise, create next endpoints
-		InsertNewVertexEndpoints(Params, NextEdgeID, NextCornerID, *NextEndpoints, OptionalOut);
-
+		bSuccess = InsertNewVertexEndpoints(Params, NextEdgeID, NextCornerID, *NextEndpoints, OptionalOut);
+		if (!bSuccess || (Progress && Progress->Cancelled()))
+		{
+			return false;
+		}
 		if (NextEndpoints->Num() == 0)
 		{
 			// Next edge wasn't long enough for the input lengths we wanted. Stop here.
@@ -338,7 +341,7 @@ bool InsertEdgeLoopEdgesInDirection(const FGroupEdgeInserter::FEdgeLoopInsertion
  *
  * Clears EndPointsOut before use.
  */
-void InsertNewVertexEndpoints(
+bool InsertNewVertexEndpoints(
 	const FGroupEdgeInserter::FEdgeLoopInsertionParams& Params,
 	int32 GroupEdgeID, int32 StartCornerID,
 	TArray<FGroupEdgeInserter::FGroupEdgeSplitPoint>& EndPointsOut,
@@ -347,7 +350,7 @@ void InsertNewVertexEndpoints(
 	EndPointsOut.Reset();
 	if (Params.SortedInputLengths->Num() == 0)
 	{
-		return;
+		return false;
 	}
 
 	const FGroupTopology::FGroupEdge& GroupEdge = Params.Topology->Edges[GroupEdgeID];
@@ -453,6 +456,15 @@ void InsertNewVertexEndpoints(
 			// Target must be on the edge that goes to the next vertex
 			int32 CurrentEid = Params.Mesh->FindEdge(CurrentVid, SpanVids[NextIndex]);
 
+			if (!ensure(CurrentEid >= 0))
+			{
+				// We shouldn't end up here, but if we do, it could be because something failed and
+				// stopped loop progress in one direction despite having performed edge splits forward,
+				// and now we have arrived at the split edge from the other direction, not realizing
+				// that SpanVids has changed
+				return false;
+			}
+
 			double SplitT = (TargetLength - CurrentArcLength) / (PerVertexLengths[NextIndex] - CurrentArcLength);
 
 			// See if the edge is stored backwards relative to the direction we're traveling
@@ -487,6 +499,8 @@ void InsertNewVertexEndpoints(
 
 		EndPointsOut.Add(SplitPoint);
 	}//end going through targets
+
+	return true;
 }
 
 void ConvertProportionsToArcLengths(
@@ -838,7 +852,28 @@ bool RetriangulateLoop(FDynamicMesh3& Mesh,
 	FSimpleHoleFiller HoleFiller(&Mesh, Loop, FSimpleHoleFiller::EFillType::PolygonEarClipping);
 	if (!HoleFiller.Fill(NewGroupID))
 	{
-		return false;
+		// If the hole filler failed, it is probably because two non-adjacent vertices in the loop
+		// already have an edge between them ("on the other side" of this group), so inserting
+		// two more triangles between the verts is not possible while keeping things manifold. 
+		// As an example, imagine a cube with one face deleted but then an obtuse triangle added
+		// to one of the boundary edges to make one cube side an irregular pentagon.
+
+		// Granted, perhaps the user shouldn't be trying to be doing simple retriangulation in a 
+		// situation like this, but let's do something reasonable anyway. We'll retriangulate with
+		// an added center vert, since this will only add to the edges that we know are boundaries.
+
+		// Delete any triangles we already added
+		FDynamicMeshEditor Editor(&Mesh);
+		Editor.RemoveTriangles(HoleFiller.NewTriangles, false); // don't remove isolated verts
+		
+		// Change the hole fill type
+		HoleFiller.FillType = FSimpleHoleFiller::EFillType::TriangleFan;
+		if (!HoleFiller.Fill(NewGroupID))
+		{
+			// Not sure how this could happen... Did we not delete the interior triangles somehow
+			// or give the wrong loop?
+			return ensure(false);
+		}
 	}
 
 	if (Mesh.HasAttributes())
@@ -1028,6 +1063,8 @@ bool CreateNewGroups(FDynamicMesh3& Mesh, TSet<int32>& PathEids, int32 OriginalG
 /** Inserts a group edge into a given group. */
 bool FGroupEdgeInserter::InsertGroupEdge(FGroupEdgeInsertionParams& Params, FGroupEdgeInserter::FOptionalOutputParams OptionalOut, FProgressCancel* Progress)
 {
+	using namespace GroupEdgeInserterLocals;
+
 	if (Progress && Progress->Cancelled())
 	{
 		return false;
diff --git a/Engine/Plugins/Runtime/GeometryProcessing/Source/DynamicMesh/Private/Operations/SimpleHoleFiller.cpp b/Engine/Plugins/Runtime/GeometryProcessing/Source/DynamicMesh/Private/Operations/SimpleHoleFiller.cpp
--- a/Engine/Plugins/Runtime/GeometryProcessing/Source/DynamicMesh/Private/Operations/SimpleHoleFiller.cpp
+++ b/Engine/Plugins/Runtime/GeometryProcessing/Source/DynamicMesh/Private/Operations/SimpleHoleFiller.cpp
@@ -105,7 +105,7 @@ bool FSimpleHoleFiller::Fill_EarClip(int GroupID)
 		}
 	}
 
-	return true;
+	return NewTriangles.Num() == Triangles.Num();
 }
 
 bool FSimpleHoleFiller::UpdateAttributes(TArray<FMeshRegionBoundaryLoops::VidOverlayMap<FVector2f>>& VidUVMaps)
@@ -117,12 +117,29 @@ bool FSimpleHoleFiller::UpdateAttributes(TArray<FMeshRegionBoundaryLoops::VidOve
 
 	FDynamicMeshAttributeSet* Attributes = Mesh->Attributes();
 
-	check(VidUVMaps.Num() == Attributes->NumUVLayers());
+	if (!ensure(VidUVMaps.Num() == Attributes->NumUVLayers()))
+	{
+		return false;
+	}
+
 	for (int i = 0; i < Attributes->NumUVLayers(); ++i)
 	{
 		FDynamicMeshUVOverlay* UVLayer = Attributes->GetUVLayer(i);
 		FMeshRegionBoundaryLoops::VidOverlayMap<FVector2f>& UVMap = VidUVMaps[i];
 
+		// Create an entry for the newly added center vert if we need it
+		if (NewVertex != IndexConstants::InvalidID && !UVMap.Contains(NewVertex))
+		{
+			FVector2f NewVertUV = FVector2f::Zero();
+			for (int32 Vid : Loop.Vertices)
+			{
+				NewVertUV += UVMap[Vid].Value;
+			}
+			NewVertUV /= Loop.Vertices.Num();
+			UVMap.Add(NewVertex, FMeshRegionBoundaryLoops::ElementIDAndValue<FVector2f>(
+				IndexConstants::InvalidID, NewVertUV));
+		}
+
 		for (int32 Tid : NewTriangles)
 		{
 			FIndex3i TriVids = Mesh->GetTriangle(Tid);
diff --git a/Engine/Plugins/Runtime/GeometryProcessing/Source/DynamicMesh/Public/Operations/SimpleHoleFiller.h b/Engine/Plugins/Runtime/GeometryProcessing/Source/DynamicMesh/Public/Operations/SimpleHoleFiller.h
--- a/Engine/Plugins/Runtime/GeometryProcessing/Source/DynamicMesh/Public/Operations/SimpleHoleFiller.h
+++ b/Engine/Plugins/Runtime/GeometryProcessing/Source/DynamicMesh/Public/Operations/SimpleHoleFiller.h
@@ -65,6 +65,12 @@ public:
 		{
 			return EOperationValidationResult::Failed_UnknownReason;
 		}
+		// TODO: For EFillType::PolygonEarClipping, we actually have more constraints; specifically, there
+		// must not be any edges between non-adjacent loop vertices that we choose to connect (this can
+		// happen in a few weird topology cases, for instance if a hole is attached to a single triangle
+		// flap), because we would end up trying to add two more triangles to an existing edge.
+		// To do a proper check we'd need to know which verts we're planning to connect, which would require
+		// going through the whole retriangulation. That does not seem worthwhile...
 
 		return EOperationValidationResult::Ok;
 	}
@@ -74,9 +80,11 @@ public:
 	/**
 	 * Updates the normals and UV's of NewTriangles. UV's are taken from VidUVMaps,
 	 * which is an array of maps (1:1 with UV layers) that map vid's of vertices on the
-	 * boundary to their UV elements and values. If the UV element for a vertex does not
-	 * yet exist in the overlay, the corresponding element ID should be InvalidID. The 
-	 * function will update it to point to the new element once it inserts it.
+	 * boundary to their UV elements and values. If an entry for NewVertex is not provided,
+	 * it is set to be the average of the boundary UV's. For other UV's, an entry must
+	 * exist, though the UV element ID can be InvalidID if an element does not exist for
+	 * that UV value. The function will update the element ID in the map to point to the
+	 * new element once it inserts it.
 
 	 * Normals are shared among NewTriangles but not with the surrounding portions of the mesh.
 	 *
