ModelingTools: Add texture filtering support to BakeTexture.

#rnx

diff --git a/Engine/Plugins/Runtime/GeometryProcessing/Source/DynamicMesh/Private/Sampling/MeshMapBaker.cpp b/Engine/Plugins/Runtime/GeometryProcessing/Source/DynamicMesh/Private/Sampling/MeshMapBaker.cpp
--- a/Engine/Plugins/Runtime/GeometryProcessing/Source/DynamicMesh/Private/Sampling/MeshMapBaker.cpp
+++ b/Engine/Plugins/Runtime/GeometryProcessing/Source/DynamicMesh/Private/Sampling/MeshMapBaker.cpp
@@ -4,7 +4,6 @@
 #include "Sampling/MeshBakerCommon.h"
 #include "Image/ImageOccupancyMap.h"
 #include "Image/ImageTile.h"
-#include "Spatial/DenseGrid2.h"
 
 #include "ExplicitUseGeometryMathTypes.h"		// using UE::Geometry::(math types)
 using namespace UE::Geometry;
@@ -14,6 +13,13 @@ using namespace UE::Geometry;
 // FMeshMapBaker
 //
 
+static const float BoxFilterRadius = 0.5f;
+static const float BCFilterRadius = 0.769f;
+
+FBoxFilter FMeshMapBaker::BoxFilter(BoxFilterRadius);
+FBSplineFilter FMeshMapBaker::BSplineFilter(BCFilterRadius);
+FMitchellNetravaliFilter FMeshMapBaker::MitchellNetravaliFilter(BCFilterRadius);
+
 
 void FMeshMapBaker::InitBake()
 {
@@ -21,7 +27,7 @@ void FMeshMapBaker::InitBake()
 	// - index lists of accumulation modes (BakeAccumulateLists)
 	// - evaluator to bake result offsets (BakeOffsets)
 	// - buffer size per sample (BakeSampleBufferSize)
-	int32 NumBakers = Bakers.Num();
+	const int32 NumBakers = Bakers.Num();
 	BakeContexts.SetNum(NumBakers);
 	BakeOffsets.SetNumUninitialized(NumBakers + 1);
 	BakeAccumulateLists.SetNum((int32) FMeshMapEvaluator::EAccumulateMode::Last);
@@ -32,7 +38,7 @@ void FMeshMapBaker::InitBake()
 		Bakers[Idx]->Setup(*this, BakeContexts[Idx]);
 		checkSlow(BakeContexts[Idx].Evaluate != nullptr && BakeContexts[Idx].EvaluateDefault != nullptr);
 		checkSlow(BakeContexts[Idx].DataLayout.Num() > 0);
-		int32 NumData = BakeContexts[Idx].DataLayout.Num();
+		const int32 NumData = BakeContexts[Idx].DataLayout.Num();
 		for (int32 DataIdx = 0; DataIdx < NumData; ++DataIdx)
 		{
 			BakeSampleBufferSize += (int32) BakeContexts[Idx].DataLayout[DataIdx];
@@ -70,6 +76,8 @@ void FMeshMapBaker::InitBake()
 	{
 		BakeResults[Idx]->Clear(BakeDefaultColors[Idx]);
 	}
+
+	InitFilter();
 }
 
 void FMeshMapBaker::InitBakeDefaults()
@@ -124,10 +132,10 @@ void FMeshMapBaker::Bake()
 		{
 			FVector2f A, B, C;
 			UVOverlay->GetTriElements(tid, A, B, C);
-			int32 VertA = FlatMesh.AppendVertex(FVector3d(A.X, A.Y, 0));
-			int32 VertB = FlatMesh.AppendVertex(FVector3d(B.X, B.Y, 0));
-			int32 VertC = FlatMesh.AppendVertex(FVector3d(C.X, C.Y, 0));
-			int32 NewTriID = FlatMesh.AppendTriangle(VertA, VertB, VertC, tid);
+			const int32 VertA = FlatMesh.AppendVertex(FVector3d(A.X, A.Y, 0));
+			const int32 VertB = FlatMesh.AppendVertex(FVector3d(B.X, B.Y, 0));
+			const int32 VertC = FlatMesh.AppendVertex(FVector3d(C.X, C.Y, 0));
+			/*int32 NewTriID =*/ FlatMesh.AppendTriangle(VertA, VertB, VertC, tid);
 		}
 	}
 
@@ -171,25 +179,36 @@ void FMeshMapBaker::Bake()
 		}
 	});
 
-	// Setup image tiling
+	// Create a temporary output float buffer for the full image dimensions.
+	const FImageTile ImageTile(FVector2i(0,0), FVector2i(Dimensions.GetWidth(), Dimensions.GetHeight()));
+	FMeshMapTileBuffer ImageTileBuffer(ImageTile, BakeSampleBufferSize);
+
+	// Tile the image
 	FImageTiling Tiles(Dimensions, TileSize, TileSize);
-	const int32 NumBakers = Bakers.Num();
-	const int32 NumResults = BakeResults.Num();
-	ParallelFor(Tiles.Num(), [this, &Tiles, &NumBakers, &NumResults](int32 TileIdx)
+	const int32 NumTiles = Tiles.Num();
+	TArray<TArray64<TTuple<int64, int64>>> GutterTexelsPerTile;
+	GutterTexelsPerTile.SetNum(NumTiles);
+	ParallelFor(NumTiles, [this, &Tiles, &ImageTileBuffer, &GutterTexelsPerTile](int32 TileIdx)
 	{
-		const FImageTile Tile = Tiles.GetTile(TileIdx);
+		// Generate unpadded and padded tiles.
+		const FImageTile Tile = Tiles.GetTile(TileIdx);	// Image area to sample
+		const FImageTile PaddedTile = Tiles.GetTile(TileIdx, TilePadding); // Filtered image area
 
 		FImageOccupancyMap OccupancyMap;
 		OccupancyMap.GutterSize = GutterSize;
 		OccupancyMap.Initialize(Dimensions, Tile, Multisampling);
 		OccupancyMap.ComputeFromUVSpaceMesh(FlatMesh, [this](int32 TriangleID) { return FlatMesh.GetTriangleGroup(TriangleID); });
+		GutterTexelsPerTile[TileIdx] = OccupancyMap.GutterTexels;
 
-		FMeshMapTileBuffer TileBuffer(Tile, BakeSampleBufferSize);
+		FMeshMapTileBuffer TileBuffer(PaddedTile, BakeSampleBufferSize);
 
 		// Calculate interior texels
-		for (FVector2i TileCoords(0,0); TileCoords.Y < Tile.GetHeight(); ++TileCoords.Y)
+		const int TileWidth = Tile.GetWidth();
+		const int TileHeight = Tile.GetHeight();
+		const int32 NumSamples = OccupancyMap.Multisampler.Num();
+		for (FVector2i TileCoords(0,0); TileCoords.Y < TileHeight; ++TileCoords.Y)
 		{
-			for (TileCoords.X = 0; TileCoords.X < Tile.GetWidth(); ++TileCoords.X)
+			for (TileCoords.X = 0; TileCoords.X < TileWidth; ++TileCoords.X)
 			{
 				if (CancelF())
 				{
@@ -203,10 +222,9 @@ void FMeshMapBaker::Bake()
 				}
 
 				const FVector2i ImageCoords = Tile.GetSourceCoords(TileCoords);
-				const int32 NumSamples = OccupancyMap.Multisampler.Num();
 				for (int32 SampleIdx = 0; SampleIdx < NumSamples; ++SampleIdx)
 				{
-					const int64 LinearIdx = TilePixelLinearIdx * OccupancyMap.Multisampler.Num() + SampleIdx;
+					const int64 LinearIdx = TilePixelLinearIdx * NumSamples + SampleIdx;
 					if (OccupancyMap.IsInterior(LinearIdx))
 					{
 						const FVector2d UVPosition = (FVector2d)OccupancyMap.TexelQueryUV[LinearIdx];
@@ -215,85 +233,169 @@ void FMeshMapBaker::Bake()
 						FMeshMapEvaluator::FCorrespondenceSample Sample;
 						DetailMeshSampler.SampleUV(UVTriangleID, UVPosition, Sample);
 
-						BakeSample(TileBuffer, Sample, Tile, TileCoords, ImageCoords, 1.0f);
+						BakeSample(TileBuffer, Sample, UVPosition, ImageCoords);
 					}
 				}
 			}
 		}
 
-		// Convert TileBuffer data to color data.
-		for (FVector2i Texel(0,0); Texel.Y < Tile.GetHeight(); ++Texel.Y)
+		auto NoopFn = [](const float& In, float& Out)
+		{
+		};
+
+		auto AddFn = [](const float& In, float& Out)
+		{
+			Out += In;
+		};
+
+		auto OverwriteFn = [](const float& In, float& Out)
 		{
-			for (Texel.X = 0; Texel.X < Tile.GetWidth(); ++Texel.X)
+			Out = In;
+		};
+
+		// WriteToOutputBuffer transfers local tile data (TileBuffer) to the image output buffer (ImageTileBuffer).
+		auto WriteToOutputBuffer = [this, &TileBuffer, &ImageTileBuffer] (const FImageTile& TargetTile, const TArray<int32>& BakeIds, auto&& Op, auto&& WeightOp)
+		{
+			const int TargetTileWidth = TargetTile.GetWidth();
+			const int TargetTileHeight = TargetTile.GetHeight();
+			for (FVector2i TileCoords(0,0); TileCoords.Y < TargetTileHeight; ++TileCoords.Y)
+			{
+				for (TileCoords.X = 0; TileCoords.X < TargetTileWidth; ++TileCoords.X)
+				{
+					if (CancelF())
+					{
+						return;
+					}
+
+					const FVector2i ImageCoords = TargetTile.GetSourceCoords(TileCoords);
+
+					const FImageTile& BufferTile = TileBuffer.GetTile();
+					const int64 TilePixelLinearIdx = BufferTile.GetIndexFromSourceCoords(ImageCoords);
+					const float& TilePixelWeight = TileBuffer.GetPixelWeight(TilePixelLinearIdx);
+					float* TilePixelBuffer = TileBuffer.GetPixel(TilePixelLinearIdx);
+
+					const int64 ImageLinearIdx = Dimensions.GetIndex(ImageCoords);
+					float& ImagePixelWeight = ImageTileBuffer.GetPixelWeight(ImageLinearIdx);
+					float* ImagePixelBuffer = ImageTileBuffer.GetPixel(ImageLinearIdx);
+
+					WeightOp(TilePixelWeight, ImagePixelWeight);
+					for( int32 Idx : BakeIds )
+					{
+						const FMeshMapEvaluator::FEvaluationContext& Context = BakeContexts[Idx];
+						const int32 NumData = Context.DataLayout.Num();
+						const int32 ResultOffset = BakeOffsets[Idx];
+						for (int32 DataIdx = 0; DataIdx < NumData; ++DataIdx)
+						{
+							const int32 ResultIdx = ResultOffset + DataIdx;
+							const int32 Offset = BakeSampleOffsets[ResultIdx];
+							float* BufferPtr = &TilePixelBuffer[Offset];
+							float* ImageBufferPtr = &ImagePixelBuffer[Offset];
+							const FMeshMapEvaluator::EComponents Stride = Context.DataLayout[DataIdx];
+
+							for (int32 FloatIdx = 0; FloatIdx < (int32)Stride; ++FloatIdx)
+							{
+								Op(BufferPtr[FloatIdx], ImageBufferPtr[FloatIdx]);
+							}
+						}
+					}
+				}
+			}
+		};
+
+		// Transfer 'Overwrite' float data to image tile buffer
+		WriteToOutputBuffer(Tile, BakeAccumulateLists[(int32)FMeshMapEvaluator::EAccumulateMode::Overwrite], OverwriteFn, NoopFn);
+
+		// Accumulate 'Add' float data to image tile buffer
+		// TODO: Optimize this write lock
+		FCriticalSection WriteLock;
+		WriteToOutputBuffer(PaddedTile, BakeAccumulateLists[(int32)FMeshMapEvaluator::EAccumulateMode::Add], AddFn, AddFn);
+	}, !bParallel ? EParallelForFlags::ForceSingleThread : EParallelForFlags::None);
+
+	// Normalize and convert ImageTileBuffer data to color data.
+	ParallelFor(NumTiles, [this, &Tiles, &ImageTileBuffer](int32 TileIdx)
+	{
+		const FImageTile Tile = Tiles.GetTile(TileIdx);
+		const int TileWidth = Tile.GetWidth();
+		const int TileHeight = Tile.GetHeight();
+		for (FVector2i TileCoords(0,0); TileCoords.Y < TileHeight; ++TileCoords.Y)
+		{
+			for (TileCoords.X = 0; TileCoords.X < TileWidth; ++TileCoords.X)
 			{
 				if (CancelF())
 				{
 					return;
 				}
 
-				const int64 PixelLinearIdx = Tile.GetIndex(Texel);
-				const float& PixelWeight = TileBuffer.GetPixelWeight(PixelLinearIdx);
-				if (PixelWeight > 0.0)
-				{
-					const FVector2i ImageCoords = Tile.GetSourceCoords(Texel);
-					const int64 ImageLinearIdx = Dimensions.GetIndex(ImageCoords);
-					float* PixelBuffer = TileBuffer.GetPixel(PixelLinearIdx);
-					const float OneOverPixelWeight = 1.0f / PixelWeight;
+				const FVector2i ImageCoords = Tile.GetSourceCoords(TileCoords);
+				const int64 ImageLinearIdx = Dimensions.GetIndex(ImageCoords);
+				const float& PixelWeight = ImageTileBuffer.GetPixelWeight(ImageLinearIdx);
+				float* PixelBuffer = ImageTileBuffer.GetPixel(ImageLinearIdx);
 
-					auto WriteToPixel = [this, &PixelBuffer, &ImageLinearIdx](TArray<int32>& BakeIds, float Weight)
+				auto WriteToPixel = [this, &PixelBuffer, &ImageLinearIdx](TArray<int32>& BakeIds, float OneOverWeight)
+				{
+					for (int32 Idx : BakeIds)
 					{
-						for (int32 Idx : BakeIds)
+						const FMeshMapEvaluator::FEvaluationContext& Context = BakeContexts[Idx];
+						const int32 NumData = Context.DataLayout.Num();
+						const int32 ResultOffset = BakeOffsets[Idx];
+						for (int32 DataIdx = 0; DataIdx < NumData; ++DataIdx)
 						{
-							const FMeshMapEvaluator::FEvaluationContext& Context = BakeContexts[Idx];
-							const int32 NumData = Context.DataLayout.Num();
-							const int32 ResultOffset = BakeOffsets[Idx];
-							for (int32 DataIdx = 0; DataIdx < NumData; ++DataIdx)
+							const int32 ResultIdx = ResultOffset + DataIdx;
+							const int32 Offset = BakeSampleOffsets[ResultIdx];
+							float* BufferPtr = &PixelBuffer[Offset];
+							const FMeshMapEvaluator::EComponents Stride = Context.DataLayout[DataIdx];
+
+							// Apply weight to raw float data.
+							for (int32 FloatIdx = 0; FloatIdx < (int32)Stride; ++FloatIdx)
 							{
-								const int32 ResultIdx = ResultOffset + DataIdx;
-								const int32 Offset = BakeSampleOffsets[ResultIdx];
-								float* BufferPtr = &PixelBuffer[Offset];
-								const FMeshMapEvaluator::EComponents Stride = Context.DataLayout[DataIdx];
-
-								// Apply weight to raw float data.
-								for (int32 FloatIdx = 0; FloatIdx < (int32)Stride; ++FloatIdx)
-								{
-									BufferPtr[FloatIdx] *= Weight;
-								}
-
-								// Convert float data to color.
-								FVector4f& Pixel = BakeResults[ResultIdx]->GetPixel(ImageLinearIdx);
-								Context.EvaluateColor(DataIdx, BufferPtr, Pixel, Context.EvalData);
+								BufferPtr[FloatIdx] *= OneOverWeight;
 							}
-						}
-					};
 
-					WriteToPixel(BakeAccumulateLists[(int32)FMeshMapEvaluator::EAccumulateMode::Add], OneOverPixelWeight);
-					WriteToPixel(BakeAccumulateLists[(int32)FMeshMapEvaluator::EAccumulateMode::Overwrite], 1.0f);
+							// Convert float data to color.
+							FVector4f& Pixel = BakeResults[ResultIdx]->GetPixel(ImageLinearIdx);
+							Context.EvaluateColor(DataIdx, BufferPtr, Pixel, Context.EvalData);
+						}
+					}
+				};
+				
+				if (PixelWeight > 0.0)
+				{
+					WriteToPixel(BakeAccumulateLists[(int32)FMeshMapEvaluator::EAccumulateMode::Add], 1.0f / PixelWeight);
 				}
+				WriteToPixel(BakeAccumulateLists[(int32)FMeshMapEvaluator::EAccumulateMode::Overwrite], 1.0f);
 			}
 		}
+	}, !bParallel ? EParallelForFlags::ForceSingleThread : EParallelForFlags::None);
 
-		for (int64 GutterIdx = 0; GutterIdx < OccupancyMap.GutterTexels.Num(); ++GutterIdx)
+	// Gutter Texel processing
+	if( bGutterEnabled )
+	{
+		const int32 NumResults = BakeResults.Num();
+		ParallelFor(NumTiles, [this, &NumResults, &GutterTexelsPerTile](int32 TileIdx)
 		{
-			int64 GutterPixelTo;
-			int64 GutterPixelFrom;
-			Tie(GutterPixelTo, GutterPixelFrom) = OccupancyMap.GutterTexels[GutterIdx];
-			for (int32 Idx = 0; Idx < NumResults; Idx++)
+			for (int64 GutterIdx = 0; GutterIdx < GutterTexelsPerTile[TileIdx].Num(); ++GutterIdx)
 			{
-				BakeResults[Idx]->CopyPixel(GutterPixelFrom, GutterPixelTo);
+				int64 GutterPixelTo;
+				int64 GutterPixelFrom;
+				Tie(GutterPixelTo, GutterPixelFrom) = GutterTexelsPerTile[TileIdx][GutterIdx];
+				FVector2i FromCoords = Dimensions.GetCoords(GutterPixelFrom);
+				FVector2i ToCoords = Dimensions.GetCoords(GutterPixelTo);
+				for (int32 Idx = 0; Idx < NumResults; Idx++)
+				{
+					BakeResults[Idx]->CopyPixel(GutterPixelFrom, GutterPixelTo);
+				}
 			}
-		}
-	}, !bParallel ? EParallelForFlags::ForceSingleThread : EParallelForFlags::None);
+		}, !bParallel ? EParallelForFlags::ForceSingleThread : EParallelForFlags::None);
+	}
 }
 
-void FMeshMapBaker::BakeSample(FMeshMapTileBuffer& TileBuffer, const FMeshMapEvaluator::FCorrespondenceSample& Sample, const FImageTile& Tile,
-	const FVector2i& TileCoords, const FVector2i& ImageCoords, const float& SampleWeight)
+void FMeshMapBaker::BakeSample(
+	FMeshMapTileBuffer& TileBuffer,
+	const FMeshMapEvaluator::FCorrespondenceSample& Sample,
+	const FVector2d& UVPosition,
+	const FVector2i& ImageCoords)
 {
-	const int64 TilePixelLinearIdx = Tile.GetIndex(TileCoords);
-	float* PixelBuffer = TileBuffer.GetPixel(TilePixelLinearIdx);
-	float& PixelWeight = TileBuffer.GetPixelWeight(TilePixelLinearIdx);
-	PixelWeight += SampleWeight;
-
+	// Evaluate each baker into stack allocated float buffer
 	float* Buffer = static_cast<float*>(FMemory_Alloca(sizeof(float) * BakeSampleBufferSize));
 	float* BufferPtr = Buffer;
 	const int32 NumBakers = Bakers.Num();
@@ -303,18 +405,58 @@ void FMeshMapBaker::BakeSample(FMeshMapTileBuffer& TileBuffer, const FMeshMapEva
 	}
 	checkSlow((BufferPtr - Buffer) == BakeSampleBufferSize);
 
-	auto AddFn = [PixelBuffer, Buffer](const int32 Idx) -> void
-	{
-		PixelBuffer[Idx] += Buffer[Idx];
-	};
+	const FImageTile& Tile = TileBuffer.GetTile();
 
-	auto OverwriteFn = [PixelBuffer, Buffer](const int32 Idx) -> void
+	auto AddFn = [this, &ImageCoords, &UVPosition, &Tile, &TileBuffer, Buffer](const TArray<int32>& BakeIds) -> void
 	{
-		PixelBuffer[Idx] = Buffer[Idx];
+		const FVector2i BoxFilterStart(
+			FMath::Clamp(ImageCoords.X - FilterKernelSize, 0, Dimensions.GetWidth()),
+			FMath::Clamp(ImageCoords.Y - FilterKernelSize, 0, Dimensions.GetHeight())
+			);
+		const FVector2i BoxFilterEnd(
+			FMath::Clamp(ImageCoords.X + FilterKernelSize + 1, 0, Dimensions.GetWidth()),
+			FMath::Clamp(ImageCoords.Y + FilterKernelSize + 1, 0, Dimensions.GetHeight())
+			);
+		const FImageTile BoxFilterTile(BoxFilterStart, BoxFilterEnd);
+
+		for (int64 FilterIdx = 0; FilterIdx < BoxFilterTile.Num(); FilterIdx++)
+		{
+			const FVector2i SourceCoords = BoxFilterTile.GetSourceCoords(FilterIdx);
+			const int64 BufferTilePixelLinearIdx = Tile.GetIndexFromSourceCoords(SourceCoords);
+			float* PixelBuffer = TileBuffer.GetPixel(BufferTilePixelLinearIdx);
+			float& PixelWeight = TileBuffer.GetPixelWeight(BufferTilePixelLinearIdx);
+
+			// Apply filter using double linear weighting (once per axis)
+			FVector2d TexelDistance = Dimensions.GetTexelUV(SourceCoords) - UVPosition;
+			TexelDistance.X *= Dimensions.GetWidth();
+			TexelDistance.Y *= Dimensions.GetHeight();
+
+			const float FilterWeight = TextureFilterEval(TexelDistance);
+			PixelWeight += FilterWeight;
+			for (int32 BakeIdx : BakeIds)
+			{
+				const FMeshMapEvaluator::FEvaluationContext& Context = BakeContexts[BakeIdx];
+				const int32 NumData = Context.DataLayout.Num();
+				const int32 ResultOffset = BakeOffsets[BakeIdx];
+				for (int32 DataIdx = 0; DataIdx < NumData; ++DataIdx)
+				{
+					const int32 ResultIdx = ResultOffset + DataIdx;
+					const int32 Offset = BakeSampleOffsets[ResultIdx];
+					const int32 Stride = (int32) Context.DataLayout[DataIdx];
+					for (int32 BufIdx = Offset; BufIdx < Offset + Stride; ++BufIdx)
+					{
+						PixelBuffer[BufIdx] += Buffer[BufIdx] * FilterWeight;
+					}
+				}
+			}
+		}
 	};
 
-	auto AccumulateFn = [this](const TArray<int32>& BakeIds, auto&& Op)
+	auto OverwriteFn = [this, &ImageCoords, &Tile, &TileBuffer, Buffer](const TArray<int32>& BakeIds) -> void
 	{
+		const int64 BufferTilePixelLinearIdx = Tile.GetIndexFromSourceCoords(ImageCoords); 
+		float* PixelBuffer = TileBuffer.GetPixel(BufferTilePixelLinearIdx);
+		
 		for (int32 Idx : BakeIds)
 		{
 			const FMeshMapEvaluator::FEvaluationContext& Context = BakeContexts[Idx];
@@ -327,14 +469,14 @@ void FMeshMapBaker::BakeSample(FMeshMapTileBuffer& TileBuffer, const FMeshMapEva
 				const int32 Stride = (int32) Context.DataLayout[DataIdx];
 				for (int32 BufIdx = Offset; BufIdx < Offset + Stride; ++BufIdx)
 				{
-					Op(BufIdx);
+					PixelBuffer[BufIdx] = Buffer[BufIdx];
 				}
 			}
 		}
 	};
 
-	AccumulateFn(BakeAccumulateLists[(int32)FMeshMapEvaluator::EAccumulateMode::Add], AddFn);
-	AccumulateFn(BakeAccumulateLists[(int32)FMeshMapEvaluator::EAccumulateMode::Overwrite], OverwriteFn);
+	AddFn(BakeAccumulateLists[(int32)FMeshMapEvaluator::EAccumulateMode::Add]);
+	OverwriteFn(BakeAccumulateLists[(int32)FMeshMapEvaluator::EAccumulateMode::Overwrite]);
 }
 
 int32 FMeshMapBaker::AddBaker(TSharedPtr<FMeshMapEvaluator> Sampler)
@@ -342,7 +484,7 @@ int32 FMeshMapBaker::AddBaker(TSharedPtr<FMeshMapEvaluator> Sampler)
 	return Bakers.Add(Sampler);
 }
 
-FMeshMapEvaluator* FMeshMapBaker::GetBaker(int32 BakerIdx)
+FMeshMapEvaluator* FMeshMapBaker::GetBaker(const int32 BakerIdx)
 {
 	return Bakers[BakerIdx].Get();
 }
@@ -353,15 +495,15 @@ void FMeshMapBaker::Reset()
 	BakeResults.Empty();
 }
 
-int32 FMeshMapBaker::NumBakers()
+int32 FMeshMapBaker::NumBakers() const
 {
 	return Bakers.Num();
 }
 
 const TArrayView<TUniquePtr<TImageBuilder<FVector4f>>> FMeshMapBaker::GetBakeResults(int32 BakerIdx)
 {
-	int32 ResultIdx = BakeOffsets[BakerIdx];
-	int32 NumResults = BakeOffsets[BakerIdx + 1] - ResultIdx;
+	const int32 ResultIdx = BakeOffsets[BakerIdx];
+	const int32 NumResults = BakeOffsets[BakerIdx + 1] - ResultIdx;
 	return TArrayView<TUniquePtr<TImageBuilder<FVector4f>>>(&BakeResults[ResultIdx], NumResults);
 }
 
@@ -370,9 +512,16 @@ void FMeshMapBaker::SetDimensions(FImageDimensions DimensionsIn)
 	Dimensions = DimensionsIn;
 }
 
+void FMeshMapBaker::SetGutterEnabled(bool bEnabled)
+{
+	bGutterEnabled = bEnabled;
+}
+
 void FMeshMapBaker::SetGutterSize(int32 GutterSizeIn)
 {
-	GutterSize = GutterSizeIn;
+	// GutterSize must be >= 1 since it is tied to MaxDistance for the
+	// OccupancyMap spatial search.
+	GutterSize = GutterSizeIn >= 1 ? GutterSizeIn : 1;
 }
 
 void FMeshMapBaker::SetMultisampling(int32 MultisamplingIn)
@@ -380,5 +529,54 @@ void FMeshMapBaker::SetMultisampling(int32 MultisamplingIn)
 	Multisampling = MultisamplingIn;
 }
 
+void FMeshMapBaker::SetFilter(EBakeFilterType FilterTypeIn)
+{
+	FilterType = FilterTypeIn;
+}
+
+void FMeshMapBaker::InitFilter()
+{
+	FilterKernelSize = TilePadding;
+	switch(FilterType)
+	{
+	case EBakeFilterType::None:
+		FilterKernelSize = 0;
+		TextureFilterEval = &EvaluateFilter<EBakeFilterType::None>;
+		break;
+	case EBakeFilterType::Box:
+		TextureFilterEval = &EvaluateFilter<EBakeFilterType::Box>;
+		break;
+	case EBakeFilterType::BSpline:
+		TextureFilterEval = &EvaluateFilter<EBakeFilterType::BSpline>;
+		break;
+	case EBakeFilterType::MitchellNetravali:
+		TextureFilterEval = &EvaluateFilter<EBakeFilterType::MitchellNetravali>;
+		break;
+	}
+}
+
+template<FMeshMapBaker::EBakeFilterType BakeFilterType>
+float FMeshMapBaker::EvaluateFilter(const FVector2d& Dist)
+{
+	float Result = 0.0f;
+	if constexpr(BakeFilterType == EBakeFilterType::None)
+	{
+		Result = 1.0f;
+	}
+	else if constexpr(BakeFilterType == EBakeFilterType::Box)
+	{
+		Result = BoxFilter.GetWeight(Dist);
+	}
+	else if constexpr(BakeFilterType == EBakeFilterType::BSpline)
+	{
+		Result = BSplineFilter.GetWeight(Dist);
+	}
+	else if constexpr(BakeFilterType == EBakeFilterType::MitchellNetravali)
+	{
+		Result = MitchellNetravaliFilter.GetWeight(Dist);
+	}
+	return Result;
+}
+
 
 
diff --git a/Engine/Plugins/Runtime/GeometryProcessing/Source/DynamicMesh/Public/Sampling/MeshBakerCommon.h b/Engine/Plugins/Runtime/GeometryProcessing/Source/DynamicMesh/Public/Sampling/MeshBakerCommon.h
--- a/Engine/Plugins/Runtime/GeometryProcessing/Source/DynamicMesh/Public/Sampling/MeshBakerCommon.h
+++ b/Engine/Plugins/Runtime/GeometryProcessing/Source/DynamicMesh/Public/Sampling/MeshBakerCommon.h
@@ -59,6 +59,11 @@ public:
 		return &Buffer[LinearIdx * PixelSize + 1];
 	}
 
+	const FImageTile& GetTile() const
+	{
+		return Tile;
+	}
+
 private:
 	const FImageTile Tile;
 	const int32 PixelSize;
diff --git a/Engine/Plugins/Runtime/GeometryProcessing/Source/DynamicMesh/Public/Sampling/MeshMapBaker.h b/Engine/Plugins/Runtime/GeometryProcessing/Source/DynamicMesh/Public/Sampling/MeshMapBaker.h
--- a/Engine/Plugins/Runtime/GeometryProcessing/Source/DynamicMesh/Public/Sampling/MeshMapBaker.h
+++ b/Engine/Plugins/Runtime/GeometryProcessing/Source/DynamicMesh/Public/Sampling/MeshMapBaker.h
@@ -5,9 +5,10 @@
 #include "Sampling/MeshBaseBaker.h"
 #include "Sampling/MeshMapEvaluator.h"
 #include "Sampling/MeshSurfaceSampler.h"
-#include "Spatial/DenseGrid2.h"
 #include "Image/ImageBuilder.h"
 #include "Image/ImageDimensions.h"
+#include "Image/BoxFilter.h"
+#include "Image/BCSplineFilter.h"
 
 namespace UE
 {
@@ -36,7 +37,7 @@ public:
 	FMeshMapEvaluator* GetBaker(int32 BakerIdx);
 
 	/** @return the number of bake evaluators on this baker. */
-	int32 NumBakers();
+	int32 NumBakers() const;
 
 	/** Reset the list of bakers. */
 	void Reset();
@@ -51,24 +52,34 @@ public:
 	//
 	// Parameters
 	//
+
+	enum class EBakeFilterType
+	{
+		None,
+		Box,
+		BSpline,
+		MitchellNetravali
+	};
 	
 	void SetDimensions(FImageDimensions DimensionsIn);
+	void SetGutterEnabled(bool bEnabled);
 	void SetGutterSize(int32 GutterSizeIn);
 	void SetMultisampling(int32 MultisamplingIn);
+	void SetFilter(EBakeFilterType FilterTypeIn);
 
 	FImageDimensions GetDimensions() const { return Dimensions; }
+	bool GetGutterEnabled() const { return bGutterEnabled; }
 	int32 GetGutterSize() const { return GutterSize; }
 	int32 GetMultisampling() const { return Multisampling; }
+	EBakeFilterType GetFilter() const { return FilterType; }
 
 protected:
 	/** Evaluate this sample. */
 	void BakeSample(
 		FMeshMapTileBuffer& TileBuffer,
 		const FMeshMapEvaluator::FCorrespondenceSample& Sample,
-		const FImageTile& Tile,
-		const FVector2i& TileCoords,
-		const FVector2i& ImageCoords,
-		const float& SampleWeight);
+		const FVector2d& UVPosition,
+		const FVector2i& ImageCoords);
 
 	/** Initialize evaluation contexts and precompute data for bake evaluation. */
 	void InitBake();
@@ -76,6 +87,9 @@ protected:
 	/** Initialize bake sample default floats and colors. */
 	void InitBakeDefaults();
 
+	/** Initialize filter */
+	void InitFilter();
+
 protected:
 	const bool bParallel = true;
 
@@ -83,6 +97,14 @@ protected:
 	TMeshSurfaceUVSampler<FMeshMapEvaluator::FCorrespondenceSample> DetailMeshSampler;
 
 	FImageDimensions Dimensions = FImageDimensions(128, 128);
+
+	/**
+	 * If true, the baker will pad the baked content past the UV borders by GutterSize.
+	 * This is useful to minimize artifacts when filtering or mipmapping.
+	 */
+	bool bGutterEnabled = true;
+
+	/** The pixel distance (in texel diagonal length) to pad baked content past the UV borders. */
 	int32 GutterSize = 4;
 
 	/** The square dimensions for multisampling each pixel. */
@@ -91,6 +113,27 @@ protected:
 	/** The square dimensions for tiled processing of the output image(s). */
 	const int32 TileSize = 32;
 
+	/** The amount of padding for tiled processing of the output image(s). */
+	const int32 TilePadding = 2;
+
+	/** The pixel distance around the sample texel to be considered by the filter. [0, TilePadding] */
+	int32 FilterKernelSize = 0;
+
+	/** The texture filter type. */
+	EBakeFilterType FilterType = EBakeFilterType::BSpline;
+
+	/** Texture filters */
+	static FBoxFilter BoxFilter;
+	static FBSplineFilter BSplineFilter;
+	static FMitchellNetravaliFilter MitchellNetravaliFilter;
+
+	/** Texture filter function */
+	using TextureFilterFn = float(*)(const FVector2d& Dist);
+	TextureFilterFn TextureFilterEval = nullptr;
+
+	template<EBakeFilterType BakeFilterType>
+	static float EvaluateFilter(const FVector2d& Dist);
+
 	/** The total size of the temporary float buffer for BakeSample. */
 	int32 BakeSampleBufferSize = 0;
 
diff --git a/Engine/Source/Runtime/GeometryCore/Public/Image/BCSplineFilter.h b/Engine/Source/Runtime/GeometryCore/Public/Image/BCSplineFilter.h
new file mode 100644
--- /dev/null
+++ b/Engine/Source/Runtime/GeometryCore/Public/Image/BCSplineFilter.h
@@ -0,0 +1,112 @@
+// Copyright Epic Games, Inc. All Rights Reserved.
+
+#pragma once
+
+#include "MathUtil.h"
+
+namespace UE
+{
+namespace Geometry
+{
+
+enum class EBCSplineType
+{
+	BSpline,
+	CatmullRom,
+	MitchellNetravali
+};
+
+/**
+ * @param SplineType The BC-spline filter type (predefined B/C parameters)
+ * @param bIsRadial When true, use Euclidean distance to compute filter weights: Weight = Filter(sqrt(X*X + Y*Y)).
+ *                  Otherwise use linear 1D distance in X/Y separately: WX = Filter(X), WY = Filter(Y), Weight = WX*WY 
+ */
+template <EBCSplineType SplineType, bool bIsRadial>
+class FBCSplineFilter
+{
+private:
+	float Radius = 1.0f;
+	float Coeff1[3] = {0.0f, 0.0f, 0.0f};
+	float Coeff2[4] = {0.0f, 0.0f, 0.0f, 0.0f};
+
+public:
+	FBCSplineFilter(const float RadiusIn);
+
+	/** @return the filter weight given a 2D distance vector. */
+	float GetWeight(const FVector2d& Dist) const;
+};
+
+template<EBCSplineType SplineType, bool bIsRadial>
+FBCSplineFilter<SplineType, bIsRadial>::FBCSplineFilter(const float RadiusIn)
+	: Radius(RadiusIn)
+{
+	auto ComputeCoeffs = [this](const float& B, const float& C)
+	{
+		// [x < 1]:       1/6 * ((12 - 9B - 6C)x^3 + (-18 + 12B + 6C)x^2 + (6 - 2B))
+		Coeff1[0] = (12.0f - 9.0f * B - 6.0f * C) / 6.0f;
+		Coeff1[1] = (-18.0f + 12.0f * B + 6.0f * C) / 6.0f;
+		Coeff1[2] = (6.0f - 2.0f * B) / 6.0f;
+
+		// [1 < x < 2]:   1/6 * ((-B - 6C)x^3 + (6B + 30C)x^2 + (-12B - 48C)x + (8B + 24C))
+		Coeff2[0] = (-B - 6.0f * C) / 6.0f;
+		Coeff2[1] = (6.0f * B + 30.0f * C) / 6.0f;
+		Coeff2[2] = (-12.0f * B - 48.0f * C) / 6.0f;
+		Coeff2[3] = (8.0f * B + 24.0f * C) / 6.0f;
+	};
+	
+	if constexpr(SplineType == EBCSplineType::BSpline)
+	{
+		const float B = 1.0f;
+		const float C = 0.0f;
+		ComputeCoeffs(B, C);
+	}
+	else if constexpr(SplineType == EBCSplineType::CatmullRom)
+	{
+		const float B = 0.0f;
+		const float C = 0.5f;
+		ComputeCoeffs(B, C);
+	}
+	else if constexpr(SplineType == EBCSplineType::MitchellNetravali)
+	{
+		const float B = 1.0f / 3.0f;
+		const float C = B;
+		ComputeCoeffs(B, C);
+	}
+}
+
+template<EBCSplineType SplineType, bool bIsRadial>
+float FBCSplineFilter<SplineType, bIsRadial>::GetWeight(const FVector2d& Dist) const
+{
+	auto ComputeWeight = [this](float X) -> float
+	{
+		const float AbsX = FMathf::Abs(X);
+		float Weight = 0.0f;
+		if(AbsX < Radius)
+		{
+			Weight = Coeff1[0] * AbsX * AbsX * AbsX + Coeff1[1] * AbsX * AbsX + Coeff1[2];
+		}
+		else if(AbsX < Radius * 2.0f)
+		{
+			Weight = Coeff2[0] * AbsX * AbsX * AbsX + Coeff2[1] * AbsX * AbsX + Coeff2[2] * AbsX + Coeff2[3];
+		}
+		return Weight;
+	};
+
+	if constexpr(bIsRadial)
+	{
+		return ComputeWeight(Dist.Length());
+	}
+	else
+	{
+		const float WeightX = ComputeWeight(Dist.X);
+		const float WeightY = ComputeWeight(Dist.Y);
+		return WeightX * WeightY;
+	}
+}
+
+using FBSplineFilter = FBCSplineFilter<EBCSplineType::BSpline, /*bIsRadial*/ false>;
+using FCatmullRomFilter = FBCSplineFilter<EBCSplineType::CatmullRom, /*bIsRadial*/ false>;
+using FMitchellNetravaliFilter = FBCSplineFilter<EBCSplineType::MitchellNetravali, /*bIsRadial*/ false>;
+	
+} // end namespace UE::Geometry
+} // end namespace UE
\ No newline at end of file
diff --git a/Engine/Source/Runtime/GeometryCore/Public/Image/BoxFilter.h b/Engine/Source/Runtime/GeometryCore/Public/Image/BoxFilter.h
new file mode 100644
--- /dev/null
+++ b/Engine/Source/Runtime/GeometryCore/Public/Image/BoxFilter.h
@@ -0,0 +1,30 @@
+// Copyright Epic Games, Inc. All Rights Reserved.
+
+#pragma once
+
+namespace UE
+{
+namespace Geometry
+{
+
+class FBoxFilter
+{
+private:
+	float Radius = 0.5f;
+	
+public:
+	FBoxFilter(const float RadiusIn)
+		: Radius(RadiusIn)
+	{		
+	}
+
+	/** @return the filter weight given a 2D distance vector. */
+	float GetWeight(const FVector2d& Dist) const
+	{
+		// [-Radius, Radius)
+		return Dist.X >= -Radius && Dist.X < Radius && Dist.Y >= -Radius && Dist.Y < Radius; 
+	}
+};
+
+} // end namespace UE::Geometry
+} // end namespace UE
\ No newline at end of file
diff --git a/Engine/Source/Runtime/GeometryCore/Public/Image/ImageTile.h b/Engine/Source/Runtime/GeometryCore/Public/Image/ImageTile.h
--- a/Engine/Source/Runtime/GeometryCore/Public/Image/ImageTile.h
+++ b/Engine/Source/Runtime/GeometryCore/Public/Image/ImageTile.h
@@ -14,34 +14,39 @@ class FImageTile
 {
 private:
 	FVector2i Start = FVector2i::Zero();
-	FVector2i End = FVector2i::Zero(); 
+	FVector2i End = FVector2i::Zero();
+	int Width = 0;
+	int Height = 0;
 
 public:
 	FImageTile() = default;
 	FImageTile(const FVector2i& InStart, const FVector2i& InEnd)
-		: Start(InStart), End(InEnd)
+		: Start(InStart)
+		, End(InEnd)
+		, Width(InEnd.X - InStart.X)
+		, Height(InEnd.Y - InStart.Y)
 	{
 	}
 
 	int32 GetWidth() const
 	{
-		return End.X - Start.X;
+		return Width;
 	}
 
 	int32 GetHeight() const
 	{
-		return End.Y - Start.Y;
+		return Height;
 	}
 
 	int64 Num() const
 	{
-		return GetWidth() * GetHeight();
+		return Width * Height;
 	}
 
 	/** @return the linear index for the given local XY coordinates into this tile. */
 	int64 GetIndex(const int32 X, const int32 Y) const
 	{
-		return Y * GetWidth() + X;
+		return Y * Width + X;
 	}
 
 	/** @return the linear index for the given local coordinates into this tile. */
@@ -50,10 +55,29 @@ public:
 		return GetIndex(LocalCoords.X, LocalCoords.Y);
 	}
 
+	/** @return true if the given source coordinates are contained by this tile. */
+	bool Contains(const int32 X, const int32 Y) const
+	{
+		return (X >= Start.X && X < End.X && Y >= Start.Y && Y < End.Y);
+	}
+
+	/** @return the linear index for the given source coordinates. */
+	int64 GetIndexFromSourceCoords(const int32 X, const int32 Y) const
+	{
+		checkSlow(Contains(X, Y));
+		return GetIndex(X - Start.X, Y - Start.Y); 
+	}
+
+	/** @return the linear index for the given source coordinates. */
+	int64 GetIndexFromSourceCoords(const FVector2i& SourceCoords) const
+	{
+		return GetIndexFromSourceCoords(SourceCoords.X, SourceCoords.Y); 
+	}
+
 	/** @return the source image's coordinates given linear index into this tile. */
 	FVector2i GetSourceCoords(const int64 LinearIdx) const
 	{
-		return FVector2i(Start.X + LinearIdx % GetWidth(), Start.Y + LinearIdx / GetWidth());
+		return FVector2i(Start.X + LinearIdx % Width, Start.Y + LinearIdx / Width);
 	}
 
 	/** @return the source image's coordinates given local XY coordinates into this tile. */
