ModelingTools: Fix BakeTexture synchronization to output buffer and re-enable BakeTexture unit test.

#rnx

diff --git a/Engine/Plugins/Runtime/GeometryProcessing/Source/DynamicMesh/Private/Sampling/MeshMapBaker.cpp b/Engine/Plugins/Runtime/GeometryProcessing/Source/DynamicMesh/Private/Sampling/MeshMapBaker.cpp
--- a/Engine/Plugins/Runtime/GeometryProcessing/Source/DynamicMesh/Private/Sampling/MeshMapBaker.cpp
+++ b/Engine/Plugins/Runtime/GeometryProcessing/Source/DynamicMesh/Private/Sampling/MeshMapBaker.cpp
@@ -188,7 +188,9 @@ void FMeshMapBaker::Bake()
 	const int32 NumTiles = Tiles.Num();
 	TArray<TArray64<TTuple<int64, int64>>> GutterTexelsPerTile;
 	GutterTexelsPerTile.SetNum(NumTiles);
-	ParallelFor(NumTiles, [this, &Tiles, &ImageTileBuffer, &GutterTexelsPerTile](int32 TileIdx)
+
+	FCriticalSection WriteLock;
+	ParallelFor(NumTiles, [this, &Tiles, &ImageTileBuffer, &GutterTexelsPerTile, &WriteLock](int32 TileIdx)
 	{
 		// Generate unpadded and padded tiles.
 		const FImageTile Tile = Tiles.GetTile(TileIdx);	// Image area to sample
@@ -307,8 +309,9 @@ void FMeshMapBaker::Bake()
 
 		// Accumulate 'Add' float data to image tile buffer
 		// TODO: Optimize this write lock
-		FCriticalSection WriteLock;
+		WriteLock.Lock();
 		WriteToOutputBuffer(PaddedTile, BakeAccumulateLists[(int32)FMeshMapEvaluator::EAccumulateMode::Add], AddFn, AddFn);
+		WriteLock.Unlock();
 	}, !bParallel ? EParallelForFlags::ForceSingleThread : EParallelForFlags::None);
 
 	// Normalize and convert ImageTileBuffer data to color data.
