GeometryProcessing: Move FFindPolygonsAlgorithm from ModelingTools to GeometryProcessing, rename to FPolygroupsGenerator. Replace usages of 'Polygon' with 'Polygroup' and update usage in ConvertToPolygonsTool.
#rnx

diff --git a/Engine/Plugins/Experimental/MeshModelingToolsetExp/Source/MeshModelingToolsExp/Private/ConvertToPolygonsTool.cpp b/Engine/Plugins/Experimental/MeshModelingToolsetExp/Source/MeshModelingToolsExp/Private/ConvertToPolygonsTool.cpp
--- a/Engine/Plugins/Experimental/MeshModelingToolsetExp/Source/MeshModelingToolsExp/Private/ConvertToPolygonsTool.cpp
+++ b/Engine/Plugins/Experimental/MeshModelingToolsetExp/Source/MeshModelingToolsExp/Private/ConvertToPolygonsTool.cpp
@@ -19,7 +19,7 @@
 #include "Polygroups/PolygroupUtil.h"
 #include "Util/ColorConstants.h"
 
-#include "FindPolygonsAlgorithm.h"
+#include "Polygroups/PolygroupsGenerator.h"
 
 #include "ExplicitUseGeometryMathTypes.h"		// using UE::Geometry::(math types)
 using namespace UE::Geometry;
@@ -42,11 +42,16 @@ public:
 	EConvertToPolygonsMode ConversionMode = EConvertToPolygonsMode::FaceNormalDeviation;
 	double AngleTolerance = 0.1f;
 	int32 NumPoints = 10;
-	FFindPolygonsAlgorithm::EWeightingType WeightingType = FFindPolygonsAlgorithm::EWeightingType::None;
+	FPolygroupsGenerator::EWeightingType WeightingType = FPolygroupsGenerator::EWeightingType::None;
 	FVector3d WeightingCoeffs = FVector3d::One();
 	int32 MinGroupSize = 2;
 	bool bCalculateNormals = false;
 
+	// input mesh
+	TSharedPtr<FDynamicMesh3, ESPMode::ThreadSafe> OriginalMesh;
+	// result
+	FPolygroupsGenerator Generator;
+
 	virtual void CalculateResult(FProgressCancel* Progress) override
 	{
 		if ((Progress && Progress->Cancelled()) || !OriginalMesh)
@@ -62,37 +67,37 @@ public:
 			return;
 		}
 
-		Polygons = FFindPolygonsAlgorithm(ResultMesh.Get());
-		Polygons.MinGroupSize = MinGroupSize;
+		Generator = FPolygroupsGenerator(ResultMesh.Get());
+		Generator.MinGroupSize = MinGroupSize;
 
 		switch (ConversionMode)
 		{
 		case EConvertToPolygonsMode::FromUVIslands:
 		{
-			Polygons.FindPolygonsFromUVIslands();
+			Generator.FindPolygroupsFromUVIslands();
 			break;
 		}
 		case EConvertToPolygonsMode::FromConnectedTris:
 		{
-			Polygons.FindPolygonsFromConnectedTris();
+			Generator.FindPolygroupsFromConnectedTris();
 			break;
 		}
 		case EConvertToPolygonsMode::FaceNormalDeviation:
 		{
 			double DotTolerance = 1.0 - FMathd::Cos(AngleTolerance * FMathd::DegToRad);
-			Polygons.FindPolygonsFromFaceNormals(DotTolerance);
+			Generator.FindPolygroupsFromFaceNormals(DotTolerance);
 			break;
 		}
 		case EConvertToPolygonsMode::FromFurthestPointSampling:
 		{
-			Polygons.FindPolygonsFromFurthestPointSampling(NumPoints, WeightingType, WeightingCoeffs);
+			Generator.FindPolygroupsFromFurthestPointSampling(NumPoints, WeightingType, WeightingCoeffs);
 			break;
 		}
 		default:
 			check(0);
 		}
 
-		Polygons.FindPolygonEdges();
+		Generator.FindPolygroupEdges();
 
 		if (bCalculateNormals && ConversionMode == EConvertToPolygonsMode::FaceNormalDeviation)
 		{
@@ -105,7 +110,7 @@ public:
 			NormalOverlay->ClearElements();
 
 			FDynamicMeshEditor Editor(ResultMesh.Get());
-			for (const TArray<int>& Polygon : Polygons.FoundPolygons)
+			for (const TArray<int>& Polygon : Generator.FoundPolygroups)
 			{
 				FVector3f Normal = (FVector3f)ResultMesh->GetTriNormal(Polygon[0]);
 				Editor.SetTriangleNormals(Polygon, Normal);
@@ -123,9 +128,7 @@ public:
 		ResultTransform = (UE::Geometry::FTransform3d)Transform;
 	}
 
-	
-	FFindPolygonsAlgorithm Polygons;
-	TSharedPtr<FDynamicMesh3, ESPMode::ThreadSafe> OriginalMesh;
+
 };
 
 TUniquePtr<FDynamicMeshOperator> UConvertToPolygonsOperatorFactory::MakeNewOperator()
@@ -195,7 +198,7 @@ void UConvertToPolygonsTool::Setup()
 		PreviewCompute->OnOpCompleted.AddLambda([this](const FDynamicMeshOperator* MeshOp) 
 		{ 
 			const FConvertToPolygonsOp*  ConvertToPolygonsOp = static_cast<const FConvertToPolygonsOp*>(MeshOp);
-			this->PolygonEdges = ConvertToPolygonsOp->Polygons.PolygonEdges;
+			this->PolygonEdges = ConvertToPolygonsOp->Generator.PolygroupEdges;
 			UpdateVisualization();
 		});
 
@@ -229,7 +232,7 @@ void UConvertToPolygonsTool::UpdateOpParameters(FConvertToPolygonsOp& ConvertToP
 	ConvertToPolygonsOp.ConversionMode    = Settings->ConversionMode;
 	ConvertToPolygonsOp.AngleTolerance    = Settings->AngleTolerance;
 	ConvertToPolygonsOp.NumPoints = Settings->NumPoints;
-	ConvertToPolygonsOp.WeightingType = (Settings->bNormalWeighted) ? FFindPolygonsAlgorithm::EWeightingType::NormalDeviation : FFindPolygonsAlgorithm::EWeightingType::None;
+	ConvertToPolygonsOp.WeightingType = (Settings->bNormalWeighted) ? FPolygroupsGenerator::EWeightingType::NormalDeviation : FPolygroupsGenerator::EWeightingType::None;
 	ConvertToPolygonsOp.WeightingCoeffs = FVector3d(Settings->NormalWeighting, 1.0, 1.0);
 	ConvertToPolygonsOp.MinGroupSize = Settings->MinGroupSize;
 	ConvertToPolygonsOp.OriginalMesh = OriginalDynamicMesh;
diff --git a/Engine/Plugins/Experimental/MeshModelingToolsetExp/Source/MeshModelingToolsExp/Private/DeformMeshPolygonsTool.cpp b/Engine/Plugins/Experimental/MeshModelingToolsetExp/Source/MeshModelingToolsExp/Private/DeformMeshPolygonsTool.cpp
--- a/Engine/Plugins/Experimental/MeshModelingToolsetExp/Source/MeshModelingToolsExp/Private/DeformMeshPolygonsTool.cpp
+++ b/Engine/Plugins/Experimental/MeshModelingToolsetExp/Source/MeshModelingToolsExp/Private/DeformMeshPolygonsTool.cpp
@@ -5,7 +5,6 @@
 #include "Async/ParallelFor.h"
 #include "Containers/BitArray.h"
 #include "DynamicMesh/DynamicMeshAttributeSet.h"
-#include "FindPolygonsAlgorithm.h"
 #include "InteractiveToolManager.h"
 #include "Intersection/IntersectionUtil.h"
 #include "DynamicMesh/MeshNormals.h"
diff --git a/Engine/Plugins/Experimental/MeshModelingToolsetExp/Source/MeshModelingToolsExp/Private/FindPolygonsAlgorithm.cpp b/Engine/Plugins/Experimental/MeshModelingToolsetExp/Source/MeshModelingToolsExp/Private/FindPolygonsAlgorithm.cpp
deleted file mode 100644
diff --git a/Engine/Plugins/Experimental/MeshModelingToolsetExp/Source/MeshModelingToolsExp/Public/FindPolygonsAlgorithm.h b/Engine/Plugins/Experimental/MeshModelingToolsetExp/Source/MeshModelingToolsExp/Public/FindPolygonsAlgorithm.h
deleted file mode 100644
diff --git a/Engine/Plugins/Runtime/GeometryProcessing/Source/DynamicMesh/Private/Polygroups/PolygroupsGenerator.cpp b/Engine/Plugins/Runtime/GeometryProcessing/Source/DynamicMesh/Private/Polygroups/PolygroupsGenerator.cpp
new file mode 100644
--- /dev/null
+++ b/Engine/Plugins/Runtime/GeometryProcessing/Source/DynamicMesh/Private/Polygroups/PolygroupsGenerator.cpp
@@ -0,0 +1,413 @@
+// Copyright Epic Games, Inc. All Rights Reserved.
+
+#include "Polygroups/PolygroupsGenerator.h"
+#include "DynamicMesh/MeshNormals.h"
+#include "Selections/MeshConnectedComponents.h"
+#include "Util/IndexUtil.h"
+#include "Parameterization/MeshDijkstra.h"
+#include "Parameterization/IncrementalMeshDijkstra.h"
+#include "Parameterization/MeshRegionGraph.h"
+
+#include "Curve/DynamicGraph3.h"
+
+
+
+#include "ExplicitUseGeometryMathTypes.h"		// using UE::Geometry::(math types)
+using namespace UE::Geometry;
+
+FPolygroupsGenerator::FPolygroupsGenerator(FDynamicMesh3* MeshIn)
+{
+	Mesh = MeshIn;
+}
+
+
+
+bool FPolygroupsGenerator::FindPolygroupsFromUVIslands()
+{
+	const FDynamicMeshUVOverlay* UV = Mesh->Attributes()->GetUVLayer(0);
+
+	FMeshConnectedComponents Components(Mesh);
+	Components.FindConnectedTriangles([&UV](int32 TriIdx0, int32 TriIdx1)
+	{
+		return UV->AreTrianglesConnected(TriIdx0, TriIdx1);
+	});
+
+	int32 NumComponents = Components.Num();
+	for (int32 ci = 0; ci < NumComponents; ++ci)
+	{
+		FoundPolygroups.Add(Components.GetComponent(ci).Indices);
+	}
+
+	if (bApplyPostProcessing)
+	{
+		PostProcessPolygroups(false);
+	}
+	if (bCopyToMesh)
+	{
+		CopyPolygroupsToMesh();
+	}
+
+	return (FoundPolygroups.Num() > 0);
+}
+
+
+
+bool FPolygroupsGenerator::FindPolygroupsFromConnectedTris()
+{
+
+	FMeshConnectedComponents Components(Mesh);
+	Components.FindConnectedTriangles([this](int32 TriIdx0, int32 TriIdx1)
+	{
+		FIndex3i NbrTris = Mesh->GetTriNeighbourTris(TriIdx0);
+		int NbrIndex = IndexUtil::FindTriIndex(TriIdx1, NbrTris);
+		return (NbrIndex != IndexConstants::InvalidID);
+	});
+
+	int32 NumComponents = Components.Num();
+	for (int32 ci = 0; ci < NumComponents; ++ci)
+	{
+		FoundPolygroups.Add(Components.GetComponent(ci).Indices);
+	}
+
+	if (bApplyPostProcessing)
+	{
+		PostProcessPolygroups(false);
+	}
+	if (bCopyToMesh)
+	{
+		CopyPolygroupsToMesh();
+	}
+
+	return (FoundPolygroups.Num() > 0);
+}
+
+
+
+bool FPolygroupsGenerator::FindPolygroupsFromFaceNormals(double DotTolerance)
+{
+	DotTolerance = 1.0 - DotTolerance;
+
+	// compute face normals
+	FMeshNormals Normals(Mesh);
+	Normals.ComputeTriangleNormals();
+
+	TArray<bool> DoneTriangle;
+	DoneTriangle.SetNum(Mesh->MaxTriangleID());
+
+	TArray<int> Stack;
+
+	// grow outward from vertices until we have no more left
+	for (int TriID : Mesh->TriangleIndicesItr())
+	{
+		if (DoneTriangle[TriID] == true)
+		{
+			continue;
+		}
+
+		TArray<int> Polygroup;
+		Polygroup.Add(TriID);
+		DoneTriangle[TriID] = true;
+
+		Stack.SetNum(0);
+		Stack.Add(TriID);
+		while (Stack.Num() > 0)
+		{
+			int CurTri = Stack.Pop(false);
+			FIndex3i NbrTris = Mesh->GetTriNeighbourTris(CurTri);
+			for (int j = 0; j < 3; ++j)
+			{
+				if (NbrTris[j] >= 0
+					&& DoneTriangle[NbrTris[j]] == false)
+				{
+					double Dot = Normals[CurTri].Dot(Normals[NbrTris[j]]);
+					if (Dot > DotTolerance)
+					{
+						Polygroup.Add(NbrTris[j]);
+						Stack.Add(NbrTris[j]);
+						DoneTriangle[NbrTris[j]] = true;
+					}
+				}
+			}
+		}
+
+		FoundPolygroups.Add(Polygroup);
+	}
+
+	if (bApplyPostProcessing)
+	{
+		PostProcessPolygroups(true);
+	}
+	if (bCopyToMesh)
+	{
+		CopyPolygroupsToMesh();
+	}
+
+	return (FoundPolygroups.Num() > 0);
+}
+
+
+
+
+/**
+ * Dual graph of mesh faces, ie graph of edges across faces between face centers.
+ * Normals and Areas are tracked for each point.
+ */
+class FMeshFaceDualGraph : public FDynamicGraph3d
+{
+public:
+	TDynamicVectorN<double, 3> Normals;
+	TDynamicVector<double> Areas;
+
+	int AppendVertex(const FVector3d& Centroid, const FVector3d& Normal, double Area)
+	{
+		int vid = FDynamicGraph3d::AppendVertex(Centroid);
+		Normals.InsertAt({ {Normal.X, Normal.Y, Normal.Z} }, vid);
+		Areas.InsertAt(Area, vid);
+		return vid;
+	}
+
+	FVector3d GetNormal(int32 vid) const
+	{
+		return Normals.AsVector3(vid);
+	}
+
+	/** Build a Face Dual Graph for a triangle mesh */
+	static void MakeFaceDualGraphForMesh(FDynamicMesh3* Mesh, FMeshFaceDualGraph& FaceGraph)
+	{
+		// if not true, code below needs updating
+		check(Mesh->IsCompactT());
+
+		for (int32 tid : Mesh->TriangleIndicesItr())
+		{
+			FVector3d Normal, Centroid;
+			double Area;
+			Mesh->GetTriInfo(tid, Normal, Area, Centroid);
+
+			int32 newid = FaceGraph.AppendVertex(Centroid, Normal, Area);
+			check(newid == tid);
+		}
+
+		for (int32 tid : Mesh->TriangleIndicesItr())
+		{
+			FIndex3i NbrT = Mesh->GetTriNeighbourTris(tid);
+			for (int32 j = 0; j < 3; ++j)
+			{
+				if (Mesh->IsTriangle(NbrT[j]))
+				{
+					FaceGraph.AppendEdge(tid, NbrT[j]);
+				}
+			}
+		}
+	}
+
+};
+
+
+
+
+
+
+bool FPolygroupsGenerator::FindPolygroupsFromFurthestPointSampling(int32 NumPoints, EWeightingType WeightingType, FVector3d WeightingCoeffs)
+{
+	NumPoints = FMath::Min(NumPoints, Mesh->VertexCount());
+
+	FMeshFaceDualGraph FaceGraph;
+	FMeshFaceDualGraph::MakeFaceDualGraphForMesh(Mesh, FaceGraph);
+
+	TIncrementalMeshDijkstra<FMeshFaceDualGraph> FurthestPoints(&FaceGraph);
+
+	TArray<int32> SeedIndices;
+
+	// need to add at least one seed point for each mesh connected component, so that all triangles are assigned a group
+	// TODO: two seed points for components that have no boundary?
+	FMeshConnectedComponents Components(Mesh);
+	Components.FindConnectedTriangles();
+	int32 NumConnected = Components.Num();
+	for (int32 k = 0; k < NumConnected; ++k)
+	{
+		SeedIndices.Add(Components[k].Indices[0]);
+	}
+
+	// initial incremental update from per-component points
+	TArray<TIncrementalMeshDijkstra<FMeshFaceDualGraph>::FSeedPoint> ComponentSeeds;
+	for (int32 vid : SeedIndices)
+	{
+		ComponentSeeds.Add(TIncrementalMeshDijkstra<FMeshFaceDualGraph>::FSeedPoint{ vid, vid, 0.0 });
+	}
+	FurthestPoints.AddSeedPoints(ComponentSeeds);
+
+	// TODO: can (approximately) bound the size of a region based on mesh area. Then can pass that in
+	// with seed point as upper distance bound. This will change the result as it will initially grow on
+	// the 'front' (and possibly no guarantee that mesh is covered?)
+	//    (initial furthest-points would be any with an invalid value...then we can ensure coverage)
+	// FindMaxGraphDistancePointID() seems like it might be somewhat expensive...
+
+	// incrementally add furthest-points
+	while ( SeedIndices.Num() < NumPoints)
+	{
+		int32 NextPointID = FurthestPoints.FindMaxGraphDistancePointID();
+		if (NextPointID >= 0)
+		{
+			TIncrementalMeshDijkstra<FMeshFaceDualGraph>::FSeedPoint NextSeedPoint{ NextPointID, NextPointID, 0.0 };
+			FurthestPoints.AddSeedPoints({ NextSeedPoint });
+			SeedIndices.Add(NextPointID);
+		}
+		else
+		{
+			break;
+		}
+	}
+
+
+	// Now that we have furthest point set, recompute a Dijkstra propagation with optional weighting
+	// (unweighted version should be the same as the FurthestPoints dijkstra though...could re-use?)
+
+	TMeshDijkstra<FMeshFaceDualGraph> SuperPixels(&FaceGraph);
+	
+	if (WeightingType == EWeightingType::NormalDeviation)
+	{
+		SuperPixels.bEnableDistanceWeighting = true;
+		SuperPixels.GetWeightedDistanceFunc = [this, WeightingCoeffs, &FaceGraph](int32 FromVID, int32 ToVID, int32 SeedVID, double Distance)
+		{
+			FVector3d NA = FaceGraph.GetNormal(ToVID);
+			FVector3d NB = FaceGraph.GetNormal(SeedVID);
+			double Dot = NA.Dot(NB);
+			if (WeightingCoeffs.X > 0.001)
+			{
+				Dot = FMathd::Pow(Dot, WeightingCoeffs.X);
+			}
+			double W = FMathd::Clamp(1.0 - Dot * Dot, 0.0, 1.0);
+			W = W * W * W;
+			double Weight = FMathd::Clamp(W, 0.001, 1.0);
+			return Weight * Distance;
+		};
+	}
+
+	TArray<TMeshDijkstra<FMeshFaceDualGraph>::FSeedPoint> SuperPixelSeeds;
+	for (int32 k = 0; k < SeedIndices.Num(); ++k)
+	{
+		SuperPixelSeeds.Add(TMeshDijkstra<FMeshFaceDualGraph>::FSeedPoint{ k, SeedIndices[k], 0.0 });
+	}
+	SuperPixels.ComputeToMaxDistance(SuperPixelSeeds, TNumericLimits<double>::Max());
+
+	TArray<TArray<int32>> TriSets;
+	TriSets.SetNum(SeedIndices.Num());
+	TArray<int32> FailedSet;
+
+	for (int32 tid : Mesh->TriangleIndicesItr())
+	{
+		int32 SeedID = SuperPixels.GetSeedExternalIDForPointSetID(tid);
+		if (SeedID >= 0)
+		{
+			TriSets[SeedID].Add(tid);
+		}
+		else
+		{
+			FailedSet.Add(tid);
+		}
+	}
+
+
+	for (int32 k = 0; k < NumPoints; ++k)
+	{
+		if (TriSets[k].Num() > 0)
+		{
+			FoundPolygroups.Add(MoveTemp(TriSets[k]));
+		}
+	}
+	
+	if (FailedSet.Num() > 0)
+	{
+		FMeshConnectedComponents FailedComponents(Mesh);
+		FailedComponents.FindConnectedTriangles(FailedSet);
+		for (FMeshConnectedComponents::FComponent& Component : FailedComponents)
+		{
+			FoundPolygroups.Add(Component.Indices);
+		}
+	}
+
+	if (bApplyPostProcessing)
+	{
+		PostProcessPolygroups(true);
+	}
+	if (bCopyToMesh)
+	{
+		CopyPolygroupsToMesh();
+	}
+
+	return (FoundPolygroups.Num() > 0);
+}
+
+
+
+
+void FPolygroupsGenerator::PostProcessPolygroups(bool bApplyMerging)
+{
+	if (bApplyMerging && MinGroupSize > 1)
+	{
+		OptimizePolygroups();
+	}
+}
+
+
+void FPolygroupsGenerator::OptimizePolygroups()
+{
+	FMeshRegionGraph RegionGraph;
+	RegionGraph.BuildFromTriangleSets(*Mesh, FoundPolygroups, [&](int32 SetIdx) { return SetIdx; });
+	bool bMerged = RegionGraph.MergeSmallRegions(MinGroupSize-1, [&](int32 A, int32 B) { return RegionGraph.GetRegionTriCount(A) > RegionGraph.GetRegionTriCount(B); });
+	bool bSwapped = RegionGraph.OptimizeBorders();
+	if (bMerged || bSwapped)
+	{
+		FoundPolygroups.Reset();
+
+		int32 N = RegionGraph.MaxRegionIndex();
+		for (int32 k = 0; k < N; ++k)
+		{
+			if (RegionGraph.IsRegion(k))
+			{
+				const TArray<int32>& Tris = RegionGraph.GetRegionTris(k);
+				FoundPolygroups.Add(Tris);
+			}
+		}
+	}
+}
+
+
+
+
+void FPolygroupsGenerator::CopyPolygroupsToMesh()
+{
+	Mesh->EnableTriangleGroups(0);
+
+	// set groups from Polygroups
+	int NumPolygroups = FoundPolygroups.Num();
+
+	// can be parallel for
+	for (int PolyIdx = 0; PolyIdx < NumPolygroups; PolyIdx++)
+	{
+		const TArray<int>& Polygroup = FoundPolygroups[PolyIdx];
+		int NumTriangles = Polygroup.Num();
+		for (int k = 0; k < NumTriangles; ++k)
+		{
+			Mesh->SetTriangleGroup(Polygroup[k], (PolyIdx + 1));
+		}
+	}
+}
+
+
+
+
+
+bool FPolygroupsGenerator::FindPolygroupEdges()
+{
+	PolygroupEdges.Reset();
+
+	for (int eid : Mesh->EdgeIndicesItr())
+	{
+		if (Mesh->IsGroupBoundaryEdge(eid))
+		{
+			PolygroupEdges.Add(eid);
+		}
+	}
+	return (PolygroupEdges.Num() > 0);
+}
diff --git a/Engine/Plugins/Runtime/GeometryProcessing/Source/DynamicMesh/Public/Polygroups/PolygroupsGenerator.h b/Engine/Plugins/Runtime/GeometryProcessing/Source/DynamicMesh/Public/Polygroups/PolygroupsGenerator.h
new file mode 100644
--- /dev/null
+++ b/Engine/Plugins/Runtime/GeometryProcessing/Source/DynamicMesh/Public/Polygroups/PolygroupsGenerator.h
@@ -0,0 +1,107 @@
+// Copyright Epic Games, Inc. All Rights Reserved.
+
+#pragma once
+
+#include "CoreMinimal.h"
+#include "DynamicMesh/DynamicMesh3.h"
+
+using UE::Geometry::FDynamicMesh3;
+
+namespace UE
+{
+namespace Geometry
+{
+
+/**
+ * FPolygroupsGenerator generates (face/tri/poly)groups for an input mesh based on the geometry and attributes of the input mesh.
+ * The various FindPolygroupsFromX() are the driving functions, each performs the full computation. 
+ */
+class DYNAMICMESH_API FPolygroupsGenerator
+{
+public:
+
+	FPolygroupsGenerator() {}
+	FPolygroupsGenerator(FDynamicMesh3* MeshIn);
+
+	// 
+	// Input parameters
+	//
+
+	// source mesh
+	FDynamicMesh3* Mesh = nullptr;
+
+	// if true, then groups will be post-processed to optimize them, based on parmeters below
+	bool bApplyPostProcessing = true;
+
+	// if > 1, groups with a triangle count smaller than this will be merged with a neighbouring group
+	int32 MinGroupSize = 0;
+
+	// if true, after groups are computed they will be copied to the output mesh
+	bool bCopyToMesh = true;
+
+	//
+	// Outputs
+	//
+
+	// lists of triangle IDs, each list defines a polygroup/polygon
+	TArray<TArray<int>> FoundPolygroups;
+
+	// list of edge IDs of mesh edges that are on polygroup borders
+	TArray<int> PolygroupEdges;
+
+
+	/**
+	 * Find Polygroups by randomly picking initial seed triangles and then flood-filling outwards,
+	 * stopping when the opening angle at an edge is larger than the angle defined by the DotTolerance.
+	 */
+	bool FindPolygroupsFromFaceNormals(double DotTolerance = 0.0001);
+
+	/**
+	 * Find Polygroups based on UV Islands, ie each UV Island becomes a Polygroup
+	 */
+	bool FindPolygroupsFromUVIslands();
+
+	/**
+	 * Find Polygroups based on mesh connectivity, ie each connected-component becomes a Polygroup
+	 */
+	bool FindPolygroupsFromConnectedTris();
+
+	/**
+	 * Weight potions for algorithms below
+	 */
+	enum class EWeightingType
+	{
+		None,
+		NormalDeviation
+	};
+
+	/**
+	 * Incrementally compute approximate-geodesic-based furthest-point sampling of the mesh until NumPoints
+	 * samples have been found, then compute local geodesic patches (eg approximate surface voronoi diagaram).
+	 * Optionally weight geodesic-distance computation, which will produce different patch shapes.
+	 */
+	bool FindPolygroupsFromFurthestPointSampling(int32 NumPoints, EWeightingType WeightingType, FVector3d WeightingCoeffs = FVector3d::One());
+
+
+	/**
+	* Copy the computed Polygroups to the input Mesh. Will be automatically called by the 
+	* FindPolygroupsFromX() functions above if bCopyToMesh=true
+	*/
+	void CopyPolygroupsToMesh();
+
+
+	/**
+	 * Initialize the PolygroupEdges output member by finding all the mesh edges that are on polygroup borders.
+	 * Requires that polygroups have been computed and copied to the mesh
+	 */
+	bool FindPolygroupEdges();
+
+
+protected:
+	void PostProcessPolygroups(bool bApplyMerging);
+	void OptimizePolygroups();
+};
+
+
+}	// end namespace Geometry
+}	// end namespace UE
