GeometryProcessing: Add full AutoUV process to PatchBasedMeshUVGenerator, based on initial superpatch decomposition computed using FPolygroupsGenerator. Add a few utility functions for setting output state to FGeometryResult. Add FMeshConnectedComponents::InitializeFromTriangleComponents(), to allow initializing from an existing triangle set (sometimes necessary when downstream code requires a FMeshConnectedComponents as input).
#rnx

diff --git a/Engine/Plugins/Runtime/GeometryProcessing/Source/DynamicMesh/Private/Parameterization/PatchBasedMeshUVGenerator.cpp b/Engine/Plugins/Runtime/GeometryProcessing/Source/DynamicMesh/Private/Parameterization/PatchBasedMeshUVGenerator.cpp
--- a/Engine/Plugins/Runtime/GeometryProcessing/Source/DynamicMesh/Private/Parameterization/PatchBasedMeshUVGenerator.cpp
+++ b/Engine/Plugins/Runtime/GeometryProcessing/Source/DynamicMesh/Private/Parameterization/PatchBasedMeshUVGenerator.cpp
@@ -6,6 +6,7 @@
 #include "DynamicSubmesh3.h"
 #include "Sampling/NormalHistogram.h"
 
+#include "Polygroups/PolygroupsGenerator.h"
 #include "Parameterization/MeshLocalParam.h"
 #include "Parameterization/DynamicMeshUVEditor.h"
 #include "Parameterization/MeshRegionGraph.h"
@@ -17,6 +18,144 @@
 using namespace UE::Geometry;
 
 
+#define LOCTEXT_NAMESPACE "PatchBasedMeshUVGenerator"
+
+
+
+FGeometryResult FPatchBasedMeshUVGenerator::AutoComputeUVs(
+	FDynamicMesh3& TargetMesh,
+	FDynamicMeshUVOverlay& TargetUVOverlay,
+	FProgressCancel* Progress)
+{
+	FGeometryResult ResultInfo = FGeometryResult(EGeometryResultType::InProgress);
+
+	// 
+	// Step 1: decompose the input mesh into small patches
+	//
+	FMeshConnectedComponents InitialPatches(&TargetMesh);
+	bool bPatchesOK = ComputeInitialMeshPatches(TargetMesh, InitialPatches, Progress);
+	if (!bPatchesOK)
+	{
+		ResultInfo.SetFailed(LOCTEXT("InitialPatchesFailed", "Failed to compute initial patch set"));
+		return ResultInfo;
+	}
+	if (ResultInfo.CheckAndSetCancelled(Progress))
+	{
+		return ResultInfo;
+	}
+
+
+	// 
+	// Step 2: merge the small patches into bigger patches, which will become UV islands
+	//
+	TArray<TArray<int32>> UVIslandTriangleSets;
+	bool bIslandsOK = ComputeIslandsByRegionMerging(TargetMesh, TargetUVOverlay, InitialPatches, UVIslandTriangleSets, Progress);
+	if (!bIslandsOK)
+	{
+		ResultInfo.SetFailed(LOCTEXT("IslandMergingFailed", "Failed to merge UV islands"));
+		return ResultInfo;
+	}
+	if (ResultInfo.CheckAndSetCancelled(Progress))
+	{
+		return ResultInfo;
+	}
+
+
+	// 
+	// Step 3: compute UVs for the UV islands
+	//
+	TArray<bool> bIslandUVsValid;
+	int32 NumSolvesFailed = ComputeUVsFromTriangleSets(TargetMesh, TargetUVOverlay, UVIslandTriangleSets, bIslandUVsValid, Progress);
+	if (NumSolvesFailed > 0)
+	{
+		ResultInfo.AddWarning(FGeometryWarning(0, LOCTEXT("PartialSolvesFailed", "Failed to compute UVs for some UV islands")));
+	}
+	if (ResultInfo.CheckAndSetCancelled(Progress))
+	{
+		return ResultInfo;
+	}
+
+
+	FDynamicMeshUVEditor UVEditor(&TargetMesh, &TargetUVOverlay);
+
+	// 
+	// Step 4 (Optional): Rotate the UV islands to optimize their orientation, to (hopefully) improve packing
+	//
+	if (bAutoAlignPatches)
+	{
+		ParallelFor(UVIslandTriangleSets.Num(), [&](int32 k)
+		{
+			if (Progress && Progress->Cancelled())
+			{
+				return;
+			}
+			if (bIslandUVsValid[k])
+			{
+				UVEditor.AutoOrientUVArea(UVIslandTriangleSets[k]);
+			}
+		});
+	}
+	if (ResultInfo.CheckAndSetCancelled(Progress))
+	{
+		return ResultInfo;
+	}
+
+
+	// 
+	// Step 5 (Optional): Automatically pack the UV islands into the unit UV rectangle
+	//
+	if (bAutoPack)
+	{
+		bool bPackingSuccess = UVEditor.QuickPack(PackingTextureResolution, PackingGutterWidth);
+		if (!bPackingSuccess)
+		{
+			ResultInfo.AddWarning(FGeometryWarning(0, LOCTEXT("IslandPackingFailed", "Failed to pack UV islands")));
+		}
+	}
+	if (ResultInfo.CheckAndSetCancelled(Progress))
+	{
+		return ResultInfo;
+	}
+
+	ResultInfo.SetSuccess();
+	return ResultInfo;
+}
+
+
+
+
+
+bool FPatchBasedMeshUVGenerator::ComputeInitialMeshPatches(
+	FDynamicMesh3& TargetMesh,
+	FMeshConnectedComponents& InitialComponentsOut,
+	FProgressCancel* Progress)
+{
+	FPolygroupsGenerator Generator(&TargetMesh);
+	Generator.bCopyToMesh = false;
+
+	Generator.MinGroupSize = this->MinPatchSize;
+	FPolygroupsGenerator::EWeightingType WeightType = (this->bNormalWeightedPatches) ?
+		FPolygroupsGenerator::EWeightingType::NormalDeviation : FPolygroupsGenerator::EWeightingType::None;
+	FVector3d GeneratorParams(this->PatchNormalWeight, 1.0, 1.0);
+	bool bOK = Generator.FindPolygroupsFromFurthestPointSampling(this->TargetPatchCount, WeightType, GeneratorParams);
+	if (!bOK)
+	{
+		return false;
+	}
+	if (Progress && Progress->Cancelled())
+	{
+		return false;
+	}
+	InitialComponentsOut.InitializeFromTriangleComponents(Generator.FoundPolygroups, true);
+
+	return true;
+}
+
+
+
+
+
+
 struct FUVAreaMetrics
 {
 	double AvgSqrAreaDistortion = 1.0;
@@ -476,3 +615,7 @@ int32 FPatchBasedMeshUVGenerator::ComputeUVsFromTriangleSets(
 
 	return NumFailed;
 }
+
+
+
+#undef LOCTEXT_NAMESPACE
\ No newline at end of file
diff --git a/Engine/Plugins/Runtime/GeometryProcessing/Source/DynamicMesh/Private/Selections/MeshConnectedComponents.cpp b/Engine/Plugins/Runtime/GeometryProcessing/Source/DynamicMesh/Private/Selections/MeshConnectedComponents.cpp
--- a/Engine/Plugins/Runtime/GeometryProcessing/Source/DynamicMesh/Private/Selections/MeshConnectedComponents.cpp
+++ b/Engine/Plugins/Runtime/GeometryProcessing/Source/DynamicMesh/Private/Selections/MeshConnectedComponents.cpp
@@ -274,6 +274,83 @@ void FMeshConnectedComponents::SortByCount(bool bLargestFirst)
 
 
 
+
+bool FMeshConnectedComponents::InitializeFromTriangleComponents(const TArray<TArray<int32>>& ComponentLists, bool bValidateIDs)
+{
+	Components.Reset();
+	Components.Reserve(ComponentLists.Num());
+	for (const TArray<int32>& ComponentList : ComponentLists)
+	{
+		// skip any empty lists
+		if (ComponentList.Num() == 0)
+		{
+			continue;
+		}
+
+		FComponent* Component = new FComponent();
+		Component->Indices.Append(ComponentList);
+
+		if (bValidateIDs)
+		{
+			for (int32 tid : Component->Indices)
+			{
+				if (Mesh->IsTriangle(tid) == false)
+				{
+					return false;
+				}
+			}
+		}
+
+		Components.Add(Component);
+	}
+
+	return true;
+}
+
+
+bool FMeshConnectedComponents::InitializeFromTriangleComponents(TArray<TArray<int32>>& ComponentLists, bool bMoveSubLists, bool bValidateIDs)
+{
+	Components.Reset();
+	Components.Reserve(ComponentLists.Num());
+	for (TArray<int32>& ComponentList : ComponentLists)
+	{
+		// skip any empty lists
+		if (ComponentList.Num() == 0)
+		{
+			continue;
+		}
+
+		FComponent* Component = new FComponent();
+		if (bMoveSubLists)
+		{
+			Component->Indices = MoveTemp(ComponentList);
+		}
+		else
+		{
+			Component->Indices.Append(ComponentList);
+		}
+
+		if (bValidateIDs)
+		{
+			for (int32 tid : Component->Indices)
+			{
+				if (Mesh->IsTriangle(tid) == false)
+				{
+					return false;
+				}
+			}
+		}
+
+		Components.Add(Component);
+	}
+
+	return true;
+}
+
+
+
+
+
 void FMeshConnectedComponents::GrowToConnectedTriangles(const FDynamicMesh3* Mesh,
 	const TArray<int>& InputROI, TArray<int>& ResultROI,
 	TArray<int32>* QueueBuffer, 
diff --git a/Engine/Plugins/Runtime/GeometryProcessing/Source/DynamicMesh/Public/Parameterization/PatchBasedMeshUVGenerator.h b/Engine/Plugins/Runtime/GeometryProcessing/Source/DynamicMesh/Public/Parameterization/PatchBasedMeshUVGenerator.h
--- a/Engine/Plugins/Runtime/GeometryProcessing/Source/DynamicMesh/Public/Parameterization/PatchBasedMeshUVGenerator.h
+++ b/Engine/Plugins/Runtime/GeometryProcessing/Source/DynamicMesh/Public/Parameterization/PatchBasedMeshUVGenerator.h
@@ -27,18 +27,54 @@ class DYNAMICMESH_API FPatchBasedMeshUVGenerator
 public:
 
 	//
-	// Parameters for island merging
+	// Parameters for patch generation (ComputeInitialMeshPatches)
+	//
+	int32 TargetPatchCount = 100;
+	bool bNormalWeightedPatches = true;
+	double PatchNormalWeight = 1.0;
+	int32 MinPatchSize = 2;
+
+	//
+	// Parameters for island merging (ComputeIslandsByRegionMerging)
 	//
 	double MergingThreshold = 1.5;
 	double CompactnessThreshold = 9999999.0;		// effectively disabled as it usually is not a good idea
 	double MaxNormalDeviationDeg = 45.0;
 
 	//
-	// ExpMap parameters
+	// ExpMap parameters (ComputeUVsFromTriangleSets)
 	//
 	int32 NormalSmoothingRounds = 0;
 	double NormalSmoothingAlpha = 0.25;
 
+	//
+	// Packing parameters (used in AutoComputeUVs)
+	//
+	bool bAutoAlignPatches = true;
+	bool bAutoPack = true;
+	int32 PackingTextureResolution = 512;
+	float PackingGutterWidth = 1.0f;
+
+
+	/**
+	 * Top-level driver function, automatically computes UVs on the given TargetMesh and stores in the
+	 * given TargetUVOverlay, using the steps below, based on the parameters above
+	 * @return result information, which may include warnings/etc
+	 */
+	FGeometryResult AutoComputeUVs(
+		FDynamicMesh3& TargetMesh,
+		FDynamicMeshUVOverlay& TargetUVOverlay,
+		FProgressCancel* Progress = nullptr);
+
+	/**
+	 * Compute a decomposition of TargetMesh into many small patches (driven by TargetPatchCount).
+	 * @param InitialComponentsOut returned decomposition into triangle sets
+	 * @return true on success
+	 */
+	bool ComputeInitialMeshPatches(
+		FDynamicMesh3& TargetMesh,
+		FMeshConnectedComponents& InitialComponentsOut,
+		FProgressCancel* Progress = nullptr);
 
 	/**
 	 * Incrementally combine existing mesh patches into larger patches that meet the various criteria defined
diff --git a/Engine/Plugins/Runtime/GeometryProcessing/Source/DynamicMesh/Public/Selections/MeshConnectedComponents.h b/Engine/Plugins/Runtime/GeometryProcessing/Source/DynamicMesh/Public/Selections/MeshConnectedComponents.h
--- a/Engine/Plugins/Runtime/GeometryProcessing/Source/DynamicMesh/Public/Selections/MeshConnectedComponents.h
+++ b/Engine/Plugins/Runtime/GeometryProcessing/Source/DynamicMesh/Public/Selections/MeshConnectedComponents.h
@@ -81,6 +81,20 @@ public:
 	 */
 	void FindTrianglesConnectedToSeeds(const TArray<int>& SeedTriangles, TFunction<bool(int32, int32)> TrisConnectedPredicate = nullptr);
 
+	/**
+	 * Initialize the internal FComponent list from the input ComponentLists, skipping any empty input lists
+	 * @param bValidateIDs if true, test that each value corresponds to a valid triangle ID on the Mesh
+	 * @return true if all IDs are valid, or if check was skipped
+	 */
+	bool InitializeFromTriangleComponents(const TArray<TArray<int32>>& ComponentLists, bool bValidateIDs = false);
+
+	/**
+	* Initialize the internal FComponent list from the input ComponentLists, skipping any empty input lists
+	* @param bMoveSubLists if true, steal the arrays inside the ComponentLists (via MoveTemp), to avoid memory copies
+	* @param bValidateIDs if true, test that each value corresponds to a valid triangle ID on the Mesh
+	* @return true if all IDs are valid, or if check was skipped
+	*/
+	bool InitializeFromTriangleComponents(TArray<TArray<int32>>& ComponentLists, bool bMoveSubLists, bool bValidateIDs = false);
 
 	//
 	// Query functions. Only valid to call after a Calculation function has been called.
diff --git a/Engine/Source/Runtime/GeometryCore/Public/Util/ProgressCancel.h b/Engine/Source/Runtime/GeometryCore/Public/Util/ProgressCancel.h
--- a/Engine/Source/Runtime/GeometryCore/Public/Util/ProgressCancel.h
+++ b/Engine/Source/Runtime/GeometryCore/Public/Util/ProgressCancel.h
@@ -121,6 +121,17 @@ struct FGeometryResult
 	 */
 	inline void SetSuccess(bool bSuccess, FProgressCancel* Progress = nullptr );
 
+	/**
+	 * Set state of the Result to Failure, and append a FGeometryError with the given ErrorMessage and ResultCode
+	 */
+	inline void SetFailed(FText ErrorMessage, int ResultCode = 0);
+
+	/**
+	 * Test if the given Progress has been cancelled, if so, set the Result to Cancelled 
+	 * @return true if cancelled, false if not cancelled
+	 */
+	inline bool CheckAndSetCancelled( FProgressCancel* Progress );
+
 	/**
 	 * Append an Error to the result
 	 */
@@ -234,4 +245,22 @@ void UE::Geometry::FGeometryResult::SetSuccess(bool bSuccess, FProgressCancel* P
 	{
 		Result = (bSuccess) ? EGeometryResultType::Success : EGeometryResultType::Failure;
 	}
+}
+
+
+void UE::Geometry::FGeometryResult::SetFailed(FText ErrorMessage, int ResultCode)
+{
+	Result = EGeometryResultType::Failure;
+	Errors.Add(FGeometryError(ResultCode, ErrorMessage));
+}
+
+
+bool UE::Geometry::FGeometryResult::CheckAndSetCancelled(FProgressCancel* Progress)
+{
+	if (Progress && Progress->Cancelled())
+	{
+		Result = EGeometryResultType::Cancelled;
+		return true;
+	}
+	return false;
 }
\ No newline at end of file
