GeometryProcessing: fix memory leak identified by PVS
#rnx

diff --git a/Engine/Plugins/Runtime/GeometryProcessing/Source/DynamicMesh/Private/Selections/MeshConnectedComponents.cpp b/Engine/Plugins/Runtime/GeometryProcessing/Source/DynamicMesh/Private/Selections/MeshConnectedComponents.cpp
--- a/Engine/Plugins/Runtime/GeometryProcessing/Source/DynamicMesh/Private/Selections/MeshConnectedComponents.cpp
+++ b/Engine/Plugins/Runtime/GeometryProcessing/Source/DynamicMesh/Private/Selections/MeshConnectedComponents.cpp
@@ -277,6 +277,20 @@ void FMeshConnectedComponents::SortByCount(bool bLargestFirst)
 
 bool FMeshConnectedComponents::InitializeFromTriangleComponents(const TArray<TArray<int32>>& ComponentLists, bool bValidateIDs)
 {
+	if (bValidateIDs)
+	{
+		for (const TArray<int32>& ComponentList : ComponentLists)
+		{
+			for (int32 tid : ComponentList)
+			{
+				if (Mesh->IsTriangle(tid) == false)
+				{
+					return false;
+				}
+			}
+		}
+	}
+
 	Components.Reset();
 	Components.Reserve(ComponentLists.Num());
 	for (const TArray<int32>& ComponentList : ComponentLists)
@@ -289,10 +303,20 @@ bool FMeshConnectedComponents::InitializeFromTriangleComponents(const TArray<TAr
 
 		FComponent* Component = new FComponent();
 		Component->Indices.Append(ComponentList);
+		Components.Add(Component);
+	}
 
-		if (bValidateIDs)
+	return true;
+}
+
+
+bool FMeshConnectedComponents::InitializeFromTriangleComponents(TArray<TArray<int32>>& ComponentLists, bool bMoveSubLists, bool bValidateIDs)
+{
+	if (bValidateIDs)
+	{
+		for (const TArray<int32>& ComponentList : ComponentLists)
 		{
-			for (int32 tid : Component->Indices)
+			for (int32 tid : ComponentList)
 			{
 				if (Mesh->IsTriangle(tid) == false)
 				{
@@ -300,16 +324,8 @@ bool FMeshConnectedComponents::InitializeFromTriangleComponents(const TArray<TAr
 				}
 			}
 		}
-
-		Components.Add(Component);
 	}
 
-	return true;
-}
-
-
-bool FMeshConnectedComponents::InitializeFromTriangleComponents(TArray<TArray<int32>>& ComponentLists, bool bMoveSubLists, bool bValidateIDs)
-{
 	Components.Reset();
 	Components.Reserve(ComponentLists.Num());
 	for (TArray<int32>& ComponentList : ComponentLists)
@@ -330,17 +346,6 @@ bool FMeshConnectedComponents::InitializeFromTriangleComponents(TArray<TArray<in
 			Component->Indices.Append(ComponentList);
 		}
 
-		if (bValidateIDs)
-		{
-			for (int32 tid : Component->Indices)
-			{
-				if (Mesh->IsTriangle(tid) == false)
-				{
-					return false;
-				}
-			}
-		}
-
 		Components.Add(Component);
 	}
 
