ModelingTools: Optimize write lock during BakeTexture.
- Added cpu trace profiler events to high level bake processes.
- Added benchmarks

#rnx

diff --git a/Engine/Plugins/Runtime/GeometryProcessing/Source/DynamicMesh/Private/Sampling/MeshMapBaker.cpp b/Engine/Plugins/Runtime/GeometryProcessing/Source/DynamicMesh/Private/Sampling/MeshMapBaker.cpp
--- a/Engine/Plugins/Runtime/GeometryProcessing/Source/DynamicMesh/Private/Sampling/MeshMapBaker.cpp
+++ b/Engine/Plugins/Runtime/GeometryProcessing/Source/DynamicMesh/Private/Sampling/MeshMapBaker.cpp
@@ -2,13 +2,13 @@
 
 #include "Sampling/MeshMapBaker.h"
 #include "Sampling/MeshBakerCommon.h"
+#include "Sampling/MeshMapBakerQueue.h"
 #include "Image/ImageOccupancyMap.h"
 #include "Image/ImageTile.h"
 
 #include "ExplicitUseGeometryMathTypes.h"		// using UE::Geometry::(math types)
 using namespace UE::Geometry;
 
-
 //
 // FMeshMapBaker
 //
@@ -113,6 +113,8 @@ void FMeshMapBaker::InitBakeDefaults()
 
 void FMeshMapBaker::Bake()
 {
+	TRACE_CPUPROFILER_EVENT_SCOPE("FMeshMapBaker::Bake");
+	
 	if (Bakers.IsEmpty())
 	{
 		return;
@@ -189,9 +191,91 @@ void FMeshMapBaker::Bake()
 	TArray<TArray64<TTuple<int64, int64>>> GutterTexelsPerTile;
 	GutterTexelsPerTile.SetNum(NumTiles);
 
-	FCriticalSection WriteLock;
-	ParallelFor(NumTiles, [this, &Tiles, &ImageTileBuffer, &GutterTexelsPerTile, &WriteLock](int32 TileIdx)
+	auto NoopFn = [](const float& In, float& Out)
+	{
+	};
+
+	auto AddFn = [](const float& In, float& Out)
 	{
+		Out += In;
+	};
+
+	auto OverwriteFn = [](const float& In, float& Out)
+	{
+		Out = In;
+	};
+
+	// WriteToOutputBuffer transfers local tile data (TileBuffer) to the image output buffer (ImageTileBuffer).
+	auto WriteToOutputBuffer = [this, &ImageTileBuffer] (FMeshMapTileBuffer& TileBufferIn, const FImageTile& TargetTile, const TArray<int32>& BakeIds, auto&& Op, auto&& WeightOp)
+	{
+		TRACE_CPUPROFILER_EVENT_SCOPE("FMeshMapBaker::WriteToOutputBuffer");
+		
+		const int TargetTileWidth = TargetTile.GetWidth();
+		const int TargetTileHeight = TargetTile.GetHeight();
+		for (FVector2i TileCoords(0,0); TileCoords.Y < TargetTileHeight; ++TileCoords.Y)
+		{
+			for (TileCoords.X = 0; TileCoords.X < TargetTileWidth; ++TileCoords.X)
+			{
+				if (CancelF())
+				{
+					return;
+				}
+
+				const FVector2i ImageCoords = TargetTile.GetSourceCoords(TileCoords);
+
+				const FImageTile& BufferTile = TileBufferIn.GetTile();
+				const int64 TilePixelLinearIdx = BufferTile.GetIndexFromSourceCoords(ImageCoords);
+				const float& TilePixelWeight = TileBufferIn.GetPixelWeight(TilePixelLinearIdx);
+				float* TilePixelBuffer = TileBufferIn.GetPixel(TilePixelLinearIdx);
+
+				const int64 ImageLinearIdx = Dimensions.GetIndex(ImageCoords);
+				float& ImagePixelWeight = ImageTileBuffer.GetPixelWeight(ImageLinearIdx);
+				float* ImagePixelBuffer = ImageTileBuffer.GetPixel(ImageLinearIdx);
+
+				WeightOp(TilePixelWeight, ImagePixelWeight);
+				for( int32 Idx : BakeIds )
+				{
+					const FMeshMapEvaluator::FEvaluationContext& Context = BakeContexts[Idx];
+					const int32 NumData = Context.DataLayout.Num();
+					const int32 ResultOffset = BakeOffsets[Idx];
+					for (int32 DataIdx = 0; DataIdx < NumData; ++DataIdx)
+					{
+						const int32 ResultIdx = ResultOffset + DataIdx;
+						const int32 Offset = BakeSampleOffsets[ResultIdx];
+						float* BufferPtr = &TilePixelBuffer[Offset];
+						float* ImageBufferPtr = &ImagePixelBuffer[Offset];
+						const FMeshMapEvaluator::EComponents Stride = Context.DataLayout[DataIdx];
+
+						for (int32 FloatIdx = 0; FloatIdx < (int32)Stride; ++FloatIdx)
+						{
+							Op(BufferPtr[FloatIdx], ImageBufferPtr[FloatIdx]);
+						}
+					}
+				}
+			}
+		}
+	};
+
+	FMeshMapBakerQueue OutputQueue(NumTiles);
+	auto WriteQueuedOutput = [this, &WriteToOutputBuffer, &AddFn](FMeshMapBakerQueue& Queue)
+	{
+		if (Queue.AcquireProcessLock())
+		{
+			void* OutputData = Queue.Process();
+			while (OutputData)
+			{
+				FMeshMapTileBuffer* TileBufferPtr = (FMeshMapTileBuffer*) OutputData;
+				WriteToOutputBuffer(*TileBufferPtr, TileBufferPtr->GetTile(), BakeAccumulateLists[(int32)FMeshMapEvaluator::EAccumulateMode::Add], AddFn, AddFn);
+				delete TileBufferPtr;
+				OutputData = Queue.Process();
+			}
+			Queue.ReleaseProcessLock();
+		}
+	};
+	
+	ParallelFor(NumTiles, [this, &Tiles, &ImageTileBuffer, &GutterTexelsPerTile, &OutputQueue, &WriteToOutputBuffer, &OverwriteFn, &AddFn, &NoopFn, &WriteQueuedOutput](int32 TileIdx)
+	{
+		TRACE_CPUPROFILER_EVENT_SCOPE("FMeshMapBaker::BakeTile");
 		// Generate unpadded and padded tiles.
 		const FImageTile Tile = Tiles.GetTile(TileIdx);	// Image area to sample
 		const FImageTile PaddedTile = Tiles.GetTile(TileIdx, TilePadding); // Filtered image area
@@ -202,7 +286,7 @@ void FMeshMapBaker::Bake()
 		OccupancyMap.ComputeFromUVSpaceMesh(FlatMesh, [this](int32 TriangleID) { return FlatMesh.GetTriangleGroup(TriangleID); });
 		GutterTexelsPerTile[TileIdx] = OccupancyMap.GutterTexels;
 
-		FMeshMapTileBuffer TileBuffer(PaddedTile, BakeSampleBufferSize);
+		FMeshMapTileBuffer* TileBuffer = new FMeshMapTileBuffer(PaddedTile, BakeSampleBufferSize);
 
 		// Calculate interior texels
 		const int TileWidth = Tile.GetWidth();
@@ -214,6 +298,7 @@ void FMeshMapBaker::Bake()
 			{
 				if (CancelF())
 				{
+					delete TileBuffer;
 					return;
 				}
 
@@ -234,89 +319,46 @@ void FMeshMapBaker::Bake()
 
 						FMeshMapEvaluator::FCorrespondenceSample Sample;
 						DetailMeshSampler.SampleUV(UVTriangleID, UVPosition, Sample);
-
-						BakeSample(TileBuffer, Sample, UVPosition, ImageCoords);
+						BakeSample(*TileBuffer, Sample, UVPosition, ImageCoords);
 					}
 				}
 			}
 		}
 
-		auto NoopFn = [](const float& In, float& Out)
-		{
-		};
-
-		auto AddFn = [](const float& In, float& Out)
-		{
-			Out += In;
-		};
+		// Transfer 'Overwrite' float data to image tile buffer
+		WriteToOutputBuffer(*TileBuffer, Tile, BakeAccumulateLists[(int32)FMeshMapEvaluator::EAccumulateMode::Overwrite], OverwriteFn, NoopFn);
 
-		auto OverwriteFn = [](const float& In, float& Out)
-		{
-			Out = In;
-		};
+		// Accumulate 'Add' float data to image tile buffer
+		OutputQueue.Post(TileIdx, TileBuffer);
+		WriteQueuedOutput(OutputQueue);
+	}, !bParallel ? EParallelForFlags::ForceSingleThread : EParallelForFlags::None);
 
-		// WriteToOutputBuffer transfers local tile data (TileBuffer) to the image output buffer (ImageTileBuffer).
-		auto WriteToOutputBuffer = [this, &TileBuffer, &ImageTileBuffer] (const FImageTile& TargetTile, const TArray<int32>& BakeIds, auto&& Op, auto&& WeightOp)
+	if (CancelF())
+	{
+		// If cancelled, delete any outstanding tile buffers in the queue.
+		while (!OutputQueue.IsDone())
 		{
-			const int TargetTileWidth = TargetTile.GetWidth();
-			const int TargetTileHeight = TargetTile.GetHeight();
-			for (FVector2i TileCoords(0,0); TileCoords.Y < TargetTileHeight; ++TileCoords.Y)
+			void* Data = OutputQueue.Process</*bFlush*/ true>();
+			if (Data)
 			{
-				for (TileCoords.X = 0; TileCoords.X < TargetTileWidth; ++TileCoords.X)
-				{
-					if (CancelF())
-					{
-						return;
-					}
-
-					const FVector2i ImageCoords = TargetTile.GetSourceCoords(TileCoords);
-
-					const FImageTile& BufferTile = TileBuffer.GetTile();
-					const int64 TilePixelLinearIdx = BufferTile.GetIndexFromSourceCoords(ImageCoords);
-					const float& TilePixelWeight = TileBuffer.GetPixelWeight(TilePixelLinearIdx);
-					float* TilePixelBuffer = TileBuffer.GetPixel(TilePixelLinearIdx);
-
-					const int64 ImageLinearIdx = Dimensions.GetIndex(ImageCoords);
-					float& ImagePixelWeight = ImageTileBuffer.GetPixelWeight(ImageLinearIdx);
-					float* ImagePixelBuffer = ImageTileBuffer.GetPixel(ImageLinearIdx);
-
-					WeightOp(TilePixelWeight, ImagePixelWeight);
-					for( int32 Idx : BakeIds )
-					{
-						const FMeshMapEvaluator::FEvaluationContext& Context = BakeContexts[Idx];
-						const int32 NumData = Context.DataLayout.Num();
-						const int32 ResultOffset = BakeOffsets[Idx];
-						for (int32 DataIdx = 0; DataIdx < NumData; ++DataIdx)
-						{
-							const int32 ResultIdx = ResultOffset + DataIdx;
-							const int32 Offset = BakeSampleOffsets[ResultIdx];
-							float* BufferPtr = &TilePixelBuffer[Offset];
-							float* ImageBufferPtr = &ImagePixelBuffer[Offset];
-							const FMeshMapEvaluator::EComponents Stride = Context.DataLayout[DataIdx];
-
-							for (int32 FloatIdx = 0; FloatIdx < (int32)Stride; ++FloatIdx)
-							{
-								Op(BufferPtr[FloatIdx], ImageBufferPtr[FloatIdx]);
-							}
-						}
-					}
-				}
+				const FMeshMapTileBuffer* TileBuffer = (FMeshMapTileBuffer*) Data;
+				delete TileBuffer;				
 			}
-		};
-
-		// Transfer 'Overwrite' float data to image tile buffer
-		WriteToOutputBuffer(Tile, BakeAccumulateLists[(int32)FMeshMapEvaluator::EAccumulateMode::Overwrite], OverwriteFn, NoopFn);
-
-		// Accumulate 'Add' float data to image tile buffer
-		// TODO: Optimize this write lock
-		WriteLock.Lock();
-		WriteToOutputBuffer(PaddedTile, BakeAccumulateLists[(int32)FMeshMapEvaluator::EAccumulateMode::Add], AddFn, AddFn);
-		WriteLock.Unlock();
-	}, !bParallel ? EParallelForFlags::ForceSingleThread : EParallelForFlags::None);
+		}
+	}
+	else
+	{
+		// The queue only acquires the process lock if the next item in the queue
+		// is ready. This could mean that there are potential leftovers in the queue
+		// after the parallel for. Write them out now.
+		WriteQueuedOutput(OutputQueue);
+	}
 
 	// Normalize and convert ImageTileBuffer data to color data.
 	ParallelFor(NumTiles, [this, &Tiles, &ImageTileBuffer](int32 TileIdx)
 	{
+		TRACE_CPUPROFILER_EVENT_SCOPE("FMeshMapBaker::WriteToImageBuffer");
+		
 		const FImageTile Tile = Tiles.GetTile(TileIdx);
 		const int TileWidth = Tile.GetWidth();
 		const int TileHeight = Tile.GetHeight();
@@ -376,6 +418,8 @@ void FMeshMapBaker::Bake()
 		const int32 NumResults = BakeResults.Num();
 		ParallelFor(NumTiles, [this, &NumResults, &GutterTexelsPerTile](int32 TileIdx)
 		{
+			TRACE_CPUPROFILER_EVENT_SCOPE("FMeshMapBaker::WriteGutterPixels");
+			
 			for (int64 GutterIdx = 0; GutterIdx < GutterTexelsPerTile[TileIdx].Num(); ++GutterIdx)
 			{
 				int64 GutterPixelTo;
@@ -503,40 +547,45 @@ int32 FMeshMapBaker::NumBakers() const
 	return Bakers.Num();
 }
 
-const TArrayView<TUniquePtr<TImageBuilder<FVector4f>>> FMeshMapBaker::GetBakeResults(int32 BakerIdx)
+const TArrayView<TUniquePtr<TImageBuilder<FVector4f>>> FMeshMapBaker::GetBakeResults(const int32 BakerIdx)
 {
 	const int32 ResultIdx = BakeOffsets[BakerIdx];
 	const int32 NumResults = BakeOffsets[BakerIdx + 1] - ResultIdx;
 	return TArrayView<TUniquePtr<TImageBuilder<FVector4f>>>(&BakeResults[ResultIdx], NumResults);
 }
 
-void FMeshMapBaker::SetDimensions(FImageDimensions DimensionsIn)
+void FMeshMapBaker::SetDimensions(const FImageDimensions DimensionsIn)
 {
 	Dimensions = DimensionsIn;
 }
 
-void FMeshMapBaker::SetGutterEnabled(bool bEnabled)
+void FMeshMapBaker::SetGutterEnabled(const bool bEnabled)
 {
 	bGutterEnabled = bEnabled;
 }
 
-void FMeshMapBaker::SetGutterSize(int32 GutterSizeIn)
+void FMeshMapBaker::SetGutterSize(const int32 GutterSizeIn)
 {
 	// GutterSize must be >= 1 since it is tied to MaxDistance for the
 	// OccupancyMap spatial search.
 	GutterSize = GutterSizeIn >= 1 ? GutterSizeIn : 1;
 }
 
-void FMeshMapBaker::SetMultisampling(int32 MultisamplingIn)
+void FMeshMapBaker::SetMultisampling(const int32 MultisamplingIn)
 {
 	Multisampling = MultisamplingIn;
 }
 
-void FMeshMapBaker::SetFilter(EBakeFilterType FilterTypeIn)
+void FMeshMapBaker::SetFilter(const EBakeFilterType FilterTypeIn)
 {
 	FilterType = FilterTypeIn;
 }
 
+void FMeshMapBaker::SetTileSize(const int TileSizeIn)
+{
+	TileSize = TileSizeIn;
+}
+
 void FMeshMapBaker::InitFilter()
 {
 	FilterKernelSize = TilePadding;
diff --git a/Engine/Plugins/Runtime/GeometryProcessing/Source/DynamicMesh/Public/Sampling/MeshMapBaker.h b/Engine/Plugins/Runtime/GeometryProcessing/Source/DynamicMesh/Public/Sampling/MeshMapBaker.h
--- a/Engine/Plugins/Runtime/GeometryProcessing/Source/DynamicMesh/Public/Sampling/MeshMapBaker.h
+++ b/Engine/Plugins/Runtime/GeometryProcessing/Source/DynamicMesh/Public/Sampling/MeshMapBaker.h
@@ -66,12 +66,14 @@ public:
 	void SetGutterSize(int32 GutterSizeIn);
 	void SetMultisampling(int32 MultisamplingIn);
 	void SetFilter(EBakeFilterType FilterTypeIn);
+	void SetTileSize(int TileSizeIn);
 
 	FImageDimensions GetDimensions() const { return Dimensions; }
 	bool GetGutterEnabled() const { return bGutterEnabled; }
 	int32 GetGutterSize() const { return GutterSize; }
 	int32 GetMultisampling() const { return Multisampling; }
 	EBakeFilterType GetFilter() const { return FilterType; }
+	int32 GetTileSize() const { return TileSize; }
 
 protected:
 	/** Evaluate this sample. */
@@ -111,10 +113,10 @@ protected:
 	int32 Multisampling = 1;
 
 	/** The square dimensions for tiled processing of the output image(s). */
-	const int32 TileSize = 32;
+	int32 TileSize = 32;
 
 	/** The amount of padding for tiled processing of the output image(s). */
-	const int32 TilePadding = 2;
+	int32 TilePadding = 2;
 
 	/** The pixel distance around the sample texel to be considered by the filter. [0, TilePadding] */
 	int32 FilterKernelSize = 0;
diff --git a/Engine/Plugins/Runtime/GeometryProcessing/Source/DynamicMesh/Public/Sampling/MeshMapBakerQueue.h b/Engine/Plugins/Runtime/GeometryProcessing/Source/DynamicMesh/Public/Sampling/MeshMapBakerQueue.h
new file mode 100644
--- /dev/null
+++ b/Engine/Plugins/Runtime/GeometryProcessing/Source/DynamicMesh/Public/Sampling/MeshMapBakerQueue.h
@@ -0,0 +1,140 @@
+// Copyright Epic Games, Inc. All Rights Reserved.
+
+#pragma once
+
+#include <atomic>
+
+#include "HAL/Platform.h"
+
+namespace UE
+{
+namespace Geometry
+{
+
+/**
+ * Multi-producer/single-consumer queue for the ordered processing of data
+ * in FMeshMapBaker.
+ * 
+ * This queue is designed to mediate the aggregation of data into a
+ * shared output buffer. The number of data to be aggregated must
+ * be provided at initialization. Producer threads then "post" data
+ * to the queue and any single thread can temporarily assume the role
+ * of consumer to "process" the data.
+ */
+class FMeshMapBakerQueue
+{
+private:
+	// Align to cache line size and pad our slot data to avoid false sharing.
+	static constexpr SIZE_T Alignment = PLATFORM_CACHE_LINE_SIZE;
+	struct alignas(Alignment) FSlotData
+	{
+		std::atomic<void*> Ptr;
+	};
+	static_assert(sizeof(FSlotData) == Alignment);
+
+	/** Cached size of the queue */
+	SIZE_T NumSlots = 0;
+
+	/** The aligned pointer for our atomic data. */
+	FSlotData* Slots = nullptr;
+
+	/** The current index to be processed. */
+	SIZE_T Current = 0;
+
+	/** Atomic flag to control ownership of consumer. */
+	std::atomic_flag Acquired = ATOMIC_FLAG_INIT;
+
+public:
+	explicit FMeshMapBakerQueue(SIZE_T Size)
+		: NumSlots(Size)
+	{
+		Slots = new FSlotData[NumSlots];
+		for (SIZE_T Idx = 0; Idx < NumSlots; ++Idx)
+		{
+			Slots[Idx].Ptr = nullptr;
+		}
+	}
+	
+	~FMeshMapBakerQueue()
+	{
+		for (SIZE_T Idx = 0; Idx < NumSlots; ++Idx)
+		{
+			ensure(Slots[Idx].Ptr.load(std::memory_order_acquire) == nullptr);
+		}
+		delete[] Slots;
+		ensure(Current == NumSlots);
+	}
+
+	/**
+	 * Post Data pointer into the queue at a given index.
+	 *
+	 * @param Idx The slot index to update.
+	 * @param Data The data to post to the queue.
+	 * @return true if the operation succeeded (the slot was empty).
+	 */
+	void Post(SIZE_T Idx, void* Data)
+	{
+		Slots[Idx].Ptr.store(Data, std::memory_order_release);
+	}
+
+	/**
+	 * Acquire ownership over processing queue data.
+	 *
+	 * It is the responsibility of the caller to release the ownership
+	 * if it was acquired via ReleaseProcessLock().
+	 *
+	 * @return true if ownership of queue processing was acquired.
+	 */
+	bool AcquireProcessLock()
+	{
+		// Only attempt to acquire process lock if data for current index is available.
+		void* Data = Current < NumSlots ? Slots[Current].Ptr.load(std::memory_order_relaxed) : nullptr;
+		return Data ? !Acquired.test_and_set(std::memory_order_acquire) : false;
+	}
+
+	/**
+	 * Retrieve the next item in the queue to process. The current index
+	 * is only advanced if the item is valid or <bFlush = true>.
+	 *
+	 * @param bFlush When true, increments the current index regardless of
+	 *               the validity of the data. Useful for flushing/clearing
+	 *               the contents of the queue.
+	 * @return A void* to the data to be processed, nullptr if the data
+	 *         for the next item in the queue is not available yet.
+	 */
+	template<bool bFlush = false>
+	void* Process()
+	{
+		void* Data = Current < NumSlots ? Slots[Current].Ptr.load(std::memory_order_acquire) : nullptr;
+		if (Data)
+		{
+			Slots[Current++].Ptr.store(nullptr, std::memory_order_release);
+		}
+		else if constexpr (bFlush)
+		{
+			Current++;
+		}
+		return Data;
+	}
+
+	/**
+	 * Release ownership of queue processing.
+	 *
+	 * This should only be invoked by the thread that successfully
+	 * acquired ownership of queue processing via AcquireProcessLock().
+	 */
+	void ReleaseProcessLock()
+	{
+		Acquired.clear(std::memory_order_release);
+	}
+
+	/** @return true if all items in the queue have been processed. */
+	bool IsDone() const
+	{
+		return Current >= NumSlots;
+	}
+};
+
+} // end namespace UE::Geometry
+} // end namespace UE
+
