- Compilation fixes for LWC changes

#rnx

diff --git a/Engine/Plugins/Experimental/MeshModelingToolsetExp/Source/ModelingOperators/Private/BaseOps/SimpleMeshProcessingBaseOp.cpp b/Engine/Plugins/Experimental/MeshModelingToolsetExp/Source/ModelingOperators/Private/BaseOps/SimpleMeshProcessingBaseOp.cpp
--- a/Engine/Plugins/Experimental/MeshModelingToolsetExp/Source/ModelingOperators/Private/BaseOps/SimpleMeshProcessingBaseOp.cpp
+++ b/Engine/Plugins/Experimental/MeshModelingToolsetExp/Source/ModelingOperators/Private/BaseOps/SimpleMeshProcessingBaseOp.cpp
@@ -41,7 +41,7 @@ FSimpleMeshProcessingBaseOp::FSimpleMeshProcessingBaseOp(const FDynamicMesh3* Me
 	}
 }
 
-void FSimpleMeshProcessingBaseOp::SetTransform(const FTransform3d& XForm)
+void FSimpleMeshProcessingBaseOp::SetTransform(const UE::Geometry::FTransform3d& XForm)
 {
 	ResultTransform = XForm;
 }
diff --git a/Engine/Plugins/Experimental/MeshModelingToolsetExp/Source/ModelingOperators/Private/SmoothingOps/SmoothingOpBase.cpp b/Engine/Plugins/Experimental/MeshModelingToolsetExp/Source/ModelingOperators/Private/SmoothingOps/SmoothingOpBase.cpp
--- a/Engine/Plugins/Experimental/MeshModelingToolsetExp/Source/ModelingOperators/Private/SmoothingOps/SmoothingOpBase.cpp
+++ b/Engine/Plugins/Experimental/MeshModelingToolsetExp/Source/ModelingOperators/Private/SmoothingOps/SmoothingOpBase.cpp
@@ -23,7 +23,7 @@ FSmoothingOpBase::FSmoothingOpBase(const FDynamicMesh3* Mesh, const FOptions& Op
 	}
 }
 
-void FSmoothingOpBase::SetTransform(const FTransform3d& XForm)
+void FSmoothingOpBase::SetTransform(const UE::Geometry::FTransform3d& XForm)
 {
 	ResultTransform = XForm;
 }
diff --git a/Engine/Plugins/Runtime/GeometryProcessing/Source/DynamicMesh/Private/DynamicMeshEditor.cpp b/Engine/Plugins/Runtime/GeometryProcessing/Source/DynamicMesh/Private/DynamicMeshEditor.cpp
--- a/Engine/Plugins/Runtime/GeometryProcessing/Source/DynamicMesh/Private/DynamicMeshEditor.cpp
+++ b/Engine/Plugins/Runtime/GeometryProcessing/Source/DynamicMesh/Private/DynamicMeshEditor.cpp
@@ -1158,7 +1158,7 @@ void FDynamicMeshEditor::SetQuadUVsFromProjection(const FIndex2i& QuadTris, cons
 }
 
 
-void FDynamicMeshEditor::RescaleAttributeUVs(float UVScale, bool bWorldSpace, int UVLayerIndex, TOptional<FTransform3d> ToWorld)
+void FDynamicMeshEditor::RescaleAttributeUVs(float UVScale, bool bWorldSpace, int UVLayerIndex, TOptional<UE::Geometry::FTransform3d> ToWorld)
 {
 	check(Mesh->HasAttributes() && Mesh->Attributes()->NumUVLayers() > UVLayerIndex );
 	FDynamicMeshUVOverlay* UVs = Mesh->Attributes()->GetUVLayer(UVLayerIndex);
diff --git a/Engine/Plugins/Runtime/GeometryProcessing/Source/DynamicMesh/Private/ShapeApproximation/SimpleShapeSet3.cpp b/Engine/Plugins/Runtime/GeometryProcessing/Source/DynamicMesh/Private/ShapeApproximation/SimpleShapeSet3.cpp
--- a/Engine/Plugins/Runtime/GeometryProcessing/Source/DynamicMesh/Private/ShapeApproximation/SimpleShapeSet3.cpp
+++ b/Engine/Plugins/Runtime/GeometryProcessing/Source/DynamicMesh/Private/ShapeApproximation/SimpleShapeSet3.cpp
@@ -415,7 +415,7 @@ void FSimpleShapeSet3d::Append(const FSimpleShapeSet3d& OtherShapeSet)
 }
 
 
-void FSimpleShapeSet3d::Append(const FSimpleShapeSet3d& OtherShapeSet, const FTransform3d& Transform)
+void FSimpleShapeSet3d::Append(const FSimpleShapeSet3d& OtherShapeSet, const UE::Geometry::FTransform3d& Transform)
 {
 	for (FSphereShape3d SphereShape : OtherShapeSet.Spheres)
 	{
@@ -445,7 +445,7 @@ void FSimpleShapeSet3d::Append(const FSimpleShapeSet3d& OtherShapeSet, const FTr
 
 
 
-void FSimpleShapeSet3d::Append(const FSimpleShapeSet3d& OtherShapeSet, const TArray<FTransform3d>& TransformSequence)
+void FSimpleShapeSet3d::Append(const FSimpleShapeSet3d& OtherShapeSet, const TArray<UE::Geometry::FTransform3d>& TransformSequence)
 {
 	for (FSphereShape3d SphereShape : OtherShapeSet.Spheres)
 	{
@@ -518,7 +518,7 @@ void FSimpleShapeSet3d::FilterByVolume(int32 MaximumCount)
 
 
 
-void FSimpleShapeSet3d::ApplyTransform(const FTransform3d& Transform)
+void FSimpleShapeSet3d::ApplyTransform(const UE::Geometry::FTransform3d& Transform)
 {
 	for (FSphereShape3d& SphereShape : Spheres)
 	{
