GeometryProcessing:
Implemented Curved PN Triangles operator.

#rnx

diff --git a/Engine/Plugins/Runtime/GeometryProcessing/Source/DynamicMesh/Private/Operations/PNTriangles.cpp b/Engine/Plugins/Runtime/GeometryProcessing/Source/DynamicMesh/Private/Operations/PNTriangles.cpp
new file mode 100644
--- /dev/null
+++ b/Engine/Plugins/Runtime/GeometryProcessing/Source/DynamicMesh/Private/Operations/PNTriangles.cpp
@@ -0,0 +1,445 @@
+// Copyright Epic Games, Inc. All Rights Reserved.
+
+#include "Operations/PNTriangles.h"
+#include "IndexTypes.h"
+#include "VectorTypes.h"
+#include "DynamicMesh/MeshNormals.h"
+#include "DynamicMesh/DynamicMesh3.h"
+#include "Async/ParallelFor.h"
+#include "ExplicitUseGeometryMathTypes.h"		// using UE::Geometry::(math types)
+#include "Util/ProgressCancel.h"
+
+using namespace UE::Geometry;
+
+namespace FPNTrianglesLocals
+{ 
+	/**
+	 * PN Triangle is represented by a control points cage. Each edge has 2 control points and 
+	 * 1 control point is inside the triangle. Additionally, each edge contains one control point 
+	 * of the normal component. We separate control points into 2 categories to avoid repeating 
+	 * computations since control points on edges are shared.
+	 */
+		
+	using FTriangleControlPoint = FVector3d;
+
+	struct FEdgeControlPoints
+	{	
+		FVector3d Point1;
+		FVector3d Point2;
+		FVector3d Normal;
+	}; 
+	
+	struct FControlPoints 
+	{
+		TArray<FEdgeControlPoints> OnEdges; // Map edge ID to the edge control points
+		TArray<FTriangleControlPoint> OnTriangles; // Map triangle ID to the triangle center control point
+	};
+
+	// Used to save just the connectivity information of a mesh
+	struct FMeshConnectivity
+	{
+		FMeshConnectivity(FDynamicMesh3& Mesh) 
+		{
+			Triangles = Mesh.GetTrianglesBuffer();
+			Edges = Mesh.GetEdgesBuffer();
+
+			//TODO: DynamicMesh3 should return this as a buffer
+			TriangleEdges.SetNum(Mesh.MaxTriangleID());
+			for (int TID = 0; TID < Mesh.MaxTriangleID(); ++TID) 
+			{
+				if (Mesh.IsTriangle(TID)) 
+				{
+					TriangleEdges[TID] = Mesh.GetTriEdgesRef(TID); 
+				}
+			}
+		}
+		
+		TDynamicVector<FIndex3i> Triangles;
+		TDynamicVector<FDynamicMesh3::FEdge> Edges;
+		TDynamicVector<FIndex3i> TriangleEdges;
+	};
+
+	/**
+	 *  Compute PN Triangle control points and optionally their normal component.
+	 * 
+	 *  @param Mesh The mesh used to compute the per-triangle control points.
+	 *  @param UseNormals If normals are disabled for the Mesh then pass manually computed normals.
+	 *  @param bComputePNNormals Compute control points for the normal component for calculating quadratically varying normals.
+	 *  @param ProgressCancel Set this to be able to cancel running operation.
+	 *  @param OutControlPoints Contains all the control points for the Mesh.
+	 * 
+	 *  @return true if the operation succeeded, false if it failed or was canceled by the user.
+	 */
+	bool ComputeControlPoints(FControlPoints& OutControlPoints,
+							  const FDynamicMesh3& Mesh, 
+							  const FMeshNormals* UseNormals, 
+							  const bool bComputePNNormals, 
+							  FProgressCancel* ProgressCancel) 
+	{
+		if (!Mesh.HasVertexNormals() && UseNormals == nullptr) 
+		{
+			return false; // Mesh must have normals
+		}
+
+		// We know ahead of time the total number of control points:
+		// 2 control points per edge and one in the middle of each triangle.
+		OutControlPoints.OnEdges.SetNum(Mesh.MaxEdgeID());
+		OutControlPoints.OnTriangles.SetNum(Mesh.MaxTriangleID());
+
+		// First compute only the control points on the edges
+		ParallelFor(Mesh.MaxEdgeID(), [&](int32 EID)
+		{
+			if (ProgressCancel && ProgressCancel->Cancelled()) 
+			{
+				return;  
+			}
+
+			if (Mesh.IsEdge(EID)) 
+			{
+				FIndex2i EdgeV = Mesh.GetEdgeV(EID);
+				
+				FVector3d Vertex1 = Mesh.GetVertex(EdgeV.A);
+				FVector3d Vertex2 = Mesh.GetVertex(EdgeV.B);
+
+				FVector3d Normal1 = (UseNormals != nullptr) ? (*UseNormals)[EdgeV.A] : (FVector3d)Mesh.GetVertexNormal(EdgeV.A);
+				FVector3d Normal2 = (UseNormals != nullptr) ? (*UseNormals)[EdgeV.B] : (FVector3d)Mesh.GetVertexNormal(EdgeV.B);
+				
+				FVector3d Edge12 = Vertex2 - Vertex1;
+
+				double Weight12 = Edge12.Dot(Normal1);
+				double Weight21 = -1.0*Edge12.Dot(Normal2); 
+
+				FEdgeControlPoints& ContolPnts = OutControlPoints.OnEdges[EID];
+				ContolPnts.Point1 = (2.0*Vertex1 + Vertex2 - Weight12*Normal1)/3.0;  
+				ContolPnts.Point2 = (2.0*Vertex2 + Vertex1 - Weight21*Normal2)/3.0;  
+
+				if (bComputePNNormals)
+				{
+					FVector3d Normal21 = Normal2 + Normal1;
+					double NWeight12 = 2.0*Edge12.Dot(Normal21)/Edge12.Dot(Edge12);
+					ContolPnts.Normal = Normalized(Normal21 - NWeight12*Edge12);
+				}
+			}
+		});
+
+		if (ProgressCancel && ProgressCancel->Cancelled())         
+		{
+			return false; 
+		}
+
+		// Compute the center control point for each triangle
+		ParallelFor(Mesh.MaxTriangleID(), [&](int32 TID)
+		{
+			if (ProgressCancel && ProgressCancel->Cancelled()) 
+			{
+				return;  
+			}
+
+			if (Mesh.IsTriangle(TID)) 
+			{
+				const FIndex3i& TriEdgeID = Mesh.GetTriEdgesRef(TID);
+				FVector3d ControlMidpoint = FVector3d::Zero();
+				for (int EIdx = 0; EIdx < 3; ++EIdx) 
+				{
+					const FEdgeControlPoints& ControlPnts = OutControlPoints.OnEdges[TriEdgeID[EIdx]];
+					ControlMidpoint += ControlPnts.Point1 + ControlPnts.Point2;
+				}
+				ControlMidpoint /= 6.0;
+
+				const FIndex3i& TriVertID = Mesh.GetTriangleRef(TID);
+				FVector3d VertexMidpoint = (Mesh.GetVertex(TriVertID.A) + 
+											Mesh.GetVertex(TriVertID.B) + 
+											Mesh.GetVertex(TriVertID.C))/3.0;
+
+				OutControlPoints.OnTriangles[TID] = ControlMidpoint + (ControlMidpoint - VertexMidpoint)/2.0;	
+			}
+		});
+
+		if (ProgressCancel && ProgressCancel->Cancelled())         
+		{
+			return false; 
+		}
+
+		return true;
+	}
+
+	/** 
+	 * Tesselate the mesh by recursively subdividing it using loop-style subdivision. Keep track of the new vertices 
+	 * added and which original triangles (before tesselation) they belong to. New vertices on original non-boundary 
+	 * edges can belong to either of the original triangles that share the edge.
+	 * 
+	 * @param Mesh The mesh we are tesselating.
+	 * @param Level How many times we are recursively subdividing the Mesh.
+	 * @param ProgressCancel Set this to be able to cancel running operation.
+	 * @param OutNewVerticies Array of tuples of the new vertex ID and the original triangle ID the vertex belongs to.
+	 * 
+	 * @return true if the operation succeeded, false if it failed or was canceled by the user.
+	 */
+	bool TesselateMesh(FDynamicMesh3& Mesh, 
+					   TArray<FIndex2i>& OutNewVerticies,
+					   const int Level, 
+					   FProgressCancel* ProgressCancel) 
+	{
+		check(Level >= 0);
+
+		if (Level < 0) 
+		{
+			return false;
+		}
+		else if (Level == 0) 
+		{
+			return true; // nothing to do
+		}
+
+		// Compute the number of new vertices we will generate with tesselation
+		int NumEdgeSegments = 1 << Level; // Number of the edges each original edge is split into
+		int NumEdgeVert = NumEdgeSegments + 1; // Number of the vertices along the edge after the tesselation
+		int NumNewEdgeVert = NumEdgeSegments - 1; // Number of the new vertices added per edge
+		// Triangular number series (n(n+1)/2) minus the 3 original vertices and new edge vertices
+		int NumNewTriangleVert = NumEdgeVert*(NumEdgeVert + 1)/2 - 3 - 3*NumNewEdgeVert;
+		int NumNewVerticies = NumNewEdgeVert * Mesh.EdgeCount() + NumNewTriangleVert* Mesh.TriangleCount(); 
+		OutNewVerticies.Reserve(NumNewVerticies);
+
+		// The number of the new triangle IDs created with Level number of subdivisions 
+		int NumNewTrianglesIDs = (NumEdgeSegments * NumEdgeSegments - 1) * Mesh.TriangleCount();
+		
+		// When we perform edge splits, we want to keep track of the original triangles that 
+		// the new added triangles belong to (i.e parent triangle). 
+		TArray<int> ParentTriangles;
+		ParentTriangles.Init(FDynamicMesh3::InvalidID, Mesh.MaxTriangleID() + NumNewTrianglesIDs);
+		for (int TID : Mesh.TriangleIndicesItr()) 
+		{ 
+			ParentTriangles[TID] = TID; // Original triangles are parents of themselves.
+		}
+
+		// Recursively subdivide the mesh Level number of times, while keeping track of new 
+		// vertices added and which original triangles they belong to.
+		for (int CurLevel = 0; CurLevel < Level; ++CurLevel) 
+		{
+			TArray<int> EdgesToProcess;
+			EdgesToProcess.Reserve(Mesh.EdgeCount());
+			for (int EID : Mesh.EdgeIndicesItr())
+			{
+				EdgesToProcess.Add(EID);
+			} 
+
+			int MaxTriangleID = Mesh.MaxTriangleID();
+
+			TArray<int> EdgesToFlip;
+			EdgesToFlip.Init(FDynamicMesh3::InvalidID, MaxTriangleID);
+
+			for (int EID : EdgesToProcess)
+			{
+				FIndex2i EdgeTris = Mesh.GetEdgeT(EID);
+
+				FDynamicMesh3::FEdgeSplitInfo SplitInfo;
+				EMeshResult Result = Mesh.SplitEdge(EID, SplitInfo);
+
+				if (Result == EMeshResult::Ok) 
+				{	  
+					ParentTriangles[SplitInfo.NewTriangles.A] = ParentTriangles[SplitInfo.OriginalTriangles.A]; 
+					if (SplitInfo.OriginalTriangles.B != FDynamicMesh3::InvalidID) 
+					{
+						ParentTriangles[SplitInfo.NewTriangles.B] = ParentTriangles[SplitInfo.OriginalTriangles.B]; 
+					}
+
+					OutNewVerticies.Add(FIndex2i(SplitInfo.NewVertex, ParentTriangles[SplitInfo.OriginalTriangles.A]));
+
+					if (EdgeTris.A < MaxTriangleID && EdgesToFlip[EdgeTris.A] == FDynamicMesh3::InvalidID)
+					{
+						EdgesToFlip[EdgeTris.A] = SplitInfo.NewEdges.B;
+					}
+					if (EdgeTris.B != FDynamicMesh3::InvalidID)
+					{
+						if (EdgeTris.B < MaxTriangleID && EdgesToFlip[EdgeTris.B] == FDynamicMesh3::InvalidID)
+						{
+							EdgesToFlip[EdgeTris.B] = SplitInfo.NewEdges.C;
+						}
+					}
+				}   
+
+			  	if (ProgressCancel && ProgressCancel->Cancelled())    
+				{
+					return false;
+				}
+			}
+
+			for (int EID : EdgesToFlip)
+			{
+				if (EID != FDynamicMesh3::InvalidID)
+				{
+					FDynamicMesh3::FEdgeFlipInfo FlipInfo;
+					Mesh.FlipEdge(EID, FlipInfo);
+
+					if (ProgressCancel && ProgressCancel->Cancelled()) 
+					{
+						return false;
+					}
+				}
+			}
+		}
+
+		checkSlow(OutNewVerticies.Num() == NumNewVerticies);
+		checkSlow(ParentTriangles.Num() == Mesh.MaxTriangleID());
+
+		return true;
+	}
+
+	/** 
+	 * Displace the vertices created from the tesselation using the cubic patch formula based on their barycentric 
+	 * coordinates.
+	 * 
+	 * @param Mesh The tessellated mesh whose vertices we are displacing.
+	 * @param FControlPoints PN Triangle control points.
+	 * @param VerticiesToDisplace Array of vertices into the Mesh that we are displacing.
+	 * @param bComputePNNormals Should we be computing quadratically varying normals using control points.
+	 * @param OriginalConnectivity Connectivity of the original mesh (before tesselation).
+	 * @param ProgressCancel Set this to be able to cancel running operation.
+	 * 
+	 * @return true if the operation succeeded, false if it failed or was canceled by the user.
+	 */
+
+	bool Displace(FDynamicMesh3& Mesh,
+				  const FControlPoints& FControlPoints,
+				  const TArray<FIndex2i>& VerticiesToDisplace,
+				  const bool bComputePNNormals, 
+				  const FMeshConnectivity& OriginalConnectivity,
+				  FProgressCancel* ProgressCancel) 
+	{
+		if (bComputePNNormals && !Mesh.HasVertexNormals()) 
+		{
+			Mesh.EnableVertexNormals(FVector3f::UnitY()); 
+		}
+
+		// Iterate over every new vertex and compute its displacement and optionally a normal
+		ParallelFor(VerticiesToDisplace.Num(), [&](int32 IDX)
+		{	
+			if (ProgressCancel && ProgressCancel->Cancelled()) 
+			{
+				return; 
+			}
+			
+			FIndex2i NewVtx = VerticiesToDisplace[IDX];
+			int VertexID = NewVtx[0]; // ID of the new vertex added with tesselation
+			int OriginalTriangleID = NewVtx[1]; // ID of the original triangle new vertex belongs to
+			
+			// Get the topology information of the original triangle
+			FIndex3i TriVertex = OriginalConnectivity.Triangles[OriginalTriangleID];
+			FIndex3i TriEdges = OriginalConnectivity.TriangleEdges[OriginalTriangleID];
+
+			FVector3d Bary = VectorUtil::BarycentricCoords(Mesh.GetVertexRef(VertexID), Mesh.GetVertexRef(TriVertex.A), 
+														   Mesh.GetVertexRef(TriVertex.B), Mesh.GetVertexRef(TriVertex.C));
+														   
+			FVector3d BarySquared = Bary*Bary;
+
+			// Displaced vertex. First compute contribution of the original control points at the original vertices
+			FVector3d NewPos = Bary[0]*BarySquared[0]*Mesh.GetVertexRef(TriVertex.A) + 
+							   Bary[1]*BarySquared[1]*Mesh.GetVertexRef(TriVertex.B) + 
+							   Bary[2]*BarySquared[2]*Mesh.GetVertexRef(TriVertex.C);
+			
+			// Compute contribution of the control points at the edges
+			for (int EIDX = 0; EIDX < 3; ++EIDX) 
+			{   
+				int EID = TriEdges[EIDX];
+				FIndex2i EdgeV = OriginalConnectivity.Edges[EID].Vert;
+				const FEdgeControlPoints& ControlPnts = FControlPoints.OnEdges[EID];
+
+				// Check that the orientation of the edge is consistent so we use the correct barycentric values
+				int EdgeVtx1 = EdgeV[0];
+				int EdgeVtx2 = EdgeV[1];
+				IndexUtil::OrientTriEdgeAndFindOtherVtx(EdgeVtx1, EdgeVtx2, TriVertex);
+				
+				int BaryIdx1 = EdgeV[0] == EdgeVtx1 ? EIDX : (EIDX + 1) % 3;
+				int BaryIdx2 = EdgeV[0] == EdgeVtx1 ? (EIDX + 1) % 3 : EIDX;
+				
+				NewPos += (3.0*BarySquared[BaryIdx1]*Bary[BaryIdx2])*ControlPnts.Point1 + 
+						  (3.0*BarySquared[BaryIdx2]*Bary[BaryIdx1])*ControlPnts.Point2;
+			}
+
+			// Finally compute contribution of the single control point inside the triangle
+			NewPos += (6.0*Bary[0]*Bary[1]*Bary[2])*FControlPoints.OnTriangles[OriginalTriangleID];
+			
+			Mesh.SetVertex(VertexID, NewPos); 
+
+			if (bComputePNNormals) 
+			{
+				// Compute contribution of the normal control points at the original vertices
+				FVector3d NewNormal = BarySquared[0]*Mesh.GetVertexNormal(TriVertex.A) + 
+									  BarySquared[1]*Mesh.GetVertexNormal(TriVertex.B) +
+									  BarySquared[2]*Mesh.GetVertexNormal(TriVertex.C);
+				
+				// Compute contribution of the normal control points at the edges
+				for (int EIDX = 0; EIDX < 3; ++EIDX) 
+				{   
+					int EID = TriEdges[EIDX];
+					const FEdgeControlPoints& ControlPnts = FControlPoints.OnEdges[EID];
+					NewNormal += Bary[EIDX]*Bary[(EIDX + 1) % 3]*ControlPnts.Normal;
+				}
+
+				Mesh.SetVertexNormal(VertexID, NewNormal);                                      
+			}
+		});
+
+		if (ProgressCancel && ProgressCancel->Cancelled()) 
+		{
+			return false; 
+		}
+		
+		return true;
+	}
+}
+
+bool FPNTriangles::Compute()
+{
+	using namespace FPNTrianglesLocals;
+
+	check(TesselationLevel >= 0);
+	check(Mesh != nullptr);
+
+	if (TesselationLevel < 0 || Mesh == nullptr) 
+	{
+		return false;
+	}
+
+	if (TesselationLevel == 0) 
+	{
+		return true; // nothing to do
+	}
+
+	FMeshNormals Normals;
+	bool bHaveVertexNormals = Mesh->HasVertexNormals();
+	if (!bHaveVertexNormals)
+	{
+		Normals = FMeshNormals(Mesh);
+		Normals.ComputeVertexNormals();
+	}
+	FMeshNormals* UseNormals = (bHaveVertexNormals) ? nullptr : &Normals;
+	
+	// Compute PN triangle control points for each flat triangle of the original mesh
+	FControlPoints FControlPoints;
+	bool bOk = ComputeControlPoints(FControlPoints, *Mesh, UseNormals, bComputePNNormals, Progress);
+	if (bOk == false) 
+	{
+		return false;
+	}
+
+	// Save the topology information of the original mesh  before we change it with tesselation
+	FMeshConnectivity OriginalConnectivity(*Mesh);
+	
+	// Tesselate the original mesh
+	TArray<FIndex2i> NewVerticies;
+	bOk = TesselateMesh(*Mesh, NewVerticies, TesselationLevel, Progress);
+	if (bOk == false) 
+	{
+		return false;
+	}
+
+	// Compute displacement and normals
+	bOk = Displace(*Mesh, FControlPoints, NewVerticies, bComputePNNormals, OriginalConnectivity, Progress);
+	if (bOk == false) 
+	{
+		return false;
+	}
+
+	return true;
+}
\ No newline at end of file
diff --git a/Engine/Plugins/Runtime/GeometryProcessing/Source/DynamicMesh/Public/Operations/PNTriangles.h b/Engine/Plugins/Runtime/GeometryProcessing/Source/DynamicMesh/Public/Operations/PNTriangles.h
new file mode 100644
--- /dev/null
+++ b/Engine/Plugins/Runtime/GeometryProcessing/Source/DynamicMesh/Public/Operations/PNTriangles.h
@@ -0,0 +1,77 @@
+// Copyright Epic Games, Inc. All Rights Reserved.
+
+#pragma once
+
+#include "GeometryTypes.h"
+
+class FProgressCancel;
+
+namespace UE
+{
+namespace Geometry
+{
+
+class FDynamicMesh3;
+
+/**
+ * FPNTriangles implements curved PN (Point-Normal) Triangles - https://alex.vlachos.com/graphics/CurvedPNTriangles.pdf.
+ * Each PN triangle replaces one original flat triangle by a curved shape that is retriangulated into 
+ * a number of small subtriangles. The geometry of a PN triangle is defined as one cubic Bezier patch using control 
+ * points. The patch matches the point and normal information at the vertices of the original flat triangle.
+ */
+
+class DYNAMICMESH_API FPNTriangles
+{
+public:
+
+	//
+	// Input & Output
+	//
+
+	/** The mesh that we are modifying */
+	FDynamicMesh3* Mesh = nullptr;
+
+	//
+	// Input
+	//
+
+	/** Set this to be able to cancel running operation. */
+	FProgressCancel* Progress = nullptr;
+	
+	/** How many times we are recursively subdividing triangles (loop style subdivision). */
+	int TesselationLevel = 1;
+
+	/**
+	 * If true, use the quadratically varying normal computation. 
+	 * Otherwise, if the vertex normals are enabled, use standard per-vertex normal computation.
+	 */
+	bool bComputePNNormals = true;
+
+public:
+	FPNTriangles(FDynamicMesh3* Mesh) : Mesh(Mesh)
+	{
+	}
+	virtual ~FPNTriangles() {}
+
+	/**
+	 * @return EOperationValidationResult::Ok if we can apply operation, or error code if we cannot.
+	 */
+	virtual EOperationValidationResult Validate()
+	{
+		if (TesselationLevel < 0 || Mesh == nullptr) 
+		{
+			return EOperationValidationResult::Failed_UnknownReason;    
+		}
+
+		return EOperationValidationResult::Ok;
+	}
+
+	/**
+	 * Generate PN Triangles geometry and optionally compute quadratically varying normals.
+	 * @return true if the algorithm succeeds, false if it failed or was canceled by the user.
+	 */
+	virtual bool Compute();
+};
+
+} // end namespace UE::Geometry
+} // end namespace UE
