fix mesh boolean and self-union simplifying across group boundaries

#rnx

diff --git a/Engine/Plugins/Runtime/GeometryProcessing/Source/DynamicMesh/Private/Operations/MeshBoolean.cpp b/Engine/Plugins/Runtime/GeometryProcessing/Source/DynamicMesh/Private/Operations/MeshBoolean.cpp
--- a/Engine/Plugins/Runtime/GeometryProcessing/Source/DynamicMesh/Private/Operations/MeshBoolean.cpp
+++ b/Engine/Plugins/Runtime/GeometryProcessing/Source/DynamicMesh/Private/Operations/MeshBoolean.cpp
@@ -838,7 +838,8 @@ bool FMeshBoolean::CollapseWouldHurtTriangleQuality(
 bool FMeshBoolean::CollapseWouldChangeShapeOrUVs(
 	const FDynamicMesh3& Mesh, const TSet<int>& CutBoundaryEdgeSet, double DotTolerance, int SourceEID,
 	int32 RemoveV, const FVector3d& RemoveVPos, int32 KeepV, const FVector3d& KeepVPos, const FVector3d& EdgeDir,
-	bool bPreserveUVsForMesh, bool bPreserveVertexUVs, bool bPreserveOverlayUVs, float UVEqualThresholdSq)
+	bool bPreserveTriangleGroups, bool bPreserveUVsForMesh, bool bPreserveVertexUVs, bool bPreserveOverlayUVs, 
+	float UVEqualThresholdSq)
 {
 	// Search the edges connected to the vertex to find one in the boundary set that points in the opposite direction
 	// If we don't find that edge, or if there are other boundary/seam edges attached, we can't remove this vertex
@@ -846,17 +847,27 @@ bool FMeshBoolean::CollapseWouldChangeShapeOrUVs(
 	bool bHasBadEdge = false;
 
 	int OpposedEdge = -1;
+	int SourceGroupID = Mesh.GetTriangleGroup(Mesh.GetEdgeT(SourceEID).A);
 	Mesh.EnumerateVertexEdges(RemoveV,
-		[&Mesh, &CutBoundaryEdgeSet, &DotTolerance,
-		&RemoveVPos, &KeepVPos, &EdgeDir, &bHasBadEdge,
-		&RemoveV, &KeepV, &SourceEID, &OpposedEdge,
-		bPreserveUVsForMesh, bPreserveVertexUVs, bPreserveOverlayUVs, UVEqualThresholdSq](int32 VertEID)
+		[&](int32 VertEID)
 		{
 			if (bHasBadEdge || VertEID == SourceEID)
 			{
 				return;
 			}
 
+			FDynamicMesh3::FEdge Edge = Mesh.GetEdge(VertEID);
+			if (bPreserveTriangleGroups && Mesh.HasTriangleGroups())
+			{
+				if (SourceGroupID != Mesh.GetTriangleGroup(Edge.Tri.A) ||
+					(Edge.Tri.B != FDynamicMesh3::InvalidID && SourceGroupID != Mesh.GetTriangleGroup(Edge.Tri.B)))
+				{
+					// RemoveV is on a group boundary, so the edge collapse would change the shape of the groups
+					bHasBadEdge = true;
+					return;
+				}
+			}
+
 			// it's a known boundary edge; check if it's the opposite-facing one we need
 			if (CutBoundaryEdgeSet.Contains(VertEID))
 			{
@@ -865,7 +876,7 @@ bool FMeshBoolean::CollapseWouldChangeShapeOrUVs(
 					bHasBadEdge = true;
 					return;
 				}
-				FIndex2i OtherEdgeV = Mesh.GetEdgeV(VertEID);
+				FIndex2i OtherEdgeV = Edge.Vert;
 				int OtherV = IndexUtil::FindEdgeOtherVertex(OtherEdgeV, RemoveV);
 				FVector3d OtherVPos = Mesh.GetVertex(OtherV);
 				FVector3d OtherEdgeDir = OtherVPos - RemoveVPos;
@@ -906,7 +917,7 @@ bool FMeshBoolean::CollapseWouldChangeShapeOrUVs(
 				{
 					int NumLayers = Mesh.Attributes()->NumUVLayers();
 					FIndex2i SourceEdgeTris = Mesh.GetEdgeT(SourceEID);
-					FIndex2i OppEdgeTris = Mesh.GetEdgeT(VertEID);
+					FIndex2i OppEdgeTris = Edge.Tri;
 
 					// special handling of seam edge when the edges aren't boundary edges
 					// -- if they're seams, we'd need to check both sides of the seams for a UV match
@@ -1094,7 +1105,7 @@ void FMeshBoolean::SimplifyAlongNewEdges(int NumMeshesToProcess, FDynamicMesh3*
 					
 					bHasBadEdge = bHasBadEdge || CollapseWouldChangeShapeOrUVs(
 						*CutMesh[MeshIdx], CutBoundaryEdgeSets[MeshIdx], DotTolerance,
-						SourceEID, RemoveV, RemoveVPos, KeepV, KeepVPos, EdgeDir,
+						SourceEID, RemoveV, RemoveVPos, KeepV, KeepVPos, EdgeDir, bPreserveTriangleGroups,
 						PreserveUVsOnlyForMesh == -1 || MeshIdx == PreserveUVsOnlyForMesh,
 						bPreserveVertexUVs, bPreserveOverlayUVs, UVDistortTolerance* UVDistortTolerance);
 				};
diff --git a/Engine/Plugins/Runtime/GeometryProcessing/Source/DynamicMesh/Private/Operations/MeshSelfUnion.cpp b/Engine/Plugins/Runtime/GeometryProcessing/Source/DynamicMesh/Private/Operations/MeshSelfUnion.cpp
--- a/Engine/Plugins/Runtime/GeometryProcessing/Source/DynamicMesh/Private/Operations/MeshSelfUnion.cpp
+++ b/Engine/Plugins/Runtime/GeometryProcessing/Source/DynamicMesh/Private/Operations/MeshSelfUnion.cpp
@@ -577,7 +577,7 @@ void FMeshSelfUnion::SimplifyAlongNewEdges(TArray<int>& CutBoundaryEdges, TMap<i
 
 					bHasBadEdge = bHasBadEdge || FMeshBoolean::CollapseWouldChangeShapeOrUVs(
 						*Mesh, CutBoundaryEdgeSet, DotTolerance,
-						SourceEID, RemoveV, RemoveVPos, KeepV, KeepVPos, EdgeDir,
+						SourceEID, RemoveV, RemoveVPos, KeepV, KeepVPos, EdgeDir, bPreserveTriangleGroups,
 						true, bPreserveVertexUVs, bPreserveOverlayUVs, UVDistortTolerance * UVDistortTolerance);
 				}
 
diff --git a/Engine/Plugins/Runtime/GeometryProcessing/Source/DynamicMesh/Public/Operations/MeshBoolean.h b/Engine/Plugins/Runtime/GeometryProcessing/Source/DynamicMesh/Public/Operations/MeshBoolean.h
--- a/Engine/Plugins/Runtime/GeometryProcessing/Source/DynamicMesh/Public/Operations/MeshBoolean.h
+++ b/Engine/Plugins/Runtime/GeometryProcessing/Source/DynamicMesh/Public/Operations/MeshBoolean.h
@@ -86,6 +86,8 @@ public:
 	 *  Equilateral triangles have value 1; Smaller values -> lower quality
 	 */
 	double TryToImproveTriQualityThreshold = .5;
+	/** Prevent simplification from distorting triangle groups */
+	bool bPreserveTriangleGroups = true;
 	/** Prevent simplification from distorting vertex UVs */
 	bool bPreserveVertexUVs = true;
 	/** Prevent simplification from distorting overlay UVs */
@@ -216,12 +218,13 @@ public:
 	);
 
 	/**
-	 * Test if a given edge collapse would change the mesh shape or UVs unacceptably
+	 * Test if a given edge collapse would change the mesh shape, mesh triangle group shape, or UVs unacceptably
 	 */
 	static bool CollapseWouldChangeShapeOrUVs(
 		const FDynamicMesh3& Mesh, const TSet<int>& CutBoundaryEdgeSet, double DotTolerance, int SourceEID,
 		int32 RemoveV, const FVector3d& RemoveVPos, int32 KeepV, const FVector3d& KeepVPos,
-		const FVector3d& EdgeDir, bool bPreserveUVsForMesh, bool bPreserveVertexUVs, bool bPreserveOverlayUVs, float UVToleranceSq);
+		const FVector3d& EdgeDir, bool bPreserveTriangleGroups, bool bPreserveUVsForMesh,
+		bool bPreserveVertexUVs, bool bPreserveOverlayUVs, float UVToleranceSq);
 };
 
 
diff --git a/Engine/Plugins/Runtime/GeometryProcessing/Source/DynamicMesh/Public/Operations/MeshSelfUnion.h b/Engine/Plugins/Runtime/GeometryProcessing/Source/DynamicMesh/Public/Operations/MeshSelfUnion.h
--- a/Engine/Plugins/Runtime/GeometryProcessing/Source/DynamicMesh/Public/Operations/MeshSelfUnion.h
+++ b/Engine/Plugins/Runtime/GeometryProcessing/Source/DynamicMesh/Public/Operations/MeshSelfUnion.h
@@ -72,6 +72,8 @@ public:
 	 *  Equilateral triangles have value 1; Smaller values -> lower quality
 	 */
 	double TryToImproveTriQualityThreshold = .25;
+	/** Prevent simplification from distorting triangle groups */
+	bool bPreserveTriangleGroups = true;
 	/** Prevent simplification from distorting vertex UVs */
 	bool bPreserveVertexUVs = true;
 	/** Prevent simplification from distorting overlay UVs */
