ModelingTools: Bake tool detail mesh sampling abstraction
- Refactored direct detail mesh sampling queries into an abstract base interface for repurposing to alternative detail mesh constucts other than DynamicMesh3 (including N-to-1).
- Move IsTriangle tests out of evaluators and into the core bake loop.

#rnx

diff --git a/Engine/Plugins/Experimental/MeshModelingToolsetExp/Source/MeshModelingToolsExp/Private/BakeMeshAttributeMapsTool.cpp b/Engine/Plugins/Experimental/MeshModelingToolsetExp/Source/MeshModelingToolsExp/Private/BakeMeshAttributeMapsTool.cpp
--- a/Engine/Plugins/Experimental/MeshModelingToolsetExp/Source/MeshModelingToolsExp/Private/BakeMeshAttributeMapsTool.cpp
+++ b/Engine/Plugins/Experimental/MeshModelingToolsetExp/Source/MeshModelingToolsExp/Private/BakeMeshAttributeMapsTool.cpp
@@ -13,7 +13,6 @@
 #include "Sampling/MeshCurvatureMapEvaluator.h"
 #include "Sampling/MeshPropertyMapEvaluator.h"
 #include "Sampling/MeshResampleImageEvaluator.h"
-#include "Util/IndexUtil.h"
 
 #include "Materials/Material.h"
 #include "Materials/MaterialInstanceDynamic.h"
@@ -81,7 +80,7 @@ const FToolTargetTypeRequirements& UBakeMeshAttributeMapsToolBuilder::GetTargetR
 
 bool UBakeMeshAttributeMapsToolBuilder::CanBuildTool(const FToolBuilderState& SceneState) const
 {
-	int32 NumTargets = SceneState.TargetManager->CountSelectedAndTargetable(SceneState, GetTargetRequirements());
+	const int32 NumTargets = SceneState.TargetManager->CountSelectedAndTargetable(SceneState, GetTargetRequirements());
 	return (NumTargets == 1 || NumTargets == 2);
 }
 
@@ -135,7 +134,6 @@ public:
 	int32 DetailMeshNormalUVLayer = 0;
 
 	// Texture2DImage & MultiTexture settings
-	const FDynamicMeshUVOverlay* UVOverlay = nullptr;
 	ImagePtr TextureImage;
 	TMap<int32, ImagePtr> MaterialToTextureImageMap;
 
@@ -147,15 +145,14 @@ public:
 			return Progress && Progress->Cancelled();
 		};
 		Baker->SetTargetMesh(BaseMesh);
-		Baker->SetDetailMesh(DetailMesh.Get(), DetailSpatial.Get());
 		Baker->SetDimensions(BakeCacheSettings.Dimensions);
 		Baker->SetUVLayer(BakeCacheSettings.UVLayer);
 		Baker->SetThickness(BakeCacheSettings.Thickness);
 		Baker->SetMultisampling(BakeCacheSettings.Multisampling);
 		Baker->SetTargetMeshTangents(BaseMeshTangents);
-		Baker->SetDetailMeshTangents(DetailMeshTangents);
-		Baker->SetDetailMeshNormalMap(DetailMeshNormalMap);
-		Baker->SetDetailMeshNormalUVLayer(DetailMeshNormalUVLayer);
+		
+		FMeshBakerDynamicMeshSampler DetailSampler(DetailMesh.Get(), DetailSpatial.Get(), DetailMeshTangents.Get());
+		Baker->SetDetailSampler(&DetailSampler);
 
 		for (const EBakeMapType MapType : ALL_BAKE_MAP_TYPES)
 		{
@@ -164,6 +161,7 @@ public:
 			case EBakeMapType::TangentSpaceNormalMap:
 			{
 				TSharedPtr<FMeshNormalMapEvaluator, ESPMode::ThreadSafe> NormalEval = MakeShared<FMeshNormalMapEvaluator, ESPMode::ThreadSafe>();
+				DetailSampler.SetNormalMap(DetailMesh.Get(), IMeshBakerDetailSampler::FBakeDetailTexture(DetailMeshNormalMap.Get(), DetailMeshNormalUVLayer));
 				Baker->AddBaker(NormalEval);
 				break;
 			}
@@ -223,6 +221,7 @@ public:
 			{
 				TSharedPtr<FMeshPropertyMapEvaluator, ESPMode::ThreadSafe> PropertyBaker = MakeShared<FMeshPropertyMapEvaluator, ESPMode::ThreadSafe>();
 				PropertyBaker->Property = EMeshPropertyMapType::Normal;
+				DetailSampler.SetNormalMap(DetailMesh.Get(), IMeshBakerDetailSampler::FBakeDetailTexture(DetailMeshNormalMap.Get(), DetailMeshNormalUVLayer));
 				Baker->AddBaker(PropertyBaker);
 				break;
 			}
@@ -257,17 +256,14 @@ public:
 			case EBakeMapType::Texture2DImage:
 			{
 				TSharedPtr<FMeshResampleImageEvaluator, ESPMode::ThreadSafe> ResampleBaker = MakeShared<FMeshResampleImageEvaluator, ESPMode::ThreadSafe>();
-				ResampleBaker->DetailUVOverlay = UVOverlay;
-				ResampleBaker->SampleFunction = [this](FVector2d UVCoord) {
-					return TextureImage->BilinearSampleUV<float>(UVCoord, FVector4f(0, 0, 0, 1));
-				};
+				DetailSampler.SetColorMap(DetailMesh.Get(), IMeshBakerDetailSampler::FBakeDetailTexture(TextureImage.Get(), TextureSettings.UVLayer));
 				Baker->AddBaker(ResampleBaker);
 				break;
 			}
 			case EBakeMapType::MultiTexture:
 			{
 				TSharedPtr<FMeshMultiResampleImageEvaluator, ESPMode::ThreadSafe> TextureBaker = MakeShared<FMeshMultiResampleImageEvaluator, ESPMode::ThreadSafe>();
-				TextureBaker->DetailUVOverlay = UVOverlay;
+				TextureBaker->DetailUVLayer = TextureSettings.UVLayer;
 				TextureBaker->MultiTextures = MaterialToTextureImageMap;
 				Baker->AddBaker(TextureBaker);
 				break;
@@ -521,14 +517,12 @@ TUniquePtr<UE::Geometry::TGenericDataOperator<FMeshMapBaker>> UBakeMeshAttribute
 	{
 		Op->TextureSettings = CachedTexture2DImageSettings;
 		Op->TextureImage = CachedTextureImage;
-		Op->UVOverlay = DetailMesh->Attributes()->GetUVLayer(CachedTexture2DImageSettings.UVLayer);
 	}
 
 	if ((bool)(CachedBakeCacheSettings.BakeMapTypes & EBakeMapType::MultiTexture))
 	{
 		Op->TextureSettings = CachedTexture2DImageSettings;
 		Op->MaterialToTextureImageMap = CachedMultiTextures;
-		Op->UVOverlay = DetailMesh->Attributes()->GetUVLayer(CachedTexture2DImageSettings.UVLayer);
 	}
 
 	return Op;
diff --git a/Engine/Plugins/Experimental/MeshModelingToolsetExp/Source/MeshModelingToolsExp/Private/BakeMeshAttributeVertexTool.cpp b/Engine/Plugins/Experimental/MeshModelingToolsetExp/Source/MeshModelingToolsExp/Private/BakeMeshAttributeVertexTool.cpp
--- a/Engine/Plugins/Experimental/MeshModelingToolsetExp/Source/MeshModelingToolsExp/Private/BakeMeshAttributeVertexTool.cpp
+++ b/Engine/Plugins/Experimental/MeshModelingToolsetExp/Source/MeshModelingToolsExp/Private/BakeMeshAttributeVertexTool.cpp
@@ -97,10 +97,12 @@ public:
 			return Progress && Progress->Cancelled();
 		};
 		Baker->SetTargetMesh(BaseMesh);
-		Baker->SetDetailMesh(DetailMesh.Get(), DetailSpatial.Get());
 		Baker->SetTargetMeshTangents(BaseMeshTangents);
 		Baker->SetThickness(BakeSettings.Thickness);
 		Baker->BakeMode = BakeSettings.VertexMode == EBakeVertexMode::Color ? FMeshVertexBaker::EBakeMode::Color : FMeshVertexBaker::EBakeMode::Channel;
+		
+		FMeshBakerDynamicMeshSampler DetailSampler(DetailMesh.Get(), DetailSpatial.Get());
+		Baker->SetDetailSampler(&DetailSampler);
 
 		auto InitOcclusionEvaluator = [this] (FMeshOcclusionMapEvaluator* OcclusionEval, const EMeshOcclusionMapType OcclusionType)
 		{
@@ -224,17 +226,14 @@ public:
 			case EBakeVertexTypeColor::Texture2DImage:
 			{
 				TSharedPtr<FMeshResampleImageEvaluator, ESPMode::ThreadSafe> TextureBaker = MakeShared<FMeshResampleImageEvaluator, ESPMode::ThreadSafe>();
-				TextureBaker->DetailUVOverlay = UVOverlay;
-				TextureBaker->SampleFunction = [this](FVector2d UVCoord) {
-					return TextureImage->BilinearSampleUV<float>(UVCoord, FVector4f(0, 0, 0, 1));
-				};
+				DetailSampler.SetColorMap(DetailMesh.Get(), IMeshBakerDetailSampler::FBakeDetailTexture(TextureImage.Get(), TextureSettings.UVLayer));
 				Baker->ColorEvaluator = TextureBaker;
 				break;
 			}
 			case EBakeVertexTypeColor::MultiTexture:
 			{
 				TSharedPtr<FMeshMultiResampleImageEvaluator, ESPMode::ThreadSafe> TextureBaker = MakeShared<FMeshMultiResampleImageEvaluator, ESPMode::ThreadSafe>();
-				TextureBaker->DetailUVOverlay = UVOverlay;
+				TextureBaker->DetailUVLayer = TextureSettings.UVLayer;
 				TextureBaker->MultiTextures = MaterialToTextureImageMap;
 				Baker->ColorEvaluator = TextureBaker;
 				break;
diff --git a/Engine/Plugins/Runtime/GeometryProcessing/Source/DynamicMesh/Private/Sampling/MeshBakerCommon.cpp b/Engine/Plugins/Runtime/GeometryProcessing/Source/DynamicMesh/Private/Sampling/MeshBakerCommon.cpp
--- a/Engine/Plugins/Runtime/GeometryProcessing/Source/DynamicMesh/Private/Sampling/MeshBakerCommon.cpp
+++ b/Engine/Plugins/Runtime/GeometryProcessing/Source/DynamicMesh/Private/Sampling/MeshBakerCommon.cpp
@@ -1,10 +1,13 @@
 // Copyright Epic Games, Inc. All Rights Reserved.
 
 #include "Sampling/MeshBakerCommon.h"
+#include "Sampling/MeshBaseBaker.h"
+#include "MeshQueries.h"
 
 #include "ExplicitUseGeometryMathTypes.h"		// using UE::Geometry::(math types)
 using namespace UE::Geometry;
 
+
 /**
  * Find point on Detail mesh that corresponds to point on Base mesh.
  * Strategy is:
@@ -15,11 +18,10 @@ using namespace UE::Geometry;
  *
  * If all of those fail, if bFailToNearestPoint is true we fall back to nearest-point,
  *
- * If all the above fail, return false
+ * If all the above fail, return nullptr
  */
-bool UE::Geometry::GetDetailMeshTrianglePoint_Raycast(
-	const FDynamicMesh3& DetailMesh,
-	const FDynamicMeshAABBTree3& DetailSpatial,
+const void* UE::Geometry::GetDetailMeshTrianglePoint_Raycast(
+	const IMeshBakerDetailSampler* DetailSampler,
 	const FVector3d& BasePoint,
 	const FVector3d& BaseNormal,
 	int32& DetailTriangleOut,
@@ -31,48 +33,44 @@ bool UE::Geometry::GetDetailMeshTrianglePoint_Raycast(
 	// TODO: should we check normals here? inverse normal should probably not be considered valid
 
 	// shoot rays forwards and backwards
-	FRay3d InwardRay = FRay3d(BasePoint + Thickness * BaseNormal, -BaseNormal);
-	FRay3d ForwardRay(BasePoint, BaseNormal);
-	FRay3d BackwardRay(BasePoint, -BaseNormal);
+	const FRay3d InwardRay(BasePoint + Thickness * BaseNormal, -BaseNormal);
+	const FRay3d ForwardRay(BasePoint, BaseNormal);
+	const FRay3d BackwardRay(BasePoint, -BaseNormal);
 	int32 ForwardHitTID = IndexConstants::InvalidID, InwardHitTID = IndexConstants::InvalidID, BackwardHitTID = IndexConstants::InvalidID;
 	double ForwardHitDist, InwardHitDist, BackwardHitDist;
 
 	IMeshSpatial::FQueryOptions Options;
 	Options.MaxDistance = Thickness;
-	bool bHitInward = DetailSpatial.FindNearestHitTriangle(InwardRay, InwardHitDist, InwardHitTID, Options);
-	bool bHitForward = DetailSpatial.FindNearestHitTriangle(ForwardRay, ForwardHitDist, ForwardHitTID, Options);
-	bool bHitBackward = DetailSpatial.FindNearestHitTriangle(BackwardRay, BackwardHitDist, BackwardHitTID, Options);
 
-	FRay3d HitRay;
 	int32 HitTID = IndexConstants::InvalidID;
-	double HitDist = TNumericLimits<double>::Max();
+	FVector3d HitBaryCoords;
 
-	if (bHitInward)
+	// Inward hit test
+	const void* HitMesh = nullptr;
+	if (const void* InwardMesh = DetailSampler->FindNearestHitTriangle(InwardRay, InwardHitDist, InwardHitTID, HitBaryCoords, Options))
 	{
-		HitRay = InwardRay;
+		HitMesh = InwardMesh;
 		HitTID = InwardHitTID;
-		HitDist = InwardHitDist;
 	}
-	else if (bHitForward)
+	// Forward hit test
+	else if (const void* ForwardMesh = DetailSampler->FindNearestHitTriangle(ForwardRay, ForwardHitDist, ForwardHitTID, HitBaryCoords, Options))
 	{
-		HitRay = ForwardRay;
+		HitMesh = ForwardMesh;
 		HitTID = ForwardHitTID;
-		HitDist = ForwardHitDist;
 	}
-	else if (bHitBackward)
+	// Backward hit test
+	else if (const void* BackwardMesh = DetailSampler->FindNearestHitTriangle(BackwardRay, BackwardHitDist, BackwardHitTID, HitBaryCoords, Options))
 	{
-		HitRay = BackwardRay;
+		HitMesh = BackwardMesh;
 		HitTID = BackwardHitTID;
-		HitDist = BackwardHitDist;
 	}
 
 	// if we got a valid ray hit, use it
-	if (DetailMesh.IsTriangle(HitTID))
+	if (HitMesh && DetailSampler->IsTriangle(HitMesh, HitTID))
 	{
 		DetailTriangleOut = HitTID;
-		FIntrRay3Triangle3d IntrQuery = TMeshQueries<FDynamicMesh3>::TriangleIntersection(DetailMesh, HitTID, HitRay);
-		DetailTriBaryCoords = IntrQuery.TriangleBaryCoords;
-		return true;
+		DetailTriBaryCoords = HitBaryCoords;
+		return HitMesh;
 	}
 	else
 	{
@@ -80,48 +78,47 @@ bool UE::Geometry::GetDetailMeshTrianglePoint_Raycast(
 		IMeshSpatial::FQueryOptions OnSurfQueryOptions;
 		OnSurfQueryOptions.MaxDistance = Thickness;
 		double NearDistSqr = 0;
-		int32 NearestTriID = -1;
+		int32 NearestTriID = IndexConstants::InvalidID;
+		const void* NearestMesh = nullptr;
 		// if we are using absolute nearest point as a fallback, then ignore max distance
 		if (bFailToNearestPoint)
 		{
-			NearestTriID = DetailSpatial.FindNearestTriangle(BasePoint, NearDistSqr);
+			NearestMesh = DetailSampler->FindNearestTriangle(BasePoint, NearDistSqr, NearestTriID, DetailTriBaryCoords);
 		}
 		else
 		{
-			NearestTriID = DetailSpatial.FindNearestTriangle(BasePoint, NearDistSqr, OnSurfQueryOptions);
+			NearestMesh = DetailSampler->FindNearestTriangle(BasePoint, NearDistSqr, NearestTriID, DetailTriBaryCoords, OnSurfQueryOptions);
 		}
-		if (DetailMesh.IsTriangle(NearestTriID))
+		if (NearestMesh && DetailSampler->IsTriangle(NearestMesh, NearestTriID))
 		{
 			DetailTriangleOut = NearestTriID;
-			FDistPoint3Triangle3d DistQuery = TMeshQueries<FDynamicMesh3>::TriangleDistance(DetailMesh, NearestTriID, BasePoint);
-			DetailTriBaryCoords = DistQuery.TriangleBaryCoords;
-			return true;
+			return NearestMesh;
 		}
 	}
 
-	return false;
+	return nullptr;
 }
 
+
 /**
  * Find point on Detail mesh that corresponds to point on Base mesh using minimum distance
  */
-bool UE::Geometry::GetDetailMeshTrianglePoint_Nearest(
-	const FDynamicMesh3& DetailMesh,
-	const FDynamicMeshAABBTree3& DetailSpatial,
+const void* UE::Geometry::GetDetailMeshTrianglePoint_Nearest(
+	const IMeshBakerDetailSampler* DetailSampler,
 	const FVector3d& BasePoint,
 	int32& DetailTriangleOut,
 	FVector3d& DetailTriBaryCoords)
 {
 	double NearDistSqr = 0;
-	int32 NearestTriID = DetailSpatial.FindNearestTriangle(BasePoint, NearDistSqr);
-	if (DetailMesh.IsTriangle(NearestTriID))
+	int32 NearestTriID = IndexConstants::InvalidID;
+	const void* NearestMesh = DetailSampler->FindNearestTriangle(BasePoint, NearDistSqr, NearestTriID, DetailTriBaryCoords);
+	if (NearestMesh && DetailSampler->IsTriangle(NearestMesh, NearestTriID))
 	{
 		DetailTriangleOut = NearestTriID;
-		FDistPoint3Triangle3d DistQuery = TMeshQueries<FDynamicMesh3>::TriangleDistance(DetailMesh, NearestTriID, BasePoint);
-		DetailTriBaryCoords = DistQuery.TriangleBaryCoords;
-		return true;
+		return NearestMesh;
 	}
 
-	return false;
+	return nullptr;
 }
 
+
diff --git a/Engine/Plugins/Runtime/GeometryProcessing/Source/DynamicMesh/Private/Sampling/MeshCurvatureMapEvaluator.cpp b/Engine/Plugins/Runtime/GeometryProcessing/Source/DynamicMesh/Private/Sampling/MeshCurvatureMapEvaluator.cpp
--- a/Engine/Plugins/Runtime/GeometryProcessing/Source/DynamicMesh/Private/Sampling/MeshCurvatureMapEvaluator.cpp
+++ b/Engine/Plugins/Runtime/GeometryProcessing/Source/DynamicMesh/Private/Sampling/MeshCurvatureMapEvaluator.cpp
@@ -3,19 +3,22 @@
 #include "Sampling/MeshCurvatureMapEvaluator.h"
 #include "Sampling/MeshMapBaker.h"
 #include "MeshCurvature.h"
-#include "MeshWeights.h"
 
 #include "ExplicitUseGeometryMathTypes.h"		// using UE::Geometry::(math types)
 using namespace UE::Geometry;
 
-void FMeshCurvatureMapEvaluator::CacheDetailCurvatures(const FDynamicMesh3* DetailMeshIn)
+void FMeshCurvatureMapEvaluator::CacheDetailCurvatures()
 {
+	// TODO: Generalize this evaluator for N meshes.
 	if (!Curvatures)
 	{
 		Curvatures = MakeShared<FMeshVertexCurvatureCache>();
-		Curvatures->BuildAll(*DetailMeshIn);
+		auto GetCurvature = [this](const void* Mesh)
+		{
+			DetailSampler->GetCurvature(Mesh, *Curvatures);
+		};
+		DetailSampler->ProcessMeshes(GetCurvature);
 	}
-	check(Curvatures->Num() == DetailMeshIn->MaxVertexID());
 }
 
 void FMeshCurvatureMapEvaluator::Setup(const FMeshBaseBaker& Baker, FEvaluationContext& Context)
@@ -28,8 +31,8 @@ void FMeshCurvatureMapEvaluator::Setup(const FMeshBaseBaker& Baker, FEvaluationC
 	Context.DataLayout = { EComponents::Float1 };
 
 	// Cache data from the baker
-	DetailMesh = Baker.GetDetailMesh();
-	CacheDetailCurvatures(DetailMesh);
+	DetailSampler = Baker.GetDetailSampler();
+	CacheDetailCurvatures();
 
 	MinPreClamp = -TNumericLimits<double>::Max();
 	MaxPreClamp = TNumericLimits<double>::Max();
@@ -74,7 +77,7 @@ void FMeshCurvatureMapEvaluator::Setup(const FMeshBaseBaker& Baker, FEvaluationC
 
 void FMeshCurvatureMapEvaluator::EvaluateSample(float*& Out, const FCorrespondenceSample& Sample, void* EvalData)
 {
-	FMeshCurvatureMapEvaluator* Eval = static_cast<FMeshCurvatureMapEvaluator*>(EvalData);
+	const FMeshCurvatureMapEvaluator* Eval = static_cast<FMeshCurvatureMapEvaluator*>(EvalData);
 	const float Curvature = (float) Eval->SampleFunction(Sample);
 	WriteToBuffer(Out, Curvature);
 }
@@ -86,7 +89,7 @@ void FMeshCurvatureMapEvaluator::EvaluateDefault(float*& Out, void* EvalData)
 
 void FMeshCurvatureMapEvaluator::EvaluateColor(const int DataIdx, float*& In, FVector4f& Out, void* EvalData)
 {
-	FMeshCurvatureMapEvaluator* Eval = static_cast<FMeshCurvatureMapEvaluator*>(EvalData);
+	const FMeshCurvatureMapEvaluator* Eval = static_cast<FMeshCurvatureMapEvaluator*>(EvalData);
 	const float Curvature = In[0];
 	const float Sign = FMathf::Sign(Curvature);
 	const float T = Eval->ClampRange.GetT(FMathf::Abs(Curvature));
@@ -96,25 +99,21 @@ void FMeshCurvatureMapEvaluator::EvaluateColor(const int DataIdx, float*& In, FV
 	In += 1;
 }
 
-double FMeshCurvatureMapEvaluator::SampleFunction(const FCorrespondenceSample& SampleData)
+double FMeshCurvatureMapEvaluator::SampleFunction(const FCorrespondenceSample& SampleData) const
 {
-	int32 DetailTriID = SampleData.DetailTriID;
-	if (DetailMesh->IsTriangle(DetailTriID))
-	{
-		FIndex3i DetailTri = DetailMesh->GetTriangle(DetailTriID);
-		double CurvatureA = GetCurvature(DetailTri.A);
-		double CurvatureB = GetCurvature(DetailTri.B);
-		double CurvatureC = GetCurvature(DetailTri.C);
-		double InterpCurvature = SampleData.DetailBaryCoords.X * CurvatureA
-			+ SampleData.DetailBaryCoords.Y * CurvatureB
-			+ SampleData.DetailBaryCoords.Z * CurvatureC;
-
-		return InterpCurvature;
-	}
-	return 0.0;
+	const void* DetailMesh = SampleData.DetailMesh;
+	const int32 DetailTriID = SampleData.DetailTriID;
+	const FIndex3i DetailTri = DetailSampler->GetTriangle(DetailMesh, DetailTriID);
+	const double CurvatureA = GetCurvature(DetailTri.A);
+	const double CurvatureB = GetCurvature(DetailTri.B);
+	const double CurvatureC = GetCurvature(DetailTri.C);
+	const double InterpCurvature = SampleData.DetailBaryCoords.X * CurvatureA
+		+ SampleData.DetailBaryCoords.Y * CurvatureB
+		+ SampleData.DetailBaryCoords.Z * CurvatureC;
+	return InterpCurvature;
 }
 
-double FMeshCurvatureMapEvaluator::GetCurvature(int32 vid)
+double FMeshCurvatureMapEvaluator::GetCurvature(const int32 vid) const
 {
 	const FMeshVertexCurvatureCache& Cache = *Curvatures;
 	switch (UseCurvatureType)
@@ -131,7 +130,7 @@ double FMeshCurvatureMapEvaluator::GetCurvature(int32 vid)
 	}
 }
 
-void FMeshCurvatureMapEvaluator::GetColorMapRange(FVector3f& NegativeColorOut, FVector3f& ZeroColorOut, FVector3f& PositiveColorOut)
+void FMeshCurvatureMapEvaluator::GetColorMapRange(FVector3f& NegativeColorOut, FVector3f& ZeroColorOut, FVector3f& PositiveColorOut) const
 {
 	switch (UseColorMode)
 	{
diff --git a/Engine/Plugins/Runtime/GeometryProcessing/Source/DynamicMesh/Private/Sampling/MeshMapBaker.cpp b/Engine/Plugins/Runtime/GeometryProcessing/Source/DynamicMesh/Private/Sampling/MeshMapBaker.cpp
--- a/Engine/Plugins/Runtime/GeometryProcessing/Source/DynamicMesh/Private/Sampling/MeshMapBaker.cpp
+++ b/Engine/Plugins/Runtime/GeometryProcessing/Source/DynamicMesh/Private/Sampling/MeshMapBaker.cpp
@@ -150,7 +150,7 @@ void FMeshMapBaker::Bake()
 	}
 
 	// This sampler finds the correspondence between target surface and detail surface.
-	DetailMeshSampler.Initialize(Mesh, UVOverlay, EMeshSurfaceSamplerQueryType::TriangleAndUV, FMeshMapEvaluator::FCorrespondenceSample(),
+	DetailCorrespondenceSampler.Initialize(Mesh, UVOverlay, EMeshSurfaceSamplerQueryType::TriangleAndUV, FMeshMapEvaluator::FCorrespondenceSample(),
 		[Mesh, NormalOverlay, UseStrategy, this](const FMeshUVSampleInfo& SampleInfo, FMeshMapEvaluator::FCorrespondenceSample& ValueOut)
 	{
 		NormalOverlay->GetTriBaryInterpolate<double>(SampleInfo.TriangleIndex, &SampleInfo.BaryCoords.X, &ValueOut.BaseNormal.X);
@@ -158,24 +158,28 @@ void FMeshMapBaker::Bake()
 		FVector3d RayDir = ValueOut.BaseNormal;
 
 		ValueOut.BaseSample = SampleInfo;
+		ValueOut.DetailMesh = nullptr;
 		ValueOut.DetailTriID = FDynamicMesh3::InvalidID;
 
-		if (UseStrategy == ECorrespondenceStrategy::Identity)
+		if (UseStrategy == ECorrespondenceStrategy::Identity && DetailSampler->SupportsIdentityCorrespondence())
 		{
+			ValueOut.DetailMesh = Mesh;
 			ValueOut.DetailTriID = SampleInfo.TriangleIndex;
 			ValueOut.DetailBaryCoords = SampleInfo.BaryCoords;
 		}
-		else if (UseStrategy == ECorrespondenceStrategy::NearestPoint)
+		else if (UseStrategy == ECorrespondenceStrategy::NearestPoint && DetailSampler->SupportsNearestPointCorrespondence())
 		{
-			bool bFoundTri = GetDetailMeshTrianglePoint_Nearest(*DetailMesh, *DetailSpatial, SampleInfo.SurfacePoint,
+			ValueOut.DetailMesh = GetDetailMeshTrianglePoint_Nearest(DetailSampler, SampleInfo.SurfacePoint,
 				ValueOut.DetailTriID, ValueOut.DetailBaryCoords);
 		}
 		else	// Fall back to raycast strategy
 		{
-			double SampleThickness = this->GetThickness();		// could modulate w/ a map here...
+			checkSlow(DetailSampler->SupportsRaycastCorrespondence());
+			
+			const double SampleThickness = this->GetThickness();		// could modulate w/ a map here...
 
 			// Find detail mesh triangle point
-			bool bFoundTri = GetDetailMeshTrianglePoint_Raycast(*DetailMesh, *DetailSpatial, SampleInfo.SurfacePoint, RayDir,
+			ValueOut.DetailMesh = GetDetailMeshTrianglePoint_Raycast(DetailSampler, SampleInfo.SurfacePoint, RayDir,
 				ValueOut.DetailTriID, ValueOut.DetailBaryCoords, SampleThickness,
 				(UseStrategy == ECorrespondenceStrategy::RaycastStandardThenNearest));
 		}
@@ -318,8 +322,11 @@ void FMeshMapBaker::Bake()
 						const int32 UVTriangleID = OccupancyMap.TexelQueryTriangle[LinearIdx];
 
 						FMeshMapEvaluator::FCorrespondenceSample Sample;
-						DetailMeshSampler.SampleUV(UVTriangleID, UVPosition, Sample);
-						BakeSample(*TileBuffer, Sample, UVPosition, ImageCoords);
+						DetailCorrespondenceSampler.SampleUV(UVTriangleID, UVPosition, Sample);
+						if (DetailSampler->IsTriangle(Sample.DetailMesh, Sample.DetailTriID))
+						{
+							BakeSample(*TileBuffer, Sample, UVPosition, ImageCoords);
+						}
 					}
 				}
 			}
diff --git a/Engine/Plugins/Runtime/GeometryProcessing/Source/DynamicMesh/Private/Sampling/MeshNormalMapEvaluator.cpp b/Engine/Plugins/Runtime/GeometryProcessing/Source/DynamicMesh/Private/Sampling/MeshNormalMapEvaluator.cpp
--- a/Engine/Plugins/Runtime/GeometryProcessing/Source/DynamicMesh/Private/Sampling/MeshNormalMapEvaluator.cpp
+++ b/Engine/Plugins/Runtime/GeometryProcessing/Source/DynamicMesh/Private/Sampling/MeshNormalMapEvaluator.cpp
@@ -10,17 +10,25 @@ using namespace UE::Geometry;
 void FMeshNormalMapEvaluator::Setup(const FMeshBaseBaker& Baker, FEvaluationContext& Context)
 {
 	// Cache data from the baker
-	DetailMesh = Baker.GetDetailMesh();
-	DetailNormalOverlay = Baker.GetDetailMeshNormals();
-	checkSlow(DetailNormalOverlay);
-	const int32 DetailUVLayer = Baker.GetDetailMeshNormalUVLayer();
-	DetailUVOverlay = Baker.GetDetailMeshUVs(DetailUVLayer);
-	checkSlow(DetailUVOverlay);
+	DetailSampler = Baker.GetDetailSampler();
+	auto GetNormalMaps = [this](const void* Mesh)
+	{
+		const FDetailNormalTexture NormalMap = DetailSampler->GetNormalMap(Mesh);
+		DetailNormalTextures.Add(Mesh, NormalMap);
+	};
+	DetailSampler->ProcessMeshes(GetNormalMaps);
+	bool bHasDetailNormalMap = false;
+	for (const auto& Map : DetailNormalTextures)
+	{
+		const TImageBuilder<FVector4f>* NormalMap;
+		int NormalUVLayer;
+		Tie(NormalMap, NormalUVLayer) = Map.Value;
+		bHasDetailNormalMap = bHasDetailNormalMap || NormalMap != nullptr;
+	}
+	
 	BaseMeshTangents = Baker.GetTargetMeshTangents();
-	DetailMeshTangents = Baker.GetDetailMeshTangents();
-	DetailMeshNormalMap = Baker.GetDetailMeshNormalMap();
 
-	Context.Evaluate = DetailMeshNormalMap ? &EvaluateSample<true> : &EvaluateSample<false>;
+	Context.Evaluate = bHasDetailNormalMap ? &EvaluateSample<true> : &EvaluateSample<false>;
 	Context.EvaluateDefault = &EvaluateDefault;
 	Context.EvaluateColor = &EvaluateColor;
 	Context.EvalData = this;
@@ -31,14 +39,14 @@ void FMeshNormalMapEvaluator::Setup(const FMeshBaseBaker& Baker, FEvaluationCont
 template <bool bUseDetailNormalMap>
 void FMeshNormalMapEvaluator::EvaluateSample(float*& Out, const FCorrespondenceSample& Sample, void* EvalData)
 {
-	FMeshNormalMapEvaluator* Eval = static_cast<FMeshNormalMapEvaluator*>(EvalData);
+	const FMeshNormalMapEvaluator* Eval = static_cast<FMeshNormalMapEvaluator*>(EvalData);
 	const FVector3f SampleResult = Eval->SampleFunction<bUseDetailNormalMap>(Sample);
 	WriteToBuffer(Out, SampleResult);
 }
 
 void FMeshNormalMapEvaluator::EvaluateDefault(float*& Out, void* EvalData)
 {
-	FMeshNormalMapEvaluator* Eval = static_cast<FMeshNormalMapEvaluator*>(EvalData);
+	const FMeshNormalMapEvaluator* Eval = static_cast<FMeshNormalMapEvaluator*>(EvalData);
 	WriteToBuffer(Out, Eval->DefaultNormal);
 }
 
@@ -55,32 +63,48 @@ template <bool bUseDetailNormalMap>
 FVector3f FMeshNormalMapEvaluator::SampleFunction(const FCorrespondenceSample& SampleData) const
 {
 	FVector3f Result = FVector3f::UnitZ();
+	const void* DetailMesh = SampleData.DetailMesh;
 	const int32 DetailTriID = SampleData.DetailTriID;
-	if (DetailMesh->IsTriangle(DetailTriID))
+
+	// get tangents on base mesh
+	FVector3d BaseTangentX, BaseTangentY;
+	BaseMeshTangents->GetInterpolatedTriangleTangent(
+		SampleData.BaseSample.TriangleIndex,
+		SampleData.BaseSample.BaryCoords,
+		BaseTangentX, BaseTangentY);
+
+	// sample normal on detail mesh
+	FVector3f DetailNormal;
+	DetailSampler->TriBaryInterpolateNormal(
+		DetailMesh,
+		DetailTriID,
+		SampleData.DetailBaryCoords,
+		DetailNormal);
+	Normalize(DetailNormal);
+	
+	if constexpr (bUseDetailNormalMap)
 	{
-		// get tangents on base mesh
-		FVector3d BaseTangentX, BaseTangentY;
-		BaseMeshTangents->GetInterpolatedTriangleTangent(
-			SampleData.BaseSample.TriangleIndex,
-			SampleData.BaseSample.BaryCoords,
-			BaseTangentX, BaseTangentY);
-
-		// sample normal on detail mesh
-		FVector3d DetailNormal;
-		DetailNormalOverlay->GetTriBaryInterpolate<double>(DetailTriID, &SampleData.DetailBaryCoords.X, &DetailNormal.X);
-		Normalize(DetailNormal);
-		
-		if constexpr (bUseDetailNormalMap)
+		const TImageBuilder<FVector4f>* DetailNormalMap;
+		int DetailNormalUVLayer;
+		Tie(DetailNormalMap, DetailNormalUVLayer) = DetailNormalTextures[DetailMesh];
+
+		if (DetailNormalMap)
 		{
 			FVector3d DetailTangentX, DetailTangentY;
-			DetailMeshTangents->GetInterpolatedTriangleTangent(
+			DetailSampler->TriBaryInterpolateTangents(
+				DetailMesh,
 				SampleData.DetailTriID,
 				SampleData.DetailBaryCoords,
 				DetailTangentX, DetailTangentY);
-			
-			FVector2d DetailUV;
-        	DetailUVOverlay->GetTriBaryInterpolate<double>(DetailTriID, &SampleData.DetailBaryCoords.X, &DetailUV.X);
-			const FVector4f DetailNormalColor4 = SampleNormalMapFunction(DetailUV);
+		
+			FVector2f DetailUV;
+			DetailSampler->TriBaryInterpolateUV(
+				DetailMesh,
+				DetailTriID,
+				SampleData.DetailBaryCoords,
+				DetailNormalUVLayer,
+				DetailUV);
+			const FVector4f DetailNormalColor4 = DetailNormalMap->BilinearSampleUV<float>(FVector2d(DetailUV), FVector4f(0, 0, 0, 1));
 
 			// Map color space [0,1] to normal space [-1,1]
 			const FVector3f DetailNormalColor(DetailNormalColor4.X, DetailNormalColor4.Y, DetailNormalColor4.Z);
@@ -91,20 +115,15 @@ FVector3f FMeshNormalMapEvaluator::SampleFunction(const FCorrespondenceSample& S
 			Normalize(DetailNormalObjectSpace);
 			DetailNormal = DetailNormalObjectSpace;
 		}
+	}
 
-		// compute normal in tangent space
-		const double dx = DetailNormal.Dot(BaseTangentX);
-		const double dy = DetailNormal.Dot(BaseTangentY);
-		const double dz = DetailNormal.Dot(SampleData.BaseNormal);
+	// compute normal in tangent space
+	const double dx = DetailNormal.Dot(BaseTangentX);
+	const double dy = DetailNormal.Dot(BaseTangentY);
+	const double dz = DetailNormal.Dot(SampleData.BaseNormal);
 
-		Result = FVector3f((float)dx, (float)dy, (float)dz);
-	}
-	return Result;
+	return FVector3f((float)dx, (float)dy, (float)dz);
 }
 
-FVector4f FMeshNormalMapEvaluator::SampleNormalMapFunction(const FVector2d& UVCoord) const
-{
-	return DetailMeshNormalMap->BilinearSampleUV<float>(UVCoord, FVector4f(0, 0, 0, 1));
-}
 
 
diff --git a/Engine/Plugins/Runtime/GeometryProcessing/Source/DynamicMesh/Private/Sampling/MeshOcclusionMapEvaluator.cpp b/Engine/Plugins/Runtime/GeometryProcessing/Source/DynamicMesh/Private/Sampling/MeshOcclusionMapEvaluator.cpp
--- a/Engine/Plugins/Runtime/GeometryProcessing/Source/DynamicMesh/Private/Sampling/MeshOcclusionMapEvaluator.cpp
+++ b/Engine/Plugins/Runtime/GeometryProcessing/Source/DynamicMesh/Private/Sampling/MeshOcclusionMapEvaluator.cpp
@@ -3,6 +3,7 @@
 #include "Sampling/MeshOcclusionMapEvaluator.h"
 #include "Sampling/MeshMapBaker.h"
 #include "Sampling/SphericalFibonacci.h"
+#include "Math/RandomStream.h"
 
 #include "ExplicitUseGeometryMathTypes.h"		// using UE::Geometry::(math types)
 using namespace UE::Geometry;
@@ -75,10 +76,7 @@ void FMeshOcclusionMapEvaluator::Setup(const FMeshBaseBaker& Baker, FEvaluationC
 	}
 
 	// Cache data from the baker
-	DetailMesh = Baker.GetDetailMesh();
-	DetailSpatial = Baker.GetDetailMeshSpatial();
-	DetailNormalOverlay = Baker.GetDetailMeshNormals();
-	check(DetailNormalOverlay);
+	DetailSampler = Baker.GetDetailSampler();
 
 	if (WANT_BENT_NORMAL(OcclusionType) && NormalSpace == ESpace::Tangent)
 	{
@@ -192,16 +190,11 @@ void FMeshOcclusionMapEvaluator::EvaluateColor(const int DataIdx, float*& In, FV
 template <EMeshOcclusionMapType ComputeType, FMeshOcclusionMapEvaluator::ESpace ComputeSpace>
 FMeshOcclusionMapEvaluator::FOcclusionTuple FMeshOcclusionMapEvaluator::SampleFunction(const FCorrespondenceSample& SampleData, const FVector3d& DefaultNormal)
 {
-	int32 DetailTriID = SampleData.DetailTriID;
-	if (!DetailMesh->IsTriangle(DetailTriID))
-	{
-		return FOcclusionTuple(0.0, DefaultNormal);
-	}
+	const void* DetailMesh = SampleData.DetailMesh;
+	const int32 DetailTriID = SampleData.DetailTriID;
 
-	FIndex3i DetailTri = DetailMesh->GetTriangle(DetailTriID);
-	//FVector3d DetailTriNormal = DetailMesh.GetTriNormal(DetailTriID);
-	FVector3d DetailTriNormal;
-	DetailNormalOverlay->GetTriBaryInterpolate<double>(DetailTriID, &SampleData.DetailBaryCoords.X, &DetailTriNormal.X);
+	FVector3f DetailTriNormal;
+	DetailSampler->TriBaryInterpolateNormal(DetailMesh, DetailTriID, SampleData.DetailBaryCoords, DetailTriNormal);
 	Normalize(DetailTriNormal);
 
 	FVector3d BaseTangentX, BaseTangentY;
@@ -214,15 +207,15 @@ FMeshOcclusionMapEvaluator::FOcclusionTuple FMeshOcclusionMapEvaluator::SampleFu
 			BaseTangentX, BaseTangentY);
 	}
 
-	FVector3d DetailBaryCoords = SampleData.DetailBaryCoords;
-	FVector3d DetailPos = DetailMesh->GetTriBaryPoint(DetailTriID, DetailBaryCoords.X, DetailBaryCoords.Y, DetailBaryCoords.Z);
+	const FVector3d DetailBaryCoords = SampleData.DetailBaryCoords;
+	FVector3d DetailPos = DetailSampler->TriBaryInterpolatePoint(DetailMesh, DetailTriID, DetailBaryCoords);
 	DetailPos += 10.0f * FMathf::ZeroTolerance * DetailTriNormal;
-	FFrame3d SurfaceFrame(DetailPos, DetailTriNormal);
+	FFrame3d SurfaceFrame(DetailPos, FVector3d(DetailTriNormal));
 
 	// TODO: Consider pregenerating random rotations.
 	FRandomStream RotationGen;
 	RotationGen.GenerateNewSeed();
-	double RotationAngle = RotationGen.GetFraction() * FMathd::TwoPi;
+	const double RotationAngle = RotationGen.GetFraction() * FMathd::TwoPi;
 	SurfaceFrame.Rotate(FQuaterniond(SurfaceFrame.Z(), RotationAngle, false));
 
 	IMeshSpatial::FQueryOptions QueryOptions;
@@ -236,14 +229,14 @@ FMeshOcclusionMapEvaluator::FOcclusionTuple FMeshOcclusionMapEvaluator::SampleFu
 		FRay3d OcclusionRay(DetailPos, SurfaceFrame.FromFrameVector(SphereDir));
 		ensure(OcclusionRay.Direction.Dot(DetailTriNormal) > 0);
 
-		bool bHit = DetailSpatial->TestAnyHitTriangle(OcclusionRay, QueryOptions);
+		const bool bHit = DetailSampler->TestAnyHitTriangle(OcclusionRay, QueryOptions);
 
 		if constexpr (WANT_AMBIENT_OCCLUSION(ComputeType))
 		{
 			// Have weight of point fall off as it becomes more coplanar with face. 
 			// This reduces faceting artifacts that we would otherwise see because geometry does not vary smoothly
 			double PointWeight = 1.0;
-			double BiasDot = OcclusionRay.Direction.Dot(DetailTriNormal);
+			const double BiasDot = OcclusionRay.Direction.Dot(DetailTriNormal);
 			if (BiasDot < BiasDotThreshold)
 			{
 				PointWeight = FMathd::Lerp(0.0, 1.0, FMathd::Clamp(BiasDot / BiasDotThreshold, 0.0, 1.0));
diff --git a/Engine/Plugins/Runtime/GeometryProcessing/Source/DynamicMesh/Private/Sampling/MeshPropertyMapEvaluator.cpp b/Engine/Plugins/Runtime/GeometryProcessing/Source/DynamicMesh/Private/Sampling/MeshPropertyMapEvaluator.cpp
--- a/Engine/Plugins/Runtime/GeometryProcessing/Source/DynamicMesh/Private/Sampling/MeshPropertyMapEvaluator.cpp
+++ b/Engine/Plugins/Runtime/GeometryProcessing/Source/DynamicMesh/Private/Sampling/MeshPropertyMapEvaluator.cpp
@@ -10,15 +10,22 @@ using namespace UE::Geometry;
 void FMeshPropertyMapEvaluator::Setup(const FMeshBaseBaker& Baker, FEvaluationContext& Context)
 {
 	// Cache data from the baker
-	DetailMesh = Baker.GetDetailMesh();
-	DetailNormalOverlay = Baker.GetDetailMeshNormals();
-	check(DetailNormalOverlay);
-	DetailUVOverlay = Baker.GetDetailMeshUVs();
-	DetailColorOverlay = Baker.GetDetailMeshColors();
-	DetailMeshTangents = Baker.GetDetailMeshTangents();
-	DetailMeshNormalMap = Baker.GetDetailMeshNormalMap();
-
-	const bool bHasDetailNormalMap = DetailMeshNormalMap != nullptr;
+	DetailSampler = Baker.GetDetailSampler();
+	auto GetNormalMaps = [this](const void* Mesh)
+	{
+		const FDetailNormalTexture NormalMap = DetailSampler->GetNormalMap(Mesh);
+		DetailNormalTextures.Add(Mesh, NormalMap);
+	};
+	DetailSampler->ProcessMeshes(GetNormalMaps);
+	bool bHasDetailNormalMap = false;
+	for (const auto& Map : DetailNormalTextures)
+	{
+		const TImageBuilder<FVector4f>* NormalMap;
+		int NormalUVLayer;
+		Tie(NormalMap, NormalUVLayer) = Map.Value;
+		bHasDetailNormalMap = bHasDetailNormalMap || NormalMap != nullptr;
+	}
+
 	Context.Evaluate = bHasDetailNormalMap ? &EvaluateSample<true> : &EvaluateSample<false>;
 	Context.EvaluateDefault = &EvaluateDefault;
 	Context.EvaluateColor = &EvaluateColor;
@@ -26,14 +33,13 @@ void FMeshPropertyMapEvaluator::Setup(const FMeshBaseBaker& Baker, FEvaluationCo
 	Context.AccumulateMode = EAccumulateMode::Add;
 	Context.DataLayout = { EComponents::Float3 };
 
-	Bounds = DetailMesh->GetBounds();
+	Bounds = DetailSampler->GetBounds();
 	for (int32 j = 0; j < 3; ++j)
 	{
 		if (Bounds.Diagonal()[j] < FMathf::ZeroTolerance)
 		{
 			Bounds.Min[j] = Bounds.Center()[j] - FMathf::ZeroTolerance;
 			Bounds.Max[j] = Bounds.Center()[j] + FMathf::ZeroTolerance;
-
 		}
 	}
 
@@ -48,13 +54,9 @@ void FMeshPropertyMapEvaluator::Setup(const FMeshBaseBaker& Baker, FEvaluationCo
 		break;
 	case EMeshPropertyMapType::Normal:
 		DefaultValue = NormalToColor(FVector3d::UnitZ());
-		if (bHasDetailNormalMap)
-		{
-			DetailUVOverlay = Baker.GetDetailMeshUVs(Baker.GetDetailMeshNormalUVLayer());
-		}
 		break;
 	case EMeshPropertyMapType::UVPosition:
-		DefaultValue = UVToColor(FVector2d::Zero());
+		DefaultValue = UVToColor(FVector2f::Zero());
 		break;
 	case EMeshPropertyMapType::MaterialID:
 		DefaultValue = FVector3f(LinearColors::LightPink3f());
@@ -68,7 +70,7 @@ void FMeshPropertyMapEvaluator::Setup(const FMeshBaseBaker& Baker, FEvaluationCo
 template <bool bUseDetailNormalMap>
 void FMeshPropertyMapEvaluator::EvaluateSample(float*& Out, const FCorrespondenceSample& Sample, void* EvalData)
 {
-	FMeshPropertyMapEvaluator* Eval = static_cast<FMeshPropertyMapEvaluator*>(EvalData);
+	const FMeshPropertyMapEvaluator* Eval = static_cast<FMeshPropertyMapEvaluator*>(EvalData);
 	const FVector3f SampleResult = Eval->SampleFunction<bUseDetailNormalMap>(Sample);
 	WriteToBuffer(Out, SampleResult);
 }
@@ -86,45 +88,52 @@ void FMeshPropertyMapEvaluator::EvaluateColor(const int DataIdx, float*& In, FVe
 }
 
 template <bool bUseDetailNormalMap>
-FVector3f FMeshPropertyMapEvaluator::SampleFunction(const FCorrespondenceSample& SampleData)
+FVector3f FMeshPropertyMapEvaluator::SampleFunction(const FCorrespondenceSample& SampleData) const
 {
+	const void* DetailMesh = SampleData.DetailMesh;
+	const FVector3d& DetailBaryCoords = SampleData.DetailBaryCoords;
 	FVector3f Color = DefaultValue;
-	int32 DetailTriID = SampleData.DetailTriID;
-	if (DetailMesh->IsTriangle(DetailTriID))
+	const int32 DetailTriID = SampleData.DetailTriID;
+
+	switch (this->Property)
 	{
-		switch (this->Property)
-		{
-		case EMeshPropertyMapType::Position:
-		{
-			FVector3d Position = DetailMesh->GetTriBaryPoint(DetailTriID, SampleData.DetailBaryCoords[0], SampleData.DetailBaryCoords[1], SampleData.DetailBaryCoords[2]);
-			Color = PositionToColor(Position, Bounds);
-		}
-		break;
-		case EMeshPropertyMapType::FacetNormal:
-		{
-			FVector3d FacetNormal = DetailMesh->GetTriNormal(DetailTriID);
-			Color = NormalToColor(FacetNormal);
-		}
-		break;
-		case EMeshPropertyMapType::Normal:
+	case EMeshPropertyMapType::Position:
+	{
+		const FVector3d Position = DetailSampler->TriBaryInterpolatePoint(DetailMesh, DetailTriID, DetailBaryCoords);
+		Color = PositionToColor(Position, Bounds);
+	}
+	break;
+	case EMeshPropertyMapType::FacetNormal:
+	{
+		const FVector3d FacetNormal = DetailSampler->GetTriNormal(DetailMesh, DetailTriID);
+		Color = NormalToColor(FacetNormal);
+	}
+	break;
+	case EMeshPropertyMapType::Normal:
+	{
+		FVector3f DetailNormal;
+		if (DetailSampler->TriBaryInterpolateNormal(DetailMesh, DetailTriID, DetailBaryCoords, DetailNormal))
 		{
-			if (DetailNormalOverlay->IsSetTriangle(DetailTriID))
+			Normalize(DetailNormal);
+
+			if constexpr (bUseDetailNormalMap)
 			{
-				FVector3d DetailNormal;
-				DetailNormalOverlay->GetTriBaryInterpolate<double>(DetailTriID, &SampleData.DetailBaryCoords.X, &DetailNormal.X);
-				Normalize(DetailNormal);
+				const TImageBuilder<FVector4f>* DetailNormalMap;
+				int DetailNormalUVLayer;
+				Tie(DetailNormalMap, DetailNormalUVLayer) = DetailNormalTextures[DetailMesh];
 
-				if constexpr (bUseDetailNormalMap)
+				if (DetailNormalMap)
 				{
 					FVector3d DetailTangentX, DetailTangentY;
-					DetailMeshTangents->GetInterpolatedTriangleTangent(
+					DetailSampler->TriBaryInterpolateTangents(
+						DetailMesh,
 						SampleData.DetailTriID,
 						SampleData.DetailBaryCoords,
 						DetailTangentX, DetailTangentY);
-			
-					FVector2d DetailUV;
-					DetailUVOverlay->GetTriBaryInterpolate<double>(DetailTriID, &SampleData.DetailBaryCoords.X, &DetailUV.X);
-					const FVector4f DetailNormalColor4 = SampleNormalMapFunction(DetailUV);
+	
+					FVector2f DetailUV;
+					DetailSampler->TriBaryInterpolateUV(DetailMesh, DetailTriID, DetailBaryCoords, DetailNormalUVLayer, DetailUV);
+					const FVector4f DetailNormalColor4 = DetailNormalMap->BilinearSampleUV<float>(FVector2d(DetailUV), FVector4f(0, 0, 0, 1));
 
 					// Map color space [0,1] to normal space [-1,1]
 					const FVector3f DetailNormalColor(DetailNormalColor4.X, DetailNormalColor4.Y, DetailNormalColor4.Z);
@@ -135,49 +144,38 @@ FVector3f FMeshPropertyMapEvaluator::SampleFunction(const FCorrespondenceSample&
 					Normalize(DetailNormalObjectSpace);
 					DetailNormal = DetailNormalObjectSpace;
 				}
-				
-				Color = NormalToColor(DetailNormal);
-			}
-		}
-		break;
-		case EMeshPropertyMapType::UVPosition:
-		{
-			if (DetailUVOverlay && DetailUVOverlay->IsSetTriangle(DetailTriID))
-			{
-				FVector2d DetailUV;
-				DetailUVOverlay->GetTriBaryInterpolate<double>(DetailTriID, &SampleData.DetailBaryCoords.X, &DetailUV.X);
-				Color = UVToColor(DetailUV);
 			}
+
+			Color = NormalToColor(DetailNormal);
 		}
-		break;
-		case EMeshPropertyMapType::MaterialID:
+	}
+	break;
+	case EMeshPropertyMapType::UVPosition:
+	{
+		FVector2f DetailUV;
+		if (DetailSampler->TriBaryInterpolateUV(DetailMesh, DetailTriID, DetailBaryCoords, 0, DetailUV))
 		{
-			if (DetailMesh->Attributes() && DetailMesh->Attributes()->HasMaterialID())
-			{
-				const FDynamicMeshMaterialAttribute* DetailMaterialIDAttrib = DetailMesh->Attributes()->GetMaterialID();
-				const int32 MatID = DetailMaterialIDAttrib->GetValue(DetailTriID);
-				Color = LinearColors::SelectColor<FVector3f>(MatID);
-			}
+			Color = UVToColor(DetailUV);
 		}
-		break;
-		case EMeshPropertyMapType::VertexColor:
+	}
+	break;
+	case EMeshPropertyMapType::MaterialID:
+	{
+		const int32 MatID = DetailSampler->GetMaterialID(DetailMesh, DetailTriID);
+		Color = LinearColors::SelectColor<FVector3f>(MatID);
+	}
+	break;
+	case EMeshPropertyMapType::VertexColor:
+	{
+		FVector4f DetailColor;
+		if (DetailSampler->TriBaryInterpolateColor(DetailMesh, DetailTriID, DetailBaryCoords, DetailColor))
 		{
-			if (DetailColorOverlay && DetailColorOverlay->IsSetTriangle(DetailTriID))
-			{
-				FVector4d DetailColor;
-				DetailColorOverlay->GetTriBaryInterpolate<double>(DetailTriID, &SampleData.DetailBaryCoords.X, &DetailColor.X);
-				Color = FVector3f(DetailColor.X, DetailColor.Y, DetailColor.Z);
-			}
-		}
-		break;
+			Color = FVector3f(DetailColor.X, DetailColor.Y, DetailColor.Z);
 		}
 	}
+	break;
+	}
 	return Color;
 }
 
-FVector4f FMeshPropertyMapEvaluator::SampleNormalMapFunction(const FVector2d& UVCoord) const
-{
-	return DetailMeshNormalMap->BilinearSampleUV<float>(UVCoord, FVector4f(0, 0, 0, 1));
-}
-
 
diff --git a/Engine/Plugins/Runtime/GeometryProcessing/Source/DynamicMesh/Private/Sampling/MeshResampleImageEvaluator.cpp b/Engine/Plugins/Runtime/GeometryProcessing/Source/DynamicMesh/Private/Sampling/MeshResampleImageEvaluator.cpp
--- a/Engine/Plugins/Runtime/GeometryProcessing/Source/DynamicMesh/Private/Sampling/MeshResampleImageEvaluator.cpp
+++ b/Engine/Plugins/Runtime/GeometryProcessing/Source/DynamicMesh/Private/Sampling/MeshResampleImageEvaluator.cpp
@@ -20,13 +20,13 @@ void FMeshResampleImageEvaluator::Setup(const FMeshBaseBaker& Baker, FEvaluation
 	Context.DataLayout = { EComponents::Float4 };
 
 	// Cache data from the baker
-	DetailMesh = Baker.GetDetailMesh();
+	DetailSampler = Baker.GetDetailSampler();
 }
 
 void FMeshResampleImageEvaluator::EvaluateSample(float*& Out, const FCorrespondenceSample& Sample, void* EvalData)
 {
-	FMeshResampleImageEvaluator* Eval = static_cast<FMeshResampleImageEvaluator*>(EvalData);
-	FVector4f SampleResult = Eval->ImageSampleFunction(Sample);
+	const FMeshResampleImageEvaluator* Eval = static_cast<FMeshResampleImageEvaluator*>(EvalData);
+	const FVector4f SampleResult = Eval->ImageSampleFunction(Sample);
 	WriteToBuffer(Out, SampleResult);
 }
 
@@ -41,17 +41,17 @@ void FMeshResampleImageEvaluator::EvaluateColor(const int DataIdx, float*& In, F
 	In += 4;
 }
 
-FVector4f FMeshResampleImageEvaluator::ImageSampleFunction(const FCorrespondenceSample& SampleData)
+FVector4f FMeshResampleImageEvaluator::ImageSampleFunction(const FCorrespondenceSample& SampleData) const
 {
-	FVector4f Color = DefaultColor;
-	int32 DetailTriID = SampleData.DetailTriID;
-	if (DetailMesh->IsTriangle(SampleData.DetailTriID) && DetailUVOverlay)
-	{
-		FVector2d DetailUV;
-		DetailUVOverlay->GetTriBaryInterpolate<double>(DetailTriID, &SampleData.DetailBaryCoords.X, &DetailUV.X);
-		Color = SampleFunction(DetailUV);
-	}
-	return Color;
+	const void* DetailMesh = SampleData.DetailMesh;
+	const int32 DetailTriID = SampleData.DetailTriID;
+	const TImageBuilder<FVector4f>* TextureImage;
+	int DetailUVLayer;
+	Tie(TextureImage, DetailUVLayer) = DetailSampler->GetColorMap(DetailMesh);
+
+	FVector2f DetailUV;
+	DetailSampler->TriBaryInterpolateUV(DetailMesh, DetailTriID, SampleData.DetailBaryCoords, DetailUVLayer, DetailUV);
+	return TextureImage->BilinearSampleUV<float>(FVector2d(DetailUV), FVector4f(0, 0, 0, 1));
 }
 
 //
@@ -68,38 +68,29 @@ void FMeshMultiResampleImageEvaluator::Setup(const FMeshBaseBaker& Baker, FEvalu
 	Context.DataLayout = { EComponents::Float4 };
 
 	// Cache data from baker
-	DetailMesh = Baker.GetDetailMesh();
-	DetailMaterialIDAttrib = ensure(DetailMesh) ? DetailMesh->Attributes()->GetMaterialID() : nullptr;
-	bValidDetailMesh = ensure(DetailMaterialIDAttrib) && ensure(DetailUVOverlay);
+	DetailSampler = Baker.GetDetailSampler();
 }
 
 void FMeshMultiResampleImageEvaluator::EvaluateSampleMulti(float*& Out, const FCorrespondenceSample& Sample, void* EvalData)
 {
-	FMeshMultiResampleImageEvaluator* Eval = static_cast<FMeshMultiResampleImageEvaluator*>(EvalData);
-	FVector4f SampleResult = Eval->ImageSampleFunction(Sample);
+	const FMeshMultiResampleImageEvaluator* Eval = static_cast<FMeshMultiResampleImageEvaluator*>(EvalData);
+	const FVector4f SampleResult = Eval->ImageSampleFunction(Sample);
 	WriteToBuffer(Out, SampleResult);
 }
 
-FVector4f FMeshMultiResampleImageEvaluator::ImageSampleFunction(const FCorrespondenceSample& Sample)
+FVector4f FMeshMultiResampleImageEvaluator::ImageSampleFunction(const FCorrespondenceSample& Sample) const
 {
 	FVector4f Color(0.0f, 0.0f, 0.0f, 1.0f);
-	if (!bValidDetailMesh)
-	{
-		return Color;
-	}
-
-	int32 DetailTriID = Sample.DetailTriID;
-	if (DetailMesh->IsTriangle(DetailTriID))
+	const void* DetailMesh = Sample.DetailMesh;
+	const int32 DetailTriID = Sample.DetailTriID;
+	const int32 MaterialID = DetailSampler->GetMaterialID(DetailMesh, DetailTriID);
+	// TODO: Decompose the Map into a sparse array lookup.
+	const TSharedPtr<UE::Geometry::TImageBuilder<FVector4f>, ESPMode::ThreadSafe> TextureImage = MultiTextures.FindRef(MaterialID);
+	if (TextureImage)
 	{
-		// TODO: Decompose the Map into a sparse array lookup.
-		int32 MaterialID = DetailMaterialIDAttrib->GetValue(DetailTriID);
-		TSharedPtr<UE::Geometry::TImageBuilder<FVector4f>, ESPMode::ThreadSafe> TextureImage = MultiTextures.FindRef(MaterialID);
-		if (TextureImage)
-		{
-			FVector2d DetailUV;
-			DetailUVOverlay->GetTriBaryInterpolate<double>(DetailTriID, &Sample.DetailBaryCoords.X, &DetailUV.X);
-			Color = TextureImage->BilinearSampleUV<float>(DetailUV, FVector4f(0, 0, 0, 1));
-		}
+		FVector2f DetailUV;
+		DetailSampler->TriBaryInterpolateUV(DetailMesh, DetailTriID, Sample.DetailBaryCoords, DetailUVLayer, DetailUV);
+		Color = TextureImage->BilinearSampleUV<float>(FVector2d(DetailUV), FVector4f(0, 0, 0, 1));
 	}
 	return Color;
 }
diff --git a/Engine/Plugins/Runtime/GeometryProcessing/Source/DynamicMesh/Private/Sampling/MeshVertexBaker.cpp b/Engine/Plugins/Runtime/GeometryProcessing/Source/DynamicMesh/Private/Sampling/MeshVertexBaker.cpp
--- a/Engine/Plugins/Runtime/GeometryProcessing/Source/DynamicMesh/Private/Sampling/MeshVertexBaker.cpp
+++ b/Engine/Plugins/Runtime/GeometryProcessing/Source/DynamicMesh/Private/Sampling/MeshVertexBaker.cpp
@@ -136,18 +136,20 @@ void FMeshVertexBaker::BakeImpl(void* Data)
 		ValueOut.BaseSample.TriangleIndex = TriangleIndex;
 		ValueOut.BaseSample.SurfacePoint = SurfacePoint;
 		ValueOut.BaseSample.BaryCoords = BaryCoords;
+		ValueOut.DetailMesh = nullptr;
 		ValueOut.DetailTriID = FDynamicMesh3::InvalidID;
 		ValueOut.BaseNormal = SurfaceNormal;
 
 		const FVector3d RayDir = SurfaceNormal;
 		if (UseStrategy == ECorrespondenceStrategy::Identity)
 		{
+			ValueOut.DetailMesh = Mesh;
 			ValueOut.DetailTriID = TriangleIndex;
 			ValueOut.DetailBaryCoords = BaryCoords;
 		}
 		else if (UseStrategy == ECorrespondenceStrategy::NearestPoint)
 		{
-			GetDetailMeshTrianglePoint_Nearest(*Baker->DetailMesh, *Baker->DetailSpatial, SurfacePoint,
+			ValueOut.DetailMesh = GetDetailMeshTrianglePoint_Nearest(Baker->DetailSampler, SurfacePoint,
 			                                   ValueOut.DetailTriID, ValueOut.DetailBaryCoords);
 		}
 		else // Fall back to raycast strategy
@@ -155,16 +157,15 @@ void FMeshVertexBaker::BakeImpl(void* Data)
 			const double SampleThickness = Baker->GetThickness(); // could modulate w/ a map here...
 
 			// Find detail mesh triangle point
-			GetDetailMeshTrianglePoint_Raycast(*Baker->DetailMesh, *Baker->DetailSpatial, SurfacePoint, RayDir,
-			                                   ValueOut.DetailTriID, ValueOut.DetailBaryCoords,
-			                                   SampleThickness,
+			ValueOut.DetailMesh = GetDetailMeshTrianglePoint_Raycast(Baker->DetailSampler, SurfacePoint, RayDir,
+			                                   ValueOut.DetailTriID, ValueOut.DetailBaryCoords, SampleThickness,
 			                                   (UseStrategy == ECorrespondenceStrategy::RaycastStandardThenNearest));
 		}
 	};
 
 	// Perform bake
-	const int32 TileWidth = 256;
-	const int32 TileHeight = 1;
+	constexpr int32 TileWidth = 256;
+	constexpr int32 TileHeight = 1;
 	const FImageTiling Tiles(Baker->Dimensions, TileWidth, TileHeight);
 	const int32 NumTiles = Tiles.Num();
 	const int NumBakers = Baker->Bakers.Num();
diff --git a/Engine/Plugins/Runtime/GeometryProcessing/Source/DynamicMesh/Public/Sampling/MeshBakerCommon.h b/Engine/Plugins/Runtime/GeometryProcessing/Source/DynamicMesh/Public/Sampling/MeshBakerCommon.h
--- a/Engine/Plugins/Runtime/GeometryProcessing/Source/DynamicMesh/Public/Sampling/MeshBakerCommon.h
+++ b/Engine/Plugins/Runtime/GeometryProcessing/Source/DynamicMesh/Public/Sampling/MeshBakerCommon.h
@@ -4,17 +4,26 @@
 
 #include "DynamicMesh/DynamicMesh3.h"
 #include "DynamicMesh/DynamicMeshAABBTree3.h"
+#include "DynamicMesh/DynamicMeshAttributeSet.h"
+#include "DynamicMesh/MeshTangents.h"
 #include "Image/ImageTile.h"
+#include "Image/ImageBuilder.h"
+#include "MeshCurvature.h"
+
 
 namespace UE
 {
 namespace Geometry
 {
 
-/** Compute base/detail mesh correspondence by raycast */
-bool GetDetailMeshTrianglePoint_Raycast(
-	const FDynamicMesh3& DetailMesh,
-	const FDynamicMeshAABBTree3& DetailSpatial,
+class IMeshBakerDetailSampler;	
+
+/**
+ * Compute base/detail mesh correspondence by raycast
+ * @return a pointer to the hit mesh, otherwise nullptr.
+ */
+const void* GetDetailMeshTrianglePoint_Raycast(
+	const IMeshBakerDetailSampler* DetailSpatial,
 	const FVector3d& BasePoint,
 	const FVector3d& BaseNormal,
 	int32& DetailTriangleOut,
@@ -23,14 +32,18 @@ bool GetDetailMeshTrianglePoint_Raycast(
 	bool bFailToNearestPoint
 );
 
-/** Compute base/detail mesh correspondence by nearest distance */
-bool GetDetailMeshTrianglePoint_Nearest(
-	const FDynamicMesh3& DetailMesh,
-	const FDynamicMeshAABBTree3& DetailSpatial,
+	
+/**
+ * Compute base/detail mesh correspondence by nearest distance
+ * @return a pointer to the nearest mesh, otherwise nullptr.
+ */
+const void* GetDetailMeshTrianglePoint_Nearest(
+	const IMeshBakerDetailSampler* DetailSpatial,
 	const FVector3d& BasePoint,
 	int32& DetailTriangleOut,
 	FVector3d& DetailTriBaryCoords);
 
+	
 /** Image tile storage for map bakes. */
 class FMeshMapTileBuffer
 {
@@ -70,6 +83,431 @@ private:
 	float* Buffer;
 };
 
+
+// TODO: Replace virtuals with templates. Dynamic dispatch isn't required here.	
+/**
+ * Base bake detail sampler class.
+ *
+ * This class defines the interface for all sample queries that
+ * FMeshBaseBaker supports including the spatial queries for
+ * base to detail correspondence sampling. 
+ */
+class IMeshBakerDetailSampler
+{
+public:
+	/** (Image, UVLayer) pair for detail textures */
+	using FBakeDetailTexture = TTuple<const TImageBuilder<FVector4f>*, int>;
+	
+	virtual ~IMeshBakerDetailSampler() = default;
+
+	/** Iterate over each mesh in the detail set. */
+	virtual void ProcessMeshes(TFunctionRef<void(const void*)> ProcessFn) const = 0;
+
+	/** Associate a color map and UV layer index for a given mesh in the detail set */
+	virtual void SetColorMap(const void* Mesh, const FBakeDetailTexture& Map) = 0;
+
+	/** Associate a normal map and UV layer index for a given mesh in the detail set */
+	virtual void SetNormalMap(const void* Mesh, const FBakeDetailTexture& Map) = 0;
+
+	/** Retrieve a color map and UV layer index from a given mesh in the detail set */
+	virtual FBakeDetailTexture GetColorMap(const void* Mesh) const = 0;
+
+	/** Retrieve a normal map and UV layer index from a given mesh in the detail set */
+	virtual FBakeDetailTexture GetNormalMap(const void* Mesh) const = 0;
+
+	/** @return true if identity correspondence is supported */
+	virtual bool SupportsIdentityCorrespondence() const = 0;
+	
+	/** @return true if nearest point to triangle distance correspondence is supported */
+	virtual bool SupportsNearestPointCorrespondence() const = 0;
+
+	/** @return true if triangle ray intersection correspondence is supported */
+	virtual bool SupportsRaycastCorrespondence() const = 0;
+
+	/**
+	 * @param Point query point
+	 * @param NearestDistSqrOut [Out] returned nearest squared distance, if triangle is found
+	 * @param TriId [Out] ID of triangle nearest to Point within MaxDistance, or InvalidID if not found
+	 * @param TriBaryCoords [Out] barycentric coordinates of the nearest point on the nearest tri.
+	 * @param Options query options
+	 * @return a pointer to the mesh that holds the nearest triangle
+	 */
+	virtual const void* FindNearestTriangle(
+		const FVector3d& Point,
+		double& NearestDistSqrOut,
+		int& TriId,
+		FVector3d& TriBaryCoords,
+		const IMeshSpatial::FQueryOptions& Options = IMeshSpatial::FQueryOptions()) const
+	{
+		return nullptr;
+	}
+
+	/**
+	 * Find nearest triangle from the given ray
+	 * @param Ray query ray
+	 * @param NearestT [Out] parameter of the nearest hit
+	 * @param TriId [Out] ID of triangle intersected by ray within MaxDistance, or InvalidID if not found
+	 * @param TriBaryCoords [Out] barycentric coordinates of intersection point on the hit tri.
+	 * @param Options query options
+	 * @return a pointer to the mesh whose triangle was intersected, nullptr otherwise.
+	 */
+	virtual const void* FindNearestHitTriangle(
+		const FRay3d& Ray,
+		double& NearestT,
+		int& TriId,
+		FVector3d& TriBaryCoords,
+		const IMeshSpatial::FQueryOptions& Options = IMeshSpatial::FQueryOptions()) const
+	{
+		return nullptr;
+	}
+
+	/**
+	 * Return true if any triangles hit by the given ray
+	 * @param Ray query ray
+	 * @param Options query options
+	 * @return true if any triangles hit by the given ray
+	 */
+	virtual bool TestAnyHitTriangle(
+		const FRay3d& Ray,
+		const IMeshSpatial::FQueryOptions& Options = IMeshSpatial::FQueryOptions()) const
+	{
+		return false;
+	}
+
+	/** @return axis aligned bounding box for the detail mesh(es). */
+	virtual FAxisAlignedBox3d GetBounds() const = 0;
+	
+	/**
+	 * @param Mesh pointer to mesh to query
+	 * @param TriId the triangle index to test
+	 * @return true if the given triangle index is valid.
+	 */
+	virtual bool IsTriangle(const void* Mesh, const int TriId) const = 0;
+
+	/**
+	 * @param Mesh pointer to mesh to query
+	 * @param TriId the triangle index to test
+	 * @return 3 vertex indices of the given triangle.
+	 */
+	virtual FIndex3i GetTriangle(const void* Mesh, const int TriId) const = 0;
+
+	/**
+	 * @param Mesh pointer to mesh to query
+	 * @param TriId the triangle index to test
+	 * @return the normal for the given triangle.
+	 */
+	virtual FVector3d GetTriNormal(const void* Mesh, const int TriId) const = 0;
+
+	/**
+	 * @param Mesh pointer to mesh to query
+	 * @param TriId the triangle index to test
+	 * @return the material ID for the given triangle.
+	 */
+	virtual int32 GetMaterialID(const void* Mesh, const int TriId) const = 0;
+
+	/**
+	 * Compute interpolated Position value inside triangle using barycentric coordinates
+	 * @param Mesh pointer to mesh to query
+	 * @param TriId index of triangle
+	 * @param BaryCoords 3 barycentric coordinates inside triangle
+	 * @return resulting interpolated Position parameter value
+	 */
+	virtual FVector3d TriBaryInterpolatePoint(
+		const void* Mesh,
+		const int32 TriId,
+		const FVector3d& BaryCoords) const = 0;
+
+	/**
+	 * Compute interpolated Normal value inside triangle using barycentric coordinates
+	 * @param Mesh pointer to mesh to query
+	 * @param TriId index of triangle
+	 * @param BaryCoords 3 barycentric coordinates inside triangle
+	 * @param NormalOut resulting interpolated Normal parameter value
+	 * @return true if a valid normal was returned, false otherwise
+	 */
+	virtual bool TriBaryInterpolateNormal(
+		const void* Mesh,
+		const int32 TriId,
+		const FVector3d& BaryCoords,
+		FVector3f& NormalOut) const = 0;
+
+	/**
+	 * Compute interpolated UV value inside triangle using barycentric coordinates
+	 * @param Mesh pointer to mesh to query
+	 * @param TriId index of triangle
+	 * @param BaryCoords 3 barycentric coordinates inside triangle
+	 * @param UVLayer UV layer to query
+	 * @param UVOut resulting interpolated UV parameter value
+	 * @return true if a valid UV was returned, false otherwise
+	 */
+	virtual bool TriBaryInterpolateUV(
+		const void* Mesh,
+		const int32 TriId,
+		const FVector3d& BaryCoords,
+		const int UVLayer,
+		FVector2f& UVOut ) const = 0;
+
+	/**
+	 * Compute interpolated vertex color value inside triangle using barycentric coordinates
+	 * @param Mesh pointer to mesh to query
+	 * @param TriId index of triangle
+	 * @param BaryCoords 3 barycentric coordinates inside triangle
+	 * @param ColorOut resulting interpolated vertex color parameter value
+	 * @return true if a valid vertex color was returned, false otherwise
+	 */
+	virtual bool TriBaryInterpolateColor(
+		const void* Mesh,
+		const int32 TriId,
+		const FVector3d& BaryCoords,
+		FVector4f& ColorOut) const = 0;
+
+	/**
+	 * Compute interpolated tangent values inside triangle using barycentric coordinates
+	 * @param Mesh pointer to mesh to query
+	 * @param TriId index of triangle
+	 * @param BaryCoords 3 barycentric coordinates inside triangle
+	 * @param TangentX resulting interpolated tangentX value
+	 * @param TangentY resulting interpolated tangentY value
+	 * @return true if valid tangents were returned, false otherwise
+	 */
+	virtual bool TriBaryInterpolateTangents(
+		const void* Mesh,
+		const int32 TriId,
+		const FVector3d& BaryCoords,
+		FVector3d& TangentX,
+		FVector3d& TangentY) const = 0;
+
+	// TODO: Rework the curvature interface to only compute the minimal set of work.
+	/**
+	 * @param Mesh pointer to mesh to query
+	 * @param CurvatureCache [Out] the curvature data to populate
+	 */
+	virtual void GetCurvature(
+		const void* Mesh,
+		FMeshVertexCurvatureCache& CurvatureCache) const = 0;
+};
+
+
+/**
+ * DynamicMesh bake detail sampler for baking 1 detail mesh to 1 target mesh.
+ */
+class FMeshBakerDynamicMeshSampler : public IMeshBakerDetailSampler
+{
+public:
+	FMeshBakerDynamicMeshSampler(const FDynamicMesh3* Mesh, const FDynamicMeshAABBTree3* Spatial, const FMeshTangentsd* Tangents = nullptr)
+		: DetailMesh(Mesh), DetailSpatial(Spatial), DetailTangents(Tangents)
+	{
+	}
+
+	virtual void ProcessMeshes(TFunctionRef<void(const void*)> ProcessFn) const override
+	{
+		ProcessFn(DetailMesh);
+	}
+
+	virtual void SetColorMap(const void* Mesh, const FBakeDetailTexture& Map) override
+	{
+		DetailColorMap = Map;
+	}
+
+	virtual void SetNormalMap(const void* Mesh, const FBakeDetailTexture& Map) override
+	{
+		DetailNormalMap = Map;
+	}
+
+	virtual FBakeDetailTexture GetColorMap(const void* Mesh) const override
+	{
+		return DetailColorMap;
+	}
+	
+	virtual FBakeDetailTexture GetNormalMap(const void* Mesh) const override
+	{
+		return DetailNormalMap;
+	}
+
+	virtual bool SupportsIdentityCorrespondence() const override
+	{
+		return true;
+	}
+
+	virtual bool SupportsNearestPointCorrespondence() const override
+	{
+		return true;
+	}
+
+	virtual bool SupportsRaycastCorrespondence() const override
+	{
+		return true;
+	}
+
+	virtual const void* FindNearestTriangle(
+		const FVector3d& Point,
+		double& NearestDistSqrOut,
+		int& TriId,
+		FVector3d& TriBaryCoords,
+		const IMeshSpatial::FQueryOptions& Options = IMeshSpatial::FQueryOptions()) const override
+	{
+		TriId = DetailSpatial->FindNearestTriangle(Point, NearestDistSqrOut, Options);
+		if (DetailMesh->IsTriangle(TriId))
+		{
+			const FDistPoint3Triangle3d DistQuery = TMeshQueries<FDynamicMesh3>::TriangleDistance(*DetailMesh, TriId, Point);
+			TriBaryCoords = DistQuery.TriangleBaryCoords;
+		}
+		return DetailMesh;
+	}
+
+	virtual const void* FindNearestHitTriangle(
+		const FRay3d& Ray,
+		double& NearestT,
+		int& TriId,
+		FVector3d& TriBaryCoords,
+		const IMeshSpatial::FQueryOptions& Options = IMeshSpatial::FQueryOptions()) const override
+	{
+		const bool bHit = DetailSpatial->FindNearestHitTriangle(Ray, NearestT, TriId, Options);
+		if (bHit)
+		{
+			const FIntrRay3Triangle3d IntrQuery = TMeshQueries<FDynamicMesh3>::TriangleIntersection(*DetailMesh, TriId, Ray);
+			TriBaryCoords = IntrQuery.TriangleBaryCoords;
+			return DetailMesh;
+		}
+		return nullptr;
+	}
+
+	virtual bool TestAnyHitTriangle(
+		const FRay3d& Ray,
+		const IMeshSpatial::FQueryOptions& Options = IMeshSpatial::FQueryOptions()) const override
+	{
+		return DetailSpatial->TestAnyHitTriangle(Ray, Options);
+	}
+
+	virtual FAxisAlignedBox3d GetBounds() const override
+	{
+		return DetailMesh->GetBounds();
+	}
+
+	virtual bool IsTriangle(const void* Mesh, const int TriId) const override
+	{
+		const FDynamicMesh3* DynamicMesh = static_cast<const FDynamicMesh3*>(Mesh);
+		return DynamicMesh->IsTriangle(TriId);
+	}
+
+	virtual FIndex3i GetTriangle(const void* Mesh, const int TriId) const override
+	{
+		const FDynamicMesh3* DynamicMesh = static_cast<const FDynamicMesh3*>(Mesh);
+		return DynamicMesh->GetTriangle(TriId);
+	}
+
+	virtual FVector3d GetTriNormal(const void* Mesh, const int TriId) const override
+	{
+		const FDynamicMesh3* DynamicMesh = static_cast<const FDynamicMesh3*>(Mesh);
+		return DynamicMesh->GetTriNormal(TriId);
+	}
+
+	virtual int32 GetMaterialID(const void* Mesh, const int TriId) const override
+	{
+		const FDynamicMesh3* DynamicMesh = static_cast<const FDynamicMesh3*>(Mesh);
+		const FDynamicMeshMaterialAttribute* MaterialIDOverlay = DynamicMesh->Attributes()->GetMaterialID();
+		return MaterialIDOverlay ? MaterialIDOverlay->GetValue(TriId) : IndexConstants::InvalidID;
+	}
+
+	virtual FVector3d TriBaryInterpolatePoint(
+		const void* Mesh,
+		const int32 TriId,
+		const FVector3d& BaryCoords) const override
+	{
+		const FDynamicMesh3* DynamicMesh = static_cast<const FDynamicMesh3*>(Mesh);
+		return DynamicMesh->GetTriBaryPoint(TriId, BaryCoords.X, BaryCoords.Y, BaryCoords.Z);
+	}
+
+	virtual bool TriBaryInterpolateNormal(
+		const void* Mesh,
+		const int32 TriId,
+		const FVector3d& BaryCoords,
+		FVector3f& NormalOut) const override
+	{
+		const FDynamicMesh3* DynamicMesh = static_cast<const FDynamicMesh3*>(Mesh);
+		const FDynamicMeshNormalOverlay* NormalOverlay = DynamicMesh->Attributes()->PrimaryNormals();
+		const bool bValid = NormalOverlay && NormalOverlay->IsSetTriangle(TriId);
+		if (bValid)
+		{
+			FVector3d Normal;
+			NormalOverlay->GetTriBaryInterpolate(TriId, &BaryCoords.X, &Normal.X);
+			NormalOut = FVector3f(Normal);
+		}
+		return bValid;
+	}
+
+	virtual bool TriBaryInterpolateUV(
+		const void* Mesh,
+		const int32 TriId,
+		const FVector3d& BaryCoords,
+		const int UVLayer,
+		FVector2f& UVOut ) const override
+	{
+		const FDynamicMesh3* DynamicMesh = static_cast<const FDynamicMesh3*>(Mesh);
+		const FDynamicMeshUVOverlay* UVOverlay = DynamicMesh->Attributes()->GetUVLayer(UVLayer);
+		const bool bValid = UVOverlay && UVOverlay->IsSetTriangle(TriId);
+		if (bValid)
+		{
+			FVector2d UV;
+			UVOverlay->GetTriBaryInterpolate(TriId, &BaryCoords.X, &UV.X);
+			UVOut = FVector2f(UV);
+		}
+		return bValid;
+	}
+
+	virtual bool TriBaryInterpolateColor(
+		const void* Mesh,
+		const int32 TriId,
+		const FVector3d& BaryCoords,
+		FVector4f& ColorOut) const override
+	{
+		const FDynamicMesh3* DynamicMesh = static_cast<const FDynamicMesh3*>(Mesh);
+		const FDynamicMeshColorOverlay* ColorOverlay = DynamicMesh->Attributes()->PrimaryColors();
+		const bool bValid = ColorOverlay && ColorOverlay->IsSetTriangle(TriId);
+		if (bValid)
+		{
+			FVector4d Color;
+			ColorOverlay->GetTriBaryInterpolate(TriId, &BaryCoords.X, &Color.X);
+			ColorOut = FVector4f(Color);
+		}
+		return bValid;
+	}
+
+	virtual bool TriBaryInterpolateTangents(
+		const void* Mesh,
+		const int32 TriId,
+		const FVector3d& BaryCoords,
+		FVector3d& TangentX,
+		FVector3d& TangentY) const override
+	{
+		// TODO: Can we use the tangent overlay here for a unified interface? (PrimaryTangents/PrimaryBiTangents)
+		bool bSuccess = false;
+		if (DetailTangents)
+		{
+			DetailTangents->GetInterpolatedTriangleTangent(TriId, BaryCoords, TangentX, TangentY);
+			bSuccess = true;
+		}
+		return bSuccess;
+	}
+
+	virtual void GetCurvature(
+		const void* Mesh,
+		FMeshVertexCurvatureCache& CurvatureCache) const override
+	{
+		const FDynamicMesh3* DynamicMesh = static_cast<const FDynamicMesh3*>(Mesh);
+		CurvatureCache.BuildAll(*DynamicMesh);
+		ensure(CurvatureCache.Num() == DetailMesh->MaxVertexID());
+	}
+	
+protected:
+	const FDynamicMesh3* DetailMesh = nullptr;
+	const FDynamicMeshAABBTree3* DetailSpatial = nullptr;
+	const FMeshTangentsd* DetailTangents = nullptr;
+	FBakeDetailTexture DetailColorMap = FBakeDetailTexture(nullptr, 0);
+	FBakeDetailTexture DetailNormalMap = FBakeDetailTexture(nullptr, 0);
+};		
+	
+
 } // end namespace UE::Geometry
 } // end namespace UE
 	
\ No newline at end of file
diff --git a/Engine/Plugins/Runtime/GeometryProcessing/Source/DynamicMesh/Public/Sampling/MeshBaseBaker.h b/Engine/Plugins/Runtime/GeometryProcessing/Source/DynamicMesh/Public/Sampling/MeshBaseBaker.h
--- a/Engine/Plugins/Runtime/GeometryProcessing/Source/DynamicMesh/Public/Sampling/MeshBaseBaker.h
+++ b/Engine/Plugins/Runtime/GeometryProcessing/Source/DynamicMesh/Public/Sampling/MeshBaseBaker.h
@@ -6,13 +6,15 @@
 #include "DynamicMesh/DynamicMeshAttributeSet.h"
 #include "DynamicMesh/DynamicMeshAABBTree3.h"
 #include "DynamicMesh/MeshTangents.h"
-#include "Image/ImageBuilder.h"
+#include "Sampling/MeshBakerCommon.h"
+
 
 namespace UE
 {
 namespace Geometry
 {
 
+
 class FMeshBaseBaker
 {
 public:
@@ -36,22 +38,9 @@ public:
 	};
 
 	// Setters
-	void SetDetailMesh(const FDynamicMesh3* Mesh, const FDynamicMeshAABBTree3* Spatial)
-	{
-		DetailMesh = Mesh;
-		DetailSpatial = Spatial;
-	}
-	void SetDetailMeshTangents(const TSharedPtr<FMeshTangentsd, ESPMode::ThreadSafe>& Tangents)
+	void SetDetailSampler(IMeshBakerDetailSampler* Sampler)
 	{
-		DetailMeshTangents = Tangents;
-	}
-	void SetDetailMeshNormalMap(const TSharedPtr<TImageBuilder<FVector4f>, ESPMode::ThreadSafe>& NormalMap)
-	{
-		DetailMeshNormalMap = NormalMap;
-	}
-	void SetDetailMeshNormalUVLayer(int32 Layer)
-	{
-		DetailMeshNormalUVLayer = Layer;
+		DetailSampler = Sampler;
 	}
 	void SetTargetMesh(const FDynamicMesh3* Mesh)
 	{
@@ -95,40 +84,9 @@ public:
 	}
 
 	// DetailMesh Getters
-	const FDynamicMesh3* GetDetailMesh() const
-	{
-		return DetailMesh;
-	}
-	const FDynamicMeshAABBTree3* GetDetailMeshSpatial() const
+	const IMeshBakerDetailSampler* GetDetailSampler() const
 	{
-		return DetailSpatial;
-	}
-	const FDynamicMeshNormalOverlay* GetDetailMeshNormals() const
-	{
-		check(DetailMesh && DetailMesh->HasAttributes());
-		return DetailMesh->Attributes()->PrimaryNormals();
-	}
-	const FDynamicMeshUVOverlay* GetDetailMeshUVs(int32 UVLayerIn=0) const
-	{
-		check(DetailMesh && DetailMesh->HasAttributes());
-		return DetailMesh->Attributes()->GetUVLayer(UVLayerIn);
-	}
-	const FDynamicMeshColorOverlay* GetDetailMeshColors() const
-	{
-		check(DetailMesh && DetailMesh->HasAttributes());
-		return DetailMesh->Attributes()->PrimaryColors();
-	}
-	const FMeshTangentsd* GetDetailMeshTangents() const
-	{
-		return DetailMeshTangents.Get();
-	}
-	const TImageBuilder<FVector4f>* GetDetailMeshNormalMap() const
-	{
-		return DetailMeshNormalMap.Get();
-	}
-	int32 GetDetailMeshNormalUVLayer() const
-	{
-		return DetailMeshNormalUVLayer;
+		return DetailSampler;
 	}
 
 	// Other Getters
@@ -149,8 +107,6 @@ protected:
 	const FDynamicMesh3* DetailMesh = nullptr;
 	const FDynamicMeshAABBTree3* DetailSpatial = nullptr;
 	TSharedPtr<FMeshTangentsd, ESPMode::ThreadSafe> DetailMeshTangents;
-	TSharedPtr<UE::Geometry::TImageBuilder<FVector4f>, ESPMode::ThreadSafe> DetailMeshNormalMap;
-	int32 DetailMeshNormalUVLayer = 0;
 	
 	const FDynamicMesh3* TargetMesh = nullptr;
 	TSharedPtr<FMeshTangentsd, ESPMode::ThreadSafe> TargetMeshTangents;
@@ -158,9 +114,11 @@ protected:
 	int32 UVLayer = 0;
 	double Thickness = 3.0;
 	ECorrespondenceStrategy CorrespondenceStrategy = ECorrespondenceStrategy::RaycastStandard;
+
+	IMeshBakerDetailSampler* DetailSampler = nullptr;
 };
-	
 
+	
 } // end namespace UE::Geometry
 } // end namespace UE
 	
\ No newline at end of file
diff --git a/Engine/Plugins/Runtime/GeometryProcessing/Source/DynamicMesh/Public/Sampling/MeshCurvatureMapEvaluator.h b/Engine/Plugins/Runtime/GeometryProcessing/Source/DynamicMesh/Public/Sampling/MeshCurvatureMapEvaluator.h
--- a/Engine/Plugins/Runtime/GeometryProcessing/Source/DynamicMesh/Public/Sampling/MeshCurvatureMapEvaluator.h
+++ b/Engine/Plugins/Runtime/GeometryProcessing/Source/DynamicMesh/Public/Sampling/MeshCurvatureMapEvaluator.h
@@ -9,6 +9,7 @@ namespace UE
 namespace Geometry
 {
 
+class IMeshBakerDetailSampler;
 class FMeshVertexCurvatureCache;
 
 /**
@@ -66,11 +67,11 @@ public:
 	static void EvaluateColor(const int DataIdx, float*& In, FVector4f& Out, void* EvalData);
 
 	/** Populate Curvatures member if valid data has not been provided */
-	void CacheDetailCurvatures(const FDynamicMesh3* DetailMesh);
+	void CacheDetailCurvatures();
 
 protected:
 	// Cached data
-	const FDynamicMesh3* DetailMesh = nullptr;
+	const IMeshBakerDetailSampler* DetailSampler = nullptr;
 	double MinPreClamp = -TNumericLimits<double>::Max();
 	double MaxPreClamp = TNumericLimits<double>::Max();
 	FInterval1d ClampRange;
@@ -79,11 +80,11 @@ protected:
 	FVector3f PositiveColor;
 
 protected:
-	double GetCurvature(int32 vid);
-	void GetColorMapRange(FVector3f& NegativeColor, FVector3f& ZeroColor, FVector3f& PositiveColor);
+	double GetCurvature(int32 vid) const;
+	void GetColorMapRange(FVector3f& NegativeColorOut, FVector3f& ZeroColorOut, FVector3f& PositiveColorOut) const;
 
 private:
-	double SampleFunction(const FCorrespondenceSample& Sample);
+	double SampleFunction(const FCorrespondenceSample& SampleData) const;
 };
 
 } // end namespace UE::Geometry
diff --git a/Engine/Plugins/Runtime/GeometryProcessing/Source/DynamicMesh/Public/Sampling/MeshMapBaker.h b/Engine/Plugins/Runtime/GeometryProcessing/Source/DynamicMesh/Public/Sampling/MeshMapBaker.h
--- a/Engine/Plugins/Runtime/GeometryProcessing/Source/DynamicMesh/Public/Sampling/MeshMapBaker.h
+++ b/Engine/Plugins/Runtime/GeometryProcessing/Source/DynamicMesh/Public/Sampling/MeshMapBaker.h
@@ -96,7 +96,7 @@ protected:
 	const bool bParallel = true;
 
 	FDynamicMesh3 FlatMesh;
-	TMeshSurfaceUVSampler<FMeshMapEvaluator::FCorrespondenceSample> DetailMeshSampler;
+	TMeshSurfaceUVSampler<FMeshMapEvaluator::FCorrespondenceSample> DetailCorrespondenceSampler;
 
 	FImageDimensions Dimensions = FImageDimensions(128, 128);
 
diff --git a/Engine/Plugins/Runtime/GeometryProcessing/Source/DynamicMesh/Public/Sampling/MeshMapEvaluator.h b/Engine/Plugins/Runtime/GeometryProcessing/Source/DynamicMesh/Public/Sampling/MeshMapEvaluator.h
--- a/Engine/Plugins/Runtime/GeometryProcessing/Source/DynamicMesh/Public/Sampling/MeshMapEvaluator.h
+++ b/Engine/Plugins/Runtime/GeometryProcessing/Source/DynamicMesh/Public/Sampling/MeshMapEvaluator.h
@@ -49,10 +49,11 @@ public:
 	struct FCorrespondenceSample
 	{
 		FMeshUVSampleInfo BaseSample;
-		FVector3d BaseNormal;
+		FVector3d BaseNormal = FVector3d::Zero();
 
-		int32 DetailTriID;
-		FVector3d DetailBaryCoords;
+		const void* DetailMesh = nullptr;
+		int32 DetailTriID = IndexConstants::InvalidID;
+		FVector3d DetailBaryCoords = FVector3d::Zero();
 	};
 
 	/**
diff --git a/Engine/Plugins/Runtime/GeometryProcessing/Source/DynamicMesh/Public/Sampling/MeshNormalMapEvaluator.h b/Engine/Plugins/Runtime/GeometryProcessing/Source/DynamicMesh/Public/Sampling/MeshNormalMapEvaluator.h
--- a/Engine/Plugins/Runtime/GeometryProcessing/Source/DynamicMesh/Public/Sampling/MeshNormalMapEvaluator.h
+++ b/Engine/Plugins/Runtime/GeometryProcessing/Source/DynamicMesh/Public/Sampling/MeshNormalMapEvaluator.h
@@ -2,9 +2,11 @@
 
 #pragma once
 
+#include "Templates/Tuple.h"
+#include "Containers/Map.h"
 #include "Sampling/MeshMapEvaluator.h"
+#include "Sampling/MeshBakerCommon.h"
 #include "DynamicMesh/MeshTangents.h"
-#include "Image/ImageBuilder.h"
 
 namespace UE
 {
@@ -32,11 +34,10 @@ public:
 
 protected:
 	// Cached data
-	const FDynamicMesh3* DetailMesh = nullptr;
-	const FDynamicMeshNormalOverlay* DetailNormalOverlay = nullptr;
-	const FDynamicMeshUVOverlay* DetailUVOverlay = nullptr;
-	const TMeshTangents<double>* DetailMeshTangents = nullptr;
-	const TImageBuilder<FVector4f>* DetailMeshNormalMap = nullptr;
+	const IMeshBakerDetailSampler* DetailSampler = nullptr;
+	using FDetailNormalTexture = IMeshBakerDetailSampler::FBakeDetailTexture;
+	using FDetailNormalTextureMap = TMap<const void*, FDetailNormalTexture>;
+	FDetailNormalTextureMap DetailNormalTextures;
 	const TMeshTangents<double>* BaseMeshTangents = nullptr;
 
 	FVector3f DefaultNormal = FVector3f::UnitZ();
@@ -44,8 +45,6 @@ protected:
 private:
 	template <bool bUseDetailNormalMap>
 	FVector3f SampleFunction(const FCorrespondenceSample& SampleData) const;
-
-	FVector4f SampleNormalMapFunction(const FVector2d& UVCoord) const;
 };
 
 } // end namespace UE::Geometry
diff --git a/Engine/Plugins/Runtime/GeometryProcessing/Source/DynamicMesh/Public/Sampling/MeshOcclusionMapEvaluator.h b/Engine/Plugins/Runtime/GeometryProcessing/Source/DynamicMesh/Public/Sampling/MeshOcclusionMapEvaluator.h
--- a/Engine/Plugins/Runtime/GeometryProcessing/Source/DynamicMesh/Public/Sampling/MeshOcclusionMapEvaluator.h
+++ b/Engine/Plugins/Runtime/GeometryProcessing/Source/DynamicMesh/Public/Sampling/MeshOcclusionMapEvaluator.h
@@ -5,7 +5,6 @@
 #include "Sampling/MeshMapEvaluator.h"
 #include "Templates/Tuple.h"
 #include "Misc/EnumClassFlags.h"
-#include "Math/RandomStream.h"
 #include "DynamicMesh/MeshTangents.h"
 
 namespace UE
@@ -13,6 +12,8 @@ namespace UE
 namespace Geometry
 {
 
+class IMeshBakerDetailSampler;	
+
 enum class EMeshOcclusionMapType : uint8
 {
 	None = 0,
@@ -72,9 +73,7 @@ public:
 
 protected:
 	// Cached data
-	const FDynamicMesh3* DetailMesh = nullptr;
-	const FDynamicMeshAABBTree3* DetailSpatial = nullptr;
-	const FDynamicMeshNormalOverlay* DetailNormalOverlay = nullptr;
+	const IMeshBakerDetailSampler* DetailSampler = nullptr;
 	const TMeshTangents<double>* BaseMeshTangents = nullptr;
 	double BiasDotThreshold = 0.25;
 	TArray<FVector3d> RayDirections;
diff --git a/Engine/Plugins/Runtime/GeometryProcessing/Source/DynamicMesh/Public/Sampling/MeshPropertyMapEvaluator.h b/Engine/Plugins/Runtime/GeometryProcessing/Source/DynamicMesh/Public/Sampling/MeshPropertyMapEvaluator.h
--- a/Engine/Plugins/Runtime/GeometryProcessing/Source/DynamicMesh/Public/Sampling/MeshPropertyMapEvaluator.h
+++ b/Engine/Plugins/Runtime/GeometryProcessing/Source/DynamicMesh/Public/Sampling/MeshPropertyMapEvaluator.h
@@ -3,7 +3,7 @@
 #pragma once
 
 #include "Sampling/MeshMapEvaluator.h"
-#include "DynamicMesh/MeshTangents.h"
+#include "Sampling/MeshBakerCommon.h"
 #include "Image/ImageBuilder.h"
 
 namespace UE
@@ -11,6 +11,8 @@ namespace UE
 namespace Geometry
 {
 
+class IMeshBakerDetailSampler;	
+
 enum class EMeshPropertyMapType
 {
 	Position = 1,
@@ -45,29 +47,27 @@ public:
 
 protected:
 	// Cached data
-	const FDynamicMesh3* DetailMesh = nullptr;
-	const FDynamicMeshNormalOverlay* DetailNormalOverlay = nullptr;
-	const FDynamicMeshUVOverlay* DetailUVOverlay = nullptr;
-	const FDynamicMeshColorOverlay* DetailColorOverlay = nullptr;
-	const TMeshTangents<double>* DetailMeshTangents = nullptr;
-	const TImageBuilder<FVector4f>* DetailMeshNormalMap = nullptr;
+	const IMeshBakerDetailSampler* DetailSampler = nullptr;
+	using FDetailNormalTexture = IMeshBakerDetailSampler::FBakeDetailTexture;
+	using FDetailNormalTextureMap = TMap<const void*, FDetailNormalTexture>;
+	FDetailNormalTextureMap DetailNormalTextures;
 	FAxisAlignedBox3d Bounds;
 	FVector3f DefaultValue = FVector3f::Zero();
 
 private:
-	FVector3f NormalToColor(const FVector3d Normal)
+	static FVector3f NormalToColor(const FVector3d Normal) 
 	{
 		return (FVector3f)((Normal + FVector3d::One()) * 0.5);
 	}
 
-	FVector3f UVToColor(const FVector2d UV)
+	static FVector3f UVToColor(const FVector2f UV)
 	{
-		double X = FMathd::Clamp(UV.X, 0.0, 1.0);
-		double Y = FMathd::Clamp(UV.Y, 0.0, 1.0);
-		return (FVector3f)FVector3d(X, Y, 0);
+		const float X = FMathf::Clamp(UV.X, 0.0, 1.0);
+		const float Y = FMathf::Clamp(UV.Y, 0.0, 1.0);
+		return FVector3f(X, Y, 0);
 	}
 
-	FVector3f PositionToColor(const FVector3d Position, const FAxisAlignedBox3d SafeBounds)
+	static FVector3f PositionToColor(const FVector3d Position, const FAxisAlignedBox3d SafeBounds)
 	{
 		double X = (Position.X - SafeBounds.Min.X) / SafeBounds.Width();
 		double Y = (Position.Y - SafeBounds.Min.Y) / SafeBounds.Height();
@@ -76,9 +76,7 @@ private:
 	}
 
 	template <bool bUseDetailNormalMap>
-	FVector3f SampleFunction(const FCorrespondenceSample& Sample);
-
-	FVector4f SampleNormalMapFunction(const FVector2d& UVCoord) const;
+	FVector3f SampleFunction(const FCorrespondenceSample& SampleData) const;
 };
 
 } // end namespace UE::Geometry
diff --git a/Engine/Plugins/Runtime/GeometryProcessing/Source/DynamicMesh/Public/Sampling/MeshResampleImageEvaluator.h b/Engine/Plugins/Runtime/GeometryProcessing/Source/DynamicMesh/Public/Sampling/MeshResampleImageEvaluator.h
--- a/Engine/Plugins/Runtime/GeometryProcessing/Source/DynamicMesh/Public/Sampling/MeshResampleImageEvaluator.h
+++ b/Engine/Plugins/Runtime/GeometryProcessing/Source/DynamicMesh/Public/Sampling/MeshResampleImageEvaluator.h
@@ -3,6 +3,7 @@
 #pragma once
 
 #include "Sampling/MeshMapEvaluator.h"
+#include "Sampling/MeshBaseBaker.h"
 #include "Image/ImageBuilder.h"
 
 namespace UE
@@ -10,15 +11,14 @@ namespace UE
 namespace Geometry
 {
 
+class IMeshBakerDetailSampler;	
+
 /**
  * A mesh evaluator for sampling 2D texture data. 
  */
 class DYNAMICMESH_API FMeshResampleImageEvaluator : public FMeshMapEvaluator
 {
 public:
-	TFunction<FVector4f(FVector2d)> SampleFunction = [](FVector2d Position) { return FVector4f::Zero(); };
-	const FDynamicMeshUVOverlay* DetailUVOverlay = nullptr;
-
 	FVector4f DefaultColor = FVector4f(0, 0, 0, 1);
 
 public:
@@ -36,12 +36,14 @@ public:
 
 protected:
 	// Cached data
-	const FDynamicMesh3* DetailMesh = nullptr;
+	const IMeshBakerDetailSampler* DetailSampler = nullptr;
 
 private:
-	FVector4f ImageSampleFunction(const FCorrespondenceSample& Sample);
+	FVector4f ImageSampleFunction(const FCorrespondenceSample& Sample) const;
 };
 
+
+// TODO: Add support for multiple color maps per mesh in IMeshBakerDetailSampler for proper MultiTexture support.	
 /**
  * A mesh evaluator for sampling multiple 2D textures by material ID
  */
@@ -49,6 +51,7 @@ class DYNAMICMESH_API FMeshMultiResampleImageEvaluator : public FMeshResampleIma
 {
 public:
 	TMap<int32, TSharedPtr<UE::Geometry::TImageBuilder<FVector4f>, ESPMode::ThreadSafe>> MultiTextures;
+	int DetailUVLayer = 0;
 
 public:
 	// Begin FMeshMapEvaluator interface
@@ -59,12 +62,8 @@ public:
 
 	static void EvaluateSampleMulti(float*& Out, const FCorrespondenceSample& Sample, void* EvalData);
 
-protected:
-	const FDynamicMeshMaterialAttribute* DetailMaterialIDAttrib = nullptr;
-	bool bValidDetailMesh = false;
-
 private:
-	FVector4f ImageSampleFunction(const FCorrespondenceSample& Sample);
+	FVector4f ImageSampleFunction(const FCorrespondenceSample& Sample) const;
 };
 
 } // end namespace UE::Geometry
