GeometryProcessing:
- add FPolygroupLayer struct which is basically an identifier for a Polygroup layer, ie either the default mesh layer or one of the extended attribute layers
- add FPolygroupsGenerator::FindPolygroupsFromHardNormalSeams()
- add option to FPolygroupsGenerator::FindPolygroupsFromFurthestPointSampling() to constrain generation relative to an input polygroup set, ie only subdivide existing polygroups
- FMeshRegionGraph::BuildFromComponents now takes optional are-triangles-connected predicate, to allow links in graph to be disabled
- FPatchBasedMeshUVGenerator can now optionally generate new UV islands relative to an input Polygoup set, via new GroupConstraint input

#rnx

diff --git a/Engine/Plugins/Runtime/GeometryProcessing/Source/DynamicMesh/Private/Parameterization/MeshRegionGraph.cpp b/Engine/Plugins/Runtime/GeometryProcessing/Source/DynamicMesh/Private/Parameterization/MeshRegionGraph.cpp
--- a/Engine/Plugins/Runtime/GeometryProcessing/Source/DynamicMesh/Private/Parameterization/MeshRegionGraph.cpp
+++ b/Engine/Plugins/Runtime/GeometryProcessing/Source/DynamicMesh/Private/Parameterization/MeshRegionGraph.cpp
@@ -8,7 +8,11 @@
 using namespace UE::Geometry;
 
 
-void FMeshRegionGraph::BuildFromComponents(const FDynamicMesh3& Mesh, const FMeshConnectedComponents& Components, TFunctionRef<int32(int32)> ExternalIDFunc)
+void FMeshRegionGraph::BuildFromComponents(
+	const FDynamicMesh3& Mesh, 
+	const FMeshConnectedComponents& Components, 
+	TFunctionRef<int32(int32)> ExternalIDFunc,
+	TFunctionRef<bool(int32, int32)> TrisConnectedPredicate)
 {
 	int32 N = Components.Num();
 	Regions.SetNum(N);
@@ -23,7 +27,15 @@ void FMeshRegionGraph::BuildFromComponents(const FDynamicMesh3& Mesh, const FMes
 		for (int32 tid : Components[k].Indices)
 		{
 			TriangleToRegionMap[tid] = k;
-			TriangleNbrTris[tid] = Mesh.GetTriNeighbourTris(tid);
+			FIndex3i NbrTris = Mesh.GetTriNeighbourTris(tid);;
+			for (int j = 0; j < 3; ++j)
+			{
+				if (NbrTris[j] >= 0 && TrisConnectedPredicate(tid, NbrTris[j]) == false)
+				{
+					NbrTris[j] = -1;
+				}
+			}
+			TriangleNbrTris[tid] = NbrTris;
 		}
 	}
 
diff --git a/Engine/Plugins/Runtime/GeometryProcessing/Source/DynamicMesh/Private/Parameterization/PatchBasedMeshUVGenerator.cpp b/Engine/Plugins/Runtime/GeometryProcessing/Source/DynamicMesh/Private/Parameterization/PatchBasedMeshUVGenerator.cpp
--- a/Engine/Plugins/Runtime/GeometryProcessing/Source/DynamicMesh/Private/Parameterization/PatchBasedMeshUVGenerator.cpp
+++ b/Engine/Plugins/Runtime/GeometryProcessing/Source/DynamicMesh/Private/Parameterization/PatchBasedMeshUVGenerator.cpp
@@ -137,7 +137,9 @@ bool FPatchBasedMeshUVGenerator::ComputeInitialMeshPatches(
 	FPolygroupsGenerator::EWeightingType WeightType = (this->bNormalWeightedPatches) ?
 		FPolygroupsGenerator::EWeightingType::NormalDeviation : FPolygroupsGenerator::EWeightingType::None;
 	FVector3d GeneratorParams(this->PatchNormalWeight, 1.0, 1.0);
-	bool bOK = Generator.FindPolygroupsFromFurthestPointSampling(this->TargetPatchCount, WeightType, GeneratorParams);
+
+	bool bOK = Generator.FindPolygroupsFromFurthestPointSampling(this->TargetPatchCount, WeightType, GeneratorParams, this->GroupConstraint);
+
 	if (!bOK)
 	{
 		return false;
@@ -336,7 +338,15 @@ bool FPatchBasedMeshUVGenerator::ComputeIslandsByRegionMerging(
 	}
 
 	FMeshRegionGraph RegionGraph;
-	RegionGraph.BuildFromComponents(TargetMesh, ConnectedComponents, [&](int32 k) { return k; } );
+	if (GroupConstraint)
+	{
+		RegionGraph.BuildFromComponents(TargetMesh, ConnectedComponents, [&](int32 k) { return k; },
+			[this](int32 a, int32 b) { return GroupConstraint->GetGroup(a) == GroupConstraint->GetGroup(b); });
+	}
+	else
+	{
+		RegionGraph.BuildFromComponents(TargetMesh, ConnectedComponents, [&](int32 k) { return k; });
+	}
 
 	if (Progress && Progress->Cancelled())
 	{
diff --git a/Engine/Plugins/Runtime/GeometryProcessing/Source/DynamicMesh/Private/Polygroups/PolygroupSet.cpp b/Engine/Plugins/Runtime/GeometryProcessing/Source/DynamicMesh/Private/Polygroups/PolygroupSet.cpp
--- a/Engine/Plugins/Runtime/GeometryProcessing/Source/DynamicMesh/Private/Polygroups/PolygroupSet.cpp
+++ b/Engine/Plugins/Runtime/GeometryProcessing/Source/DynamicMesh/Private/Polygroups/PolygroupSet.cpp
@@ -7,6 +7,29 @@
 using namespace UE::Geometry;
 
 
+bool FPolygroupLayer::CheckExists(const FDynamicMesh3* Mesh)
+{
+	if (Mesh)
+	{
+		if (bIsDefaultLayer)
+		{
+			if (Mesh->HasTriangleGroups())
+			{
+				return true;
+			}
+		}
+		else
+		{
+			if (LayerIndex >= 0 && Mesh->HasAttributes() && LayerIndex < Mesh->Attributes()->NumPolygroupLayers())
+			{
+				return true;
+			}
+		}
+	}
+	return false;
+}
+
+
 FPolygroupSet::FPolygroupSet(const FPolygroupSet* CopyIn)
 {
 	Mesh = CopyIn->Mesh;
@@ -22,6 +45,30 @@ FPolygroupSet::FPolygroupSet(const FDynamicMesh3* MeshIn)
 	RecalculateMaxGroupID();
 }
 
+/** Initialize a PolygroupSet for the given Mesh, and standard triangle group layer */
+FPolygroupSet::FPolygroupSet(const FDynamicMesh3* MeshIn, FPolygroupLayer GroupLayer)
+{
+	Mesh = MeshIn;
+	GroupLayerIndex = -1;
+	if (! GroupLayer.bIsDefaultLayer )
+	{
+		if (ensure(Mesh->Attributes()))
+		{
+			if (GroupLayer.LayerIndex < Mesh->Attributes()->NumPolygroupLayers())
+			{
+				PolygroupAttrib = Mesh->Attributes()->GetPolygroupLayer(GroupLayer.LayerIndex);
+				GroupLayerIndex = GroupLayer.LayerIndex;
+			}
+		}
+		if (GroupLayerIndex == -1)
+		{
+			ensureMsgf(false, TEXT("FPolygroupSet: Attribute index missing!"));
+		}
+	}
+	RecalculateMaxGroupID();
+}
+
+
 FPolygroupSet::FPolygroupSet(const FDynamicMesh3* MeshIn, const FDynamicMeshPolygroupAttribute* PolygroupAttribIn)
 {
 	Mesh = MeshIn;
diff --git a/Engine/Plugins/Runtime/GeometryProcessing/Source/DynamicMesh/Private/Polygroups/PolygroupsGenerator.cpp b/Engine/Plugins/Runtime/GeometryProcessing/Source/DynamicMesh/Private/Polygroups/PolygroupsGenerator.cpp
--- a/Engine/Plugins/Runtime/GeometryProcessing/Source/DynamicMesh/Private/Polygroups/PolygroupsGenerator.cpp
+++ b/Engine/Plugins/Runtime/GeometryProcessing/Source/DynamicMesh/Private/Polygroups/PolygroupsGenerator.cpp
@@ -7,6 +7,7 @@
 #include "Parameterization/MeshDijkstra.h"
 #include "Parameterization/IncrementalMeshDijkstra.h"
 #include "Parameterization/MeshRegionGraph.h"
+#include "Async/ParallelFor.h"
 
 #include "Curve/DynamicGraph3.h"
 
@@ -22,9 +23,13 @@ FPolygroupsGenerator::FPolygroupsGenerator(FDynamicMesh3* MeshIn)
 
 
 
-bool FPolygroupsGenerator::FindPolygroupsFromUVIslands()
+bool FPolygroupsGenerator::FindPolygroupsFromUVIslands(int32 UVLayer)
 {
-	const FDynamicMeshUVOverlay* UV = Mesh->Attributes()->GetUVLayer(0);
+	const FDynamicMeshUVOverlay* UV = Mesh->Attributes()->GetUVLayer(UVLayer);
+	if (!UV)
+	{
+		return false;
+	}
 
 	FMeshConnectedComponents Components(Mesh);
 	Components.FindConnectedTriangles([&UV](int32 TriIdx0, int32 TriIdx1)
@@ -51,6 +56,38 @@ bool FPolygroupsGenerator::FindPolygroupsFromUVIslands()
 }
 
 
+bool FPolygroupsGenerator::FindPolygroupsFromHardNormalSeams()
+{
+	const FDynamicMeshNormalOverlay* Normals = Mesh->Attributes()->GetNormalLayer(0);
+	if (!Normals)
+	{
+		return false;
+	}
+
+	FMeshConnectedComponents Components(Mesh);
+	Components.FindConnectedTriangles([Normals](int32 TriIdx0, int32 TriIdx1)
+	{
+		return Normals->AreTrianglesConnected(TriIdx0, TriIdx1);
+	});
+
+	int32 NumComponents = Components.Num();
+	for (int32 ci = 0; ci < NumComponents; ++ci)
+	{
+		FoundPolygroups.Add(Components.GetComponent(ci).Indices);
+	}
+
+	if (bApplyPostProcessing)
+	{
+		PostProcessPolygroups(false);
+	}
+	if (bCopyToMesh)
+	{
+		CopyPolygroupsToMesh();
+	}
+
+	return (FoundPolygroups.Num() > 0);
+}
+
 
 bool FPolygroupsGenerator::FindPolygroupsFromConnectedTris()
 {
@@ -147,7 +184,6 @@ bool FPolygroupsGenerator::FindPolygroupsFromFaceNormals(double DotTolerance)
 
 
 
-
 /**
  * Dual graph of mesh faces, ie graph of edges across faces between face centers.
  * Normals and Areas are tracked for each point.
@@ -172,7 +208,10 @@ public:
 	}
 
 	/** Build a Face Dual Graph for a triangle mesh */
-	static void MakeFaceDualGraphForMesh(FDynamicMesh3* Mesh, FMeshFaceDualGraph& FaceGraph)
+	static void MakeFaceDualGraphForMesh(
+		FDynamicMesh3* Mesh, 
+		FMeshFaceDualGraph& FaceGraph,
+		TFunctionRef<bool(int,int)> TrisConnectedPredicate )
 	{
 		// if not true, code below needs updating
 		check(Mesh->IsCompactT());
@@ -192,7 +231,7 @@ public:
 			FIndex3i NbrT = Mesh->GetTriNeighbourTris(tid);
 			for (int32 j = 0; j < 3; ++j)
 			{
-				if (Mesh->IsTriangle(NbrT[j]))
+				if (Mesh->IsTriangle(NbrT[j]) && TrisConnectedPredicate(tid, NbrT[j]) )
 				{
 					FaceGraph.AppendEdge(tid, NbrT[j]);
 				}
@@ -207,12 +246,23 @@ public:
 
 
 
-bool FPolygroupsGenerator::FindPolygroupsFromFurthestPointSampling(int32 NumPoints, EWeightingType WeightingType, FVector3d WeightingCoeffs)
+bool FPolygroupsGenerator::FindPolygroupsFromFurthestPointSampling(
+	int32 NumPoints, 
+	EWeightingType WeightingType, 
+	FVector3d WeightingCoeffs,
+	FPolygroupSet* StartingGroups)
 {
 	NumPoints = FMath::Min(NumPoints, Mesh->VertexCount());
 
+	// cannot seem to use auto or TUniqueFunction here...
+	TFunction<bool(int,int)> TrisConnectedPredicate = [](int, int) -> bool { return true; };
+	if (StartingGroups != nullptr)
+	{
+		TrisConnectedPredicate = [StartingGroups](int a, int b) -> bool { return StartingGroups->GetGroup(a) == StartingGroups->GetGroup(b) ? true : false; };
+	}
+
 	FMeshFaceDualGraph FaceGraph;
-	FMeshFaceDualGraph::MakeFaceDualGraphForMesh(Mesh, FaceGraph);
+	FMeshFaceDualGraph::MakeFaceDualGraphForMesh(Mesh, FaceGraph, TrisConnectedPredicate);
 
 	TIncrementalMeshDijkstra<FMeshFaceDualGraph> FurthestPoints(&FaceGraph);
 
@@ -221,7 +271,7 @@ bool FPolygroupsGenerator::FindPolygroupsFromFurthestPointSampling(int32 NumPoin
 	// need to add at least one seed point for each mesh connected component, so that all triangles are assigned a group
 	// TODO: two seed points for components that have no boundary?
 	FMeshConnectedComponents Components(Mesh);
-	Components.FindConnectedTriangles();
+	Components.FindConnectedTriangles(TrisConnectedPredicate);
 	int32 NumConnected = Components.Num();
 	for (int32 k = 0; k < NumConnected; ++k)
 	{
diff --git a/Engine/Plugins/Runtime/GeometryProcessing/Source/DynamicMesh/Public/Parameterization/MeshRegionGraph.h b/Engine/Plugins/Runtime/GeometryProcessing/Source/DynamicMesh/Public/Parameterization/MeshRegionGraph.h
--- a/Engine/Plugins/Runtime/GeometryProcessing/Source/DynamicMesh/Public/Parameterization/MeshRegionGraph.h
+++ b/Engine/Plugins/Runtime/GeometryProcessing/Source/DynamicMesh/Public/Parameterization/MeshRegionGraph.h
@@ -33,7 +33,8 @@ public:
 	 */
 	void BuildFromComponents(const FDynamicMesh3& Mesh, 
 							 const FMeshConnectedComponents& Components,
-							 TFunctionRef<int32(int32)> ExternalIDFunc);
+							 TFunctionRef<int32(int32)> ExternalIDFunc,
+							 TFunctionRef<bool(int32, int32)> TrisConnectedPredicate = [](int a, int b) { return true; } );
 
 	/**
 	 * Build a region graph for a Mesh using the given triangle sets
diff --git a/Engine/Plugins/Runtime/GeometryProcessing/Source/DynamicMesh/Public/Parameterization/PatchBasedMeshUVGenerator.h b/Engine/Plugins/Runtime/GeometryProcessing/Source/DynamicMesh/Public/Parameterization/PatchBasedMeshUVGenerator.h
--- a/Engine/Plugins/Runtime/GeometryProcessing/Source/DynamicMesh/Public/Parameterization/PatchBasedMeshUVGenerator.h
+++ b/Engine/Plugins/Runtime/GeometryProcessing/Source/DynamicMesh/Public/Parameterization/PatchBasedMeshUVGenerator.h
@@ -5,6 +5,7 @@
 #include "DynamicMesh/DynamicMesh3.h"
 #include "DynamicMesh/DynamicMeshAttributeSet.h"
 #include "Selections/MeshConnectedComponents.h"
+#include "Polygroups/PolygroupSet.h"
 #include "Util/ProgressCancel.h"
 
 namespace UE
@@ -26,6 +27,14 @@ class DYNAMICMESH_API FPatchBasedMeshUVGenerator
 {
 public:
 
+	//
+	// If this is set, then the output patches will respect the borders of the
+	// input polygroups, eg as if they were disconnected components.
+	// Must be a PolygroupSet for the TargetMesh passed to functions below
+	//
+	FPolygroupSet* GroupConstraint = nullptr;
+
+
 	//
 	// Parameters for patch generation (ComputeInitialMeshPatches)
 	//
diff --git a/Engine/Plugins/Runtime/GeometryProcessing/Source/DynamicMesh/Public/Polygroups/PolygroupSet.h b/Engine/Plugins/Runtime/GeometryProcessing/Source/DynamicMesh/Public/Polygroups/PolygroupSet.h
--- a/Engine/Plugins/Runtime/GeometryProcessing/Source/DynamicMesh/Public/Polygroups/PolygroupSet.h
+++ b/Engine/Plugins/Runtime/GeometryProcessing/Source/DynamicMesh/Public/Polygroups/PolygroupSet.h
@@ -12,6 +12,27 @@ namespace UE
 namespace Geometry
 {
 
+/**
+ * FPolygroupLayer represents a polygroup set on a FDynamicMesh3, which supports a "default"
+ * group set stored on the mesh, and then N extended group layers stored in the mesh AttributeSet.
+ * This struct can represent either.
+ */
+struct DYNAMICMESH_API FPolygroupLayer
+{
+	/** If true, layer is the default FDynamicMesh3 triangle groups layer */
+	bool bIsDefaultLayer = true;
+	/** If bIsDefaultLayer is false, this is the index of the AttributeSet Polygroup Layer */
+	int32 LayerIndex = -1;
+
+	/** Construct a FPolygroupLayer for the default layer */
+	static FPolygroupLayer Default() { return FPolygroupLayer{ true, -1 }; }
+	/** Construct a FPolygroupLayer for an extended layer */
+	static FPolygroupLayer Layer(int32 Index) { return FPolygroupLayer{ false, Index }; }
+
+	/** @return true if the specified layer (default or extended) exist and is initialized on the given Mesh */
+	bool CheckExists(const FDynamicMesh3* Mesh);
+};
+
 
 /**
  * Polygroup sets can be stored in multiple places. The default location is in the per-triangle group integer stored
@@ -34,6 +55,9 @@ struct DYNAMICMESH_API FPolygroupSet
 	/** Initialize a PolygroupSet for the given Mesh, and standard triangle group layer */
 	explicit FPolygroupSet(const FDynamicMesh3* MeshIn);
 
+	/** Initialize a PolygroupSet for the given Mesh, and standard triangle group layer */
+	explicit FPolygroupSet(const FDynamicMesh3* MeshIn, FPolygroupLayer GroupLayer);
+
 	/** Initialize a PolygroupSet for given Mesh and specific Polygroup attribute layer */
 	explicit FPolygroupSet(const FDynamicMesh3* MeshIn, const FDynamicMeshPolygroupAttribute* PolygroupAttribIn);
 
diff --git a/Engine/Plugins/Runtime/GeometryProcessing/Source/DynamicMesh/Public/Polygroups/PolygroupsGenerator.h b/Engine/Plugins/Runtime/GeometryProcessing/Source/DynamicMesh/Public/Polygroups/PolygroupsGenerator.h
--- a/Engine/Plugins/Runtime/GeometryProcessing/Source/DynamicMesh/Public/Polygroups/PolygroupsGenerator.h
+++ b/Engine/Plugins/Runtime/GeometryProcessing/Source/DynamicMesh/Public/Polygroups/PolygroupsGenerator.h
@@ -4,9 +4,11 @@
 
 #include "CoreMinimal.h"
 #include "DynamicMesh/DynamicMesh3.h"
+#include "Polygroups/PolygroupSet.h"
 
 using UE::Geometry::FDynamicMesh3;
 
+
 namespace UE
 {
 namespace Geometry
@@ -59,7 +61,12 @@ public:
 	/**
 	 * Find Polygroups based on UV Islands, ie each UV Island becomes a Polygroup
 	 */
-	bool FindPolygroupsFromUVIslands();
+	bool FindPolygroupsFromUVIslands(int32 UVLayer = 0);
+
+	/**
+	* Find Polygroups based on Seams in UV Overlay
+	*/
+	bool FindPolygroupsFromHardNormalSeams();
 
 	/**
 	 * Find Polygroups based on mesh connectivity, ie each connected-component becomes a Polygroup
@@ -79,8 +86,14 @@ public:
 	 * Incrementally compute approximate-geodesic-based furthest-point sampling of the mesh until NumPoints
 	 * samples have been found, then compute local geodesic patches (eg approximate surface voronoi diagaram).
 	 * Optionally weight geodesic-distance computation, which will produce different patch shapes.
+	 * If StartingGroups are provided, then new groups are constrained to subdivide existing groups, ie
+	 * none of the new groups will cross the boundaries of the StartingGroups
 	 */
-	bool FindPolygroupsFromFurthestPointSampling(int32 NumPoints, EWeightingType WeightingType, FVector3d WeightingCoeffs = FVector3d::One());
+	bool FindPolygroupsFromFurthestPointSampling(
+		int32 NumPoints, 
+		EWeightingType WeightingType, 
+		FVector3d WeightingCoeffs = FVector3d::One(),
+		FPolygroupSet* StartingGroups = nullptr);
 
 
 	/**
