GeometryProcessing: add FPolygroupsGenerator::CopyPolygroupsToPolygroupSet(), to copy to arbitrary polygroup set (instead of only to default mesh polygroups)
#rnx

diff --git a/Engine/Plugins/Runtime/GeometryProcessing/Source/DynamicMesh/Private/Polygroups/PolygroupsGenerator.cpp b/Engine/Plugins/Runtime/GeometryProcessing/Source/DynamicMesh/Private/Polygroups/PolygroupsGenerator.cpp
--- a/Engine/Plugins/Runtime/GeometryProcessing/Source/DynamicMesh/Private/Polygroups/PolygroupsGenerator.cpp
+++ b/Engine/Plugins/Runtime/GeometryProcessing/Source/DynamicMesh/Private/Polygroups/PolygroupsGenerator.cpp
@@ -444,7 +444,22 @@ void FPolygroupsGenerator::CopyPolygroupsToMesh()
 	}
 }
 
+void FPolygroupsGenerator::CopyPolygroupsToPolygroupSet(FPolygroupSet& Polygroups, FDynamicMesh3& TargetMesh)
+{
+	// set groups from Polygroups
+	int NumPolygroups = FoundPolygroups.Num();
 
+	// can be parallel for
+	for (int PolyIdx = 0; PolyIdx < NumPolygroups; PolyIdx++)
+	{
+		const TArray<int>& Polygroup = FoundPolygroups[PolyIdx];
+		int NumTriangles = Polygroup.Num();
+		for (int k = 0; k < NumTriangles; ++k)
+		{
+			Polygroups.SetGroup(Polygroup[k], (PolyIdx + 1), TargetMesh);
+		}
+	}
+}
 
 
 
diff --git a/Engine/Plugins/Runtime/GeometryProcessing/Source/DynamicMesh/Public/Polygroups/PolygroupsGenerator.h b/Engine/Plugins/Runtime/GeometryProcessing/Source/DynamicMesh/Public/Polygroups/PolygroupsGenerator.h
--- a/Engine/Plugins/Runtime/GeometryProcessing/Source/DynamicMesh/Public/Polygroups/PolygroupsGenerator.h
+++ b/Engine/Plugins/Runtime/GeometryProcessing/Source/DynamicMesh/Public/Polygroups/PolygroupsGenerator.h
@@ -102,6 +102,10 @@ public:
 	*/
 	void CopyPolygroupsToMesh();
 
+	/**
+	* Copy the computed Polygroups to the given PolygroupSet and Mesh
+	*/
+	void CopyPolygroupsToPolygroupSet(FPolygroupSet& Polygroups, FDynamicMesh3& TargetMesh);
 
 	/**
 	 * Initialize the PolygroupEdges output member by finding all the mesh edges that are on polygroup borders.
