fix incorrect comment on boolean tool compute method

#rnx

diff --git a/Engine/Plugins/Runtime/GeometryProcessing/Source/DynamicMesh/Public/Operations/MeshBoolean.h b/Engine/Plugins/Runtime/GeometryProcessing/Source/DynamicMesh/Public/Operations/MeshBoolean.h
--- a/Engine/Plugins/Runtime/GeometryProcessing/Source/DynamicMesh/Public/Operations/MeshBoolean.h
+++ b/Engine/Plugins/Runtime/GeometryProcessing/Source/DynamicMesh/Public/Operations/MeshBoolean.h
@@ -171,8 +171,7 @@ public:
 	}
 
 	/**
-	 * Compute the plane cut by splitting mesh edges that cross the cut plane, and then deleting any triangles
-	 * on the positive side of the cutting plane.
+	 * Compute the Boolean operation
 	 * @return true if operation succeeds
 	 */
 	bool Compute();
