Dynamic Mesh Plugin: Operations -  Adding first pass at code that  uses edge flips of an intrinsic mesh to reduce the length of an input surface curve (defined on a FDynamicMesh)

diff --git a/Engine/Plugins/Runtime/GeometryProcessing/Source/DynamicMesh/Private/Operations/GeodesicPath.cpp b/Engine/Plugins/Runtime/GeometryProcessing/Source/DynamicMesh/Private/Operations/GeodesicPath.cpp
new file mode 100644
--- /dev/null
+++ b/Engine/Plugins/Runtime/GeometryProcessing/Source/DynamicMesh/Private/Operations/GeodesicPath.cpp
@@ -0,0 +1,623 @@
+// Copyright Epic Games, Inc. All Rights Reserved.
+
+#include "Operations/GeodesicPath.h"
+#include "Math/UnrealMathVectorCommon.h"
+
+using namespace UE::Geometry;
+
+
+// -- FEdgePath methods ---//
+
+int32 FEdgePath::GetHeadSegmentID() const
+{
+	NodeType* Node = PathLinkList.GetHead();
+	if (!Node)
+	{
+		return InvalidID;
+	}
+
+	int32 SID = Node->GetValue().SID;
+	return SID;
+}
+
+int32 FEdgePath::GetTailSegmentID() const
+{
+	NodeType* Node = PathLinkList.GetTail();
+	if (!Node)
+	{
+		return InvalidID;
+	}
+
+	int32 SID = Node->GetValue().SID;
+	return SID;
+}
+
+FEdgePath::FDirectedSegment FEdgePath::GetSegment(int32 SID) const
+{
+	FDirectedSegment Result = { InvalidID, InvalidID };
+	if (!IsSegment(SID))
+	{
+		return Result;
+	}
+	const FSegmentAndSID& SegmentAndSID = SIDtoNode[SID]->GetValue();
+	Result.EID = SegmentAndSID.EID;
+	Result.HeadIndex = SegmentAndSID.HeadIndex;
+	return Result;
+}
+
+int32 FEdgePath::GetPrevSegmentID(int32 SID) const
+{
+	if (!IsSegment(SID))
+	{
+		return InvalidID;
+	}
+	const NodeType* Node = SIDtoNode[SID];
+	const NodeType* PrevNode = Node->GetPrevNode();
+	if (PrevNode)
+	{
+		return PrevNode->GetValue().SID;
+	}
+	else
+	{
+		return InvalidID;
+	}
+}
+
+int32 FEdgePath::GetNextSegmentID(int32 SID) const
+{
+	if (!IsSegment(SID))
+	{
+		return InvalidID;
+	}
+	const NodeType* Node = SIDtoNode[SID];
+	const NodeType* NextNode = Node->GetNextNode();
+	if (NextNode)
+	{
+		return NextNode->GetValue().SID;
+	}
+	else
+	{
+		return InvalidID;
+	}
+}
+
+int32 FEdgePath::AppendSegment(const FDirectedSegment& Segment)
+{
+	
+	// record the mesh edge as a new segment
+	const int32 SID = DirectedSegmentRefCounts.Allocate();
+	FSegmentAndSID SegmentAndSID = {{Segment.EID, Segment.HeadIndex}, SID};
+	PathLinkList.AddTail(SegmentAndSID);
+
+	NodeType* Node = PathLinkList.GetTail();
+	SIDtoNode.InsertAt(Node, SID);
+
+	return SID;
+}
+
+int32 FEdgePath::InsertSegmentBefore(const FEdgePath::FDirectedSegment& Segment, int32 SegmentIDToInsertBefore)
+{
+	NodeType* NodeToInsertBefore = (SegmentIDToInsertBefore > -1) ? SIDtoNode[SegmentIDToInsertBefore] : nullptr;
+
+	int32 SID = DirectedSegmentRefCounts.Allocate();
+	FSegmentAndSID SegmentAndSID = {{Segment.EID, Segment.HeadIndex}, SID};
+	PathLinkList.InsertNode(SegmentAndSID, NodeToInsertBefore);
+
+	NodeType* Node = (NodeToInsertBefore) ? NodeToInsertBefore->GetPrevNode() : PathLinkList.GetTail();
+	SIDtoNode.InsertAt(Node, SID);
+
+	return SID;
+}
+
+void FEdgePath::RemoveSegment(int32 SID)
+{
+	if (!IsSegment(SID))
+	{
+		return;
+	}
+
+	NodeType* Node = SIDtoNode[SID];
+	PathLinkList.RemoveNode(Node); // does delete
+	SIDtoNode[SID] = nullptr;
+
+	DirectedSegmentRefCounts.Decrement(SID);
+}
+
+
+/**
+* Visits a subset of the triangle neighborhood of vertex CVID, 
+* i.e. faces in CW or CCW-order about vertex CVID starting with edge StartEID and ending at either a boundary or edge EndEID,
+* returns false if a boundary is encountered.
+*
+* At each face a Visitor functor is called.
+*
+* @param EdgeFlipMesh - IntrinsicEdgeFlip mesh
+* @param StartEID     - edge incident on CVID
+* @param EndEID       - edge incident on CVID
+* @param CVID         - vertex at center when visiting triangle neighborhood
+* @param Visitor      - functor called on each tri between StartEID and EndEID
+* @param bClockwise   - specifies the direction of travel around VID as clockwise (CW) or counter clockwise (CCW).
+*
+* the functor Visitor(int32 TriID, int32 IndexOfCenterVID, FIndex3i& TriEIDs)
+*                           TriID             - the triangle visited
+*                           IndexOfCenterVID  - the corner index of the specified VID in the triangle
+*                           TriEIDs           - edge IDs for the triangle.
+*
+*/
+template <typename FunctorType>
+bool VisitWedgeTriangles(const FIntrinsicEdgeFlipMesh& EdgeFlipMesh, const int32 StartEID, const int32 EndEID, const int32 CVID, 
+                         FunctorType& Visitor, bool bClockwise = false)
+{
+	int32 CurEID = StartEID;
+
+	FIndex2i CurEdgeT   = EdgeFlipMesh.GetEdgeT(CurEID);
+	int32 CurTID        = CurEdgeT[0];
+	FIndex3i CurTriEIDs = EdgeFlipMesh.GetTriEdges(CurTID);
+	int32 IndexOf       = CurTriEIDs.IndexOf(CurEID);
+
+	// make sure we are starting with the triangle on the correct side of the start edge for clockwise / counter clockwise traversal
+	if (( bClockwise && (EdgeFlipMesh.GetTriangle(CurTID)[IndexOf] == CVID)) ||
+		(!bClockwise && (EdgeFlipMesh.GetTriangle(CurTID)[IndexOf] != CVID)))
+	{
+		CurTID     = CurEdgeT[1];
+		if (CurTID == FIntrinsicEdgeFlipMesh::InvalidID)
+		{
+			return false;
+		}
+		CurTriEIDs = EdgeFlipMesh.GetTriEdges(CurTID);
+		IndexOf    = CurTriEIDs.IndexOf(CurEID);
+	}
+	const int32 NextEdgeOffset = (bClockwise) ? 1 : 2;
+
+	do
+	{
+		Visitor(CurTID, IndexOf, CurTriEIDs);
+
+		// advance to next edge
+		CurEID   = CurTriEIDs[(NextEdgeOffset + IndexOf) % 3];
+		CurEdgeT = EdgeFlipMesh.GetEdgeT(CurEID);
+		CurTID = (CurEdgeT[0] == CurTID) ? CurEdgeT[1] : CurEdgeT[0];
+
+		if (CurTID == FIntrinsicEdgeFlipMesh::InvalidID)
+		{
+			return (CurEID == EndEID);
+		}
+
+		CurTriEIDs = EdgeFlipMesh.GetTriEdges(CurTID);
+		IndexOf = CurTriEIDs.IndexOf(CurEID);
+
+	} while (CurEID != EndEID);
+
+	return true;
+}
+
+// -- Deformable Edge Path --//
+
+
+FDeformableEdgePath::FDeformableEdgePath(const FDynamicMesh3& SurfaceMesh, const TArray<FEdgePath::FDirectedSegment>& PathAsDirectedSegments)
+	: EdgeFlipMesh(SurfaceMesh)
+	, PathLength(0.)
+	, NumFlips(0)
+{
+
+	for (const FEdgePath::FDirectedSegment& DirectedSegment : PathAsDirectedSegments)
+	{
+		const int32 EID = DirectedSegment.EID;
+		if (!EdgeFlipMesh.IsEdge(EID))
+		{
+			continue;
+		}
+	
+		
+		// compare with last segment eliminate double-back
+		const int32 TailSID = EdgePath.GetTailSegmentID();
+		const FEdgePath::FDirectedSegment TailSegment = EdgePath.GetSegment(TailSID);
+		
+		if (TailSegment.EID == DirectedSegment.EID && TailSegment.HeadIndex != DirectedSegment.HeadIndex)
+		{
+			EdgePath.RemoveSegment(TailSID);
+		}
+		else
+		{ 
+			int32 SID = EdgePath.AppendSegment(DirectedSegment);
+			TArray<int32>& Segments = EIDToSIDsMap.FindOrAdd(EID);
+			Segments.Add(SID);
+
+			UpdateJointAndQueue(SID);
+		}
+	}
+
+	// no segments in this path.
+	if (EdgePath.NumSegments() == 0)
+	{
+		return;
+	}
+
+	// compute the initial length of the path
+	const int32 MaxSID = EdgePath.MaxSegmentID();
+	for (int32 SID = EdgePath.GetHeadSegmentID(); SID < MaxSID; ++SID)
+	{
+		if (!EdgePath.IsSegment(SID))
+		{
+			continue;
+		}
+		const int32 EID = EdgePath.GetSegment(SID).EID;
+		const double EdgeLength = EdgeFlipMesh.GetEdgeLength(EID);
+		PathLength += EdgeLength;
+	}
+}
+
+void FDeformableEdgePath::Minimize(FDeformableEdgePath::FEdgePathDeformationInfo& DeformedPathInfo, const int32 MaxNumIterations)
+{
+
+	DeformedPathInfo = FEdgePathDeformationInfo();
+	DeformedPathInfo.OriginalLength = PathLength;
+
+	constexpr double AngleEpsilon = 1.e-3;
+
+	int32 NumIterations = 0;
+	while (JointAngleQueue.Num() > 0 && MaxNumIterations > NumIterations)
+	{
+		// get ID of Outgoing segment in the joint, and the associated smaller internal angle.
+		const uint32 QueueSegmentID = JointAngleQueue.Top();
+		const double InternalAngle = JointAngleQueue.GetKey(QueueSegmentID);
+
+		JointAngleQueue.Pop();
+
+		if (InternalAngle > TMathUtilConstants<double>::Pi - AngleEpsilon)
+		{
+			// the smaller internal angle was greater than Pi, this joint can't be further straightened 
+			continue;
+		}
+
+		const int32 SID = int32(QueueSegmentID);
+		FPathJoint& PathJoint = PathJoints[SID];
+
+		if (!IsJointFlexible(PathJoint))
+		{
+			continue;
+		}
+
+		DeformJoint(PathJoint);
+		NumIterations++;
+	}
+
+	DeformedPathInfo.NumIterations = NumIterations;
+	DeformedPathInfo.NumEdgeFlips = NumFlips;
+	DeformedPathInfo.FinalLength = PathLength;
+}
+
+
+void FDeformableEdgePath::DeformJoint(FDeformableEdgePath::FPathJoint& PathJoint) 
+{
+	const int32 JointVID = PathJoint.VID;
+
+	const int32 OutgoingSID = PathJoint.OutgoingSID;
+	const int32 IncomingSID = EdgePath.GetPrevSegmentID(OutgoingSID);
+
+	if (IncomingSID < 0 || JointVID == InvalidID) // a joint is composed of incoming and outgoing edges meeting at VID 
+	{
+		return;
+	}
+
+	const int32 IncomingEID = EdgePath.GetSegment(IncomingSID).EID;
+	const int32 OutgoingEID = EdgePath.GetSegment(OutgoingSID).EID;
+
+	TArray<FEdgePath::FDirectedSegment> DeformedPath;
+
+	switch (PathJoint.InteriorSide)
+	{
+		case ESide::Right:
+		{
+			const bool bClockwise = false;
+			DeformJoint(IncomingEID, OutgoingEID, JointVID, DeformedPath, bClockwise);
+		}
+		break;
+
+		case ESide::Left:
+		{
+			const bool bClockwise = true;
+			DeformJoint(IncomingEID, OutgoingEID, JointVID, DeformedPath, bClockwise);
+		}
+		break;
+
+		default:
+			check(0);
+	}
+
+	ReplaceJointWithPath(PathJoint, DeformedPath);
+}
+
+bool FDeformableEdgePath::DeformJoint( const int32 IncomingEID, const int32 OutgoingEID, int32 JointVID, 
+                                       TArray<FEdgePath::FDirectedSegment>& DeformedPath, const bool bClockwise )
+{
+	bool Result = OuterArcFlipEdges(IncomingEID, OutgoingEID, JointVID, bClockwise);
+
+	// collect deformed path
+	if (Result)
+	{
+		const int32 HeadOffset = (bClockwise) ? 1 : 2;
+		auto PathAccumulator = [this, HeadOffset, &DeformedPath, bClockwise, JointVID](int32 TriID, int32 IndexOfEID, FIndex3i& TriEIDs)
+		{
+			const int32 IndexOfCenterVID = (bClockwise) ? (IndexOfEID + 1) % 3 : IndexOfEID;
+			const int32 OpposingEID = TriEIDs[(IndexOfCenterVID + 1) % 3];
+			const FIndex3i TriVIDs = EdgeFlipMesh.GetTriangle(TriID);
+			checkSlow(JointVID == TriVIDs[IndexOfCenterVID]);
+			const int32 HeadVID = TriVIDs[(IndexOfCenterVID + HeadOffset) % 3];
+			const FIndex2i OpposingEdgeV = EdgeFlipMesh.GetEdgeV(OpposingEID);
+
+			FEdgePath::FDirectedSegment DirectedSegment;
+			DirectedSegment.EID = OpposingEID;
+			DirectedSegment.HeadIndex = (OpposingEdgeV.A == HeadVID) ? 0 : 1;
+			checkSlow(OpposingEdgeV[DirectedSegment.HeadIndex] == HeadVID);
+
+			DeformedPath.Add(DirectedSegment);
+
+		};
+
+		Result = VisitWedgeTriangles(EdgeFlipMesh, IncomingEID, OutgoingEID, JointVID, PathAccumulator, bClockwise);
+		checkSlow(Result); // since OuterArcFlipEdges was true, this should be too.
+	}
+
+	return Result;
+}
+
+bool FDeformableEdgePath::IsJointFlexible(const FDeformableEdgePath::FPathJoint& Joint) const
+{
+
+	const int32 SID = Joint.OutgoingSID;
+	const int32 PrevSID = EdgePath.GetPrevSegmentID(SID);
+
+	if ( SID < 0 || PrevSID < 0)
+	{
+		return false;
+	}
+
+	// Get the Interior Side and JointVID 
+	const int32 JointVID     = Joint.VID;
+	const ESide InteriorSide = Joint.InteriorSide;
+
+
+	bool bIsBlocked = false;
+	int32 StartEID = EdgePath.GetSegment(SID).EID;
+	int32 EndEID   = EdgePath.GetSegment(PrevSID).EID;
+
+	const bool bClockwise = (InteriorSide == ESide::Right);
+
+	auto OccupiedEdgeTestor = [&bIsBlocked, StartEID, EndEID, this](int32 TriID, int32 IndexOfCenterVID, FIndex3i& TriEIDs)
+	{
+
+		for (int32 i = 0; i < 3 && !bIsBlocked; ++i)
+		{
+			const int32 CurEID = TriEIDs[i];
+			if (CurEID == StartEID || CurEID == EndEID)
+			{
+				continue;
+			}
+			const TArray<int32>* SegmentArray = EIDToSIDsMap.Find(CurEID);
+			bIsBlocked = (SegmentArray) || bIsBlocked;
+		}
+	};
+
+	// check StartEID and EndEID.  These should have one entry each.
+	// To prevent the introduction of a path crossing, this  prevents the local joint from being deformed
+	// in the case that other segments of the path have doubled back on either of these edges.
+	// [todo] when multiple path segments share the same edge, ascertain their relative interleaving (e.g. left to right)
+	//        so we can identify situations when the local joint can be deformed without producing a crossing.
+	{
+		if (const TArray<int32>* SegmentArray = EIDToSIDsMap.Find(StartEID))
+		{
+			bIsBlocked = bIsBlocked ||(SegmentArray->Num() > 1);
+		}
+		if (const TArray<int32>* SegmentArray = EIDToSIDsMap.Find(EndEID))
+		{
+			bIsBlocked = bIsBlocked || (SegmentArray->Num() > 1);
+		}
+	}
+
+	const bool bHasBoundary = !VisitWedgeTriangles(EdgeFlipMesh, StartEID, EndEID, JointVID,
+												   OccupiedEdgeTestor, bClockwise);
+
+	return (!bHasBoundary && !bIsBlocked);
+}
+
+
+void FDeformableEdgePath::UpdateJointAndQueue(int32 SID)
+{
+	
+	int32 PrevSID = EdgePath.GetPrevSegmentID(SID);
+
+	// a joint is composed of cur segment and prev segment
+	if (PrevSID != InvalidID)
+	{
+		const int32 JointVID = SegmentHeadVID(PrevSID);
+		checkSlow(JointVID == SegmentTailVID(SID)); 
+
+		const int32 IncomingEID = EdgePath.GetSegment(PrevSID).EID;
+		const int32 OutgoingEID = EdgePath.GetSegment(SID).EID;
+		
+
+		double LAngle = 0., RAngle = 0.;
+		ComputeWedgeAngles(IncomingEID, OutgoingEID, JointVID, LAngle, RAngle);
+		const ESide InteriorSide = (LAngle < RAngle) ? ESide::Left : ESide::Right;
+
+		// encode the joint
+		FPathJoint PathJoint = { SID, InteriorSide, JointVID };
+		PathJoints.InsertAt(PathJoint, SID);
+
+
+		uint32 QueueSegmentID = uint32(SID);
+		double SortKey = FMath::Min(LAngle, RAngle);
+
+		if ( JointAngleQueue.IsPresent(QueueSegmentID) )
+		{
+			JointAngleQueue.Update(SortKey, QueueSegmentID);
+		}
+		else
+		{
+			JointAngleQueue.Add(SortKey, QueueSegmentID);
+		}
+	}
+}
+
+void FDeformableEdgePath::RemoveSegment(int32 SID)
+{
+	if (SID < 0)
+	{
+		return;
+	}
+
+	const int32 EID = EdgePath.GetSegment(SID).EID;
+	JointAngleQueue.Remove(SID);
+	EdgePath.RemoveSegment(SID);
+
+	if (TArray<int32>* SegmentArray = EIDToSIDsMap.Find(EID))
+	{
+		SegmentArray->RemoveSwap(SID);
+		if (SegmentArray->Num() == 0)
+		{
+			EIDToSIDsMap.Remove(EID);
+		}
+	}
+
+	// subtract length of segment from total path length
+	PathLength -= EdgeFlipMesh.GetEdgeLength(EID);
+
+	checkSlow(!EdgePath.IsSegment(SID));
+}
+
+void FDeformableEdgePath::ReplaceJointWithPath(const FDeformableEdgePath::FPathJoint PathJoint, const TArray<FEdgePath::FDirectedSegment>& PathAsDirectedEdges)
+{
+	int32 OutgoingSID = PathJoint.OutgoingSID;
+	if (OutgoingSID < 0)
+	{
+		return;
+	}
+	// The Joint is comprised of AtoB and BtoC
+	int32 IncomingSID = EdgePath.GetPrevSegmentID(OutgoingSID);  
+
+	// segment for first edge element after Joint.
+	int32 NextSegmentSID = EdgePath.GetNextSegmentID(OutgoingSID); // Note this can be InvalidID, and that is ok
+
+	// remove the joint
+	RemoveSegment(IncomingSID);
+	RemoveSegment(OutgoingSID);
+
+	// add path before NextSegment
+	if (NextSegmentSID != InvalidID)
+	{ 
+		for (const FEdgePath::FDirectedSegment& DirectedSegment : PathAsDirectedEdges)
+		{
+			int32 NewSID = EdgePath.InsertSegmentBefore(DirectedSegment, NextSegmentSID);
+			PathLength += EdgeFlipMesh.GetEdgeLength(DirectedSegment.EID);
+			UpdateJointAndQueue(NewSID);
+		}
+
+		// update the joint that now connects NextEdgeNode to the last NewNode.
+		UpdateJointAndQueue(NextSegmentSID);
+	}
+	else   // add to tail
+	{
+		for (const FEdgePath::FDirectedSegment& DirectedSegment : PathAsDirectedEdges)
+		{
+			int32 NewSID = EdgePath.AppendSegment(DirectedSegment);
+			PathLength += EdgeFlipMesh.GetEdgeLength(DirectedSegment.EID);
+			UpdateJointAndQueue(NewSID);
+		}
+	}
+
+}
+
+bool FDeformableEdgePath::OuterArcFlipEdges(int32 StartEID, int32 EndEID, int32 CVID, bool bClockwise)
+{
+	int32 CurEID = StartEID;
+
+	FIndex2i CurEdgeT   = EdgeFlipMesh.GetEdgeT(CurEID);
+	int32 CurTID        = CurEdgeT[0];
+	FIndex3i CurTriEIDs = EdgeFlipMesh.GetTriEdges(CurTID);
+	int32 IndexOf       = CurTriEIDs.IndexOf(CurEID);
+
+	// make sure we are starting with the triangle on the correct side of the start edge for clockwise / counter clockwise traversal
+	if (( bClockwise && (EdgeFlipMesh.GetTriangle(CurTID)[IndexOf] == CVID)) ||
+		(!bClockwise && (EdgeFlipMesh.GetTriangle(CurTID)[IndexOf] != CVID)))
+	{
+		CurTID     = CurEdgeT[1];
+		if (CurTID == InvalidID)
+		{
+			return false;
+		}
+		CurTriEIDs = EdgeFlipMesh.GetTriEdges(CurTID);
+		IndexOf    = CurTriEIDs.IndexOf(CurEID);
+	}
+	const int32 NextEdgeOffset = (bClockwise) ? 1 : 2;
+
+	int32 NextEID = CurTriEIDs[(IndexOf + NextEdgeOffset) % 3];
+
+	while (NextEID != EndEID)
+	{
+		FIndex2i PreFlipT = EdgeFlipMesh.GetEdgeT(NextEID);
+		// The ID of the triangle adjacent to CurEID after flipping NextEID
+		// relies on the logic in FIntrinsicEdgeFlipMesh::FlipEdgeTopology().. if that changes this will break
+		int32 TmpTID = CurTID;
+		if (bClockwise)
+		{
+			TmpTID = (CurTID == PreFlipT[1]) ? PreFlipT[0] : PreFlipT[1];
+		}
+
+		UE::Geometry::FIntrinsicEdgeFlipMesh::FEdgeFlipInfo EdgeFlipInfo;
+		UE::Geometry::EMeshResult Result = EdgeFlipMesh.FlipEdge(NextEID, EdgeFlipInfo);
+
+		if (Result == UE::Geometry::EMeshResult::Failed_IsBoundaryEdge)
+		{
+			return false;
+		}
+
+		if (Result == UE::Geometry::EMeshResult::Ok)
+		{
+			this->NumFlips++;
+
+			CurTID     = TmpTID;
+			CurTriEIDs = EdgeFlipMesh.GetTriEdges(CurTID);
+			IndexOf    = CurTriEIDs.IndexOf(CurEID);
+			checkSlow(IndexOf != -1);  // would only fail if the TmpID logic above is wrong.
+			NextEID    = CurTriEIDs[(IndexOf + NextEdgeOffset) % 3];
+		}
+		else
+		{
+			// didn't flip edge.  advance to next triangle on other side of that edge
+			CurEID     = NextEID;
+			CurTID     = (PreFlipT[0] == CurTID) ? PreFlipT[1] : PreFlipT[0];
+			CurTriEIDs = EdgeFlipMesh.GetTriEdges(CurTID);
+			IndexOf    = CurTriEIDs.IndexOf(CurEID);
+			NextEID    = CurTriEIDs[(IndexOf + NextEdgeOffset) % 3];
+		}
+	}
+
+	return true;
+}
+
+void FDeformableEdgePath::ComputeWedgeAngles(int32 IncomingEID, 
+                                             int32 OutgoingEID, 
+											 int32 CenterVID,
+	                                         double& LeftSideAngle, double& RightSideAngle) const
+{
+
+	const bool bClockwise = false;
+
+	double WedgeAngle = 0.;
+	auto WedgeAngleAccumulator = [this, &WedgeAngle](int32 TID, int32 IndexOf, const FIndex3i& TriEIDs)
+	{
+		WedgeAngle += EdgeFlipMesh.GetTriInternalAngleR(TID, IndexOf);
+	};
+
+
+	bool bLeftContainsBoundary = !VisitWedgeTriangles(EdgeFlipMesh, OutgoingEID, IncomingEID, CenterVID, WedgeAngleAccumulator, bClockwise);
+	double LeftWedgeAngle = WedgeAngle;
+	
+	WedgeAngle = 0.;
+	bool bRightContainsBoundary = !VisitWedgeTriangles(EdgeFlipMesh, IncomingEID, OutgoingEID, CenterVID, WedgeAngleAccumulator, bClockwise);
+	double RightWedgeAngle = WedgeAngle;
+
+	LeftSideAngle  = (bLeftContainsBoundary)  ? TMathUtilConstants<double>::MaxReal : LeftWedgeAngle;
+	RightSideAngle = (bRightContainsBoundary) ? TMathUtilConstants<double>::MaxReal : RightWedgeAngle;
+}
\ No newline at end of file
diff --git a/Engine/Plugins/Runtime/GeometryProcessing/Source/DynamicMesh/Public/Operations/GeodesicPath.h b/Engine/Plugins/Runtime/GeometryProcessing/Source/DynamicMesh/Public/Operations/GeodesicPath.h
new file mode 100644
--- /dev/null
+++ b/Engine/Plugins/Runtime/GeometryProcessing/Source/DynamicMesh/Public/Operations/GeodesicPath.h
@@ -0,0 +1,305 @@
+// Copyright Epic Games, Inc. All Rights Reserved.
+
+#pragma once
+
+#include "CoreMinimal.h"
+
+#include "Containers/Array.h"
+#include "Containers/BinaryHeap.h"
+#include "Containers/List.h"
+#include "Containers/Map.h"
+#include "DynamicMesh/DynamicMesh3.h"
+#include "IndexTypes.h"
+#include "IntrinsicTriangulationMesh.h"
+#include "Util/DynamicVector.h"
+#include "Util/RefCountVector.h"
+#include "MathUtil.h"
+
+
+#include "ExplicitUseGeometryMathTypes.h"
+
+namespace UE
+{
+namespace Geometry
+{
+
+/**
+* FEdgePath is is little more than a double link-list that exposes IDs instead of node pointers.
+* The ID of a removed segment may be reused for a subsequently appended or inserted segment.
+* When used in conjunction with a mesh, it represents a directed, edge-aligned path.
+* 
+* Note: this exists primarily for use with FDeformableEdgePath and 
+*       no checking is done to ensure the path is continuous, in fact it could 
+*       be used to represent multiple disconnected paths.
+*/
+class DYNAMICMESH_API FEdgePath
+{
+public:
+
+	
+	constexpr static int InvalidID = IndexConstants::InvalidID;
+	struct FDirectedSegment
+	{
+		int32 EID  = InvalidID;         // expected to be a mesh edge ID
+		int32 HeadIndex = InvalidID;    // 0 or 1.  GetEdgeV()[HeadIndex] is the 'to' vertex,  GetEdgeV()[1-HeadIndex] is the 'from' vertex
+	};
+	
+	FEdgePath() = default;
+	FEdgePath(const FEdgePath&) = delete;
+
+	/** @return number of segments in path */
+	inline int32 NumSegments() const;
+
+	/** @return upper bound on the segment ID used in path.  i.e. all segment IDs < MaxSegmentID */
+	inline int32 MaxSegmentID() const;
+
+	/** @return true if the specified segment ID corresponds to a valid segment in path */
+	inline bool IsSegment(int32 SID) const;
+	
+	/** @return requested segment if it is part of the path, otherwise return invalid segment (with EID = InvalidID) */
+	FDirectedSegment GetSegment(int32 SID) const;
+
+	/** @return the segment ID of the {first, last} segment in the path, or InvalidID if the path is empty */
+	int32 GetHeadSegmentID() const;
+	int32 GetTailSegmentID() const;
+
+	/** @return segment ID of {previous, next} segment in the path, or InvalidID if if no {previous, next} exists or if SID itself is not in the path  */
+	int32 GetPrevSegmentID(int32 SID) const;
+	int32 GetNextSegmentID(int32 SID) const;
+
+	/** @return segment ID of a new segment appended to the tail of the path */
+	int32 AppendSegment(const FDirectedSegment& Segment);
+
+	/** @return segment ID of a new segment inserted before the specified existing segment. If SegmentIDToInsertBefore is InvalidID, then insert at head  */
+	int32 InsertSegmentBefore(const FDirectedSegment& DirectedSegment, int32 SegmentIDToInsertBefore = InvalidID);
+	
+	/** remove specified segment from path. Note SID is added to an internal free list and maybe reused  */
+	void RemoveSegment(int32 SID);
+
+
+protected:
+
+	struct FSegmentAndSID : FDirectedSegment
+	{
+		int32 SID = InvalidID; // segment id
+	};
+
+	typedef typename TDoubleLinkedList<FSegmentAndSID>::TDoubleLinkedListNode NodeType;
+	
+	TDoubleLinkedList<FSegmentAndSID>  PathLinkList;                // The path is represented by Segments in a double link list.
+	
+	TDynamicVector<NodeType*>          SIDtoNode;                   // Nodes per segment indexed by SegmentID (SID)
+	FRefCountVector                    DirectedSegmentRefCounts{};  // SID list.
+};
+
+
+/**
+* Given a single path initially defined by edges on a FDynamicMesh (the SurfaceMesh), this class uses an intrinsic mesh with 
+* the same vertex set to minimize the length of the path on the surface mesh.  This is done iteratively by a succession 
+* of local path deformations comprised of (local) intrisic edge flips and path re-routing, each deformation locally reduces the path length.  
+* The resulting path is comprised of intrinsic mesh edges.
+* 
+* This is an implementation of the ideas in 
+*       "You Can Find Geodesic Paths in Triangle Meshes by Just Flipping Edges" by Nicholas Sharp and Keenan Crane
+*        ACM Trans. Graph. Vol. 39, No. 6, 2020
+* 
+* The original path is constructed as a sequence of mesh edges relative to the  SurfaceMesh ( the FDynamicMesh3 ),
+* these edges are encoded with path direction as FDirectedSegments.
+* Note: The current implementation does not support closed paths or multiple paths.
+*/
+class DYNAMICMESH_API FDeformableEdgePath
+{
+public:
+	
+	constexpr static int InvalidID = IndexConstants::InvalidID;
+	
+	/**
+	* Constructor assumes the directed segments are ordered from tail (at index 0) to head
+	*/
+	FDeformableEdgePath(const FDynamicMesh3& SurfaceMesh, const TArray<FEdgePath::FDirectedSegment>& OriginalPathAsDirectedSegments);
+	
+	virtual ~FDeformableEdgePath(){}
+
+	// only allow construction from a path and mesh.
+	FDeformableEdgePath() = delete;
+	FDeformableEdgePath(const FDeformableEdgePath& ) = delete;
+	
+	// information gathered during minimization
+	struct FEdgePathDeformationInfo
+	{
+		int32 NumIterations   = 0;     // number of two-edge joints that were replaced by shorter paths (but perhaps more segments)
+		int32 NumEdgeFlips    = 0;     // number of edge flips performed in minimizing the path
+		double OriginalLength = 0.;    // path length prior to minimization
+		double FinalLength    = 0.;    // path length after minimization
+	};
+
+	/**
+	* Minimize the deformable edge path with respect to the total path length.  
+	* This minimization finds a local minimum in path length but may not find the shortest over all distance between to points.
+	* In this sense the operation can be thought of as a path-straightening. 
+	* @param DeformedPathInfo  - on return this info struct is populated
+	* @param MaxNumIterations  - the maximum number of two-edge joints that may be replaced by shorter paths.
+	*/ 
+	virtual void Minimize(FEdgePathDeformationInfo&  DeformedPathInfo, const int32 MaxNumIterations = TMathUtilConstants<int32>::MaxReal);
+	
+	/**
+	* @return a const reference to the edge path.
+	*/
+	inline const FEdgePath& GetEdgePath() const;
+
+	/**
+	* @return Length of the path.
+	*/
+	inline double GetPathLength() const;
+
+	/**
+	* @return const reference to the intrinsic mesh on which the EdgePath is defined.
+	*/ 
+	inline const FIntrinsicEdgeFlipMesh& GetIntrinsicMesh() const;
+
+protected:
+	
+	enum class ESide { Left, Right };      // indicates which wedge (Left or Right) has the smaller internal angle, where in 2d the path travels "up" 
+	struct FPathJoint
+	{
+		int32  OutgoingSID = InvalidID;    // Outgoing segment.  EdgePath.GetPrevSegmentID(OutgoingSID) returns incoming segment
+		ESide  InteriorSide;               // Side of the one-ring with the smaller angle.
+		int32  VID = InvalidID;            // incoming and Outgoing edges meet at this vertex
+	};
+
+
+	/**
+	* Replace the specified two connected edge segments (FPathJoint) in the main path with a new sub-path 
+	* that connects the same end points. This new path is formed by flipping the adjacent edges to the 
+	* left (or right) of the joint vertex and may have one, two, or more segments
+	*/
+	void DeformJoint(FPathJoint& PathJoint);
+
+	/**
+	* Edge flip the adjacent edges to the JointVID and collect the DeformedPath that can be used to replace the local incoming and outgoing edge path
+	* @param IncomingEID  - ID of incoming edge of joint
+	* @param OutgoingEID  - ID of outgoing edge of joint
+	* @param JointVID     - ID of Vertex where incoming and outgoing edges meet
+	* 
+	* @param DeformedPath - On return, holds directed segments that can be used to replace the joint
+	* @param bClockwise   - When true, the deformation and new path are on the "left" of the original path, otherwise on the "right"
+	*                       where the path is locally viewed as traveling from bottom to top
+	*/ 
+	bool DeformJoint( const int32 IncomingEID, const int32 OutgoingEID, int32 JointVID, 
+	                  TArray<FEdgePath::FDirectedSegment>& DeformedPath, const bool bClockwise );
+	
+	/**
+	* @return the ID of the vertex at the "head" of the specified segment
+	*/
+	inline int32 SegmentHeadVID(int32 SID) const;
+	
+	/**
+	* @return the ID of the vertex at the "tail" of the specified segment
+	*/
+	inline int32 SegmentTailVID(int32 SID) const;
+	
+	/**
+	* @return true if no part of the global path lies in the wedge (partial one-ring) on PathJoint.Interior side
+	*/ 
+	bool IsJointFlexible(const FPathJoint& PathJoint) const;
+	
+	/**
+	* Compute data associated with the joint formed by the specified segment (outgoing), and the preceding segment (incoming) in the path.
+	* Also update the JointAngleQueue that tracks straightness of each joint.
+	*/ 
+	void UpdateJointAndQueue(int32 OutgoingSegmentID);
+	
+	/**
+	* remove specified segment from the linked-list path
+	* and clean up joint-related data for the joint that ends with this segment.
+	* note: the joint data (e.g. joint angle queue) for the joint that started with this segment will be dirty.
+	*/ 
+	void RemoveSegment(int32 SID);
+
+	/**
+	* Splice in a path to replace the joint and update the joint angle queue.
+	*/ 
+	void ReplaceJointWithPath(const FPathJoint PathJoint, const TArray<FEdgePath::FDirectedSegment>& PathAsDirectedEdges);
+
+	/**
+	* Traveling either CCW or CW around vertex CenterVID, from edge StartEID to edge EndEID, flip edges adjacent to CVID.
+	* return true on success, and false if one of the edges was a boundary edge.  
+	* 
+	* It is assumed the specified edges do in fact have vertex CVID in common, but this is not checked.
+	*
+	* Note: in the false case, some of the edges may have been flipped before encountering the boundary edge
+	*       - no attempt is made to restore the mesh.
+	*       
+	*/
+	bool OuterArcFlipEdges(int32 StartEID, int32 EndEID, int32 CenterVID, bool bClockwise = false);
+	
+	/**
+	* Compute the two angles formed at the vertex CenterVID by the two wedges formed by splitting the local one-ring by local path
+	* specified as incoming and outgoing edges.
+	* 
+	* It is assumed the specified edges do in fact have vertex CVID in common, but this is not checked.
+	*/
+	void ComputeWedgeAngles(int32 IncomingEID, int32 OutgoingEID, int32 CenterVID, 
+	                        double& LeftSideAngle, double& RightSideAngle) const ;
+	
+protected:
+
+	FIntrinsicEdgeFlipMesh         EdgeFlipMesh;        // Intrinsic mesh with the same vertices as the original surface mesh. The geodesic path is comprised of edges in this mesh.
+
+	double                         PathLength;          // Total length of the path
+	int32                          NumFlips;            // Count of the number of edge flips performed in shortening the path.
+
+	FEdgePath                      EdgePath;            // Path defined relative to the edges in the EdgeFlipMesh.
+	TMap<int32, TArray<int32>>     EIDToSIDsMap;        // Map mesh edge ID to array of path segments.  A path could traverse the same mesh edge multiple times.
+	TDynamicVector<FPathJoint>     PathJoints;          // List of directed segment joints in the path, indexed by SID of Outgoing segment.	
+	FBinaryHeap<double>            JointAngleQueue;     // Priority queue that can automatically resize if needed.  Priority based on the smaller angle between the two edges in a joint
+	
+};
+
+const FEdgePath& FDeformableEdgePath::GetEdgePath() const
+{
+	return EdgePath;
+}
+
+double  FDeformableEdgePath::GetPathLength() const
+{
+	return PathLength;
+}
+
+const FIntrinsicEdgeFlipMesh& FDeformableEdgePath::GetIntrinsicMesh() const
+{
+	return EdgeFlipMesh;
+}
+
+int32 FDeformableEdgePath::SegmentHeadVID(int32 SID) const
+{
+	const auto Directed = EdgePath.GetSegment(SID);
+	const FIndex2i EdgeV = EdgeFlipMesh.GetEdgeV(Directed.EID);
+	return EdgeV[Directed.HeadIndex];
+}
+int32 FDeformableEdgePath::SegmentTailVID(int32 SID) const
+{
+	const auto Directed = EdgePath.GetSegment(SID);
+	const FIndex2i EdgeV = EdgeFlipMesh.GetEdgeV(Directed.EID);
+	return EdgeV[1 - Directed.HeadIndex];
+}
+
+
+int FEdgePath::NumSegments() const
+{
+	return PathLinkList.Num();
+}
+
+int32 FEdgePath::MaxSegmentID() const
+{
+	return DirectedSegmentRefCounts.GetMaxIndex();
+}
+
+bool FEdgePath::IsSegment(int32 SID) const
+{
+	return DirectedSegmentRefCounts.IsValid(SID);
+}
+
+
+}; // end namespace Geometry
+}; // end namespace UE
