DynamicMesh: fix trivial bug in PolygroupsGenerator, was using wrong array size variable
#rnx

diff --git a/Engine/Plugins/Runtime/GeometryProcessing/Source/DynamicMesh/Private/Polygroups/PolygroupsGenerator.cpp b/Engine/Plugins/Runtime/GeometryProcessing/Source/DynamicMesh/Private/Polygroups/PolygroupsGenerator.cpp
--- a/Engine/Plugins/Runtime/GeometryProcessing/Source/DynamicMesh/Private/Polygroups/PolygroupsGenerator.cpp
+++ b/Engine/Plugins/Runtime/GeometryProcessing/Source/DynamicMesh/Private/Polygroups/PolygroupsGenerator.cpp
@@ -358,7 +358,7 @@ bool FPolygroupsGenerator::FindPolygroupsFromFurthestPointSampling(
 	}
 
 
-	for (int32 k = 0; k < NumPoints; ++k)
+	for (int32 k = 0; k < TriSets.Num(); ++k)
 	{
 		if (TriSets[k].Num() > 0)
 		{
