make dynamicmesh3 refuse to merge an edge if it would create a duplicate triangle in the mesh
+ make mesh editor's mesh append function robust to failures to append triangles

#rnx

diff --git a/Engine/Plugins/Runtime/GeometryProcessing/Source/DynamicMesh/Private/DynamicMeshEditor.cpp b/Engine/Plugins/Runtime/GeometryProcessing/Source/DynamicMesh/Private/DynamicMeshEditor.cpp
--- a/Engine/Plugins/Runtime/GeometryProcessing/Source/DynamicMesh/Private/DynamicMeshEditor.cpp
+++ b/Engine/Plugins/Runtime/GeometryProcessing/Source/DynamicMesh/Private/DynamicMeshEditor.cpp
@@ -1497,7 +1497,10 @@ void FDynamicMeshEditor::AppendMesh(const FDynamicMesh3* AppendMesh,
 
 		FIndex3i Tri = AppendMesh->GetTriangle(TriID);
 		int NewTriID = Mesh->AppendTriangle(VertexMap.GetTo(Tri.A), VertexMap.GetTo(Tri.B), VertexMap.GetTo(Tri.C), GroupID);
-		TriangleMap.Add(TriID, NewTriID);
+		if (ensure(NewTriID >= 0))
+		{
+			TriangleMap.Add(TriID, NewTriID);
+		}
 	}
 
 
@@ -1550,9 +1553,9 @@ void FDynamicMeshEditor::AppendMesh(const FDynamicMesh3* AppendMesh,
 		{
 			const FDynamicMeshMaterialAttribute* FromMaterialIDs = AppendMesh->Attributes()->GetMaterialID();
 			FDynamicMeshMaterialAttribute* ToMaterialIDs = Mesh->Attributes()->GetMaterialID();
-			for (int TriID : AppendMesh->TriangleIndicesItr())
+			for (const TPair<int32, int32>& MapTID : TriangleMap.GetForwardMap())
 			{
-				ToMaterialIDs->SetValue(TriangleMap.GetTo(TriID), FromMaterialIDs->GetValue(TriID));
+				ToMaterialIDs->SetValue(MapTID.Value, FromMaterialIDs->GetValue(MapTID.Key));
 			}
 		}
 
@@ -1562,9 +1565,9 @@ void FDynamicMeshEditor::AppendMesh(const FDynamicMesh3* AppendMesh,
 			// TODO: remap groups? this will be somewhat expensive...
 			const FDynamicMeshPolygroupAttribute* FromPolygroups = AppendMesh->Attributes()->GetPolygroupLayer(PolygroupLayerIndex);
 			FDynamicMeshPolygroupAttribute* ToPolygroups = Mesh->Attributes()->GetPolygroupLayer(PolygroupLayerIndex);
-			for (int TriID : AppendMesh->TriangleIndicesItr())
+			for (const TPair<int32, int32>& MapTID : TriangleMap.GetForwardMap())
 			{
-				ToPolygroups->SetValue(TriangleMap.GetTo(TriID), FromPolygroups->GetValue(TriID));
+				ToPolygroups->SetValue(MapTID.Value, FromPolygroups->GetValue(MapTID.Key));
 			}
 		}
 
@@ -1668,17 +1671,16 @@ void FDynamicMeshEditor::AppendNormals(const FDynamicMesh3* AppendMesh,
 	}
 
 	// now set new triangles
-	for (int TriID : AppendMesh->TriangleIndicesItr())
+	for (const TPair<int32, int32>& MapTID : TriangleMap.GetForwardMap())
 	{
-		if (FromNormals->IsSetTriangle(TriID))
+		if (FromNormals->IsSetTriangle(MapTID.Key))
 		{
-			FIndex3i ElemTri = FromNormals->GetTriangle(TriID);
-			int NewTriID = TriangleMap.GetTo(TriID);
+			FIndex3i ElemTri = FromNormals->GetTriangle(MapTID.Key);
 			for (int j = 0; j < 3; ++j)
 			{
 				ElemTri[j] = FromNormals->IsElement(ElemTri[j]) ? NormalMapOut.GetTo(ElemTri[j]) : FDynamicMesh3::InvalidID;
 			}
-			ToNormals->SetTriangle(NewTriID, ElemTri);
+			ToNormals->SetTriangle(MapTID.Value, ElemTri);
 		}
 	}
 }
@@ -1698,17 +1700,16 @@ void FDynamicMeshEditor::AppendUVs(const FDynamicMesh3* AppendMesh,
 	}
 
 	// now set new triangles
-	for (int TriID : AppendMesh->TriangleIndicesItr())
+	for (const TPair<int32, int32>& MapTID : TriangleMap.GetForwardMap())
 	{
-		if (FromUVs->IsSetTriangle(TriID))
+		if (FromUVs->IsSetTriangle(MapTID.Key))
 		{
-			FIndex3i ElemTri = FromUVs->GetTriangle(TriID);
-			int NewTriID = TriangleMap.GetTo(TriID);
+			FIndex3i ElemTri = FromUVs->GetTriangle(MapTID.Key);
 			for (int j = 0; j < 3; ++j)
 			{
 				ElemTri[j] = FromUVs->IsElement(ElemTri[j]) ? UVMapOut.GetTo(ElemTri[j]) : FDynamicMesh3::InvalidID;
 			}
-			ToUVs->SetTriangle(NewTriID, ElemTri);
+			ToUVs->SetTriangle(MapTID.Value, ElemTri);
 		}
 	}
 }
@@ -1726,17 +1727,16 @@ void FDynamicMeshEditor::AppendColors(const FDynamicMesh3* AppendMesh,
 	}
 
 	// now set new triangles
-	for (int TriID : AppendMesh->TriangleIndicesItr())
+	for (const TPair<int32, int32>& MapTID : TriangleMap.GetForwardMap())
 	{
-		if (FromOverlay->IsSetTriangle(TriID))
+		if (FromOverlay->IsSetTriangle(MapTID.Key))
 		{
-			FIndex3i ElemTri = FromOverlay->GetTriangle(TriID);
-			int NewTriID = TriangleMap.GetTo(TriID);
+			FIndex3i ElemTri = FromOverlay->GetTriangle(MapTID.Key);
 			for (int j = 0; j < 3; ++j)
 			{
 				ElemTri[j] = FromOverlay->IsElement(ElemTri[j]) ? MapOut.GetTo(ElemTri[j]) : FDynamicMesh3::InvalidID;
 			}
-			ToOverlay->SetTriangle(NewTriID, ElemTri);
+			ToOverlay->SetTriangle(MapTID.Value, ElemTri);
 		}
 	}
 }
diff --git a/Engine/Source/Runtime/GeometryCore/Private/DynamicMesh/DynamicMesh3_Edits.cpp b/Engine/Source/Runtime/GeometryCore/Private/DynamicMesh/DynamicMesh3_Edits.cpp
--- a/Engine/Source/Runtime/GeometryCore/Private/DynamicMesh/DynamicMesh3_Edits.cpp
+++ b/Engine/Source/Runtime/GeometryCore/Private/DynamicMesh/DynamicMesh3_Edits.cpp
@@ -1645,9 +1645,15 @@ EMeshResult FDynamicMesh3::MergeEdges(int eKeep, int eDiscard, FMergeEdgesInfo&
 	// we will not join triangles with incompatible winding order
 	// I can't see how to do this purely topologically.
 	// So relying on closest-pairs testing.
-	IndexUtil::OrientTriEdge(a, b, GetTriangle(tab));
-	//int tcd_otherv = OrientTriEdgeAndFindOtherVtx(ref c, ref d, GetTriangle(tcd));
-	IndexUtil::OrientTriEdge(c, d, GetTriangle(tcd));
+	int OppAB = IndexUtil::OrientTriEdgeAndFindOtherVtx(a, b, GetTriangle(tab));
+	int OppCD = IndexUtil::OrientTriEdgeAndFindOtherVtx(c, d, GetTriangle(tcd));
+
+	// Refuse to merge if doing so would create a duplicate triangle
+	if (OppAB == OppCD)
+	{
+		return EMeshResult::Failed_InvalidNeighbourhood;
+	}
+
 	int x = c; c = d; d = x;   // joinable bdry edges have opposing orientations, so flip to get ac and b/d correspondences
 	FVector3d Va = GetVertex(a), Vb = GetVertex(b), Vc = GetVertex(c), Vd = GetVertex(d);
 	if (bCheckValidOrientation && 
