Dynamic Mesh, Operations, GeodesicPath - remove unwanted include.
#rnx

diff --git a/Engine/Plugins/Runtime/GeometryProcessing/Source/DynamicMesh/Public/Operations/GeodesicPath.h b/Engine/Plugins/Runtime/GeometryProcessing/Source/DynamicMesh/Public/Operations/GeodesicPath.h
--- a/Engine/Plugins/Runtime/GeometryProcessing/Source/DynamicMesh/Public/Operations/GeodesicPath.h
+++ b/Engine/Plugins/Runtime/GeometryProcessing/Source/DynamicMesh/Public/Operations/GeodesicPath.h
@@ -16,8 +16,6 @@
 #include "MathUtil.h"
 
 
-#include "ExplicitUseGeometryMathTypes.h"
-
 namespace UE
 {
 namespace Geometry
