fix simplification sometimes bailing early due to invalid entries in queue

diff --git a/Engine/Plugins/Runtime/GeometryProcessing/Source/DynamicMesh/Private/MeshSimplification.cpp b/Engine/Plugins/Runtime/GeometryProcessing/Source/DynamicMesh/Private/MeshSimplification.cpp
--- a/Engine/Plugins/Runtime/GeometryProcessing/Source/DynamicMesh/Private/MeshSimplification.cpp
+++ b/Engine/Plugins/Runtime/GeometryProcessing/Source/DynamicMesh/Private/MeshSimplification.cpp
@@ -219,7 +219,7 @@ void TMeshSimplification<QuadricErrorType>::InitializeQueue()
 	EdgeQuadrics.SetNum(MaxEID);
 	EdgeQueue.Initialize(MaxEID);
 	TArray<FEdgeError> EdgeErrors;
-	EdgeErrors.SetNum(MaxEID);
+	EdgeErrors.Init(FEdgeError{MAX_FLT, -1}, MaxEID);
 
 	// @todo vertex quadrics can be computed in parallel
 	//gParallel.BlockStartEnd(0, MaxEID - 1, (start_eid, end_eid) = > {
