Dynamic Mesh, Operations: Intrinsic Mesh - export intrinsic mesh edges as array of surface points.

diff --git a/Engine/Plugins/Runtime/GeometryProcessing/Source/DynamicMesh/Private/Operations/IntrinsicTriangulationMesh.cpp b/Engine/Plugins/Runtime/GeometryProcessing/Source/DynamicMesh/Private/Operations/IntrinsicTriangulationMesh.cpp
--- a/Engine/Plugins/Runtime/GeometryProcessing/Source/DynamicMesh/Private/Operations/IntrinsicTriangulationMesh.cpp
+++ b/Engine/Plugins/Runtime/GeometryProcessing/Source/DynamicMesh/Private/Operations/IntrinsicTriangulationMesh.cpp
@@ -759,6 +759,57 @@ FIntrinsicTriangulation::FSurfacePoint::FSurfacePoint(int32 TriID, const FVector
 	Position.TriPosition.BarycentricCoords[2] = BaryCentrics[2];
 }
 
+FVector3d FIntrinsicTriangulation::FSurfacePoint::FSurfacePoint::AsR3Position(const FDynamicMesh3& Mesh, bool& bIsValid) const
+{
+	FVector3d Result(0., 0., 0.);
+	switch (PositionType)
+	{
+		case FSurfacePoint::EPositionType::Vertex:
+		{
+			const int32 SurfaceVID = Position.VertexPosition.VID;
+			bIsValid = Mesh.IsVertex(SurfaceVID);
+			if (bIsValid)
+			{
+				Result = Mesh.GetVertex(SurfaceVID);
+			}
+		}
+		break;
+		case FSurfacePoint::EPositionType::Edge:
+		{
+			const int32 SurfaceEID = Position.EdgePosition.EdgeID;
+			double Alpha = Position.EdgePosition.Alpha;
+			bIsValid = Mesh.IsEdge(SurfaceEID);
+			if (bIsValid)
+			{
+				const FIndex2i SurfaceEdgeV = Mesh.GetEdgeV(SurfaceEID);
+				Result = Alpha * (Mesh.GetVertex(SurfaceEdgeV.A)) + (1. - Alpha) * (Mesh.GetVertex(SurfaceEdgeV.B));
+			}
+		}
+		break;
+		case FSurfacePoint::EPositionType::Triangle:
+		{
+			const int32 SurfaceTID = Position.TriPosition.TriID;
+			bIsValid = Mesh.IsTriangle(SurfaceTID);
+
+			if (bIsValid)
+			{
+				Result = Mesh.GetTriBaryPoint(SurfaceTID,
+					Position.TriPosition.BarycentricCoords[0],
+					Position.TriPosition.BarycentricCoords[1],
+					Position.TriPosition.BarycentricCoords[2]);
+			}
+
+		}
+		break;
+		default:
+		{
+			bIsValid = false;
+			// shouldn't be able to reach this point
+			check(0);
+		}
+	}
+	return Result;
+}
 
 UE::Geometry::FIntrinsicTriangulation::FIntrinsicTriangulation(const FDynamicMesh3& SrcMesh)
 	: MyBase(SrcMesh)
@@ -778,13 +829,35 @@ UE::Geometry::FIntrinsicTriangulation::FIntrinsicTriangulation(const FDynamicMes
 		{
 			if (ExtrinsicMesh->IsVertex(VID))
 			{
-				int32 FirstEID = -1;
-				for (int32 EID : ExtrinsicMesh->VtxEdgesItr(VID))
+				int32 RefEID = FDynamicMesh3::InvalidID;
+				if (ExtrinsicMesh->IsBoundaryVertex(VID) && !ExtrinsicMesh->IsBowtieVertex(VID))
+				{
+					// vertex is on the mesh boundary, 
+					// choose the boundary edge such that traveling CCW (about the vertex) moves into the mesh
+					for (int32 EID : ExtrinsicMesh->VtxEdgesItr(VID))
+					{
+						if (ExtrinsicMesh->IsBoundaryEdge(EID))
+						{
+							const int32 TID = ExtrinsicMesh->GetEdgeT(EID).A;
+							const int32 IndexOf = ExtrinsicMesh->GetTriEdges(TID).IndexOf(EID);
+							if (ExtrinsicMesh->GetTriangle(TID)[IndexOf] == VID)
+							{
+								RefEID = EID;
+								break;
+							}
+						}
+					}
+				}
+				else
 				{
-					FirstEID = EID;
-					break;
+					// vertex isn't on a mesh boundary, just pick the first adj edge to be the reference edge
+					for (int32 EID : ExtrinsicMesh->VtxEdgesItr(VID))
+					{
+						RefEID = EID;
+						break;
+					}
 				}
-				VIDToReferenceEID[VID] = FirstEID;
+				VIDToReferenceEID[VID] = RefEID;
 			}
 		}
 		for (int32 TID = 0; TID < MaxTriangleID; ++TID)
@@ -925,6 +998,182 @@ UE::Geometry::FIntrinsicTriangulation::FIntrinsicTriangulation(const FDynamicMes
 
 }
 
+TArray<FIntrinsicTriangulation::FSurfacePoint> UE::Geometry::FIntrinsicTriangulation::TraceEdge(int32 EID, double CoalesceThreshold, bool bReverse) const
+{
+	TArray<FIntrinsicTriangulation::FSurfacePoint> SurfacePoints;
+	if (!IsEdge(EID))
+	{
+		return SurfacePoints;
+	}
+
+
+	const FDynamicMesh3* SurfaceMesh = GetExtrinsicMesh();
+	const double IntrinsicEdgeLength = GetEdgeLength(EID);
+	const FIndex2i EdgeV = GetEdgeV(EID);
+	int32 StartVID = EdgeV.A;
+	int32 EndVID = EdgeV.B;
+	if (bReverse)
+	{
+		StartVID = EdgeV.B;
+		EndVID = EdgeV.A;
+	}
+
+
+	FIndex2i AdjTIDs = GetEdgeT(EID);
+
+	// look-up the polar angle of this edge as it leaves StartVID ( this angle is relative to a specified reference mesh edge )
+	const double PolarAngle = 
+		[&]{
+			const int32 AdjTID = AdjTIDs.A;
+			const int32 IndexOf = GetTriEdges(AdjTID).IndexOf(EID);
+			checkSlow(IndexOf > -1);
+			if (GetTriangle(AdjTID)[IndexOf] == StartVID)
+			{
+				// IntrinsicEdgeAngles hold the local angle of each out-going edge relative to the vertex the edge exits
+				return IntrinsicEdgeAngles[AdjTID][IndexOf];
+			}
+			else
+			{
+				const double ToRadians = GeometricVertexInfo[StartVID].ToRadians;
+				// polar angle of prev edge just before (clockwise) the edge we want
+				const double PrevPolarAngle = IntrinsicEdgeAngles[AdjTID][(IndexOf + 1) % 3];
+				// angle between previous edge and this one
+				const double InternalAngle = InternalAngles[AdjTID][(IndexOf + 1) % 3];
+				// add the internal angle to rotate from Prev Edge to this edge
+				return ToRadians * InternalAngle + PrevPolarAngle;
+			}
+		}();
+
+	// Trace the surface mesh from the intrinsic StartVID, in the PolarAngle direction, a distance of IntrinsicEdgeLength.
+	// 
+	// Note: this intrisic vertex may or may not correspond to a vertex in the surface mesh if vertices were added to the intrinsic mesh
+	// by doing an edge split or a triangle poke.
+	FMeshGeodesicSurfaceTracer SurfaceTracer(*SurfaceMesh);
+	{
+		const FSurfacePoint& StartSurfacePoint = GetVertexSurfacePoint(StartVID);
+		const FSurfacePoint::FSurfacePositionUnion& SurfacePosition = StartSurfacePoint.Position;
+		switch (StartSurfacePoint.PositionType)
+		{
+			case FSurfacePoint::EPositionType::Vertex:
+			{
+				const int32 SurfaceVID = SurfacePosition.VertexPosition.VID;
+				const int32 RefSurfaceEID = VIDToReferenceEID[SurfaceVID];
+
+				FMeshGeodesicSurfaceTracer::FMeshTangentDirection TangentAtVertex = { SurfaceVID, RefSurfaceEID, PolarAngle };
+				SurfaceTracer.TraceMeshFromVertex(TangentAtVertex, IntrinsicEdgeLength);
+			}
+			break;
+			case FSurfacePoint::EPositionType::Edge:
+			{
+				const int32 RefSurfaceEID = SurfacePosition.EdgePosition.EdgeID;
+				double Alpha = SurfacePosition.EdgePosition.Alpha;
+				// convert to BC
+				const int32 SurfaceTID = SurfaceMesh->GetEdgeT(RefSurfaceEID).A;
+				const FIndex2i SurfaceEdgeV = SurfaceMesh->GetEdgeV(RefSurfaceEID);
+				const int32 IndexOf = SurfaceMesh->GetTriEdges(SurfaceTID).IndexOf(RefSurfaceEID);
+				const bool bSameOrientation = (SurfaceMesh->GetTriangle(SurfaceTID)[IndexOf] == SurfaceEdgeV.A);
+				FVector3d BaryPoint(0., 0., 0.);
+				if (!bSameOrientation)
+				{
+					Alpha = 1. - Alpha;
+				}
+				BaryPoint[IndexOf] = Alpha;
+				BaryPoint[(IndexOf + 1) % 3] = 1. - Alpha;
+
+				FMeshGeodesicSurfaceTracer::FMeshSurfaceDirection DirectionOnSurface(RefSurfaceEID, PolarAngle);
+				SurfaceTracer.TraceMeshFromBaryPoint(SurfaceTID, BaryPoint, DirectionOnSurface, IntrinsicEdgeLength);
+			}
+			break;
+			case FSurfacePoint::EPositionType::Triangle:
+			{
+				const int32 SurfaceTID = SurfacePosition.TriPosition.TriID;
+
+				const FVector3d BaryPoint(SurfacePosition.TriPosition.BarycentricCoords[0],
+					SurfacePosition.TriPosition.BarycentricCoords[1],
+					SurfacePosition.TriPosition.BarycentricCoords[2]);
+				const int32 RefSurfaceEID = TIDToReferenceEID[SurfaceTID];
+
+				FMeshGeodesicSurfaceTracer::FMeshSurfaceDirection DirectionOnSurface(RefSurfaceEID, PolarAngle);
+				SurfaceTracer.TraceMeshFromBaryPoint(SurfaceTID, BaryPoint, DirectionOnSurface, IntrinsicEdgeLength);
+			}
+			break;
+			default:
+			{
+				// shouldn't be able to reach this point
+				check(0);
+			}
+		}
+	}
+
+	// util to convert the trace result to a surface point, potentially snapping to a vertex if within the coalesce threshold 
+	auto TraceResultToSurfacePoint = [CoalesceThreshold, SurfaceMesh](const FMeshGeodesicSurfaceTracer::FTraceResult& TraceResult)->FSurfacePoint
+	{
+
+	
+		if (TraceResult.bIsEdgePoint)
+		{
+			// TraceResult alpha is defined as EdgeV.A * (1-Alpha) + EdgeV.B Alpha;  This is the complement of what we want
+			const double Alpha = FMath::Clamp((1. - TraceResult.EdgeAlpha), 0., 1.);
+			const int32 EID = TraceResult.EdgeID;
+
+			if (Alpha <= 0.5 && Alpha < CoalesceThreshold)
+			{
+				return FSurfacePoint(SurfaceMesh->GetEdgeV(EID).B);
+			}
+			else if (Alpha > 0.5 && (1. - Alpha) < CoalesceThreshold)
+			{
+				return FSurfacePoint(SurfaceMesh->GetEdgeV(EID).A);
+			}
+
+			return FSurfacePoint(EID, Alpha);
+		}
+		else
+		{
+			// TODO - should snap to vertex / edge if close?
+			const int32 TID = TraceResult.TriID;
+			const FVector3d BC = TraceResult.Barycentric;
+			return FSurfacePoint(TID, BC);
+		}
+	};
+
+	// Add surface point to the outgoing array, but don't allow for duplicate vertex points
+	auto AddSurfacePoint = [&SurfacePoints](const FSurfacePoint& PointA)
+	{
+		if (SurfacePoints.Num() == 0)
+		{
+			SurfacePoints.Add(PointA);
+		}
+		else
+		{
+			const FSurfacePoint& PointB = SurfacePoints.Last();
+			const bool bAreSameVertexPoint = ( PointA.PositionType == FSurfacePoint::EPositionType::Vertex &&
+				                               PointB.PositionType == FSurfacePoint::EPositionType::Vertex &&
+				                               PointA.Position.VertexPosition.VID == PointB.Position.VertexPosition.VID );
+			if (!bAreSameVertexPoint)
+			{
+				SurfacePoints.Add(PointA);
+			}
+		}
+	};
+
+	// package the surface trace results as series of surface points.  Note, because this is a trace along an intrinsic edge
+	// the first and last result will be an intrinsic vertex (StartVID and EndVID)
+	TArray<FMeshGeodesicSurfaceTracer::FTraceResult>& TraceResults = SurfaceTracer.GetTraceResults();
+	{
+		int32 NumTraceResults = TraceResults.Num();
+		AddSurfacePoint(GetVertexSurfacePoint(StartVID));
+		for (int32 i = 1; i < NumTraceResults - 1; ++i)
+		{
+			FMeshGeodesicSurfaceTracer::FTraceResult& TraceResult = TraceResults[i];
+			FSurfacePoint SurfacePoint = TraceResultToSurfacePoint(TraceResult);
+			AddSurfacePoint(SurfacePoint);
+		}
+		AddSurfacePoint(GetVertexSurfacePoint(EndVID));
+	}
+
+
+	return SurfacePoints;
+}
 
 EMeshResult UE::Geometry::FIntrinsicTriangulation::FlipEdge(int32 EID, FEdgeFlipInfo& EdgeFlipInfo)
 {
diff --git a/Engine/Plugins/Runtime/GeometryProcessing/Source/DynamicMesh/Private/Operations/MeshGeodesicSurfaceTracer.cpp b/Engine/Plugins/Runtime/GeometryProcessing/Source/DynamicMesh/Private/Operations/MeshGeodesicSurfaceTracer.cpp
--- a/Engine/Plugins/Runtime/GeometryProcessing/Source/DynamicMesh/Private/Operations/MeshGeodesicSurfaceTracer.cpp
+++ b/Engine/Plugins/Runtime/GeometryProcessing/Source/DynamicMesh/Private/Operations/MeshGeodesicSurfaceTracer.cpp
@@ -347,7 +347,7 @@ GeodesicSingleTriangleUtils::FTraceResult  GeodesicSingleTriangleUtils::TraceTan
 		const FVector2d RotatedRayDir = TangentTri2.ChangeBasis(RayDir, HitEdgeIndex);
 
 		// populate trace result
-		TraceResult.Classification = ETraceClassification::Continuing;
+		TraceResult.Classification = (FMath::Abs(RayDistance - MaxDistance) < TMathUtilConstants<double>::ZeroTolerance) ? ETraceClassification::DistanceTerminated : ETraceClassification::Continuing;
 		TraceResult.TraceDist = RayDistance;
 		TraceResult.TriID = TangentTri2.Tri3ID;
 
@@ -500,102 +500,62 @@ GeodesicSingleTriangleUtils::FTraceResult GeodesicSingleTriangleUtils::TraceFrom
 }
 GeodesicSingleTriangleUtils::FTraceResult GeodesicSingleTriangleUtils::TraceFromVertex(const FDynamicMesh3& Mesh, int32 VID, const FMeshSurfaceDirection& Direction, FTangentTri2& ScratchTangentTri2, double MaxDistance)
 {
-
-	int32 EID = Direction.EdgeID;
-	FIndex2i Edge = Mesh.GetEdgeV(EID);
-	check(Edge.A == VID || Edge.B == VID);
-
-	// This doesn't work at a bowtie, or a boundary vertex. 
+	// This doesn't work at a bowtie
+	// [todo] make this work form a bowtie, potentially returning an immediately terminated trace
 	check(!Mesh.IsBowtieVertex(VID));
-	check(!Mesh.IsBoundaryVertex(VID));
-
-	// @todo: return a boundary terminated result in place of the checks
-
-
-
-	// Compute local tangent plane information. 
-	struct FInternalAngle
-	{
-		int32 TriID;
-		double Angle; // internal angle
-		double OffsetAngle; // sum of prior internal angles.
-	};
 
+	// Find total angle around the vertex.
 	double TotalAngle = 0;
-	TArray<FInternalAngle> InternalAngles;
-	FIndex2i FirstLastTri = Mesh.GetEdgeT(EID);
-
-	// Travel counter clockwise around the vertex, visiting each triangle in order.
-	// Start with the first triangle we visit when traveling counter clockwise around VID, starting with edge EID.
-	int CurTriID = FirstLastTri.A;
-	if (GetPermutedEdges(Mesh, FirstLastTri.B, VID).A == EID)
+	const bool bIsBoundaryVertex = Mesh.IsBoundaryVertex(VID);
+	const int32 RefEID = Direction.EdgeID;
 	{
-		CurTriID = FirstLastTri.B;
-	}
-	int32 FirstTriID = CurTriID;
-	do
-	{
-		FIndex3i TriVIDS = Mesh.GetTriangle(CurTriID);
-
-		const int32 TriIndex = IndexUtil::FindTriIndex(VID, TriVIDS);
-		const double AngleR = Mesh.GetTriInternalAngleR(CurTriID, TriIndex);
-
-
-		FInternalAngle& InternalAngle = InternalAngles.AddDefaulted_GetRef();
-		InternalAngle.TriID = CurTriID;
-		InternalAngle.Angle = AngleR;
-		InternalAngle.OffsetAngle = TotalAngle;
-
-		CurTriID = NextCCWTriangle(Mesh, CurTriID, VID);
-		TotalAngle += AngleR;
-
-	} while (CurTriID != FirstTriID && CurTriID != FDynamicMesh3::InvalidID);
-
-	// The direction of the ray was encoded as RayDir by assuming the edge is directed from Edge.A to Edge.B
-	// and aligned with the x-axis.
-	// if we trace from the vertex Edge.B, we need to flip the ray
-	const FVector2d RayDir = (VID == Edge.B) ? -Direction.Dir : Direction.Dir;
-
-
-	FVector2d InComingRay = -RayDir;
-	double IndicentAngle = FMath::Atan2(InComingRay.Y, InComingRay.X); // angle on a triangle face, not the tangent plane
-
-	double ExitPolarAngle = IndicentAngle * (TMathUtilConstants<double>::TwoPi / TotalAngle) + TMathUtilConstants<double>::Pi;
-	ExitPolarAngle = AsZeroToTwoPi(ExitPolarAngle);
+		int32 FirstEID = RefEID;
+		if (bIsBoundaryVertex)
+		{
+			// vertex is on the mesh boundary, 
+			// choose the boundary edge such that traveling CCW (about the vertex) moves into the mesh
 
-	double ExitAngle = ExitPolarAngle * (TotalAngle / TMathUtilConstants<double>::TwoPi);
+			for (int32 EID : Mesh.VtxEdgesItr(VID))
+			{
+				if (Mesh.IsBoundaryEdge(EID))
+				{
+					const int32 TID = Mesh.GetEdgeT(EID).A;
+					const int32 IndexOf = Mesh.GetTriEdges(TID).IndexOf(EID);
+					if (Mesh.GetTriangle(TID)[IndexOf] == VID)
+					{
+						FirstEID = EID;
+						break;
+					}
+				}
+			}
+		}
+		const FIndex2i FirstLastTri = Mesh.GetEdgeT(FirstEID);
 
-	// find the triangle that the ray enters, and angle from reference to the first side of that triangle
-	int32 DstTriID = FDynamicMesh3::InvalidID;
-	double DstOffsetAngle = 0;
-	for (const FInternalAngle& AngleData : InternalAngles)
-	{
-		if (AngleData.OffsetAngle <= ExitAngle && (AngleData.OffsetAngle + AngleData.Angle) >= ExitAngle)
+		// Travel counter-clockwise around the vertex, visiting each triangle in order.
+		// Start with the first triangle we visit when traveling counter clockwise around VID, starting with edge EID.
+		int CurTriID = FirstLastTri.A;
+		if (!bIsBoundaryVertex && GetPermutedEdges(Mesh, FirstLastTri.B, VID).A == RefEID)
 		{
-			DstTriID = AngleData.TriID;
-			DstOffsetAngle = AngleData.OffsetAngle;
-			break;
+			CurTriID = FirstLastTri.B;
 		}
-	}
-	checkSlow(DstTriID != FDynamicMesh3::InvalidID);
-
-	// What is the first edge of the Dst Tri?
-	const FIndex3i PermutedDstEdges = GetPermutedEdges(Mesh, DstTriID, VID);
-	int DstEdgeID = PermutedDstEdges.A; // First Edge.
+		int FirstTriID = CurTriID;
+		do
+		{
+			const FIndex3i TriVIDS = Mesh.GetTriangle(CurTriID);
+			const int32 TriIndex = IndexUtil::FindTriIndex(VID, TriVIDS);
+			const double AngleR = Mesh.GetTriInternalAngleR(CurTriID, TriIndex);
 
-	double AngleEdgeToRay = ExitAngle - DstOffsetAngle;
+			CurTriID = NextCCWTriangle(Mesh, CurTriID, VID);
+			TotalAngle += AngleR;
 
-	// ray relative to edge exiting VID
-	double DstAlpha = 0.;
-	FVector2d DstRayDir(FMath::Cos(AngleEdgeToRay), FMath::Sin(AngleEdgeToRay));
-	FIndex2i DstEdge = Mesh.GetEdgeV(DstEdgeID);
-	if (DstEdge.B == VID)
-	{
-		DstAlpha = 1. - DstAlpha;
-		DstRayDir = -DstRayDir;
+		} while (CurTriID != FirstTriID && CurTriID != FDynamicMesh3::InvalidID);
 	}
-	FMeshSurfaceDirection OutGoingDirection(DstEdgeID, DstRayDir);
-	return TraceTriangleFromEdge(Mesh, DstTriID, DstAlpha, OutGoingDirection, ScratchTangentTri2, MaxDistance);
+
+	//convert input from a vector defined relative to a surface tri to a vector on the tangent plane.
+	const double ToRadians = TMathUtilConstants<double>::TwoPi / FMath::Max(TotalAngle, TMathUtilConstants<double>::ZeroTolerance);
+	const double PolarAngleR = AsZeroToTwoPi(ToRadians * FMath::Atan2(Direction.Dir[1], Direction.Dir[0]));
+	FMeshTangentDirection TangentDirAtVertex = { VID, RefEID, PolarAngleR };
+	return TraceFromVertex(Mesh, TangentDirAtVertex, ScratchTangentTri2, MaxDistance);
 }
 
 
@@ -607,15 +567,16 @@ GeodesicSingleTriangleUtils::FTraceResult  GeodesicSingleTriangleUtils::TraceFro
 }
 GeodesicSingleTriangleUtils::FTraceResult  GeodesicSingleTriangleUtils::TraceFromVertex(const FDynamicMesh3& Mesh, const FMeshTangentDirection& TangentDirection, FTangentTri2& ScratchTangentTri2, double MaxDistance)
 {
-	int32 EID = TangentDirection.EdgeID;
-	int32 VID = TangentDirection.VID;
+	const int32 RefEID = TangentDirection.EdgeID;
+	const int32 VID    = TangentDirection.VID;
+
+	FIndex2i RefEdgeV = Mesh.GetEdgeV(RefEID);
+	check(RefEdgeV.A == VID || RefEdgeV.B == VID);
 
-	FIndex2i EdgeV = Mesh.GetEdgeV(EID);
-	check(EdgeV.A == VID || EdgeV.B == VID);
-	// This doesn't work at a bowtie, or a boundary vertex. 
-	// @todo: return a boundary terminated result in place of the checks
+	// This doesn't work at a bowtie
+	// [todo] make this work form a bowtie, potentially returning an imediately terminated trace
 	check(!Mesh.IsBowtieVertex(VID));
-	check(!Mesh.IsBoundaryVertex(VID));
+	const bool bIsBoundaryVertex = Mesh.IsBoundaryVertex(VID);
 
 
 	// Compute local tangent plane information. 
@@ -626,39 +587,68 @@ GeodesicSingleTriangleUtils::FTraceResult  GeodesicSingleTriangleUtils::TraceFro
 		double OffsetAngle; // sum of prior internal angles.
 	};
 
+	int32 FirstEID = RefEID;
+	if (bIsBoundaryVertex)
+	{
+		// vertex is on the mesh boundary, 
+		// choose the boundary edge such that traveling CCW (about the vertex) moves into the mesh
+
+		for (int32 EID : Mesh.VtxEdgesItr(VID))
+		{
+			if (Mesh.IsBoundaryEdge(EID))
+			{
+				const int32 TID = Mesh.GetEdgeT(EID).A;
+				const int32 IndexOf = Mesh.GetTriEdges(TID).IndexOf(EID);
+				if (Mesh.GetTriangle(TID)[IndexOf] == VID)
+				{
+					FirstEID = EID;
+					break;
+				}
+			}
+		}
+	}
+	const FIndex2i FirstLastTri = Mesh.GetEdgeT(FirstEID);
 	double TotalAngle = 0;
 	TArray<FInternalAngle> InternalAngles;
-	FIndex2i FirstLastTri = Mesh.GetEdgeT(EID);
-
 	// Travel counter-clockwise around the vertex, visiting each triangle in order.
 	// Start with the first triangle we visit when traveling counter clockwise around VID, starting with edge EID.
 	int CurTriID = FirstLastTri.A;
-	if (GetPermutedEdges(Mesh, FirstLastTri.B, VID).A == EID)
+	if (!bIsBoundaryVertex && GetPermutedEdges(Mesh, FirstLastTri.B, VID).A == RefEID)
 	{
 		CurTriID = FirstLastTri.B;
 	}
 	int FirstTriID = CurTriID;
+	double AngleOfRefEID = 0.; // to be computed. the offset from the first edge to the reference edge.
 	do
 	{
 		FIndex3i TriVIDS = Mesh.GetTriangle(CurTriID);
 
 		const int32 TriIndex = IndexUtil::FindTriIndex(VID, TriVIDS);
+		const int32 CurEID = Mesh.GetTriEdges(CurTriID)[TriIndex];
 		const double AngleR = Mesh.GetTriInternalAngleR(CurTriID, TriIndex);
 
 
 		FInternalAngle& InternalAngle = InternalAngles.AddDefaulted_GetRef();
 		InternalAngle.TriID = CurTriID;
 		InternalAngle.Angle = AngleR;
-		InternalAngle.OffsetAngle = TotalAngle;
+		InternalAngle.OffsetAngle = TotalAngle; // counting from FirstEID;
+		if (CurEID == RefEID)
+		{
+			AngleOfRefEID = TotalAngle;
+		}
 
 		CurTriID = NextCCWTriangle(Mesh, CurTriID, VID);
 		TotalAngle += AngleR;
 
 	} while (CurTriID != FirstTriID && CurTriID != FDynamicMesh3::InvalidID);
 
+	for (FInternalAngle& InternalAngle : InternalAngles)
+	{
+		InternalAngle.OffsetAngle -= AngleOfRefEID;
+	}
 	const double ToSurfaceAngle = TotalAngle / TMathUtilConstants<double>::TwoPi;
 	// if the edge "ends" at VID then it aligns with Pi in the tangent space (i.e. the vector A->B aligns with 0 radians, and the vector B->A aligns with pi )
-	double PolarAngle = (EdgeV.B == VID) ? TangentDirection.PolarAngle + TMathUtilConstants<double>::Pi : TangentDirection.PolarAngle;
+	double PolarAngle = (RefEdgeV.B == VID) ? TangentDirection.PolarAngle + TMathUtilConstants<double>::Pi : TangentDirection.PolarAngle;
 	const double SurfaceAngle = AsZeroToTwoPi(PolarAngle) * ToSurfaceAngle;
 
 
@@ -674,7 +664,21 @@ GeodesicSingleTriangleUtils::FTraceResult  GeodesicSingleTriangleUtils::TraceFro
 			break;
 		}
 	}
-	checkSlow(DstTriID != FDynamicMesh3::InvalidID);
+
+	if (DstTriID == FDynamicMesh3::InvalidID)
+	{
+		// this case should only happen if we are exiting the mesh
+		check(bIsBoundaryVertex);
+
+		GeodesicSingleTriangleUtils::FTraceResult TraceResult;
+		TraceResult.Classification = ETraceClassification::BoundaryTerminated;
+		TraceResult.TraceDist = 0.;
+		TraceResult.bIsEdgePoint = true;
+		TraceResult.EdgeID = RefEID;
+		TraceResult.EdgeAlpha = (RefEdgeV.A == VID) ? 0. : 1.;
+
+		return TraceResult;
+	}
 
 	// What is the first edge of the Dst Tri?
 	const FIndex3i PermutedDstEdges = GetPermutedEdges(Mesh, DstTriID, VID);
diff --git a/Engine/Plugins/Runtime/GeometryProcessing/Source/DynamicMesh/Public/Operations/IntrinsicTriangulationMesh.h b/Engine/Plugins/Runtime/GeometryProcessing/Source/DynamicMesh/Public/Operations/IntrinsicTriangulationMesh.h
--- a/Engine/Plugins/Runtime/GeometryProcessing/Source/DynamicMesh/Public/Operations/IntrinsicTriangulationMesh.h
+++ b/Engine/Plugins/Runtime/GeometryProcessing/Source/DynamicMesh/Public/Operations/IntrinsicTriangulationMesh.h
@@ -376,12 +376,21 @@ public:
 		return ExtrinsicMesh;
 	}
 	
-	// vertex information
+	/** Location relative to the surface mesh : a surface mesh vertex, surface mesh edge - crossing or barycentric coords */
+	struct FSurfacePoint;
 
+	/**
+	* @return Array of FSurfacePoints relative to the ExtrinsicMesh that represent the specified intrinsic mesh edge.
+	* This is directed from vertex EdgeV.A to vertex EdgeV.B unless bReverse is true
+	*
+	* @param CoalesceThreshold - in barycentric units [0,1], edge crossings with in this threshold are snapped to the nearest vertex.
+	*                            and any resulting repetitions of vertex surface points are replaced with a single vertex surface point.
+	* @param bReverse          - if true, trace edge from EdgeV.B to EdgeV.A
+	*/
+	TArray<FIntrinsicTriangulation::FSurfacePoint> TraceEdge(int32 EID, double CoalesceThreshold = 0., bool bReverse = false) const;
 
-	struct FSurfacePoint;
 	/**
-	* return FSurfacePoint information used to map an intrinsic vertex to the surface of the original mesh.
+	* @return FSurfacePoint for the specified intrinsic vertex.
 	*/
 	const FIntrinsicTriangulation::FSurfacePoint& GetVertexSurfacePoint(int32 VID) const
 	{
@@ -481,6 +490,13 @@ public:
 
 		EPositionType PositionType;      // Position type, needed to correctly interpret the Position union
 		FSurfacePositionUnion Position;  // Position relative to vertex, edge, or triangle. 
+
+		/**
+		* @return the R3 postition of this surface point as applied to the input Mesh
+		* on return bValidPoint will be false if the edge / vertex / triangle ID referenced by this surface
+		* point is not an element of the mesh in question
+		*/
+		FVector3d AsR3Position(const FDynamicMesh3& Mesh, bool& bValidPoint) const;
 	};
 
 protected:
diff --git a/Engine/Plugins/Runtime/GeometryProcessing/Source/DynamicMesh/Public/Operations/MeshGeodesicSurfaceTracer.h b/Engine/Plugins/Runtime/GeometryProcessing/Source/DynamicMesh/Public/Operations/MeshGeodesicSurfaceTracer.h
--- a/Engine/Plugins/Runtime/GeometryProcessing/Source/DynamicMesh/Public/Operations/MeshGeodesicSurfaceTracer.h
+++ b/Engine/Plugins/Runtime/GeometryProcessing/Source/DynamicMesh/Public/Operations/MeshGeodesicSurfaceTracer.h
@@ -335,7 +335,7 @@ struct GeodesicSingleTriangleUtils::FTraceResult
 	FVector3d Barycentric;                     // Surface point as Barycentric coords.  Can be reconstructed with FDynamicMesh::GetTriBaryPoint(TriID, Barycentric[0], Barycentric[1], Barycentric[2]);
 	bool bIsEdgePoint;                         // Duplicate description of the surface point in the case of edge termination
 	int32 EdgeID = IndexConstants::InvalidID;  // In the case of Edge termination, SurfaceDirection.EdgeID = EdgeID
-	double EdgeAlpha = -1.;                    // When edge point, locates position between the vertices lerping parameter, e.g. Alpha Edge.A + ( 1-Alpha) Edge.B
+	double EdgeAlpha = -1.;                    // When edge point: locates position between the vertices as (1-Alpha) Edge.A + Alpha Edge.B
 };
 
 bool GeodesicSingleTriangleUtils::IsTerminated(const FTraceResult& Result)
