DynamicMesh: Fix the conformal solve functionality in the DynamicMeshUVEditor when setting using existing UV flag to true. Old code was clearing UVs before processing and incorrectly tracking relationships between the internal solver and the final overlay.

#rnx

diff --git a/Engine/Plugins/Runtime/GeometryProcessing/Source/DynamicMesh/Private/Parameterization/DynamicMeshUVEditor.cpp b/Engine/Plugins/Runtime/GeometryProcessing/Source/DynamicMesh/Private/Parameterization/DynamicMeshUVEditor.cpp
--- a/Engine/Plugins/Runtime/GeometryProcessing/Source/DynamicMesh/Private/Parameterization/DynamicMeshUVEditor.cpp
+++ b/Engine/Plugins/Runtime/GeometryProcessing/Source/DynamicMesh/Private/Parameterization/DynamicMeshUVEditor.cpp
@@ -527,6 +527,11 @@ bool FDynamicMeshUVEditor::SetTriangleUVsFromFreeBoundaryConformal(const TArray<
 	if (ensure(UVOverlay) == false) return false;
 	if (Triangles.Num() == 0) return false;
 
+	if (!bUseExistingUVTopology)
+	{
+		ResetUVs(Triangles);
+	}
+
 	FDynamicMesh3 Submesh(EMeshComponents::None);
 	TMap<int32, int32> BaseToSubmeshV;
 	TArray<int32> SubmeshToBaseV;
@@ -548,7 +553,7 @@ bool FDynamicMeshUVEditor::SetTriangleUVsFromFreeBoundaryConformal(const TArray<
 				FVector3d Position = Mesh->GetVertex((bUseExistingUVTopology) ? UVOverlay->GetParentVertex(Triangle[j]) : Triangle[j]);
 				int32 NewVtxID = Submesh.AppendVertex(Position);
 				check(NewVtxID == SubmeshToBaseV.Num());
-				SubmeshToBaseV.Add(NewVtxID);
+				SubmeshToBaseV.Add(Triangle[j]);
 				BaseToSubmeshV.Add(Triangle[j], NewVtxID);
 				NewTriangle[j] = NewVtxID;
 			}
@@ -620,9 +625,6 @@ bool FDynamicMeshUVEditor::SetTriangleUVsFromFreeBoundaryConformal(const TArray<
 	}
 	else
 	{
-		// clear existing triangles
-		ResetUVs(Triangles);
-
 		// copy back to target UVOverlay
 		TArray<int32> VtxElementIDs;
 		TArray<int32> NewElementIDs;
