fix polycut w/ degenerate rounded rect (corner ratio = 0) in three different ways:
 - prevent parameter from going all the way to zero
 - extend cutting polygon a bit outside the bounding box to avoid creating coplanar cuts that stress the boolean system
 - when finding the initial nearest triangle to check for coplanarity, filter out degenerate triangles w/ zero vector normals

#rnx

diff --git a/Engine/Plugins/Experimental/MeshModelingToolsetExp/Source/MeshModelingToolsEditorOnly/Public/PolygonOnMeshTool.h b/Engine/Plugins/Experimental/MeshModelingToolsetExp/Source/MeshModelingToolsEditorOnly/Public/PolygonOnMeshTool.h
--- a/Engine/Plugins/Experimental/MeshModelingToolsetExp/Source/MeshModelingToolsEditorOnly/Public/PolygonOnMeshTool.h
+++ b/Engine/Plugins/Experimental/MeshModelingToolsetExp/Source/MeshModelingToolsEditorOnly/Public/PolygonOnMeshTool.h
@@ -97,7 +97,7 @@ public:
 	float Height = 50.0f;
 
 	/** Corner Ratio of RoundRect Polygon */
-	UPROPERTY(EditAnywhere, Category = Shape, meta = (UIMin = "0.0", UIMax = "1.0", EditCondition = "Shape == EPolygonType::RoundRect"))
+	UPROPERTY(EditAnywhere, Category = Shape, meta = (UIMin = "0.01", ClampMin = "0.00001", ClampMax = "1.0", EditCondition = "Shape == EPolygonType::RoundRect"))
 	float CornerRatio = 0.5f;
 
 	/** Number of sides in Circle or RoundRect Corner */
diff --git a/Engine/Plugins/Experimental/MeshModelingToolsetExp/Source/ModelingOperatorsEditorOnly/Private/CuttingOps/EmbedPolygonsOp.cpp b/Engine/Plugins/Experimental/MeshModelingToolsetExp/Source/ModelingOperatorsEditorOnly/Private/CuttingOps/EmbedPolygonsOp.cpp
--- a/Engine/Plugins/Experimental/MeshModelingToolsetExp/Source/ModelingOperatorsEditorOnly/Private/CuttingOps/EmbedPolygonsOp.cpp
+++ b/Engine/Plugins/Experimental/MeshModelingToolsetExp/Source/ModelingOperatorsEditorOnly/Private/CuttingOps/EmbedPolygonsOp.cpp
@@ -119,6 +119,7 @@ void CollapseDegenerateEdgesOnVertexPath(FDynamicMesh3& Mesh, TArray<int>& Verte
 void FEmbedPolygonsOp::BooleanPath(FProgressCancel* Progress)
 {
 	FAxisAlignedBox3d Bounds = OriginalMesh->GetBounds();
+	Bounds.Expand(.01); // expand a little beyond bounds to avoid creating coplanar cases
 	double MeshDiameter = Bounds.MaxDim();
 
 	FFrame3d Frame = PolygonFrame;
diff --git a/Engine/Plugins/Runtime/GeometryProcessing/Source/DynamicMesh/Private/Operations/MeshBoolean.cpp b/Engine/Plugins/Runtime/GeometryProcessing/Source/DynamicMesh/Private/Operations/MeshBoolean.cpp
--- a/Engine/Plugins/Runtime/GeometryProcessing/Source/DynamicMesh/Private/Operations/MeshBoolean.cpp
+++ b/Engine/Plugins/Runtime/GeometryProcessing/Source/DynamicMesh/Private/Operations/MeshBoolean.cpp
@@ -8,6 +8,7 @@
 
 #include "Async/ParallelFor.h"
 #include "DynamicMesh/MeshTransforms.h"
+#include "DynamicMesh/MeshNormals.h"
 #include "Spatial/SparseDynamicOctree3.h"
 
 #include "Algo/RemoveIf.h"
@@ -371,36 +372,46 @@ bool FMeshBoolean::Compute()
 			{
 				bRemoveInside = 0;
 			}
-			ParallelFor(MaxTriID, [this, &KeepTri, &MeshIdx, &Winding, &OtherSpatial, &ProcessMesh, bCoplanarKeepSameDir, bRemoveInside](int TID)
+			FMeshNormals OtherNormals(OtherSpatial.GetMesh());
+			OtherNormals.ComputeTriangleNormals();
+			const double OnPlaneTolerance = SnapTolerance;
+			IMeshSpatial::FQueryOptions NonDegenCoplanarCandidateFilter(OnPlaneTolerance,
+				[&OtherNormals](int TID) -> bool // filter degenerate triangles from matching
+				{
+					// By convention, the normal for degenerate triangles is the zero vector
+					return !OtherNormals[TID].IsZero();
+				});
+			ParallelFor(MaxTriID, [&](int TID)
 				{
 					if (!ProcessMesh.IsTriangle(TID))
 					{
 						return;
 					}
 					
-					FVector3d Centroid = ProcessMesh.GetTriCentroid(TID);
+					FTriangle3d Tri;
+					ProcessMesh.GetTriVertices(TID, Tri.V[0], Tri.V[1], Tri.V[2]);
+					FVector3d Centroid = Tri.Centroid();
 
 					// first check for the coplanar case
 					{
 						double DSq;
-						double OnPlaneTolerance = SnapTolerance;
-						int OtherTID = OtherSpatial.FindNearestTriangle(Centroid, DSq, OnPlaneTolerance);
+						int OtherTID = OtherSpatial.FindNearestTriangle(Centroid, DSq, NonDegenCoplanarCandidateFilter);
 						if (OtherTID > -1) // only consider it coplanar if there is a matching tri
 						{
 							
-							FVector3d OtherNormal = OtherSpatial.GetMesh()->GetTriNormal(OtherTID);
+							FVector3d OtherNormal = OtherNormals[OtherTID];
 							FVector3d Normal = ProcessMesh.GetTriNormal(TID);
 							double DotNormals = OtherNormal.Dot(Normal);
 
 							//if (FMath::Abs(DotNormals) > .9) // TODO: do we actually want to check for a normal match? coplanar vertex check below is more robust?
 							{
 								// To be extra sure it's a coplanar match, check the vertices are *also* on the other mesh (w/in SnapTolerance)
-								FTriangle3d Tri;
-								ProcessMesh.GetTriVertices(TID, Tri.V[0], Tri.V[1], Tri.V[2]);
+
 								bool bAllTrisOnOtherMesh = true;
 								for (int Idx = 0; Idx < 3; Idx++)
 								{
-									if (OtherSpatial.FindNearestTriangle(Tri.V[Idx], DSq, OnPlaneTolerance) == FDynamicMesh3::InvalidID)
+									// use a slightly more forgiving tolerance to account for the likelihood that these vertices were mesh-cut right to the boundary of the coplanar region and have some additional error
+									if (OtherSpatial.FindNearestTriangle(Tri.V[Idx], DSq, OnPlaneTolerance * 2) == FDynamicMesh3::InvalidID)
 									{
 										bAllTrisOnOtherMesh = false;
 										break;
@@ -408,12 +419,16 @@ bool FMeshBoolean::Compute()
 								}
 								if (bAllTrisOnOtherMesh)
 								{
-									if (MeshIdx != 0) // for coplanar tris favor the first mesh; just delete from the other mesh
+									// for coplanar tris favor the first mesh; just delete from the other mesh
+									// for fully degenerate tris, favor deletion also
+									//  (Note: For degenerate tris we have no orientation info, so we are choosing between
+									//         potentially leaving 'cracks' in solid regions or 'spikes' in empty regions)
+									if (MeshIdx != 0 || Normal.IsZero())
 									{
 										KeepTri[MeshIdx][TID] = false;
 										return;
 									}
-									else // for the first mesh, logic depends on orientation of matching tri
+									else // for the first mesh, & with a valid normal, logic depends on orientation of matching tri
 									{
 										KeepTri[MeshIdx][TID] = DotNormals > 0 == bCoplanarKeepSameDir;
 										return;
