GeometryProcessing:

Removed redundant code since GetTriangleEdges method was added to DynamicMesh3 class that does the same thing.
Fixed a typo ContolPnts -> ControlPoints.
Previously only per-vertex normals were set. Now normals are correclty set per-element as well.

#rnx

diff --git a/Engine/Plugins/Runtime/GeometryProcessing/Source/DynamicMesh/Private/Operations/PNTriangles.cpp b/Engine/Plugins/Runtime/GeometryProcessing/Source/DynamicMesh/Private/Operations/PNTriangles.cpp
--- a/Engine/Plugins/Runtime/GeometryProcessing/Source/DynamicMesh/Private/Operations/PNTriangles.cpp
+++ b/Engine/Plugins/Runtime/GeometryProcessing/Source/DynamicMesh/Private/Operations/PNTriangles.cpp
@@ -26,7 +26,8 @@ namespace FPNTrianglesLocals
 	{	
 		FVector3d Point1;
 		FVector3d Point2;
-		FVector3d Normal;
+		FVector3d NormalTriA;
+		FVector3d NormalTriB;
 	}; 
 	
 	struct FControlPoints 
@@ -42,16 +43,7 @@ namespace FPNTrianglesLocals
 		{
 			Triangles = Mesh.GetTrianglesBuffer();
 			Edges = Mesh.GetEdgesBuffer();
-
-			//TODO: DynamicMesh3 should return this as a buffer
-			TriangleEdges.SetNum(Mesh.MaxTriangleID());
-			for (int TID = 0; TID < Mesh.MaxTriangleID(); ++TID) 
-			{
-				if (Mesh.IsTriangle(TID)) 
-				{
-					TriangleEdges[TID] = Mesh.GetTriEdgesRef(TID); 
-				}
-			}
+			TriangleEdges = Mesh.GetTriangleEdges();
 		}
 		
 		TDynamicVector<FIndex3i> Triangles;
@@ -76,7 +68,7 @@ namespace FPNTrianglesLocals
 							  const bool bComputePNNormals, 
 							  FProgressCancel* ProgressCancel) 
 	{
-		if (!Mesh.HasVertexNormals() && UseNormals == nullptr) 
+		if (!Mesh.HasAttributes() && !Mesh.HasVertexNormals() && UseNormals == nullptr)
 		{
 			return false; // Mesh must have normals
 		}
@@ -86,6 +78,35 @@ namespace FPNTrianglesLocals
 		OutControlPoints.OnEdges.SetNum(Mesh.MaxEdgeID());
 		OutControlPoints.OnTriangles.SetNum(Mesh.MaxTriangleID());
 
+		auto ComputeControlPoint = [](const FVector3d& Vertex1, const FVector3d& Vertex2, const FVector3f& Normal)
+		{ 
+			FVector3d Edge12 = Vertex2 - Vertex1;
+			double Weight12 = Edge12.Dot(Normal);
+			FVector3d Point = (2.0*Vertex1 + Vertex2 - Weight12*Normal)/3.0;  
+			return Point;
+		 };
+
+		auto ComputeControlNormal = [](const FVector3d& Vertex1, const FVector3d& Vertex2, const FVector3f& Normal1, const FVector3f& Normal2)
+		{ 	
+			FVector3d Edge12 = Vertex2 - Vertex1;
+			FVector3d Normal21 = Normal2 + Normal1;
+
+			double Divisor = Edge12.Dot(Edge12);
+			
+			FVector3f Normal;
+			if (FMath::IsNearlyZero(Divisor)) 
+			{
+				Normal = Normal21 / 2.0; // If edge is degenerate then simply interpolate between the normals
+			} 
+			else
+			{ 
+				double NWeight12 = 2.0*Edge12.Dot(Normal21)/Divisor;
+				Normal = Normalized(Normal21 - NWeight12*Edge12);
+			}
+			
+			return Normal;
+		 };
+
 		// First compute only the control points on the edges
 		ParallelFor(Mesh.MaxEdgeID(), [&](int32 EID)
 		{
@@ -101,24 +122,55 @@ namespace FPNTrianglesLocals
 				FVector3d Vertex1 = Mesh.GetVertex(EdgeV.A);
 				FVector3d Vertex2 = Mesh.GetVertex(EdgeV.B);
 
-				FVector3d Normal1 = (UseNormals != nullptr) ? (*UseNormals)[EdgeV.A] : (FVector3d)Mesh.GetVertexNormal(EdgeV.A);
-				FVector3d Normal2 = (UseNormals != nullptr) ? (*UseNormals)[EdgeV.B] : (FVector3d)Mesh.GetVertexNormal(EdgeV.B);
-				
-				FVector3d Edge12 = Vertex2 - Vertex1;
+				FIndex2i EdgeTri = Mesh.GetEdgeT(EID);
 
-				double Weight12 = Edge12.Dot(Normal1);
-				double Weight21 = -1.0*Edge12.Dot(Normal2); 
+				FVector3f Normal1, Normal2;
+				if (Mesh.HasAttributes()) 
+				{
+					const FDynamicMeshNormalOverlay* NormalOverlay = Mesh.Attributes()->PrimaryNormals();
+					NormalOverlay->GetElementAtVertex(EdgeTri.A, EdgeV.A, Normal1);
+					NormalOverlay->GetElementAtVertex(EdgeTri.A, EdgeV.B, Normal2);
+				}
+				else 
+				{
+					Normal1 = (UseNormals != nullptr) ? (FVector3f)(*UseNormals)[EdgeV.A] : Mesh.GetVertexNormal(EdgeV.A);
+					Normal2 = (UseNormals != nullptr) ? (FVector3f)(*UseNormals)[EdgeV.B] : Mesh.GetVertexNormal(EdgeV.B);
+				}
 
-				FEdgeControlPoints& ContolPnts = OutControlPoints.OnEdges[EID];
-				ContolPnts.Point1 = (2.0*Vertex1 + Vertex2 - Weight12*Normal1)/3.0;  
-				ContolPnts.Point2 = (2.0*Vertex2 + Vertex1 - Weight21*Normal2)/3.0;  
+				FEdgeControlPoints& ControlPnts = OutControlPoints.OnEdges[EID];
 
 				if (bComputePNNormals)
 				{
-					FVector3d Normal21 = Normal2 + Normal1;
-					double NWeight12 = 2.0*Edge12.Dot(Normal21)/Edge12.Dot(Edge12);
-					ContolPnts.Normal = Normalized(Normal21 - NWeight12*Edge12);
+					ControlPnts.NormalTriA = ComputeControlNormal(Vertex1, Vertex2, Normal1, Normal2);
 				}
+
+				FVector3d Point1ForTriA = ComputeControlPoint(Vertex1, Vertex2, Normal1);
+				FVector3d Point2ForTriA = ComputeControlPoint(Vertex2, Vertex1, Normal2);
+
+				if (Mesh.HasAttributes()) 
+				{
+					const FDynamicMeshNormalOverlay* NormalOverlay = Mesh.Attributes()->PrimaryNormals();
+					if (EdgeTri.B != FDynamicMesh3::InvalidID && NormalOverlay->IsSeamEdge(EID)) 
+					{	
+						NormalOverlay->GetElementAtVertex(EdgeTri.B, EdgeV.A, Normal1);
+						NormalOverlay->GetElementAtVertex(EdgeTri.B, EdgeV.B, Normal2); 
+						FVector3d Point1ForTriB = ComputeControlPoint(Vertex1, Vertex2, Normal1);
+						FVector3d Point2ForTriB = ComputeControlPoint(Vertex2, Vertex1, Normal2);
+
+						// Following the core idea of Crack-free PN Triangles we average the control point we'd like 
+						// with the control point our neighbor triangle would like. 
+						Point1ForTriA = (Point1ForTriA + Point1ForTriB) / 2.0;
+						Point2ForTriA = (Point2ForTriA + Point2ForTriB) / 2.0;
+
+						if (bComputePNNormals)
+						{
+							ControlPnts.NormalTriB = ComputeControlNormal(Vertex1, Vertex2, Normal1, Normal2);
+						}
+					}
+				}
+
+				ControlPnts.Point1 = Point1ForTriA;
+				ControlPnts.Point2 = Point2ForTriA;
 			}
 		});
 
@@ -171,12 +223,12 @@ namespace FPNTrianglesLocals
 	 * @param Mesh The mesh we are tesselating.
 	 * @param Level How many times we are recursively subdividing the Mesh.
 	 * @param ProgressCancel Set this to be able to cancel running operation.
-	 * @param OutNewVerticies Array of tuples of the new vertex ID and the original triangle ID the vertex belongs to.
+	 * @param OutNewVertices Array of tuples of the new vertex ID and the original triangle ID the vertex belongs to.
 	 * 
 	 * @return true if the operation succeeded, false if it failed or was canceled by the user.
 	 */
 	bool TesselateMesh(FDynamicMesh3& Mesh, 
-					   TArray<FIndex2i>& OutNewVerticies,
+					   TArray<FIndex2i>& OutNewVertices,
 					   const int Level, 
 					   FProgressCancel* ProgressCancel) 
 	{
@@ -197,8 +249,8 @@ namespace FPNTrianglesLocals
 		int NumNewEdgeVert = NumEdgeSegments - 1; // Number of the new vertices added per edge
 		// Triangular number series (n(n+1)/2) minus the 3 original vertices and new edge vertices
 		int NumNewTriangleVert = NumEdgeVert*(NumEdgeVert + 1)/2 - 3 - 3*NumNewEdgeVert;
-		int NumNewVerticies = NumNewEdgeVert * Mesh.EdgeCount() + NumNewTriangleVert* Mesh.TriangleCount(); 
-		OutNewVerticies.Reserve(NumNewVerticies);
+		int NumNewVertices = NumNewEdgeVert * Mesh.EdgeCount() + NumNewTriangleVert* Mesh.TriangleCount(); 
+		OutNewVertices.Reserve(NumNewVertices);
 
 		// The number of the new triangle IDs created with Level number of subdivisions 
 		int NumNewTrianglesIDs = (NumEdgeSegments * NumEdgeSegments - 1) * Mesh.TriangleCount();
@@ -243,7 +295,7 @@ namespace FPNTrianglesLocals
 						ParentTriangles[SplitInfo.NewTriangles.B] = ParentTriangles[SplitInfo.OriginalTriangles.B]; 
 					}
 
-					OutNewVerticies.Add(FIndex2i(SplitInfo.NewVertex, ParentTriangles[SplitInfo.OriginalTriangles.A]));
+					OutNewVertices.Add(FIndex2i(SplitInfo.NewVertex, ParentTriangles[SplitInfo.OriginalTriangles.A]));
 
 					if (EdgeTris.A < MaxTriangleID && EdgesToFlip[EdgeTris.A] == FDynamicMesh3::InvalidID)
 					{
@@ -279,7 +331,7 @@ namespace FPNTrianglesLocals
 			}
 		}
 
-		checkSlow(OutNewVerticies.Num() == NumNewVerticies);
+		checkSlow(OutNewVertices.Num() == NumNewVertices);
 		checkSlow(ParentTriangles.Num() == Mesh.MaxTriangleID());
 
 		return true;
@@ -291,7 +343,7 @@ namespace FPNTrianglesLocals
 	 * 
 	 * @param Mesh The tessellated mesh whose vertices we are displacing.
 	 * @param FControlPoints PN Triangle control points.
-	 * @param VerticiesToDisplace Array of vertices into the Mesh that we are displacing.
+	 * @param VerticesToDisplace Array of vertices into the Mesh that we are displacing.
 	 * @param bComputePNNormals Should we be computing quadratically varying normals using control points.
 	 * @param OriginalConnectivity Connectivity of the original mesh (before tesselation).
 	 * @param ProgressCancel Set this to be able to cancel running operation.
@@ -299,27 +351,25 @@ namespace FPNTrianglesLocals
 	 * @return true if the operation succeeded, false if it failed or was canceled by the user.
 	 */
 
-	bool Displace(FDynamicMesh3& Mesh,
-				  const FControlPoints& FControlPoints,
-				  const TArray<FIndex2i>& VerticiesToDisplace,
-				  const bool bComputePNNormals, 
-				  const FMeshConnectivity& OriginalConnectivity,
-				  FProgressCancel* ProgressCancel) 
+	bool DisplaceAndSetNormals(FDynamicMesh3& Mesh,
+					           const FControlPoints& FControlPoints,
+					           const TArray<FIndex2i>& VerticesToDisplace,
+					           const bool bComputePNNormals,
+					           const FMeshConnectivity& OriginalConnectivity,
+					           FProgressCancel* ProgressCancel) 
 	{
-		if (bComputePNNormals && !Mesh.HasVertexNormals()) 
-		{
-			Mesh.EnableVertexNormals(FVector3f::UnitY()); 
-		}
-
+		bool bHasVertexNormals = Mesh.HasVertexNormals();
+		bool bHasAttributes = Mesh.HasAttributes();
+		
 		// Iterate over every new vertex and compute its displacement and optionally a normal
-		ParallelFor(VerticiesToDisplace.Num(), [&](int32 IDX)
+		ParallelFor(VerticesToDisplace.Num(), [&](int32 IDX)
 		{	
 			if (ProgressCancel && ProgressCancel->Cancelled()) 
 			{
 				return; 
 			}
 			
-			FIndex2i NewVtx = VerticiesToDisplace[IDX];
+			FIndex2i NewVtx = VerticesToDisplace[IDX];
 			int VertexID = NewVtx[0]; // ID of the new vertex added with tesselation
 			int OriginalTriangleID = NewVtx[1]; // ID of the original triangle new vertex belongs to
 			
@@ -363,20 +413,30 @@ namespace FPNTrianglesLocals
 
 			if (bComputePNNormals) 
 			{
-				// Compute contribution of the normal control points at the original vertices
-				FVector3d NewNormal = BarySquared[0]*Mesh.GetVertexNormal(TriVertex.A) + 
-									  BarySquared[1]*Mesh.GetVertexNormal(TriVertex.B) +
-									  BarySquared[2]*Mesh.GetVertexNormal(TriVertex.C);
+				if (bHasVertexNormals) 
+				{
+					// Compute contribution of the normal control points at the original vertices
+					FVector3d NewNormal = BarySquared[0]*Mesh.GetVertexNormal(TriVertex.A) + 
+										  BarySquared[1]*Mesh.GetVertexNormal(TriVertex.B) +
+										  BarySquared[2]*Mesh.GetVertexNormal(TriVertex.C);
+					
+					// Compute contribution of the normal control points at the edges
+					for (int EIDX = 0; EIDX < 3; ++EIDX) 
+					{   
+						int EID = TriEdges[EIDX];
+						const FEdgeControlPoints& ControlPnts = FControlPoints.OnEdges[EID];
+						NewNormal += Bary[EIDX]*Bary[(EIDX + 1) % 3]*ControlPnts.NormalTriA;
+					}
+
+					Normalize(NewNormal);
+
+					Mesh.SetVertexNormal(VertexID, NewNormal); 
+				} 
 				
-				// Compute contribution of the normal control points at the edges
-				for (int EIDX = 0; EIDX < 3; ++EIDX) 
-				{   
-					int EID = TriEdges[EIDX];
-					const FEdgeControlPoints& ControlPnts = FControlPoints.OnEdges[EID];
-					NewNormal += Bary[EIDX]*Bary[(EIDX + 1) % 3]*ControlPnts.Normal;
+				if (bHasAttributes) 
+				{
+					//TODO: Compute per-element quadractic PN normals
 				}
-
-				Mesh.SetVertexNormal(VertexID, NewNormal);                                      
 			}
 		});
 
@@ -406,14 +466,15 @@ bool FPNTriangles::Compute()
 		return true; // nothing to do
 	}
 
+	// Compute per-vertex normals if no normals exist
 	FMeshNormals Normals;
-	bool bHaveVertexNormals = Mesh->HasVertexNormals();
-	if (!bHaveVertexNormals)
+	bool bHasNormals = Mesh->HasVertexNormals() || Mesh->HasAttributes();
+	if (!bHasNormals)
 	{
 		Normals = FMeshNormals(Mesh);
 		Normals.ComputeVertexNormals();
 	}
-	FMeshNormals* UseNormals = (bHaveVertexNormals) ? nullptr : &Normals;
+	FMeshNormals* UseNormals = (bHasNormals) ? nullptr : &Normals;
 	
 	// Compute PN triangle control points for each flat triangle of the original mesh
 	FControlPoints FControlPoints;
@@ -423,19 +484,19 @@ bool FPNTriangles::Compute()
 		return false;
 	}
 
-	// Save the topology information of the original mesh  before we change it with tesselation
+	// Save the topology information of the original mesh before we change it with tesselation
 	FMeshConnectivity OriginalConnectivity(*Mesh);
 	
 	// Tesselate the original mesh
-	TArray<FIndex2i> NewVerticies;
-	bOk = TesselateMesh(*Mesh, NewVerticies, TesselationLevel, Progress);
+	TArray<FIndex2i> NewVertices;
+	bOk = TesselateMesh(*Mesh, NewVertices, TesselationLevel, Progress);
 	if (bOk == false) 
 	{
 		return false;
 	}
 
 	// Compute displacement and normals
-	bOk = Displace(*Mesh, FControlPoints, NewVerticies, bComputePNNormals, OriginalConnectivity, Progress);
+	bOk = DisplaceAndSetNormals(*Mesh, FControlPoints, NewVertices, bComputePNNormals, OriginalConnectivity, Progress);
 	if (bOk == false) 
 	{
 		return false;
diff --git a/Engine/Source/Runtime/GeometryCore/Public/DynamicMesh/DynamicMeshOverlay.h b/Engine/Source/Runtime/GeometryCore/Public/DynamicMesh/DynamicMeshOverlay.h
--- a/Engine/Source/Runtime/GeometryCore/Public/DynamicMesh/DynamicMeshOverlay.h
+++ b/Engine/Source/Runtime/GeometryCore/Public/DynamicMesh/DynamicMeshOverlay.h
@@ -641,6 +641,26 @@ public:
 		BaseType::GetElement(ElementID, V);
 	}
 
+	/**
+	 * Get the Element value associated with a vertex of a triangle.
+	 * @param TriangleID ID of a triangle containing the Element
+	 * @param VertexID ID of the Element's parent vertex 
+	 * @param Value Value contained at the Element
+	 */
+	inline void GetElementAtVertex(int TriangleID, int VertexID, VectorType& Value) const
+	{
+		FIndex3i Triangle = BaseType::GetTriangle(TriangleID);
+		for (int IDX = 0; IDX < 3; ++IDX)
+		{	
+			int ElementID = Triangle[IDX];
+			if (BaseType::ParentVertices[ElementID] == VertexID)
+			{
+				GetElement(ElementID, Value);
+				return;
+			}
+		}
+	}	
+
 	/**
 	 * Get the Element associated with a vertex of a triangle
 	 * @param TriVertexIndex index of vertex in triangle, valid values are 0,1,2
