Dynamic Mesh, Operations, GeodesicPath - allow geodesic path to be exported as array of surface points.

diff --git a/Engine/Plugins/Runtime/GeometryProcessing/Source/DynamicMesh/Private/Operations/GeodesicPath.cpp b/Engine/Plugins/Runtime/GeometryProcessing/Source/DynamicMesh/Private/Operations/GeodesicPath.cpp
--- a/Engine/Plugins/Runtime/GeometryProcessing/Source/DynamicMesh/Private/Operations/GeodesicPath.cpp
+++ b/Engine/Plugins/Runtime/GeometryProcessing/Source/DynamicMesh/Private/Operations/GeodesicPath.cpp
@@ -2,6 +2,8 @@
 
 #include "Operations/GeodesicPath.h"
 #include "Math/UnrealMathVectorCommon.h"
+#include "Operations/MeshGeodesicSurfaceTracer.h" 
+
 
 using namespace UE::Geometry;
 
@@ -194,8 +196,8 @@ bool VisitWedgeTriangles(const FIntrinsicEdgeFlipMesh& EdgeFlipMesh, const int32
 // -- Deformable Edge Path --//
 
 
-FDeformableEdgePath::FDeformableEdgePath(const FDynamicMesh3& SurfaceMesh, const TArray<FEdgePath::FDirectedSegment>& PathAsDirectedSegments)
-	: EdgeFlipMesh(SurfaceMesh)
+FDeformableEdgePath::FDeformableEdgePath(const FDynamicMesh3& SurfaceMeshIn, const TArray<FEdgePath::FDirectedSegment>& PathAsDirectedSegments)
+	: EdgeFlipMesh(SurfaceMeshIn)
 	, PathLength(0.)
 	, NumFlips(0)
 {
@@ -620,4 +622,69 @@ void FDeformableEdgePath::ComputeWedgeAngles(int32 IncomingEID,
 
 	LeftSideAngle  = (bLeftContainsBoundary)  ? TMathUtilConstants<double>::MaxReal : LeftWedgeAngle;
 	RightSideAngle = (bRightContainsBoundary) ? TMathUtilConstants<double>::MaxReal : RightWedgeAngle;
+}
+
+TArray<FDeformableEdgePath::FSurfacePoint> FDeformableEdgePath::AsSurfacePoints(double CoalesceThreshold) const
+{
+	TArray<FIntrinsicTriangulation::FSurfacePoint> PathSurfacePoints;
+
+	const int32 NumIntrinsicPathSegments = EdgePath.NumSegments();
+	
+	// empty path case
+	if (NumIntrinsicPathSegments == 0)
+	{
+		return PathSurfacePoints;
+	}
+
+	int32 SID = EdgePath.GetHeadSegmentID();
+	while (SID != InvalidID)
+	{ 
+	
+		const int32 NextSID = EdgePath.GetNextSegmentID(SID);
+		const bool bIsLastSegment = (NextSID == InvalidID);
+		
+		const int32 EID = EdgePath.GetSegment(SID).EID;
+		const bool bReverseEdge = (EdgePath.GetSegment(SID).HeadIndex == 0);
+
+		// get intrinsic edge as a sequence of surface points, note each segment starts and ends at a surface mesh vertex but may cross several surface mesh edges.
+		TArray<FIntrinsicTriangulation::FSurfacePoint> SegmentSurfacePoints = EdgeFlipMesh.TraceEdge(EID, CoalesceThreshold, bReverseEdge);
+
+		PathSurfacePoints.Append(MoveTemp(SegmentSurfacePoints));
+		if (!bIsLastSegment)
+		{
+			// delete last element since it will be the same as the first element in the next segment
+			PathSurfacePoints.Pop(); 
+		}
+		
+		SID = NextSID;
+	}
+	
+	return MoveTemp(PathSurfacePoints);
+}
+
+double UE::Geometry::SumPathLength(const FDeformableEdgePath& DeformableEdgePath)
+{
+	double TotalPathLength = 0;
+	const FEdgePath& EdgePath = DeformableEdgePath.GetEdgePath();
+	const FIntrinsicTriangulation& EdgeFlipMesh = DeformableEdgePath.GetIntrinsicMesh();
+
+	const int32 NumIntrinsicPathSegments = EdgePath.NumSegments();
+
+	// empty path case
+	if (NumIntrinsicPathSegments == 0)
+	{
+		return TotalPathLength;
+	}
+
+	int32 SID = EdgePath.GetHeadSegmentID();
+	while (SID != FDeformableEdgePath::InvalidID)
+	{
+		const int32 EID = EdgePath.GetSegment(SID).EID;
+		const double SegmentLength = EdgeFlipMesh.GetEdgeLength(EID);
+		TotalPathLength += SegmentLength;
+			
+		SID = EdgePath.GetNextSegmentID(SID);
+	}
+
+	return TotalPathLength;
 }
\ No newline at end of file
diff --git a/Engine/Plugins/Runtime/GeometryProcessing/Source/DynamicMesh/Public/Operations/GeodesicPath.h b/Engine/Plugins/Runtime/GeometryProcessing/Source/DynamicMesh/Public/Operations/GeodesicPath.h
--- a/Engine/Plugins/Runtime/GeometryProcessing/Source/DynamicMesh/Public/Operations/GeodesicPath.h
+++ b/Engine/Plugins/Runtime/GeometryProcessing/Source/DynamicMesh/Public/Operations/GeodesicPath.h
@@ -114,7 +114,7 @@ public:
 	/**
 	* Constructor assumes the directed segments are ordered from tail (at index 0) to head
 	*/
-	FDeformableEdgePath(const FDynamicMesh3& SurfaceMesh, const TArray<FEdgePath::FDirectedSegment>& OriginalPathAsDirectedSegments);
+	FDeformableEdgePath(const FDynamicMesh3& SurfaceMeshIn, const TArray<FEdgePath::FDirectedSegment>& OriginalPathAsDirectedSegments);
 	
 	virtual ~FDeformableEdgePath(){}
 
@@ -153,7 +153,24 @@ public:
 	/**
 	* @return const reference to the intrinsic mesh on which the EdgePath is defined.
 	*/ 
-	inline const FIntrinsicEdgeFlipMesh& GetIntrinsicMesh() const;
+	inline const FIntrinsicTriangulation& GetIntrinsicMesh() const;
+
+	/**
+	* struct that references a point on a mesh, by vertex, by edge-crossing, or barycentric coords
+	*/ 
+	using FSurfacePoint =  FIntrinsicTriangulation::FSurfacePoint;
+	/**
+	* @return an array of surface points relative to the SurfaceMesh that define this path.
+	* Note the first and last surfacepoints correspond to the start and end vertex, but all 
+	* other surface points may be either edge crossings or vertex points.
+	* 
+	* @param CoalesceThreshold - In barycentric units [0,1], edge-crossings within this threshold are snapped to the nearest vertex
+	*                            and any resulting repitition of vertex surface points are replaced with a single vertex surface point.
+	*                            Due to numerical precision issues, a path that 'should' intersect a surface mesh vertex may appear
+	*                            as a sequence of vertex adjacent edge-crossings very close to the vertex.. the threshold is applied in
+	*                            a post-process coalesce those crossings into a single vertex crossing.
+	*/
+	TArray<FSurfacePoint> AsSurfacePoints(double CoalesceThreshold = 0.) const;
 
 protected:
 	
@@ -241,8 +258,8 @@ protected:
 	                        double& LeftSideAngle, double& RightSideAngle) const ;
 	
 protected:
-
-	FIntrinsicEdgeFlipMesh         EdgeFlipMesh;        // Intrinsic mesh with the same vertices as the original surface mesh. The geodesic path is comprised of edges in this mesh.
+     
+	FIntrinsicTriangulation        EdgeFlipMesh;        // Intrinsic mesh with the same vertices as the original surface mesh. The geodesic path is comprised of edges in this mesh.
 
 	double                         PathLength;          // Total length of the path
 	int32                          NumFlips;            // Count of the number of edge flips performed in shortening the path.
@@ -254,6 +271,13 @@ protected:
 	
 };
 
+/**
+* @return the length of the path, computed by suming the path segment lengths.  Primarily for testing
+* as the result should match FDeformableEdgePath::GetPathLength()
+*/ 
+double SumPathLength(const FDeformableEdgePath& DeformableEdgePath);
+
+
 const FEdgePath& FDeformableEdgePath::GetEdgePath() const
 {
 	return EdgePath;
@@ -264,7 +288,7 @@ double  FDeformableEdgePath::GetPathLength() const
 	return PathLength;
 }
 
-const FIntrinsicEdgeFlipMesh& FDeformableEdgePath::GetIntrinsicMesh() const
+const FIntrinsicTriangulation& FDeformableEdgePath::GetIntrinsicMesh() const
 {
 	return EdgeFlipMesh;
 }
