make coplanar removal test more double-check its assumption that the coplanar pair was actually removed; fixed sliver/edge cases where from the sliver side it looks like a coplanar surface but from the other it typically does not

 + also bring skip-degenerates-in-nearest-tri-search improvement previously added to MeshBoolean into MeshSelfUnion

diff --git a/Engine/Plugins/Runtime/GeometryProcessing/Source/DynamicMesh/Private/Operations/MeshBoolean.cpp b/Engine/Plugins/Runtime/GeometryProcessing/Source/DynamicMesh/Private/Operations/MeshBoolean.cpp
--- a/Engine/Plugins/Runtime/GeometryProcessing/Source/DynamicMesh/Private/Operations/MeshBoolean.cpp
+++ b/Engine/Plugins/Runtime/GeometryProcessing/Source/DynamicMesh/Private/Operations/MeshBoolean.cpp
@@ -359,6 +359,10 @@ bool FMeshBoolean::Compute()
 	{ // (just for scope)
 		// first decide what triangles to delete for both meshes (*before* deleting anything so winding doesn't get messed up!)
 		TArray<bool> KeepTri[2];
+		// This array is used to double-check the assumption that we will delete the other surface when we keep a coplanar tri
+		// Note we only need it for mesh 0 (i.e., the mesh we try to keep triangles from when we preserve coplanar surfaces)
+		TArray<int32> DeleteIfOtherKept;
+		DeleteIfOtherKept.Init(-1, CutMesh[0]->MaxTriangleID());
 		for (int MeshIdx = 0; MeshIdx < NumMeshesToProcess; MeshIdx++)
 		{
 			TFastWindingTree<FDynamicMesh3> Winding(&Spatial[1 - MeshIdx]);
@@ -430,7 +434,15 @@ bool FMeshBoolean::Compute()
 									}
 									else // for the first mesh, & with a valid normal, logic depends on orientation of matching tri
 									{
-										KeepTri[MeshIdx][TID] = DotNormals > 0 == bCoplanarKeepSameDir;
+										bool bKeep = DotNormals > 0 == bCoplanarKeepSameDir;
+										KeepTri[MeshIdx][TID] = bKeep;
+										if (bKeep)
+										{
+											// If we kept this tri, remember the coplanar pair we expect to be deleted, in case
+											// it isn't deleted (e.g. because it wasn't coplanar); to then delete this one instead.
+											// This can help clean up sliver triangles near a cut boundary that look locally coplanar
+											DeleteIfOtherKept[TID] = OtherTID;
+										}
 										return;
 									}
 								}
@@ -442,6 +454,21 @@ bool FMeshBoolean::Compute()
 					double WindingNum = Winding.FastWindingNumber(Centroid);
 					KeepTri[MeshIdx][TID] = (WindingNum > WindingThreshold) != bRemoveInside;
 				});
+		}
+
+		// Don't keep coplanar tris if the matched, second-mesh tri that we expected to delete was actually kept
+		for (int TID : CutMesh[0]->TriangleIndicesItr())
+		{
+			int32 DeleteIfOtherKeptTID = DeleteIfOtherKept[TID];
+			if (DeleteIfOtherKeptTID > -1 && KeepTri[1][DeleteIfOtherKeptTID])
+			{
+				KeepTri[0][TID] = false;
+			}
+		}
+
+		for (int MeshIdx = 0; MeshIdx < NumMeshesToProcess; MeshIdx++)
+		{
+			FDynamicMesh3& ProcessMesh = *CutMesh[MeshIdx];
 			for (int EID : ProcessMesh.EdgeIndicesItr())
 			{
 				FDynamicMesh3::FEdge Edge = ProcessMesh.GetEdge(EID);
diff --git a/Engine/Plugins/Runtime/GeometryProcessing/Source/DynamicMesh/Private/Operations/MeshSelfUnion.cpp b/Engine/Plugins/Runtime/GeometryProcessing/Source/DynamicMesh/Private/Operations/MeshSelfUnion.cpp
--- a/Engine/Plugins/Runtime/GeometryProcessing/Source/DynamicMesh/Private/Operations/MeshSelfUnion.cpp
+++ b/Engine/Plugins/Runtime/GeometryProcessing/Source/DynamicMesh/Private/Operations/MeshSelfUnion.cpp
@@ -175,11 +175,13 @@ bool FMeshSelfUnion::Compute()
 	{ // (just for scope)
 		// decide what triangles to delete
 		TArray<bool> KeepTri;
+		TArray<int32> DeleteIfOtherKept;
 		TFastWindingTree<FDynamicMesh3> Winding(&Spatial);
 		int MaxTriID = Mesh->MaxTriangleID();
 		KeepTri.SetNumUninitialized(MaxTriID);
+		DeleteIfOtherKept.Init(-1, MaxTriID);
 		
-		ParallelFor(MaxTriID, [this, &Spatial, &Normals, &TriToComponentID, &KeepTri, &Winding](int TID)
+		ParallelFor(MaxTriID, [this, &Spatial, &Normals, &TriToComponentID, &KeepTri, &DeleteIfOtherKept, &Winding](int TID)
 		{
 			if (!Mesh->IsTriangle(TID))
 			{
@@ -200,9 +202,10 @@ bool FMeshSelfUnion::Compute()
 				double DSq;
 				int MyComponentID = TriToComponentID[TID];
 				IMeshSpatial::FQueryOptions QueryOptions(SnapTolerance,
-					[&TriToComponentID, MyComponentID](int OtherTID)
+					[&Normals, &TriToComponentID, MyComponentID](int OtherTID)
 					{
-						return TriToComponentID[OtherTID] != MyComponentID;
+						// By convention, the normal for degenerate triangles is the zero vector
+						return !Normals[OtherTID].IsZero() && TriToComponentID[OtherTID] != MyComponentID;
 					}
 				);
 				int OtherTID = Spatial.FindNearestTriangle(Centroid, DSq, QueryOptions);
@@ -225,7 +228,7 @@ bool FMeshSelfUnion::Compute()
 						}
 						if (bAllTrisOnOtherComponent)
 						{
-							if (DotNormals < 0)
+							if (DotNormals <= 0) // include zero in range to also discard degenerate triangles w/ zero normals
 							{
 								KeepTri[TID] = false;
 							}
@@ -236,6 +239,13 @@ bool FMeshSelfUnion::Compute()
 								int OtherComponentID = TriToComponentID[OtherTID];
 								bool bHasPriority = MyComponentID < OtherComponentID;
 								KeepTri[TID] = bHasPriority;
+								if (bHasPriority)
+								{
+									// If we kept this tri, remember the coplanar pair we expect to be deleted, in case
+									// it isn't deleted (e.g. because it wasn't coplanar); to then delete this one instead.
+									// This can help clean up sliver triangles near a cut boundary that look locally coplanar
+									DeleteIfOtherKept[TID] = OtherTID;
+								}
 							}
 							return;
 						}
@@ -246,6 +256,16 @@ bool FMeshSelfUnion::Compute()
 			KeepTri[TID] = bKeep;
 		});
 
+		// Don't keep coplanar tris if the matched, "lower priority" tri that we expected to delete was actually kept
+		for (int TID : Mesh->TriangleIndicesItr())
+		{
+			int32 DeleteIfOtherKeptTID = DeleteIfOtherKept[TID];
+			if (DeleteIfOtherKeptTID > -1 && KeepTri[DeleteIfOtherKeptTID])
+			{
+				KeepTri[TID] = false;
+			}
+		}
+
 		// track where we will create new boundary edges
 		for (int EID : Mesh->EdgeIndicesItr())
 		{
