GeometryProcessing: handle case where boolean generates a very tiny hole, and then tries to minimal-fill it. The MinimalFill code does an initial remesh pass to create some "density" to collapse. If the hole is tiny, the remesh target edge length can become near-zero, and then the remesher explodes.
Also adds TMeshQueries::TotalEdgeLength() function which computes accumulated length of an edge set.
#rnx

diff --git a/Engine/Plugins/Runtime/GeometryProcessing/Source/DynamicMesh/Private/Operations/MinimalHoleFiller.cpp b/Engine/Plugins/Runtime/GeometryProcessing/Source/DynamicMesh/Private/Operations/MinimalHoleFiller.cpp
--- a/Engine/Plugins/Runtime/GeometryProcessing/Source/DynamicMesh/Private/Operations/MinimalHoleFiller.cpp
+++ b/Engine/Plugins/Runtime/GeometryProcessing/Source/DynamicMesh/Private/Operations/MinimalHoleFiller.cpp
@@ -1,6 +1,7 @@
 // Copyright Epic Games, Inc. All Rights Reserved.
 
 #include "Operations/MinimalHoleFiller.h"
+#include "Operations/SimpleHoleFiller.h"
 #include "MeshQueries.h"
 
 #include "ExplicitUseGeometryMathTypes.h"		// using UE::Geometry::(math types)
@@ -614,6 +615,14 @@ bool FMinimalHoleFiller::Fill(int32 GroupID)
 		return true;
 	}
 
+	// if the loop is tiny just use the simple fill
+	double LoopTotalLength = TMeshQueries<FDynamicMesh3>::TotalEdgeLength(*Mesh, FillLoop.Edges);
+	if (LoopTotalLength < 100.0f*FMathf::ZeroTolerance)
+	{
+		NewTriangles = Simplefill.NewTriangles;
+		return true;
+	}
+
 	// extract the simple fill mesh as a submesh, via RegionOperator, so we can backsub later
 	TSet<int> IntialFillTris(Simplefill.NewTriangles);
 
@@ -672,6 +681,8 @@ bool FMinimalHoleFiller::Fill(int32 GroupID)
 		Remesher.SetExternalConstraints(MoveTemp(Constraints));
 	}
 
+	// TODO: we should estimate/valide and/or try to cap the number of triangles being created here.
+	// If RemeshTargetEdgeLength value is extremely tiny this could be trying to create infinity triangles.
 	Remesher.SetTargetEdgeLength(RemeshTargetEdgeLength);
 	Remesher.FastestRemesh();
 
diff --git a/Engine/Source/Runtime/GeometryCore/Public/MeshQueries.h b/Engine/Source/Runtime/GeometryCore/Public/MeshQueries.h
--- a/Engine/Source/Runtime/GeometryCore/Public/MeshQueries.h
+++ b/Engine/Source/Runtime/GeometryCore/Public/MeshQueries.h
@@ -417,6 +417,23 @@ public:
 		return MinLength;
 	}
 
+	/// Given a mesh and a subset of mesh edges, compute the total length of all the edges
+	static double TotalEdgeLength(const TriangleMeshType& Mesh, const TArray<int>& Edges)
+	{
+		double AccumulatedLength = 0;
+		for (int EdgeID : Edges)
+		{
+			if (Mesh.IsEdge(EdgeID))
+			{
+				FVector3d A, B;
+				Mesh.GetEdgeV(EdgeID, A, B);
+				AccumulatedLength += Distance(A, B);
+			}
+		}
+		return AccumulatedLength;
+	}
+
+
 	/// Given a mesh and a subset of mesh edges, compute the min, max, and mean edge lengths
 	static void EdgeLengthStatsFromEdges(const TriangleMeshType& Mesh, const TArray<int>& Edges, double& MinEdgeLength,
 		double& MaxEdgeLength, double& AverageEdgeLength)
