ModelingTools: Fix trace event names.

#rnx

diff --git a/Engine/Plugins/Runtime/GeometryProcessing/Source/DynamicMesh/Private/Sampling/MeshMapBaker.cpp b/Engine/Plugins/Runtime/GeometryProcessing/Source/DynamicMesh/Private/Sampling/MeshMapBaker.cpp
--- a/Engine/Plugins/Runtime/GeometryProcessing/Source/DynamicMesh/Private/Sampling/MeshMapBaker.cpp
+++ b/Engine/Plugins/Runtime/GeometryProcessing/Source/DynamicMesh/Private/Sampling/MeshMapBaker.cpp
@@ -113,7 +113,7 @@ void FMeshMapBaker::InitBakeDefaults()
 
 void FMeshMapBaker::Bake()
 {
-	TRACE_CPUPROFILER_EVENT_SCOPE("FMeshMapBaker::Bake");
+	TRACE_CPUPROFILER_EVENT_SCOPE(FMeshMapBaker::Bake);
 	
 	if (Bakers.IsEmpty())
 	{
@@ -220,7 +220,7 @@ void FMeshMapBaker::Bake()
 	// WriteToOutputBuffer transfers local tile data (TileBuffer) to the image output buffer (ImageTileBuffer).
 	auto WriteToOutputBuffer = [this, &ImageTileBuffer] (FMeshMapTileBuffer& TileBufferIn, const FImageTile& TargetTile, const TArray<int32>& BakeIds, auto&& Op, auto&& WeightOp)
 	{
-		TRACE_CPUPROFILER_EVENT_SCOPE("FMeshMapBaker::WriteToOutputBuffer");
+		TRACE_CPUPROFILER_EVENT_SCOPE(FMeshMapBaker::WriteToOutputBuffer);
 		
 		const int TargetTileWidth = TargetTile.GetWidth();
 		const int TargetTileHeight = TargetTile.GetHeight();
@@ -287,7 +287,7 @@ void FMeshMapBaker::Bake()
 	
 	ParallelFor(NumTiles, [this, &Tiles, &GutterTexelsPerTile, &OutputQueue, &WriteToOutputBuffer, &OverwriteFn, &NoopFn, &WriteQueuedOutput](int32 TileIdx)
 	{
-		TRACE_CPUPROFILER_EVENT_SCOPE("FMeshMapBaker::BakeTile");
+		TRACE_CPUPROFILER_EVENT_SCOPE(FMeshMapBaker::BakeTile);
 		// Generate unpadded and padded tiles.
 		const FImageTile Tile = Tiles.GetTile(TileIdx);	// Image area to sample
 		const FImageTile PaddedTile = Tiles.GetTile(TileIdx, TilePadding); // Filtered image area
@@ -372,7 +372,7 @@ void FMeshMapBaker::Bake()
 	// Normalize and convert ImageTileBuffer data to color data.
 	ParallelFor(NumTiles, [this, &Tiles, &ImageTileBuffer](int32 TileIdx)
 	{
-		TRACE_CPUPROFILER_EVENT_SCOPE("FMeshMapBaker::WriteToImageBuffer");
+		TRACE_CPUPROFILER_EVENT_SCOPE(FMeshMapBaker::WriteToImageBuffer);
 		
 		const FImageTile Tile = Tiles.GetTile(TileIdx);
 		const int TileWidth = Tile.GetWidth();
@@ -433,7 +433,7 @@ void FMeshMapBaker::Bake()
 		const int32 NumResults = BakeResults.Num();
 		ParallelFor(NumTiles, [this, &NumResults, &GutterTexelsPerTile](int32 TileIdx)
 		{
-			TRACE_CPUPROFILER_EVENT_SCOPE("FMeshMapBaker::WriteGutterPixels");
+			TRACE_CPUPROFILER_EVENT_SCOPE(FMeshMapBaker::WriteGutterPixels);
 			
 			for (int64 GutterIdx = 0; GutterIdx < GutterTexelsPerTile[TileIdx].Num(); ++GutterIdx)
 			{
