ModelingTools: Update BakeVertex normal computation to use NormalOverlay and Angle & Area weighting.

#rnx

diff --git a/Engine/Plugins/Runtime/GeometryProcessing/Source/DynamicMesh/Private/Sampling/MeshVertexBaker.cpp b/Engine/Plugins/Runtime/GeometryProcessing/Source/DynamicMesh/Private/Sampling/MeshVertexBaker.cpp
--- a/Engine/Plugins/Runtime/GeometryProcessing/Source/DynamicMesh/Private/Sampling/MeshVertexBaker.cpp
+++ b/Engine/Plugins/Runtime/GeometryProcessing/Source/DynamicMesh/Private/Sampling/MeshVertexBaker.cpp
@@ -2,6 +2,7 @@
 
 #include "Sampling/MeshVertexBaker.h"
 #include "Sampling/MeshBakerCommon.h"
+#include "DynamicMesh/MeshNormals.h"
 
 #include "ExplicitUseGeometryMathTypes.h"		// using UE::Geometry::(math types)
 using namespace UE::Geometry;
@@ -95,7 +96,7 @@ void FMeshVertexBaker::BakeImpl(void* Data)
 	ECorrespondenceStrategy UseStrategy = Baker->CorrespondenceStrategy;
 	bool bIsIdentity = true;
 	int NumDetailMeshes = 0;
-	auto CheckIdentity = [&Baker, &bIsIdentity, &NumDetailMeshes](const void* DetailMesh)
+	auto CheckIdentity = [Baker, &bIsIdentity, &NumDetailMeshes](const void* DetailMesh)
 	{
 		bIsIdentity = bIsIdentity && (DetailMesh == Baker->TargetMesh);
 		++NumDetailMeshes;
@@ -110,37 +111,46 @@ void FMeshVertexBaker::BakeImpl(void* Data)
 	
 	const FDynamicMesh3* Mesh = Baker->TargetMesh;
 	const FDynamicMeshColorOverlay* ColorOverlay = Baker->TargetMesh->Attributes()->PrimaryColors();
+	const FDynamicMeshNormalOverlay* NormalOverlay = Baker->TargetMesh->Attributes()->PrimaryNormals();
 
 	// TODO: Refactor into TMeshSurfaceSampler class (future vertex bake enhancements will require non-UV based surface sampling)
-	auto SampleSurface = [&Baker, &Mesh, &ColorOverlay, &UseStrategy](int32 ElementIdx,
+	auto SampleSurface = [Baker, Mesh, ColorOverlay, NormalOverlay, UseStrategy](int32 ElementIdx,
 	                                                                FMeshMapEvaluator::FCorrespondenceSample& ValueOut)
 	{
-		const int32 VId = ColorOverlay->GetParentVertex(ElementIdx);
+		const int32 VertexId = ColorOverlay->GetParentVertex(ElementIdx);
 
 		// Compute ray direction
-		// TODO: Precompute a color to normal overlay mapping data structure and replace this with a lookup.
 		FVector3f SurfaceNormal = FVector3f::Zero();
-		TArray<int> ColorElemTris;
-		ColorOverlay->GetElementTriangles(ElementIdx, ColorElemTris);
-		for (int TId : ColorElemTris)
+		TArray<int> ColorElementTris;
+		ColorOverlay->GetElementTriangles(ElementIdx, ColorElementTris);
+		for (const int TriId : ColorElementTris)
 		{
-			SurfaceNormal += Mesh->GetTriNormal(TId);
+			FVector3d TriNormal, TriCentroid; double TriArea;
+			Mesh->GetTriInfo(TriId, TriNormal, TriArea, TriCentroid);
+			FVector3d NormalWeights = FMeshNormals::GetVertexWeightsOnTriangle(Mesh, TriId, TriArea, true, true);
+			const FIndex3i TriVerts = Mesh->GetTriangle(TriId);
+			const int TriVertexId = IndexUtil::FindTriIndex(VertexId, TriVerts);
+			
+			if (NormalOverlay->IsSetTriangle(TriId))
+			{
+				FVector3f Normal;
+				NormalOverlay->GetElementAtVertex(TriId, VertexId, Normal);
+				SurfaceNormal += NormalWeights[TriVertexId] * Normal;
+			}
+			else
+			{
+				SurfaceNormal += NormalWeights[TriVertexId] * TriNormal;
+			}
 		}
 		Normalize(SurfaceNormal);
 
 		// Compute surface point and barycentric coords
-		const FVector3d SurfacePoint = Mesh->GetVertex(VId);
-		const int32 TriangleIndex = ColorElemTris[0];
+		const FVector3d SurfacePoint = Mesh->GetVertex(VertexId);
+		const int32 TriangleIndex = ColorElementTris[0];
+		const FIndex3i TriVerts = Mesh->GetTriangle(TriangleIndex);
+		const int TriVertexId = IndexUtil::FindTriIndex(VertexId, TriVerts);
 		FVector3d BaryCoords = FVector3d::Zero();
-		FIndex3i TriVerts = Mesh->GetTriangle(TriangleIndex);
-		for (int Idx = 0; Idx < 3; ++Idx)
-		{
-			if (VId == TriVerts[Idx])
-			{
-				BaryCoords[Idx] = 1.0;
-				break;
-			}
-		}
+		BaryCoords[TriVertexId] = 1.0;
 
 		ValueOut.BaseSample.TriangleIndex = TriangleIndex;
 		ValueOut.BaseSample.SurfacePoint = SurfacePoint;
@@ -149,7 +159,6 @@ void FMeshVertexBaker::BakeImpl(void* Data)
 		ValueOut.DetailTriID = FDynamicMesh3::InvalidID;
 		ValueOut.BaseNormal = SurfaceNormal;
 
-		const FVector3d RayDir = SurfaceNormal;
 		if (UseStrategy == ECorrespondenceStrategy::Identity)
 		{
 			ValueOut.DetailMesh = Mesh;
@@ -166,7 +175,7 @@ void FMeshVertexBaker::BakeImpl(void* Data)
 			const double SampleThickness = Baker->GetThickness(); // could modulate w/ a map here...
 
 			// Find detail mesh triangle point
-			ValueOut.DetailMesh = GetDetailMeshTrianglePoint_Raycast(Baker->DetailSampler, SurfacePoint, RayDir,
+			ValueOut.DetailMesh = GetDetailMeshTrianglePoint_Raycast(Baker->DetailSampler, SurfacePoint, SurfaceNormal,
 			                                   ValueOut.DetailTriID, ValueOut.DetailBaryCoords, SampleThickness,
 			                                   (UseStrategy == ECorrespondenceStrategy::RaycastStandardThenNearest));
 		}
@@ -178,7 +187,7 @@ void FMeshVertexBaker::BakeImpl(void* Data)
 	const FImageTiling Tiles(Baker->Dimensions, TileWidth, TileHeight);
 	const int32 NumTiles = Tiles.Num();
 	const int NumBakers = Baker->Bakers.Num();
-	ParallelFor(NumTiles, [&Baker, &Tiles, &TileWidth, &NumBakers, &SampleSurface](const int32 TileIdx)
+	ParallelFor(NumTiles, [Baker, &Tiles, TileWidth, NumBakers, &SampleSurface](const int32 TileIdx)
 	{
 		const FImageTile Tile = Tiles.GetTile(TileIdx);
 		const int Width = Tile.GetWidth();
diff --git a/Engine/Source/Runtime/GeometryCore/Public/DynamicMesh/MeshNormals.h b/Engine/Source/Runtime/GeometryCore/Public/DynamicMesh/MeshNormals.h
--- a/Engine/Source/Runtime/GeometryCore/Public/DynamicMesh/MeshNormals.h
+++ b/Engine/Source/Runtime/GeometryCore/Public/DynamicMesh/MeshNormals.h
@@ -189,20 +189,14 @@ public:
 	static bool QuickRecomputeOverlayNormals(FDynamicMesh3& Mesh, bool bInvert = false, bool bWeightByArea = true, bool bWeightByAngle = true);
 
 
-protected:
-	/** Compute per-vertex normals using area-weighted averaging of one-ring triangle normals */
-	void Compute_FaceAvg_AreaWeighted();
-	/** Compute per-vertex normals using a custom combination of area-weighted and angle-weighted averaging of one-ring triangle normals */
-	void Compute_FaceAvg(bool bWeightByArea, bool bWeightByAngle);
-
-	/** Compute per-triangle normals */
-	void Compute_Triangle();
-
-	/** Recompute the element Normals of the given attribute overlay using area-weighted averaging of one-ring triangle normals */
-	void Compute_Overlay_FaceAvg_AreaWeighted(const FDynamicMeshNormalOverlay* NormalOverlay);
-	/** Recompute the element Normals of the given attribute overlay using a custom combination of area-weighted and angle-weighted averaging of one-ring triangle normals */
-	void Compute_Overlay_FaceAvg(const FDynamicMeshNormalOverlay* NormalOverlay, bool bWeightByArea, bool bWeightByAngle);
-
+	/**
+	 * Retrieve the area and/or angle weights for each vertex of a triangle
+	 * @param Mesh the mesh to query
+	 * @param TriID the triangle index of the mesh to query
+	 * @param TriArea the area of the triangle
+	 * @param bWeightByArea if true, include weighting by the area of the triangle
+	 * @param bWeightByAngle if true, include weighting by the interior angles of the triangle
+	 */
 	static FVector3d GetVertexWeightsOnTriangle(const FDynamicMesh3* Mesh, int TriID, double TriArea, bool bWeightByArea, bool bWeightByAngle)
 	{
 		FVector3d TriNormalWeights = FVector3d::One();
@@ -217,6 +211,21 @@ protected:
 		return TriNormalWeights;
 	}
 
+
+protected:
+	/** Compute per-vertex normals using area-weighted averaging of one-ring triangle normals */
+	void Compute_FaceAvg_AreaWeighted();
+	/** Compute per-vertex normals using a custom combination of area-weighted and angle-weighted averaging of one-ring triangle normals */
+	void Compute_FaceAvg(bool bWeightByArea, bool bWeightByAngle);
+
+	/** Compute per-triangle normals */
+	void Compute_Triangle();
+
+	/** Recompute the element Normals of the given attribute overlay using area-weighted averaging of one-ring triangle normals */
+	void Compute_Overlay_FaceAvg_AreaWeighted(const FDynamicMeshNormalOverlay* NormalOverlay);
+	/** Recompute the element Normals of the given attribute overlay using a custom combination of area-weighted and angle-weighted averaging of one-ring triangle normals */
+	void Compute_Overlay_FaceAvg(const FDynamicMeshNormalOverlay* NormalOverlay, bool bWeightByArea, bool bWeightByAngle);
+
 };
 
 
