GeometryProcessing: remove empty GeometricObjects module, contents was moved to GeometryCore and this is no longer needed
#rnx

diff --git a/Engine/Plugins/Runtime/GeoReferencing/Source/GeoReferencing/GeoReferencing.Build.cs b/Engine/Plugins/Runtime/GeoReferencing/Source/GeoReferencing/GeoReferencing.Build.cs
--- a/Engine/Plugins/Runtime/GeoReferencing/Source/GeoReferencing/GeoReferencing.Build.cs
+++ b/Engine/Plugins/Runtime/GeoReferencing/Source/GeoReferencing/GeoReferencing.Build.cs
@@ -16,8 +16,7 @@ public class GeoReferencing : ModuleRules
 			new string[]
 			{
 				"Core",
-				"Engine",
-				"GeometricObjects"
+				"Engine"
 			}
 		);
 
diff --git a/Engine/Plugins/Runtime/GeometryProcessing/GeometryProcessing.uplugin b/Engine/Plugins/Runtime/GeometryProcessing/GeometryProcessing.uplugin
--- a/Engine/Plugins/Runtime/GeometryProcessing/GeometryProcessing.uplugin
+++ b/Engine/Plugins/Runtime/GeometryProcessing/GeometryProcessing.uplugin
@@ -14,11 +14,6 @@
 	"IsBetaVersion": true,
 	"Installed": false,
 	"Modules": [
-		{
-			"Name": "GeometricObjects",
-			"Type": "Runtime",
-			"LoadingPhase": "Default"
-		},
 		{
 			"Name": "GeometryAlgorithms",
 			"Type": "Runtime",
diff --git a/Engine/Plugins/Runtime/GeometryProcessing/Source/GeometricObjects/GeometricObjects.Build.cs b/Engine/Plugins/Runtime/GeometryProcessing/Source/GeometricObjects/GeometricObjects.Build.cs
deleted file mode 100644
diff --git a/Engine/Plugins/Runtime/GeometryProcessing/Source/GeometricObjects/Private/GeometricObjectsModule.cpp b/Engine/Plugins/Runtime/GeometryProcessing/Source/GeometricObjects/Private/GeometricObjectsModule.cpp
deleted file mode 100644
