UV Editor: Add some tracing to the UV Editor layout tool call stack

#rnx

diff --git a/Engine/Plugins/Experimental/MeshModelingToolsetExp/Source/ModelingOperators/Private/ParameterizationOps/UVLayoutOp.cpp b/Engine/Plugins/Experimental/MeshModelingToolsetExp/Source/ModelingOperators/Private/ParameterizationOps/UVLayoutOp.cpp
--- a/Engine/Plugins/Experimental/MeshModelingToolsetExp/Source/ModelingOperators/Private/ParameterizationOps/UVLayoutOp.cpp
+++ b/Engine/Plugins/Experimental/MeshModelingToolsetExp/Source/ModelingOperators/Private/ParameterizationOps/UVLayoutOp.cpp
@@ -19,6 +19,8 @@ void FUVLayoutOp::SetTransform(const FTransform& Transform) {
 
 void FUVLayoutOp::CalculateResult(FProgressCancel* Progress)
 {
+	TRACE_CPUPROFILER_EVENT_SCOPE(UVLayoutOp_CalculateResult);
+	
 	if (Progress && Progress->Cancelled())
 	{
 		return;
diff --git a/Engine/Plugins/Experimental/MeshModelingToolsetExp/Source/ModelingOperators/Public/BackgroundModelingComputeSource.h b/Engine/Plugins/Experimental/MeshModelingToolsetExp/Source/ModelingOperators/Public/BackgroundModelingComputeSource.h
--- a/Engine/Plugins/Experimental/MeshModelingToolsetExp/Source/ModelingOperators/Public/BackgroundModelingComputeSource.h
+++ b/Engine/Plugins/Experimental/MeshModelingToolsetExp/Source/ModelingOperators/Public/BackgroundModelingComputeSource.h
@@ -48,6 +48,8 @@ protected:
 	// 	FAbortableBackgroundTask API
 	void DoWork()
 	{
+		TRACE_CPUPROFILER_EVENT_SCOPE(ModelingOpTask_DoWork);
+		
 		if (Operator)
 		{
 			Operator->CalculateResult(GetProgress());
diff --git a/Engine/Plugins/Experimental/UVEditor/Source/UVEditorTools/Private/UVIslandConformalUnwrapAction.cpp b/Engine/Plugins/Experimental/UVEditor/Source/UVEditorTools/Private/UVIslandConformalUnwrapAction.cpp
--- a/Engine/Plugins/Experimental/UVEditor/Source/UVEditorTools/Private/UVIslandConformalUnwrapAction.cpp
+++ b/Engine/Plugins/Experimental/UVEditor/Source/UVEditorTools/Private/UVIslandConformalUnwrapAction.cpp
@@ -40,6 +40,8 @@ void UUVIslandConformalUnwrapAction::Shutdown()
 
 void UUVIslandConformalUnwrapAction::SetSelection(int32 SelectionTargetIndexIn, const UE::Geometry::FDynamicMeshSelection* NewSelection)
 {
+	TRACE_CPUPROFILER_EVENT_SCOPE(UVIslandConformalUnwrapAction_SetSelection);
+	
 	SelectionTargetIndex = SelectionTargetIndexIn;
 	if (NewSelection)
 	{
diff --git a/Engine/Plugins/Experimental/UVEditor/Source/UVEditorTools/Private/UVSeamSewAction.cpp b/Engine/Plugins/Experimental/UVEditor/Source/UVEditorTools/Private/UVSeamSewAction.cpp
--- a/Engine/Plugins/Experimental/UVEditor/Source/UVEditorTools/Private/UVSeamSewAction.cpp
+++ b/Engine/Plugins/Experimental/UVEditor/Source/UVEditorTools/Private/UVSeamSewAction.cpp
@@ -71,6 +71,8 @@ void UUVSeamSewAction::Shutdown()
 
 void UUVSeamSewAction::SetSelection(int32 SelectionTargetIndexIn, const UE::Geometry::FDynamicMeshSelection* NewSelection)
 {
+	TRACE_CPUPROFILER_EVENT_SCOPE(UVSeamSewAction_SetSelection);
+	
 	SelectionTargetIndex = SelectionTargetIndexIn;
 	if (NewSelection)
 	{
@@ -114,6 +116,8 @@ void UUVSeamSewAction::UpdateVisualizations()
 
 void UUVSeamSewAction::UpdateSewEdgePreviewLines()
 {
+	TRACE_CPUPROFILER_EVENT_SCOPE(UVSeamSewAction_UpdateSewEdgePreviewLines);
+	
 	SewEdgePairingLineSet->Clear();
 	if (CurrentSelection && !CurrentSelection->IsEmpty())
 	{
diff --git a/Engine/Plugins/Runtime/GeometryProcessing/Source/DynamicMesh/Private/Parameterization/MeshUVPacking.cpp b/Engine/Plugins/Runtime/GeometryProcessing/Source/DynamicMesh/Private/Parameterization/MeshUVPacking.cpp
--- a/Engine/Plugins/Runtime/GeometryProcessing/Source/DynamicMesh/Private/Parameterization/MeshUVPacking.cpp
+++ b/Engine/Plugins/Runtime/GeometryProcessing/Source/DynamicMesh/Private/Parameterization/MeshUVPacking.cpp
@@ -56,6 +56,8 @@ FDynamicMeshUVPacker::FDynamicMeshUVPacker(FDynamicMeshUVOverlay* UVOverlayIn)
 
 bool FDynamicMeshUVPacker::StandardPack()
 {
+	TRACE_CPUPROFILER_EVENT_SCOPE(DynamicMeshUVPacker_StandardPack);
+	
 	FUVPacker Packer;
 	Packer.bAllowFlips = bAllowFlips;
 	Packer.GutterSize = GutterSize;
@@ -78,6 +80,8 @@ bool FDynamicMeshUVPacker::StandardPack()
 
 bool FDynamicMeshUVPacker::StackPack()
 {
+	TRACE_CPUPROFILER_EVENT_SCOPE(DynamicMeshUVPacker_StackPack);
+	
 	FUVPacker Packer;
 	Packer.bAllowFlips = bAllowFlips;
 	Packer.GutterSize = GutterSize;
diff --git a/Engine/Plugins/Runtime/GeometryProcessing/Source/DynamicMesh/Private/Selections/MeshConnectedComponents.cpp b/Engine/Plugins/Runtime/GeometryProcessing/Source/DynamicMesh/Private/Selections/MeshConnectedComponents.cpp
--- a/Engine/Plugins/Runtime/GeometryProcessing/Source/DynamicMesh/Private/Selections/MeshConnectedComponents.cpp
+++ b/Engine/Plugins/Runtime/GeometryProcessing/Source/DynamicMesh/Private/Selections/MeshConnectedComponents.cpp
@@ -8,6 +8,8 @@ using namespace UE::Geometry;
 
 void FMeshConnectedComponents::FindConnectedTriangles(TFunction<bool(int32, int32)> TrisConnectedPredicate)
 {
+	TRACE_CPUPROFILER_EVENT_SCOPE(MeshConnectedComponents_FindConnectedTriangles);
+	
 	// initial active set contains all valid triangles
 	// active values are as follows:   0: unprocessed,  1: in queue,  2: done,  3: invalid
 	TArray<uint8> ActiveSet;
@@ -30,6 +32,8 @@ void FMeshConnectedComponents::FindConnectedTriangles(TFunction<bool(int32, int3
 
 void FMeshConnectedComponents::FindConnectedTriangles(TFunctionRef<bool(int)> IndexFilterFunc, TFunction<bool(int32, int32)> TrisConnectedPredicate)
 {
+	TRACE_CPUPROFILER_EVENT_SCOPE(MeshConnectedComponents_FindConnectedTriangles_Filtered);
+	
 	// initial active set contains all valid triangles
 	// active values are as follows:   0: unprocessed,  1: in queue,  2: done,  3: invalid
 	TArray<uint8> ActiveSet;
@@ -52,6 +56,8 @@ void FMeshConnectedComponents::FindConnectedTriangles(TFunctionRef<bool(int)> In
 
 void FMeshConnectedComponents::FindConnectedTriangles(const TArray<int>& TriangleROI, TFunction<bool(int32, int32)> TrisConnectedPredicate)
 {
+	TRACE_CPUPROFILER_EVENT_SCOPE(MeshConnectedComponents_FindConnectedTriangles_TriangleROI);
+
 	// initial active set contains all valid triangles
 	// active values are as follows:   0: unprocessed,  1: in queue,  2: done,  3: invalid
 	TArray<uint8> ActiveSet;
@@ -75,6 +81,8 @@ void FMeshConnectedComponents::FindConnectedTriangles(const TArray<int>& Triangl
 
 void FMeshConnectedComponents::FindTrianglesConnectedToSeeds(const TArray<int>& SeedTriangles, TFunction<bool(int32, int32)> TrisConnectedPredicate)
 {
+	TRACE_CPUPROFILER_EVENT_SCOPE(MeshConnectedComponents_FindConnectedTriangles_SeedTriangles);
+	
 	// initial active set contains all valid triangles
 	// active values are as follows:   0: unprocessed,  1: in queue,  2: done,  3: invalid
 	TArray<uint8> ActiveSet;
@@ -97,6 +105,8 @@ void FMeshConnectedComponents::FindTrianglesConnectedToSeeds(const TArray<int>&
 
 void FMeshConnectedComponents::FindTriComponents(FInterval1i ActiveRange, TArray<uint8>& ActiveSet, TFunction<bool(int32, int32)> TrisConnectedPredicate)
 {
+	TRACE_CPUPROFILER_EVENT_SCOPE(FindTriComponents_ActiveRange);
+	
 	Components.Empty();
 
 	// temporary queue
@@ -136,6 +146,8 @@ void FMeshConnectedComponents::FindTriComponents(FInterval1i ActiveRange, TArray
 
 void FMeshConnectedComponents::FindTriComponents(const TArray<int32>& SeedList, TArray<uint8>& ActiveSet, TFunction<bool(int32, int32)> TrisConnectedPredicate)
 {
+	TRACE_CPUPROFILER_EVENT_SCOPE(FindTriComponents_SeedList);
+	
 	Components.Empty();
 
 	// temporary queue
@@ -176,6 +188,8 @@ void FMeshConnectedComponents::FindTriComponents(const TArray<int32>& SeedList,
 
 void FMeshConnectedComponents::FindTriComponent(FComponent* Component, TArray<int32>& ComponentQueue, TArray<uint8>& ActiveSet)
 {
+	TRACE_CPUPROFILER_EVENT_SCOPE(FindTriComponents_Component);
+
 	while (ComponentQueue.Num() > 0)
 	{
 		int32 CurTriangle = ComponentQueue.Pop(false);
@@ -200,6 +214,8 @@ void FMeshConnectedComponents::FindTriComponent(FComponent* Component, TArray<in
 void FMeshConnectedComponents::FindTriComponent(FComponent* Component, TArray<int32>& ComponentQueue, TArray<uint8>& ActiveSet, 
 	TFunctionRef<bool(int32, int32)> TriConnectedPredicate)
 {
+	TRACE_CPUPROFILER_EVENT_SCOPE(FindTriComponents_Component_Filtered);
+	
 	while (ComponentQueue.Num() > 0)
 	{
 		int32 CurTriangle = ComponentQueue.Pop(false);
diff --git a/Engine/Source/Runtime/GeometryCore/Private/DynamicMesh/DynamicMeshOverlay.cpp b/Engine/Source/Runtime/GeometryCore/Private/DynamicMesh/DynamicMeshOverlay.cpp
--- a/Engine/Source/Runtime/GeometryCore/Private/DynamicMesh/DynamicMeshOverlay.cpp
+++ b/Engine/Source/Runtime/GeometryCore/Private/DynamicMesh/DynamicMeshOverlay.cpp
@@ -250,6 +250,8 @@ int TDynamicMeshOverlay<RealType, ElementSize>::SplitElementWithNewParent(int El
 template<typename RealType, int ElementSize>
 void TDynamicMeshOverlay<RealType, ElementSize>::SplitBowties()
 {
+	TRACE_CPUPROFILER_EVENT_SCOPE(DynamicMeshOverlay_SplitBowties);
+	
 	// arrays for storing contiguous triangle groups from parentmesh
 	TArray<int> TrianglesOut, ContiguousGroupLengths;
 	TArray<bool> GroupIsLoop;
diff --git a/Engine/Source/Runtime/GeometryCore/Private/Parameterization/UVPacking.cpp b/Engine/Source/Runtime/GeometryCore/Private/Parameterization/UVPacking.cpp
--- a/Engine/Source/Runtime/GeometryCore/Private/Parameterization/UVPacking.cpp
+++ b/Engine/Source/Runtime/GeometryCore/Private/Parameterization/UVPacking.cpp
@@ -93,6 +93,8 @@ protected:
 
 bool FStandardChartPacker::FindBestPacking(TArray<FUVIsland>& Charts)
 {
+	TRACE_CPUPROFILER_EVENT_SCOPE(StandardChartPacker_FindBestPacking);
+	
 	if ( (uint32)Charts.Num() > TextureResolution * TextureResolution)
 	{
 		// More charts than texels
@@ -398,6 +400,8 @@ void FStandardChartPacker::ScaleCharts( TArray<FUVIsland>& Charts, double UVScal
 
 bool FStandardChartPacker::PackCharts(TArray<FUVIsland>& Charts, double UVScale, double& OutEfficiency, TAtomic<bool>& bAbort)
 {
+	TRACE_CPUPROFILER_EVENT_SCOPE(StandardChartPacker_PackCharts);
+	
 	ScaleCharts(Charts, UVScale);
 
 	FUVSpaceAllocator BestChartRaster(FUVSpaceAllocator::EMode::UsedSegments, TextureResolution, TextureResolution);
@@ -807,6 +811,8 @@ void FStandardChartPacker::RasterizeChart(const FUVIsland& Chart, uint32 RectW,
 
 bool FUVPacker::StandardPack(IUVMeshView* Mesh, int NumIslands, TFunctionRef<void(int, TArray<int32>&)> CopyIsland)
 {
+	TRACE_CPUPROFILER_EVENT_SCOPE(UVPacker_StandardPack);
+	
 	using namespace UE::InternalUVPacking;
 	FStandardChartPacker Packer;
 
@@ -883,6 +889,8 @@ bool FUVPacker::StandardPack(IUVMeshView* Mesh, int NumIslands, TFunctionRef<voi
 
 bool FUVPacker::StackPack(IUVMeshView* Mesh, int NumIslands, TFunctionRef<void(int, TArray<int32>&)> CopyIsland)
 {
+	TRACE_CPUPROFILER_EVENT_SCOPE(UVPacker_StackPack);
+	
 	double GutterWidth = (double)GutterSize / (double)TextureResolution;
 
 	int32 NumCharts = NumIslands;
@@ -952,6 +960,8 @@ bool FUVPacker::StackPack(IUVMeshView* Mesh, int NumIslands, TFunctionRef<void(i
 
 void FUVPacker::GetIslandStats(IUVMeshView* Mesh, const TArray<int32>& Island, FAxisAlignedBox2d& IslandBoundsOut, double& IslandScaleFactorOut, double& UVAreaOut)
 {
+	TRACE_CPUPROFILER_EVENT_SCOPE(UVPacker_GetIslandStats);
+
 	IslandBoundsOut = FAxisAlignedBox2d::Empty();
 	UVAreaOut = 0.0f;
 	double UVLengthSum = 0, WorldLengthSum = 0;
