GeometryProcessing: add initial MeshBevel operation. Still a WIP, does not handle UVs at all, or interior-region vertices, and the current solve for vertex positions produces horrible-looking bevels, but the main topological problems seem to be solved.
Various utility functions added to MeshIndexUtil.h and FIndex2i/3i.
#rnx

diff --git a/Engine/Plugins/Runtime/GeometryProcessing/Source/DynamicMesh/Private/Operations/MeshBevel.cpp b/Engine/Plugins/Runtime/GeometryProcessing/Source/DynamicMesh/Private/Operations/MeshBevel.cpp
new file mode 100644
--- /dev/null
+++ b/Engine/Plugins/Runtime/GeometryProcessing/Source/DynamicMesh/Private/Operations/MeshBevel.cpp
@@ -0,0 +1,1101 @@
+// Copyright Epic Games, Inc. All Rights Reserved.
+
+#include "Operations/MeshBevel.h"
+
+#include "GroupTopology.h"
+#include "EdgeLoop.h"
+#include "DynamicMesh/DynamicMeshChangeTracker.h"
+#include "MeshWeights.h"
+#include "CompGeom/PolygonTriangulation.h"
+#include "DynamicMesh/MeshNormals.h"
+#include "DynamicMesh/MeshIndexUtil.h"
+#include "Algo/Count.h"
+
+using namespace UE::Geometry;
+
+
+
+void FMeshBevel::InitializeFromGroupTopology(const FDynamicMesh3& Mesh, const FGroupTopology& Topology)
+{
+	ResultInfo = FGeometryResult(EGeometryResultType::InProgress);
+
+	// set up initial problem inputs
+	for (int32 TopoEdgeID = 0; TopoEdgeID < Topology.Edges.Num(); ++TopoEdgeID)
+	{
+		if (Topology.IsIsolatedLoop(TopoEdgeID))
+		{
+			FEdgeLoop NewLoop;
+			NewLoop.InitializeFromEdges(&Mesh, Topology.Edges[TopoEdgeID].Span.Edges);
+			AddBevelEdgeLoop(Mesh, NewLoop);
+		}
+		else
+		{
+			AddBevelGroupEdge(Mesh, Topology, TopoEdgeID);
+		}
+
+		if (ResultInfo.CheckAndSetCancelled(Progress))
+		{
+			return;
+		}
+	}
+
+	// precompute topological information necessary to apply bevel to vertices/edges/loops
+	BuildVertexSets(Mesh);
+}
+
+
+void FMeshBevel::InitializeFromGroupTopology(const FDynamicMesh3& Mesh, const FGroupTopology& Topology, const TArray<int32>& GroupEdges)
+{
+	ResultInfo = FGeometryResult(EGeometryResultType::InProgress);
+
+	// set up initial problem inputs
+	for (int32 TopoEdgeID : GroupEdges)
+	{
+		if (Topology.IsIsolatedLoop(TopoEdgeID))
+		{
+			FEdgeLoop NewLoop;
+			NewLoop.InitializeFromEdges(&Mesh, Topology.Edges[TopoEdgeID].Span.Edges);
+			AddBevelEdgeLoop(Mesh, NewLoop);
+		}
+		else
+		{
+			AddBevelGroupEdge(Mesh, Topology, TopoEdgeID);
+		}
+
+		if (ResultInfo.CheckAndSetCancelled(Progress))
+		{
+			return;
+		}
+	}
+
+	// precompute topological information necessary to apply bevel to vertices/edges/loops
+	BuildVertexSets(Mesh);
+}
+
+
+bool FMeshBevel::Apply(FDynamicMesh3& Mesh, FDynamicMeshChangeTracker* ChangeTracker)
+{
+	// disconnect along bevel graph edges/vertices and save necessary info
+	UnlinkEdges(Mesh, ChangeTracker);
+	if (ResultInfo.CheckAndSetCancelled(Progress))
+	{
+		return false;
+	}
+	UnlinkLoops(Mesh, ChangeTracker);
+	if (ResultInfo.CheckAndSetCancelled(Progress))
+	{
+		return false;
+	}
+	UnlinkVertices(Mesh, ChangeTracker);
+	if (ResultInfo.CheckAndSetCancelled(Progress))
+	{
+		return false;
+	}
+
+	// update vertex positions
+	DisplaceVertices(Mesh, InsetDistance);
+	if (ResultInfo.CheckAndSetCancelled(Progress))
+	{
+		return false;
+	}
+
+	// mesh the bevel corners and edges
+	CreateBevelMeshing(Mesh);
+	if (ResultInfo.CheckAndSetCancelled(Progress))
+	{
+		return false;
+	}
+
+	// compute normals
+	ComputeNormals(Mesh);
+	if (ResultInfo.CheckAndSetCancelled(Progress))
+	{
+		return false;
+	}
+
+	// todo: compute UVs, other attribs
+
+	ResultInfo.SetSuccess(true, Progress);
+	return true;
+}
+
+
+
+FMeshBevel::FBevelVertex* FMeshBevel::GetBevelVertexFromVertexID(int32 VertexID)
+{
+	int32* FoundIndex = VertexIDToIndexMap.Find(VertexID);
+	if (FoundIndex == nullptr)
+	{
+		return nullptr;
+	}
+	return &Vertices[*FoundIndex];
+}
+
+
+
+void FMeshBevel::AddBevelGroupEdge(const FDynamicMesh3& Mesh, const FGroupTopology& Topology, int32 GroupEdgeID)
+{
+	const TArray<int32>& MeshEdgeList = Topology.Edges[GroupEdgeID].Span.Edges;
+
+	// currently cannot handle boundary edges
+	if ( Algo::CountIf(MeshEdgeList, [&Mesh](int EdgeID) { return Mesh.IsBoundaryEdge(EdgeID); }) > 0 )
+	{
+		return;
+	}
+
+	FIndex2i EdgeCornerIDs = Topology.Edges[GroupEdgeID].EndpointCorners;
+
+	FBevelEdge Edge;
+
+	for (int32 ci = 0; ci < 2; ++ci)
+	{
+		int32 CornerID = EdgeCornerIDs[ci];
+		FGroupTopology::FCorner Corner = Topology.Corners[CornerID];
+		int32 VertexID = Corner.VertexID;
+		Edge.bEndpointBoundaryFlag[ci] = Mesh.IsBoundaryVertex(VertexID);
+		int32 IncomingEdgeID = (ci == 0) ? MeshEdgeList[0] : MeshEdgeList.Last();
+
+		FBevelVertex* VertInfo = GetBevelVertexFromVertexID(VertexID);
+		if (VertInfo == nullptr)
+		{
+			FBevelVertex NewVertex;
+			NewVertex.CornerID = CornerID;
+			NewVertex.VertexID = VertexID;
+			int32 NewIndex = Vertices.Num();
+			Vertices.Add(NewVertex);
+			VertexIDToIndexMap.Add(VertexID, NewIndex);
+			VertInfo = &Vertices[NewIndex];
+		}
+		VertInfo->IncomingBevelMeshEdges.Add(IncomingEdgeID);
+		VertInfo->IncomingBevelTopoEdges.Add(GroupEdgeID);
+	}
+
+	Edge.MeshEdges.Append(MeshEdgeList);
+	Edge.MeshVertices.Append(Topology.Edges[GroupEdgeID].Span.Vertices);
+	Edge.GroupEdgeID = GroupEdgeID;
+	Edge.GroupIDs = Topology.Edges[GroupEdgeID].Groups;
+	Edges.Add(MoveTemp(Edge));
+}
+
+
+
+
+
+void FMeshBevel::AddBevelEdgeLoop(const FDynamicMesh3& Mesh, const FEdgeLoop& MeshEdgeLoop)
+{
+	// currently cannot handle boundary edges
+	if ( Algo::CountIf(MeshEdgeLoop.Edges, [&Mesh](int EdgeID) { return Mesh.IsBoundaryEdge(EdgeID); }) > 0 )
+	{
+		return;
+	}
+
+	FBevelLoop Loop;
+	Loop.MeshEdges = MeshEdgeLoop.Edges;
+	Loop.MeshVertices = MeshEdgeLoop.Vertices;
+
+	Loops.Add(Loop);
+}
+
+
+
+
+void FMeshBevel::BuildVertexSets(const FDynamicMesh3& Mesh)
+{
+	// can be parallel
+	for (FBevelVertex& Vertex : Vertices)
+	{
+		// get sorted list of triangles around the vertex
+		TArray<int> GroupLengths;
+		TArray<bool> bGroupIsLoop;
+		EMeshResult Result = Mesh.GetVtxContiguousTriangles(Vertex.VertexID, Vertex.SortedTriangles, GroupLengths, bGroupIsLoop);
+		if ( Result != EMeshResult::Ok || GroupLengths.Num() != 1 || Vertex.SortedTriangles.Num() < 2)
+		{
+			Vertex.VertexType = EBevelVertexType::Unknown;
+			continue;
+		}
+
+		// GetVtxContiguousTriangles does not return triangles sorted in a consistent direction. This check will
+		// reverse the ordering such that it is consistently walking counter-clockwise around the vertex (I think...)
+		FIndex3i Tri0 = Mesh.GetTriangle(Vertex.SortedTriangles[0]).GetCycled(Vertex.VertexID);
+		FIndex3i Tri1 = Mesh.GetTriangle(Vertex.SortedTriangles[1]).GetCycled(Vertex.VertexID);
+		if (Tri0.C == Tri1.B)
+		{
+			Algo::Reverse(Vertex.SortedTriangles);
+		}
+
+		if (Mesh.IsBoundaryVertex(Vertex.VertexID))
+		{
+			Vertex.VertexType = EBevelVertexType::BoundaryVertex;
+			continue;
+		}
+
+
+		checkSlow(Vertex.IncomingBevelMeshEdges.Num() != 0);		// shouldn't ever happen
+		if (Vertex.IncomingBevelMeshEdges.Num() == 1)
+		{
+			BuildTerminatorVertex(Vertex, Mesh);
+		}
+		else
+		{
+			BuildJunctionVertex(Vertex, Mesh);
+		}
+
+		if (ResultInfo.CheckAndSetCancelled(Progress))
+		{
+			return;
+		}
+	}
+}
+
+
+
+void FMeshBevel::BuildJunctionVertex(FBevelVertex& Vertex, const FDynamicMesh3& Mesh)
+{
+	//
+	// Now split up the single contiguous one-ring into "Wedges" created by the incoming split-edges
+	//
+
+	// find first split edge, and the triangle/index "after" that first split edge
+	int32 NT = Vertex.SortedTriangles.Num();
+	int32 StartTriIndex = -1;
+	for (int32 k = 0; k < NT; ++k)
+	{
+		if (FindSharedEdgeInTriangles(Mesh, Vertex.SortedTriangles[k], Vertex.SortedTriangles[(k + 1) % NT]) == Vertex.IncomingBevelMeshEdges[0])
+		{
+			StartTriIndex = (k + 1) % NT;		// start at second tri, so that bevel-edge is first edge in wedge
+			break;
+		}
+	}
+	if (StartTriIndex == -1)
+	{
+		Vertex.VertexType = EBevelVertexType::Unknown;
+		return;
+	}
+
+	// now walk around the one-ring tris, accumulating into current Wedge until we hit another split-edge,
+	// at which point a new Wedge is spawned
+	int32 CurTriIndex = StartTriIndex;
+	FOneRingWedge CurWedge;
+	CurWedge.WedgeVertex = Vertex.VertexID;
+	CurWedge.Triangles.Add(Vertex.SortedTriangles[CurTriIndex]);
+	CurWedge.BorderEdges.A = Vertex.IncomingBevelMeshEdges[0];
+	for (int32 k = 0; k < NT; ++k)
+	{
+		int32 CurTri = Vertex.SortedTriangles[CurTriIndex % NT];
+		int32 NextTri = Vertex.SortedTriangles[(CurTriIndex + 1) % NT];
+		int32 SharedEdge = FindSharedEdgeInTriangles(Mesh, CurTri, NextTri);
+		checkSlow(SharedEdge != -1);
+		if (Vertex.IncomingBevelMeshEdges.Contains(SharedEdge))
+		{
+			// if we found a bevel-edge, close the current wedge and start a new one
+			CurWedge.BorderEdges.B = SharedEdge;
+			Vertex.Wedges.Add(CurWedge);
+			CurWedge = FOneRingWedge();
+			CurWedge.WedgeVertex = Vertex.VertexID;
+			CurWedge.BorderEdges.A = SharedEdge;
+		}
+		CurWedge.Triangles.Add(NextTri);
+		CurTriIndex++;
+	}
+	// ?? is there a chance that we have a final open wedge here? we iterate one extra time so it shouldn't happen (or could we get an extra wedge then??)
+
+	for (FOneRingWedge& Wedge : Vertex.Wedges)
+	{
+		Wedge.BorderEdgeTriEdgeIndices.A = Mesh.GetTriEdges(Wedge.Triangles[0]).IndexOf(Wedge.BorderEdges.A);
+		Wedge.BorderEdgeTriEdgeIndices.B = Mesh.GetTriEdges(Wedge.Triangles.Last()).IndexOf(Wedge.BorderEdges.B);
+	}
+
+	if (Vertex.Wedges.Num() > 1)
+	{
+		Vertex.VertexType = EBevelVertexType::JunctionVertex;
+	}
+	else
+	{
+		Vertex.VertexType = EBevelVertexType::Unknown;
+	}
+}
+
+
+void FMeshBevel::BuildTerminatorVertex(FBevelVertex& Vertex, const FDynamicMesh3& Mesh)
+{
+	Vertex.VertexType = EBevelVertexType::Unknown;
+	if (ensure(Vertex.IncomingBevelMeshEdges.Num() == 1) == false)
+	{
+		return;
+	}
+
+	int32 IncomingEdgeID = Vertex.IncomingBevelMeshEdges[0];
+	int32 NumTris = Vertex.SortedTriangles.Num();
+
+	// We have one edge coming into the vertex one ring, our main problem is to pick a second
+	// edge to split the one-ring with. There is no obvious right answer in many cases. 
+
+	// "other" split edge that we pick
+	int32 RingSplitEdgeID = -1;
+
+	// NOTE: code below assumes we have polygroups on the mesh. In theory we could also support not having polygroups,
+	// by (eg) picking an arbitrary edge "furthest" from IncomingEdgeID in the exponential map
+
+	// Find the ordered set of triangles that are not in either of the groups connected to IncomingEdgeID. Eg imagine
+	// a cube corner, if we want to bevel IncomingEdgeID along one of the cube edges, we want to add the new edge
+	// in the "furthest" face (perpendicular to the edge)
+	FIndex2i IncomingEdgeT = Mesh.GetEdgeT(IncomingEdgeID);
+	FIndex2i IncomingEdgeGroups(Mesh.GetTriangleGroup(IncomingEdgeT.A), Mesh.GetTriangleGroup(IncomingEdgeT.B));
+	TArray<int32> VertexGroups;
+	Mesh.GetAllVertexGroups(Vertex.VertexID, VertexGroups);
+	TArray<int32> OtherGroupTris;	// sorted wedge of triangles that are not in either of the groups connected to incoming edge
+	TArray<int32> OtherGroups;		// list of group IDs encountered, in-order
+	for (int32 tid : Vertex.SortedTriangles)
+	{
+		int32 gid = Mesh.GetTriangleGroup(tid);
+		if (IncomingEdgeGroups.Contains(gid) == false)
+		{
+			OtherGroupTris.Add(tid);
+			OtherGroups.AddUnique(gid);
+		}
+	}
+
+	// Determine which edge to split at in the "other" group triangles. If we only have one group
+	// then we can try to pick the "middlest" edge. The worst case is when there is only one triangle,
+	// then we are picking a not-very-good edge no matter what (potentially we should do an edge split or
+	// face poke in that situation). If we have multiple groups then we probably want to pick one
+	// of the group-boundary edges inside the triangle-span, ideally the "middlest" but currently we
+	// are just picking one arbitrarily...
+	if (OtherGroups.Num() == 1)
+	{
+		Vertex.NewGroupID = OtherGroups[0];
+		if (OtherGroupTris.Num() == 1)
+		{
+			FIndex3i TriEdges = Mesh.GetTriEdges(OtherGroupTris[0]);
+			for (int32 j = 0; j < 3; ++j)
+			{
+				if (Mesh.GetEdgeV(TriEdges[j]).Contains(Vertex.VertexID))
+				{
+					RingSplitEdgeID = TriEdges[j];
+					break;
+				}
+			}
+		}
+		else if (OtherGroupTris.Num() == 2)
+		{
+			RingSplitEdgeID = FindSharedEdgeInTriangles(Mesh, OtherGroupTris[0], OtherGroupTris[1]);
+		}
+		else
+		{
+			// TODO: should compute opening angles here and pick the edge closest to the middle of the angular span!!
+			int32 j = OtherGroupTris.Num() / 2;
+			RingSplitEdgeID = FindSharedEdgeInTriangles(Mesh, OtherGroupTris[j], OtherGroupTris[j+1]);
+		}
+	}
+	else
+	{
+		for (int32 k = 1; k < OtherGroupTris.Num(); ++k)
+		{
+			if (Mesh.GetTriangleGroup(OtherGroupTris[k-1]) != Mesh.GetTriangleGroup(OtherGroupTris[k]))
+			{
+				RingSplitEdgeID = FindSharedEdgeInTriangles(Mesh, OtherGroupTris[k-1], OtherGroupTris[k]);
+				Vertex.NewGroupID = Mesh.GetTriangleGroup(OtherGroupTris[k-1]);		// this seems very arbitrary
+				break;
+			}
+		}
+	}
+
+	if (RingSplitEdgeID == -1)
+	{
+		return;
+	}
+
+	FIndex2i SplitEdgeV = Mesh.GetEdgeV(RingSplitEdgeID);
+	Vertex.TerminatorInfo = FIndex2i(RingSplitEdgeID, SplitEdgeV.OtherElement(Vertex.VertexID));
+
+	TArray<int32> SplitTriSets[2];
+	if (SplitInteriorVertexTrianglesIntoSubsets(&Mesh, Vertex.VertexID, IncomingEdgeID, RingSplitEdgeID, SplitTriSets[0], SplitTriSets[1]) == false)
+	{
+		return;
+	}
+
+	Vertex.Wedges.SetNum(2);
+	Vertex.Wedges[0].WedgeVertex = Vertex.VertexID;
+	Vertex.Wedges[0].Triangles.Append(SplitTriSets[0]);
+	Vertex.Wedges[1].WedgeVertex = Vertex.VertexID;
+	Vertex.Wedges[1].Triangles.Append(SplitTriSets[1]);
+
+	Vertex.VertexType = EBevelVertexType::TerminatorVertex;
+}
+
+
+void FMeshBevel::UnlinkEdges(FDynamicMesh3& Mesh, FDynamicMeshChangeTracker* ChangeTracker)
+{
+	for (FBevelEdge& Edge : Edges)
+	{
+		UnlinkBevelEdgeInterior(Mesh, Edge, ChangeTracker);
+	}
+}
+
+
+
+namespace UELocal
+{
+	// decomposition of a vertex one-ring into two connected triangle subsets
+	struct FVertexSplit
+	{
+		int32 VertexID;
+		bool bOK;
+		TArray<int32> TriSets[2];
+	};
+
+	// walk along a sequence of vertex-splits and make sure that the split triangle sets
+	// maintain consistent "sides" (see call in UnlinkBevelEdgeInterior for more details)
+	static void ReconcileTriangleSets(TArray<FVertexSplit>& SplitSequence)
+	{
+		int32 N = SplitSequence.Num();
+		TArray<int32> PrevTriSet0;
+		for (int32 k = 0; k < N; ++k)
+		{
+			if (PrevTriSet0.Num() == 0 && SplitSequence[k].TriSets[0].Num() > 0)
+			{
+				PrevTriSet0 = SplitSequence[k].TriSets[0];
+			}
+			else
+			{
+				bool bFoundInSet0 = false;
+				for (int32 tid : SplitSequence[k].TriSets[0])
+				{
+					if (PrevTriSet0.Contains(tid))
+					{
+						bFoundInSet0 = true;
+						break;
+					}
+				}
+				if (!bFoundInSet0)
+				{
+					Swap(SplitSequence[k].TriSets[0], SplitSequence[k].TriSets[1]);
+				}
+				PrevTriSet0 = SplitSequence[k].TriSets[0];
+			}
+		}
+	}
+
+};
+
+
+void FMeshBevel::UnlinkBevelEdgeInterior(
+	FDynamicMesh3& Mesh,
+	FBevelEdge& BevelEdge,
+	FDynamicMeshChangeTracker* ChangeTracker)
+{
+	// figure out what sets of triangles to split each vertex into
+	int32 N = BevelEdge.MeshVertices.Num();
+
+	TArray<UELocal::FVertexSplit> SplitsToProcess;
+	SplitsToProcess.SetNum(N);
+
+	// precompute triangle sets for each vertex we want to split, by "cutting" the one ring into two halves based
+	// on edges - 2 edges for interior vertices, and 1 edge for a boundary vertex at the start/end of the edge-span
+
+	SplitsToProcess[0] = UELocal::FVertexSplit{ BevelEdge.MeshVertices[0], false };
+	if (BevelEdge.bEndpointBoundaryFlag[0])
+	{
+		SplitsToProcess[0].bOK = SplitBoundaryVertexTrianglesIntoSubsets(&Mesh, SplitsToProcess[0].VertexID, BevelEdge.MeshEdges[0], 
+			SplitsToProcess[0].TriSets[0], SplitsToProcess[0].TriSets[1]);
+	}
+	for (int32 k = 1; k < N - 1; ++k)
+	{
+		SplitsToProcess[k].VertexID = BevelEdge.MeshVertices[k];
+		if (Mesh.IsBoundaryVertex(SplitsToProcess[k].VertexID))
+		{
+			SplitsToProcess[k].bOK = false;
+		}
+		else
+		{
+			SplitsToProcess[k].bOK = SplitInteriorVertexTrianglesIntoSubsets(&Mesh, SplitsToProcess[k].VertexID,
+				BevelEdge.MeshEdges[k-1], BevelEdge.MeshEdges[k], SplitsToProcess[k].TriSets[0], SplitsToProcess[k].TriSets[1]);
+		}
+	}
+	SplitsToProcess[N-1] = UELocal::FVertexSplit{ BevelEdge.MeshVertices[N - 1], false };
+	if (BevelEdge.bEndpointBoundaryFlag[1])
+	{
+		SplitsToProcess[N-1].bOK = SplitBoundaryVertexTrianglesIntoSubsets(&Mesh, SplitsToProcess[N-1].VertexID, BevelEdge.MeshEdges[N-2], 
+			SplitsToProcess[N-1].TriSets[0], SplitsToProcess[N-1].TriSets[1]);
+	}
+
+	// SplitInteriorVertexTrianglesIntoSubsets does not consistently order its output sets, ie, if you imagine [Edge0,Edge1] as a path
+	// cutting through the one ring, the "side" that Set0 and Set1 end up is arbitrary, and depends on the ordering of edges in the triangles of Edge1.
+	// This might ideally be fixed in the future, but for the time being, all we need is consistency. So we walk from the start of the 
+	// edge to the end, checking for overlap between each tri-one-ring-wedge. If Split[k].TriSet0 does not overlap with Split[k-1].TriSet0, then
+	// we want to swap TriSet0 and TriSet1 at Split[k].
+	UELocal::ReconcileTriangleSets(SplitsToProcess);
+
+	// apply vertex splits and accumulate new list
+	N = SplitsToProcess.Num();
+	for (int32 k = 0; k < N; ++k)
+	{
+		const UELocal::FVertexSplit& Split = SplitsToProcess[k];
+		if (ChangeTracker)
+		{
+			ChangeTracker->SaveVertexOneRingTriangles(Split.VertexID, true);
+		}
+
+		bool bDone = false;
+		if (Split.bOK)
+		{
+			FDynamicMesh3::FVertexSplitInfo SplitInfo;
+			EMeshResult Result = Mesh.SplitVertex(Split.VertexID, Split.TriSets[0], SplitInfo);
+			if (ensure(Result == EMeshResult::Ok))
+			{
+				BevelEdge.NewMeshVertices.Add(SplitInfo.NewVertex);
+				bDone = true;
+			}
+		}
+		if (!bDone)
+		{
+			BevelEdge.NewMeshVertices.Add(Split.VertexID);
+		}
+	}
+
+	// now build edge correspondences
+	N = BevelEdge.MeshVertices.Num();
+	checkSlow(N == BevelEdge.NewMeshVertices.Num());
+	for (int32 k = 0; k < N-1; ++k)
+	{
+		int32 Edge0 = BevelEdge.MeshEdges[k];
+		int32 Edge1 = Mesh.FindEdge(BevelEdge.NewMeshVertices[k], BevelEdge.NewMeshVertices[k + 1]);
+		BevelEdge.NewMeshEdges.Add(Edge1);
+		checkSlow(Edge1 >= 0);
+		if ( Mesh.IsEdge(Edge1) && Edge0 != Edge1 && MeshEdgePairs.Contains(Edge0) == false )
+		{
+			MeshEdgePairs.Add(Edge0, Edge1);
+			MeshEdgePairs.Add(Edge1, Edge0);
+		}
+	}
+}
+
+
+
+
+void FMeshBevel::UnlinkBevelLoop(FDynamicMesh3& Mesh, FBevelLoop& BevelLoop, FDynamicMeshChangeTracker* ChangeTracker)
+{
+	int32 N = BevelLoop.MeshVertices.Num();
+
+	TArray<UELocal::FVertexSplit> SplitsToProcess;
+	SplitsToProcess.SetNum(N);
+
+	// precompute triangle sets for each vertex we want to split
+	for (int32 k = 0; k < N; ++k)
+	{
+		SplitsToProcess[k].VertexID = BevelLoop.MeshVertices[k];
+		if (Mesh.IsBoundaryVertex(SplitsToProcess[k].VertexID))
+		{
+			// cannot split boundary vertex
+			SplitsToProcess[k].bOK = false;
+		}
+		else
+		{
+			int32 PrevEdge = (k == 0) ? BevelLoop.MeshEdges.Last() : BevelLoop.MeshEdges[k-1];
+			int32 CurEdge = BevelLoop.MeshEdges[k];
+			SplitsToProcess[k].bOK = SplitInteriorVertexTrianglesIntoSubsets(&Mesh, SplitsToProcess[k].VertexID,
+				PrevEdge, CurEdge, SplitsToProcess[k].TriSets[0], SplitsToProcess[k].TriSets[1]);
+		}
+	}
+
+	// fix up triangle sets - see call in UnlinkBevelEdgeInterior() for more info
+	UELocal::ReconcileTriangleSets(SplitsToProcess);
+
+	// apply vertex splits and accumulate new list
+	N = SplitsToProcess.Num();
+	for (int32 k = 0; k < N; ++k)
+	{
+		const UELocal::FVertexSplit& Split = SplitsToProcess[k];
+		if (ChangeTracker)
+		{
+			ChangeTracker->SaveVertexOneRingTriangles(Split.VertexID, true);
+		}
+
+		bool bDone = false;
+		if (Split.bOK)
+		{
+			FDynamicMesh3::FVertexSplitInfo SplitInfo;
+			EMeshResult Result = Mesh.SplitVertex(Split.VertexID, Split.TriSets[1], SplitInfo);
+			if (Result == EMeshResult::Ok)
+			{
+				BevelLoop.NewMeshVertices.Add(SplitInfo.NewVertex);
+				bDone = true;
+			}
+		}
+		if (!bDone)
+		{
+			BevelLoop.NewMeshVertices.Add(Split.VertexID);		// failed to split, so we have a shared vertex on both "sides"
+		}
+	}
+
+	// now build edge correspondences
+	N = BevelLoop.MeshVertices.Num();
+	checkSlow(N == BevelLoop.NewMeshVertices.Num());
+	for (int32 k = 0; k < N; ++k)
+	{
+		int32 Edge0 = BevelLoop.MeshEdges[k];
+		int32 Edge1 = Mesh.FindEdge(BevelLoop.NewMeshVertices[k], BevelLoop.NewMeshVertices[(k + 1)%N]);
+		BevelLoop.NewMeshEdges.Add(Edge1);
+		checkSlow(Edge1 >= 0);
+		if (Mesh.IsEdge(Edge1) && Edge0 != Edge1 && MeshEdgePairs.Contains(Edge0) == false )
+		{
+			MeshEdgePairs.Add(Edge0, Edge1);
+			MeshEdgePairs.Add(Edge1, Edge0);
+		}
+	}
+}
+
+void FMeshBevel::UnlinkLoops(FDynamicMesh3& Mesh, FDynamicMeshChangeTracker* ChangeTracker)
+{
+	for (FBevelLoop& Loop : Loops)
+	{
+		UnlinkBevelLoop(Mesh, Loop, ChangeTracker);
+	}
+}
+
+
+
+
+
+void FMeshBevel::UnlinkVertices(FDynamicMesh3& Mesh, FDynamicMeshChangeTracker* ChangeTracker)
+{
+	// TODO: currently have to do terminator vertices first because we do some of the 
+	// determination inside the unlink code...
+
+	for (FBevelVertex& Vertex : Vertices)
+	{
+		if (Vertex.VertexType == EBevelVertexType::TerminatorVertex)
+		{
+			UnlinkTerminatorVertex(Mesh, Vertex, ChangeTracker);
+		}
+	}
+
+	for (FBevelVertex& Vertex : Vertices)
+	{
+		if (Vertex.VertexType == EBevelVertexType::JunctionVertex)
+		{
+			UnlinkJunctionVertex(Mesh, Vertex, ChangeTracker);
+		}
+	}
+}
+
+
+void FMeshBevel::UnlinkJunctionVertex(FDynamicMesh3& Mesh, FBevelVertex& Vertex, FDynamicMeshChangeTracker* ChangeTracker)
+{
+	check(Vertex.VertexType == EBevelVertexType::JunctionVertex);
+
+	if (ChangeTracker)
+	{
+		ChangeTracker->SaveVertexOneRingTriangles(Vertex.VertexID, true);
+	}
+
+	int32 NumWedges = Vertex.Wedges.Num();
+	checkSlow(NumWedges > 1);
+
+	// Split triangles around vertex into separate tri-sets based on wedges.
+	// This will create a new vertex for each wedge.
+	for (int32 k = 1; k < NumWedges; ++k)
+	{
+		FOneRingWedge& Wedge = Vertex.Wedges[k];
+
+		FDynamicMesh3::FVertexSplitInfo SplitInfo;
+		EMeshResult Result = Mesh.SplitVertex(Vertex.VertexID, Wedge.Triangles, SplitInfo);
+		if (Result == EMeshResult::Ok)
+		{
+			Wedge.WedgeVertex = SplitInfo.NewVertex;
+		}
+	}
+
+	// update end start/end pairs for each wedge. If we created new edges above, this is
+	// the first time we will encounter them, so save in edge correspondence map
+	for (int32 k = 0; k < NumWedges; ++k)
+	{
+		FOneRingWedge& Wedge = Vertex.Wedges[k];
+		for (int32 j = 0; j < 2; ++j)
+		{
+			int32 OldWedgeEdgeID = Wedge.BorderEdges[j];
+			int32 OldWedgeEdgeIndex = Wedge.BorderEdgeTriEdgeIndices[j];
+			int32 TriangleID = (j == 0) ? Wedge.Triangles[0] : Wedge.Triangles.Last();
+			int32 CurWedgeEdgeID = Mesh.GetTriEdges(TriangleID)[OldWedgeEdgeIndex];
+			FIndex3i TriVerts = Mesh.GetTriangle(TriangleID);
+			if (OldWedgeEdgeID != CurWedgeEdgeID)
+			{
+				if (MeshEdgePairs.Contains(OldWedgeEdgeID) == false)
+				{
+					MeshEdgePairs.Add(OldWedgeEdgeID, CurWedgeEdgeID);
+					MeshEdgePairs.Add(CurWedgeEdgeID, OldWedgeEdgeID);
+				}
+				Wedge.BorderEdges[j] = CurWedgeEdgeID;
+			}
+		}
+	}
+
+}
+
+
+
+
+void FMeshBevel::UnlinkTerminatorVertex(FDynamicMesh3& Mesh, FBevelVertex& BevelVertex, FDynamicMeshChangeTracker* ChangeTracker)
+{
+	check(BevelVertex.VertexType == EBevelVertexType::TerminatorVertex);
+	ensure(BevelVertex.Wedges.Num() == 2);
+
+	if (ChangeTracker)
+	{
+		ChangeTracker->SaveVertexOneRingTriangles(BevelVertex.VertexID, true);
+	}
+
+	// TODO: do we need to update Wedge BorderEdges here??
+
+	// split the vertex
+	FDynamicMesh3::FVertexSplitInfo SplitInfo;
+	EMeshResult Result = Mesh.SplitVertex(BevelVertex.VertexID, BevelVertex.Wedges[1].Triangles, SplitInfo);
+	if (Result == EMeshResult::Ok)
+	{
+		BevelVertex.Wedges[1].WedgeVertex = SplitInfo.NewVertex;
+	}
+
+}
+
+
+
+void FMeshBevel::DisplaceVertices(FDynamicMesh3& Mesh, double Distance)
+{
+	//
+	// This displacement method produces not-very-nice-looking bevels, improvements TBD
+	//
+
+	auto GetDisplacedVertexPos = [Distance](const FDynamicMesh3& Mesh, int32 VertexID) -> FVector3d
+	{
+		FVector3d CurPos = Mesh.GetVertex(VertexID);
+		FVector3d Centroid = FMeshWeights::MeanValueCentroid(Mesh, VertexID);
+		FVector3d MoveDir = Normalized(Centroid - CurPos);
+		return CurPos + Distance * MoveDir;
+	};
+
+
+	auto DisplacePairedVertexLists = [&GetDisplacedVertexPos](
+		const FDynamicMesh3& Mesh, 
+		TArray<int32>& Vertices0, TArray<int32>& Vertices1, 
+		TArray<FVector3d>& Positions0, TArray<FVector3d>& Positions1, int32 InsetStart, int32 InsetEnd)
+	{
+		int32 NumVertices = Vertices0.Num();
+		if (NumVertices == Vertices1.Num())
+		{
+			Positions0.SetNum(NumVertices);
+			Positions1.SetNum(NumVertices);
+			int32 Stop = NumVertices - InsetEnd;
+			for (int32 k = InsetStart; k < Stop; ++k)
+			{
+				if (Vertices0[k] == Vertices1[k])
+				{
+					Positions0[k] = Positions1[k] = Mesh.GetVertex(Vertices0[k]);
+				}
+				else
+				{
+					Positions0[k] = GetDisplacedVertexPos(Mesh, Vertices0[k]);
+					Positions1[k] = GetDisplacedVertexPos(Mesh, Vertices1[k]);
+				}
+			}
+		}
+	};
+
+
+	for (FBevelEdge& Edge : Edges)
+	{
+		DisplacePairedVertexLists(Mesh, Edge.MeshVertices, Edge.NewMeshVertices, Edge.NewPositions0, Edge.NewPositions1, 
+			Edge.bEndpointBoundaryFlag[0]?0:1, Edge.bEndpointBoundaryFlag[1]?0:1 );
+	}
+	for (FBevelLoop& Loop : Loops)
+	{
+		DisplacePairedVertexLists(Mesh, Loop.MeshVertices, Loop.NewMeshVertices, Loop.NewPositions0, Loop.NewPositions1, 0, 0);
+	}
+
+
+	// corners
+	for (FBevelVertex& Vertex : Vertices)
+	{
+		if ( (Vertex.VertexType == EBevelVertexType::JunctionVertex)
+			|| (Vertex.VertexType == EBevelVertexType::TerminatorVertex) )
+		{
+			int32 NumWedges = Vertex.Wedges.Num();
+			for (int32 k = 0; k < NumWedges; ++k)
+			{
+				FOneRingWedge& Wedge = Vertex.Wedges[k];
+				Wedge.NewPosition = GetDisplacedVertexPos(Mesh, Wedge.WedgeVertex);
+			}
+		}
+	}
+
+
+	auto SetDisplacedPositions = [&GetDisplacedVertexPos](FDynamicMesh3& Mesh, TArray<int32>& VerticesIn, TArray<FVector3d>& PositionsIn, int32 InsetStart, int32 InsetEnd)
+	{
+		int32 NumVertices = VerticesIn.Num();
+		if (PositionsIn.Num() == NumVertices)
+		{
+			int32 Stop = NumVertices - InsetEnd;
+			for (int32 k = InsetStart; k < Stop; ++k)
+			{
+				Mesh.SetVertex(VerticesIn[k], PositionsIn[k]);
+			}
+		}
+	};
+
+
+	// now bake in new positions
+	for (FBevelEdge& Edge : Edges)
+	{
+		SetDisplacedPositions(Mesh, Edge.MeshVertices, Edge.NewPositions0, Edge.bEndpointBoundaryFlag[0]?0:1, Edge.bEndpointBoundaryFlag[1]?0:1);
+		SetDisplacedPositions(Mesh, Edge.NewMeshVertices, Edge.NewPositions1, Edge.bEndpointBoundaryFlag[0]?0:1, Edge.bEndpointBoundaryFlag[1]?0:1);
+	}
+	for (FBevelLoop& Loop : Loops)
+	{
+		SetDisplacedPositions(Mesh, Loop.MeshVertices, Loop.NewPositions0, 0, 0);
+		SetDisplacedPositions(Mesh, Loop.NewMeshVertices, Loop.NewPositions1, 0, 0);
+	}
+	for (FBevelVertex& Vertex : Vertices)
+	{
+		if ( (Vertex.VertexType == EBevelVertexType::JunctionVertex)
+			|| (Vertex.VertexType == EBevelVertexType::TerminatorVertex) )
+		{
+			for (FOneRingWedge& Wedge : Vertex.Wedges)
+			{
+				Mesh.SetVertex(Wedge.WedgeVertex, Wedge.NewPosition);
+			}
+		}
+	}
+}
+
+
+
+
+void FMeshBevel::AppendJunctionVertexPolygon(FDynamicMesh3& Mesh, FBevelVertex& Vertex)
+{
+	check(Vertex.VertexType == EBevelVertexType::JunctionVertex);
+
+	// UnlinkJunctionVertex() split the terminator vertex into N vertices, one for each
+	// (now disconnected) triangle-wedge. The wedges are ordered such that their wedge-vertices
+	// define a polygon with correct winding, so we can just mesh it and append the triangles
+
+	TArray<FVector3d> PolygonPoints;
+	for (FOneRingWedge& Wedge : Vertex.Wedges)
+	{
+		PolygonPoints.Add(Mesh.GetVertex(Wedge.WedgeVertex));
+	}
+
+	TArray<FIndex3i> Triangles;
+	PolygonTriangulation::TriangulateSimplePolygon<double>(PolygonPoints, Triangles);
+	Vertex.NewGroupID = Mesh.AllocateTriangleGroup();
+	for (FIndex3i Tri : Triangles)
+	{
+		int32 A = Vertex.Wedges[Tri.A].WedgeVertex;
+		int32 B = Vertex.Wedges[Tri.B].WedgeVertex;
+		int32 C = Vertex.Wedges[Tri.C].WedgeVertex;
+		int32 tid = Mesh.AppendTriangle(A, C, B, Vertex.NewGroupID);		// cycle orientation here because TriangulateSimplePolygon is wrong
+		if (Mesh.IsTriangle(tid))
+		{
+			Vertex.NewTriangles.Add(tid);
+		}
+	}
+}
+
+
+void FMeshBevel::AppendTerminatorVertexTriangle(FDynamicMesh3& Mesh, FBevelVertex& Vertex)
+{
+	check(Vertex.VertexType == EBevelVertexType::TerminatorVertex);
+
+	// UnlinkTerminatorVertex() opened up a triangle-shaped hole adjacent to the incoming edge quad-strip
+	// at the terminator vertex. The Wedges of the terminator vertex contain the vertex IDs of the two
+	// verts on the quad-strip edge. We need the third vertex. We stored [SplitEdge, FarVertexID] in
+	// .TerminatorInfo, however FarVertexID may have become a different vertex when we unlinked other
+	// vertices. So, we will try to use SplitEdge to find it.
+	// If this turns out to have problems, basically the QuadEdgeID is on the boundary of a 3-edge hole,
+	// and so it should be straightforward to find the two other boundary edges and that gives the vertex.
+	int32 RingSplitEdgeID = Vertex.TerminatorInfo.A;
+	if (Mesh.IsEdge(RingSplitEdgeID))
+	{
+		FIndex2i SplitEdgeV = Mesh.GetEdgeV(RingSplitEdgeID);
+		int32 FarVertexID = SplitEdgeV.OtherElement(Vertex.VertexID);
+
+		int32 QuadEdgeID = Mesh.FindEdge(Vertex.Wedges[0].WedgeVertex, Vertex.Wedges[1].WedgeVertex);
+		if (Mesh.IsEdge(QuadEdgeID))
+		{
+			FIndex2i QuadEdgeV = Mesh.GetOrientedBoundaryEdgeV(QuadEdgeID);
+			// should have computed this GroupID in initial setup
+			int32 UseGroupID = (Vertex.NewGroupID >= 0) ? Vertex.NewGroupID : Mesh.AllocateTriangleGroup();
+			int32 tid = Mesh.AppendTriangle(QuadEdgeV.B, QuadEdgeV.A, FarVertexID, UseGroupID);
+			if (Mesh.IsTriangle(tid))
+			{
+				Vertex.NewTriangles.Add(tid);
+			}
+		}
+	}
+}
+
+
+void FMeshBevel::AppendEdgeQuads(FDynamicMesh3& Mesh, FBevelEdge& Edge)
+{
+	int32 NumEdges = Edge.MeshEdges.Num();
+	if (NumEdges != Edge.NewMeshEdges.Num())
+	{
+		return;
+	}
+
+	Edge.NewGroupID = Mesh.AllocateTriangleGroup();
+
+	// At this point each edge-span should be fully disconnected into a set of paired edges, 
+	// so we can trivially join each edge pair with a quad.
+	for (int32 k = 0; k < NumEdges; ++k)
+	{
+		int32 EdgeID0 = Edge.MeshEdges[k];
+		int32 EdgeID1 = Edge.NewMeshEdges[k];
+
+		// in certain cases, like bevel topo-edges with a single mesh-edge, we would not
+		// have been able to construct the "other" mesh edge when processing the topo-edge
+		// (where .NewMeshEdges is computed), it would only have been created when processing the
+		// junction vertex. Currently we do not go back and update .NewMeshEdges in that case, but
+		// we do store the edge-pair-correspondence in the MeshEdgePairs map. 
+		if (EdgeID0 == EdgeID1)
+		{
+			int32* FoundEdgeID1 = MeshEdgePairs.Find(EdgeID0);
+			if (FoundEdgeID1 != nullptr)
+			{
+				EdgeID1 = *FoundEdgeID1;
+			}
+		}
+
+		FIndex2i QuadTris(IndexConstants::InvalidID, IndexConstants::InvalidID);
+		if (EdgeID0 != EdgeID1 && Mesh.IsEdge(EdgeID1) )
+		{
+			FIndex2i EdgeV0 = Mesh.GetOrientedBoundaryEdgeV(EdgeID0);
+			FIndex2i EdgeV1 = Mesh.GetOrientedBoundaryEdgeV(EdgeID1);
+			QuadTris.A = Mesh.AppendTriangle(EdgeV0.B, EdgeV0.A, EdgeV1.B, Edge.NewGroupID);
+			QuadTris.B = Mesh.AppendTriangle(EdgeV1.B, EdgeV1.A, EdgeV0.B, Edge.NewGroupID);
+		}
+		Edge.StripQuads.Add(QuadTris);
+	}
+}
+
+
+
+void FMeshBevel::AppendLoopQuads(FDynamicMesh3& Mesh, FBevelLoop& Loop)
+{
+	int32 NumEdges = Loop.MeshEdges.Num();
+	if (NumEdges != Loop.NewMeshEdges.Num())
+	{
+		return;
+	}
+
+	Loop.NewGroupID = Mesh.AllocateTriangleGroup();
+
+	// At this point each edge-span should be fully disconnected into a set of paired edges, 
+	// so we can trivially join each edge pair with a quad.
+	for (int32 k = 0; k < NumEdges; ++k)
+	{
+		int32 EdgeID0 = Loop.MeshEdges[k];
+		int32 EdgeID1 = Loop.NewMeshEdges[k];
+
+		// case that happens in AppendEdgeQuads() should never happen for loops...
+
+		FIndex2i QuadTris(IndexConstants::InvalidID, IndexConstants::InvalidID);
+		if (EdgeID0 != EdgeID1 && Mesh.IsEdge(EdgeID1))
+		{
+			FIndex2i EdgeV0 = Mesh.GetOrientedBoundaryEdgeV(EdgeID0);
+			FIndex2i EdgeV1 = Mesh.GetOrientedBoundaryEdgeV(EdgeID1);
+			QuadTris.A = Mesh.AppendTriangle(EdgeV0.B, EdgeV0.A, EdgeV1.B, Loop.NewGroupID);
+			QuadTris.B = Mesh.AppendTriangle(EdgeV1.B, EdgeV1.A, EdgeV0.B, Loop.NewGroupID);
+		}
+		Loop.StripQuads.Add(QuadTris);
+	}
+}
+
+
+
+
+void FMeshBevel::CreateBevelMeshing(FDynamicMesh3& Mesh)
+{
+	for (FBevelVertex& Vertex : Vertices)
+	{
+		if (Vertex.VertexType == EBevelVertexType::JunctionVertex)
+		{
+			if (Vertex.Wedges.Num() > 2)
+			{
+				AppendJunctionVertexPolygon(Mesh, Vertex);
+			}
+		}
+	}
+
+	for (FBevelEdge& Edge : Edges)
+	{
+		AppendEdgeQuads(Mesh, Edge);
+	}
+	for (FBevelLoop& Loop : Loops)
+	{
+		AppendLoopQuads(Mesh, Loop);
+	}
+
+	// easier to do these last so that we can use quad edge to orient the triangle
+	for (FBevelVertex& Vertex : Vertices)
+	{
+		if (Vertex.VertexType == EBevelVertexType::TerminatorVertex)
+		{
+			AppendTerminatorVertexTriangle(Mesh, Vertex);
+		}
+	}
+}
+
+
+
+
+void FMeshBevel::ComputeNormals(FDynamicMesh3& Mesh)
+{
+	if (Mesh.HasAttributes() == false)
+	{
+		return;
+	}
+	FDynamicMeshNormalOverlay* NormalOverlay = Mesh.Attributes()->PrimaryNormals();
+
+	auto SetNormalsOnTriRegion = [NormalOverlay](const TArray<int32>& Triangles)
+	{
+		if (Triangles.Num() > 0)
+		{
+			FMeshNormals::InitializeOverlayRegionToPerVertexNormals(NormalOverlay, Triangles);
+		}
+	};
+
+	auto QuadsToTris = [](const FDynamicMesh3& Mesh, const TArray<FIndex2i>& Quads, TArray<int32>& TrisOut)
+	{
+		int32 N = Quads.Num();
+		TrisOut.Reset();
+		TrisOut.Reserve(2 * N);
+		for (const FIndex2i& Quad : Quads)
+		{
+			if (Mesh.IsTriangle(Quad.A))
+			{
+				TrisOut.Add(Quad.A);
+			}
+			if (Mesh.IsTriangle(Quad.B))
+			{
+				TrisOut.Add(Quad.B);
+			}
+		}
+	};
+
+
+	for (FBevelVertex& Vertex : Vertices)
+	{
+		SetNormalsOnTriRegion(Vertex.NewTriangles);
+	}
+
+	TArray<int32> TriList;
+	for (FBevelEdge& Edge : Edges)
+	{
+		QuadsToTris(Mesh, Edge.StripQuads, TriList);
+		SetNormalsOnTriRegion(TriList);
+	}
+	for (FBevelLoop& Loop : Loops)
+	{
+		QuadsToTris(Mesh, Loop.StripQuads, TriList);
+		SetNormalsOnTriRegion(TriList);
+	}
+}
diff --git a/Engine/Plugins/Runtime/GeometryProcessing/Source/DynamicMesh/Public/Operations/MeshBevel.h b/Engine/Plugins/Runtime/GeometryProcessing/Source/DynamicMesh/Public/Operations/MeshBevel.h
new file mode 100644
--- /dev/null
+++ b/Engine/Plugins/Runtime/GeometryProcessing/Source/DynamicMesh/Public/Operations/MeshBevel.h
@@ -0,0 +1,247 @@
+// Copyright Epic Games, Inc. All Rights Reserved.
+#pragma once
+
+#include "DynamicMesh/DynamicMesh3.h"
+#include "Util/ProgressCancel.h"
+
+namespace UE
+{
+namespace Geometry
+{
+
+class FGroupTopology;
+class FDynamicMeshChangeTracker;
+class FEdgeLoop;
+
+
+/**
+ * FMeshBevel applies a "Bevel" operation to edges of a FDynamicMesh3. Bevel is not strictly well-defined,
+ * there are a wide range of possible cases to handle and currently only some are supported.
+ * See this website for a discussion of many interesting cases (most are not currently supported): https://wiki.blender.org/wiki/User:Howardt/Bevel
+ * 
+ * The bevel operation is applied in-place to the input mesh. The bevel is mesh-topological, ie implemented by
+ * un-stitching mesh edges and inserting new triangles, rather than via booleans/etc. 
+ * 
+ * Currently supports:
+ *   - Bevel an isolated closed loop of edges, that edge-loop becomes a quad-strip (simplest case)
+ *   - Bevel a set of open edge spans that may meet a T-junctions
+ *       - if incoming-span-valence at a vertex is >= 3, vertex is replaced by a polygon
+ *       - if incoming-span-valence at a vertex is 1, bevel is "terminated" by expanding the vertex into an edge (this is the messiest case)
+ *          - vertex-on-boundary is a special case that is simpler
+ *       - edge spans are replaced by quad-strips
+ * 
+ * Currently does not support:
+ *   - Beveling an isolated vertex
+ *   - partial bevels of a GroupTopology Edge
+ *   - multiple-segment bevel (eg to do rounds/etc)
+ *   - ???
+ * 
+ * Generally this FMeshBevel is intended to be applied to group-boundary "edges" in a FGroupTopology. However the FGroupTopology
+ * is currently only used in the initial setup, so other methods could be used to construct the bevel topology. 
+ * 
+ * Currently "updating" the bevel shape, once the topology is known, is not supported, but could be implemented
+ * relatively easily, as all the relevant information is already tracked and stored.
+ */
+class DYNAMICMESH_API FMeshBevel
+{
+public:
+
+	/** Distance that bevel edges/vertices are inset from their initial position. Not guaranteed to hold for all vertices, though. */
+	double InsetDistance = 5.0;
+
+	/** Set this member to support progress/cancel in the computations below */
+	FProgressCancel* Progress = nullptr;
+
+	/**
+	 * Initialize the bevel with all edges of the given GroupTopology
+	 */
+	void InitializeFromGroupTopology(const FDynamicMesh3& Mesh, const FGroupTopology& Topology);
+
+	/**
+	 * Initialize the bevel with the specified edges of a GroupTopology
+	 */
+	void InitializeFromGroupTopology(const FDynamicMesh3& Mesh, const FGroupTopology& Topology, const TArray<int32>& GroupEdges);
+
+	/**
+	 * Apply the bevel operation to the mesh, and optionally track changes
+	 */
+	bool Apply(FDynamicMesh3& Mesh, FDynamicMeshChangeTracker* ChangeTracker);
+
+	/** status of the operation, warnings/errors may be returned here */
+	FGeometryResult ResultInfo = FGeometryResult(EGeometryResultType::InProgress);
+
+
+public:
+
+	//
+	// Current Bevel computation strategy is basically to fully precompute all the necessary info for the entire bevel,
+	// then "unlink" all the edges and vertices, and then stitch it all back together. 
+	// Various incremental strategies were attempted, however certain cases like a two "bevel vertices" connected by a 
+	// single mesh edge greatly complicate any attempt to decompose the problem into sub-parts. 
+	// 
+	// The data structures below are used to track this topological information during the operation.
+	// Note that some parts of these data structures may become invalid/incorrect as the operation proceeds...
+	//
+
+	// POSSIBLE IMPROVEMENTS:
+	// * compute wedges for Loop/Edge in setup? would avoid having to deal w/ possibly weird 
+	//   configurations introduced by unlink of corners...
+
+	// FBevelLoop is the accumulated data for a closed loop of mesh-edges, with no T-junctions w/ other bevel-edges.
+	// This is the easiest case as each mesh-edge of the loop expands out into a quad, with no complex vertex-polygons/etc
+	struct FBevelLoop
+	{
+		// initial topological information that defines what happens in unlink/displace/mesh steps
+		TArray<int32> MeshVertices;		// sequential list of mesh vertex IDs along edge loop
+		TArray<int32> MeshEdges;		// sequential list of mesh edge IDs along edge loop
+
+		// new mesh topology computed during unlink step
+		TArray<int32> NewMeshVertices;		// list of new vertices on "other" side of unlinked edge, 1-1 with MeshVertices
+		TArray<int32> NewMeshEdges;			// list of new edges on "other" side of unlinked edge, 1-1 with MeshEdges
+
+		// buffers for new vertex positions computed during displace step
+		TArray<FVector3d> NewPositions0;	// new positions for MeshVertices list
+		TArray<FVector3d> NewPositions1;	// new positions for NewMeshVertices list
+
+		// new geometry computed during mesh step
+		int32 NewGroupID;
+		TArray<FIndex2i> StripQuads;		// triangle-ID-pairs for each new quad added along edge, 1-1 with MeshEdges
+	};
+
+	// FBevelEdge is the accumulated data for an open span of mesh-edges, which possibly meets up with other bevel-edges
+	// at the vertices on either end of the span. Each mesh-edge of the bevel-edge will become a quad
+	struct FBevelEdge
+	{
+		// initial topological information that defines what happens in unlink/displace/mesh steps
+		TArray<int32> MeshVertices;		// sequential list of mesh vertex IDs along edge
+		TArray<int32> MeshEdges;		// sequential list of mesh edge IDs along edge
+		int32 GroupEdgeID;				// ID of this edge in external topology (eg FGroupTopology)
+		FIndex2i GroupIDs;				// topological IDs of groups on either side of topological edge
+		bool bEndpointBoundaryFlag[2];	// flag defining whether vertex at start/end of MeshVertices was a boundary vertex
+
+		// new mesh topology computed during unlink step
+		TArray<int32> NewMeshVertices;		// list of new vertices on "other" side of unlinked edge, 1-1 with MeshVertices
+		TArray<int32> NewMeshEdges;			// list of new edges on "other" side of unlinked edge, 1-1 with MeshEdges
+
+		// buffers for new vertex positions computed during displace step
+		TArray<FVector3d> NewPositions0;	// new positions for MeshVertices list
+		TArray<FVector3d> NewPositions1;	// new positions for NewMeshVertices list
+
+		// new geometry computed during mesh step
+		int32 NewGroupID;
+		TArray<FIndex2i> StripQuads;		// triangle-ID-pairs for each new quad added along edge, 1-1 with MeshEdges
+	};
+
+
+	// A FOneRingWedge represents a sequential set of connected triangles around a vertex, ie a subset of an ordered triangle one-ring.
+	// Used to represent the desired bevel topology in FBevelVertex
+	struct FOneRingWedge
+	{
+		TArray<int32> Triangles;				// list of sequential triangles in this wedge
+		FIndex2i BorderEdges;					// "first" and "last" Edges of sequential triangles in Triangles list (connected to central Vertex)
+		FIndex2i BorderEdgeTriEdgeIndices;		// index 0/1/2 of BorderEdges[j] in start/ed Triangles
+
+		int32 WedgeVertex;						// central vertex of this wedge (updated by unlink functions)
+
+		FVector3d NewPosition;					// new calculated position for vertex of this wedge
+	};
+
+	// a FBevelVertex can have various types, depending on the topology of the bevel edge graph and input mesh
+	enum class EBevelVertexType
+	{
+		// A JunctionVertex is a vertex at which 2 or more FBevelEdges meet (ie is an endpoint of 2 or more of those vertex-spans)
+		// If N>=3 or more edges meet at a JunctionVertex, it will become a polygon with N vertices, one for each "wedge"
+		JunctionVertex,
+		// A TerminatorVertex is a vertex at which a single FBevelEdge terminates, ie the N=1 case. This requires different handling
+		// because we essentially want to turn that vertex into an edge, which means inserting a triangle into the adjacent one-ring
+		TerminatorVertex,
+		// a BoundaryVertex is a junction/terminator on the mesh boundary
+		BoundaryVertex,
+		// An Unknown vertex is one at which we don't know what to do, or some error occurred while processing as a Junction/Terminator.
+		Unknown
+	};
+
+	// A FBevelVertex repesents/stores the accumulated data at a "bevel vertex", which is the mesh vertex at the end of a FBevelEdge.
+	// A FBevelVertex may be expanded out into a polygon or just an edge, depending on its Type
+	struct FBevelVertex
+	{
+		int32 VertexID;												// Initial Mesh Vertex ID for the Bevel Vertex
+		int32 CornerID;												// Initial Group Topology Corner ID for the Bevel Vertex (if exists)
+		EBevelVertexType VertexType = EBevelVertexType::Unknown;	// Type of the Bevel Vertex
+
+		TArray<int32> IncomingBevelMeshEdges;						// Set of (unsorted) Mesh Edges that are destined to be Beveled, coming into the vertex
+		TArray<int32> IncomingBevelTopoEdges;						// Set of (unsorted) Group Topology Edges that are to be Beveled, coming into the vertex
+
+		TArray<int32> SortedTriangles;			// ordered triangle one-ring around VertexID
+		TArray<FOneRingWedge> Wedges;			// ordered decomposition of one-ring into "wedges" between incoming bevel edges (no correspondence w/ IncomingBevelMeshEdges list)
+
+		int32 NewGroupID;						// new polygroup allocated for the beveled polygon generated by this vertex (if NumEdges > 2)
+		TArray<int32> NewTriangles;				// new triangles that make up the beveled polygon for this vertex (if NumEdges > 2)
+
+		FIndex2i TerminatorInfo;				// for TerminatorVertex type, store [EdgeID, FarVertexID] in one-ring, used to unlink/fill (see usage)
+	};
+
+
+
+public:
+	TMap<int32, int32> VertexIDToIndexMap;		// map of mesh-vertex-IDs to indices into Vertices list
+	TArray<FBevelVertex> Vertices;				// list of FBevelVertex data structures for mesh vertices that need beveling
+
+	TArray<FBevelEdge> Edges;					// list of FBevelEdge data structures for mesh edge-spans that need beveling
+	TArray<FBevelLoop> Loops;					// list of FBevelLoop data structures for mesh edge-loops that need beveling
+
+	TMap<int32, int32> MeshEdgePairs;			// Many edges of the input mesh will be split into edge pairs, which are then stitched together with quads.
+												// This map stores the authoritative correspondences between these edge pairs. Both pairs, ie (a,b) and (b,a) are stored.
+
+
+
+protected:
+
+	FBevelVertex* GetBevelVertexFromVertexID(int32 VertexID);
+
+	// Setup phase: register Edges (spans) and (isolated) Loops that need to be beveled and precompute/store any mesh topology that must be tracked across the operation
+	// Required BevelVertex's are added by AddBevelGroupEdge()
+	// Once edges are configured, BuildVertexSets() is called to precompute the vertex topological information
+
+	void AddBevelGroupEdge(const FDynamicMesh3& Mesh, const FGroupTopology& Topology, int32 GroupEdgeID);
+	void AddBevelEdgeLoop(const FDynamicMesh3& Mesh, const FEdgeLoop& Loop);
+	void BuildVertexSets(const FDynamicMesh3& Mesh);
+	void BuildJunctionVertex(FBevelVertex& Vertex, const FDynamicMesh3& Mesh);
+	void BuildTerminatorVertex(FBevelVertex& Vertex, const FDynamicMesh3& Mesh);
+
+	// Unlink phase - disconnect triangles along bevel edges/loops, and at vertices. 
+	// Vertices may expand out into multiple "wedges" depending on incoming bevel-edge topology.
+
+	void UnlinkEdges( FDynamicMesh3& Mesh, FDynamicMeshChangeTracker* ChangeTracker);
+	void UnlinkBevelEdgeInterior( FDynamicMesh3& Mesh, FBevelEdge& BevelEdge, FDynamicMeshChangeTracker* ChangeTracker);
+
+	void UnlinkLoops( FDynamicMesh3& Mesh, FDynamicMeshChangeTracker* ChangeTracker);
+	void UnlinkBevelLoop( FDynamicMesh3& Mesh, FBevelLoop& BevelLoop, FDynamicMeshChangeTracker* ChangeTracker);
+
+	void UnlinkVertices( FDynamicMesh3& Mesh, FDynamicMeshChangeTracker* ChangeTracker);
+	void UnlinkJunctionVertex(FDynamicMesh3& Mesh, FBevelVertex& BevelVertex, FDynamicMeshChangeTracker* ChangeTracker);
+	void UnlinkTerminatorVertex(FDynamicMesh3& Mesh, FBevelVertex& BevelVertex, FDynamicMeshChangeTracker* ChangeTracker);
+
+	// Displace phase - move unlinked vertices to new positions
+
+	void DisplaceVertices(FDynamicMesh3& Mesh, double Distance);
+
+	// Meshing phase - append quad-strips between unlinked edge spans/loops, polygons at junction vertices where required,
+	// and triangles at terminator vertices
+
+	void CreateBevelMeshing(FDynamicMesh3& Mesh);
+	void AppendJunctionVertexPolygon(FDynamicMesh3& Mesh, FBevelVertex& Vertex);
+	void AppendTerminatorVertexTriangle(FDynamicMesh3& Mesh, FBevelVertex& Vertex);
+	void AppendEdgeQuads(FDynamicMesh3& Mesh, FBevelEdge& Edge);
+	void AppendLoopQuads(FDynamicMesh3& Mesh, FBevelLoop& Loop);
+
+	// Normals phase - calculate normals for new geometry
+
+	void ComputeNormals(FDynamicMesh3& Mesh);
+
+};
+
+
+
+} // end namespace UE::Geometry
+} // end namespace UE
\ No newline at end of file
diff --git a/Engine/Source/Runtime/GeometryCore/Private/DynamicMesh/MeshIndexUtil.cpp b/Engine/Source/Runtime/GeometryCore/Private/DynamicMesh/MeshIndexUtil.cpp
--- a/Engine/Source/Runtime/GeometryCore/Private/DynamicMesh/MeshIndexUtil.cpp
+++ b/Engine/Source/Runtime/GeometryCore/Private/DynamicMesh/MeshIndexUtil.cpp
@@ -66,4 +66,152 @@ void UE::Geometry::VertexToTriangleOneRing(const FDynamicMesh3* Mesh, const TArr
 			if (EdgeT.B != IndexConstants::InvalidID) TriangleIDsOut.Add(EdgeT.B);
 		});
 	}
+}
+
+
+FIndex2i UE::Geometry::FindVertexEdgesInTriangle(const FDynamicMesh3& Mesh, int32 TriangleID, int32 VertexID)
+{
+	FIndex2i Result(IndexConstants::InvalidID, IndexConstants::InvalidID);
+	if (Mesh.IsTriangle(TriangleID))
+	{
+		FIndex3i TriV = Mesh.GetTriangle(TriangleID);
+		FIndex3i TriEdges = Mesh.GetTriEdges(TriangleID);
+		for (int32 j = 0; j < 3; ++j)
+		{
+			if (TriV[j] == VertexID)
+			{
+				return FIndex2i( TriEdges[(j==0)?2:j-1], TriEdges[j] );
+			}
+		}
+	}
+	return FIndex2i(IndexConstants::InvalidID, IndexConstants::InvalidID);;
+};
+
+
+int32 UE::Geometry::FindSharedEdgeInTriangles(const FDynamicMesh3& Mesh, int32 Triangle0, int32 Triangle1)
+{
+	if (Mesh.IsTriangle(Triangle0) && Mesh.IsTriangle(Triangle1))
+	{
+		FIndex3i Edges0 = Mesh.GetTriEdges(Triangle0);
+		FIndex3i Edges1 = Mesh.GetTriEdges(Triangle1);
+		for (int32 j = 0; j < 3; ++j)
+		{
+			if (Edges1.Contains(Edges0[j]))
+			{
+				return Edges0[j];
+			}
+		}
+	}
+	return IndexConstants::InvalidID;
+}
+
+
+
+bool UE::Geometry::SplitBoundaryVertexTrianglesIntoSubsets(
+	const FDynamicMesh3* Mesh,
+	int32 VertexID,
+	int32 SplitEdgeID,
+	TArray<int32>& TriangleSet0, TArray<int32>& TriangleSet1)
+{
+	if (Mesh->IsVertex(VertexID) == false || Mesh->IsEdge(SplitEdgeID) == false)
+	{
+		return false;
+	}
+	FIndex2i StartTris = Mesh->GetEdgeT(SplitEdgeID);
+	if (StartTris.B == IndexConstants::InvalidID)
+	{
+		return false;
+	}
+
+	for (int32 si = 0; si < 2; ++si)
+	{
+		TArray<int32>& CurrentTriangleSet = (si == 0) ? TriangleSet0 : TriangleSet1;
+		int32 StartTri = StartTris[si];
+		CurrentTriangleSet.Add(StartTri);
+
+		int32 EdgeOtherTri = StartTris[si == 0 ? 1 : 0];
+		int32 CurTri = StartTri;
+		int32 PrevTri = EdgeOtherTri;
+
+		bool bDone = false;
+		while (!bDone)
+		{
+			FIndex3i NextTri = UE::Geometry::FindNextAdjacentTriangleAroundVtx(Mesh, VertexID, CurTri, PrevTri,
+				[&](int32 Tri0, int32 Tri1, int32 Edge) { return Edge != SplitEdgeID; }
+			);
+			if (NextTri.A != IndexConstants::InvalidID)
+			{
+				if (NextTri.A == EdgeOtherTri)		// if we somehow looped around, then the arguments were bad and we need to abort
+				{
+					return false;
+				}
+
+				CurrentTriangleSet.Add(NextTri.A);
+				PrevTri = CurTri;
+				CurTri = NextTri.A;
+			}
+			else
+			{
+				bDone = true;
+			}
+		}
+	}
+
+	return (TriangleSet0.Num() > 0 && TriangleSet1.Num() > 0);
+}
+
+
+
+bool UE::Geometry::SplitInteriorVertexTrianglesIntoSubsets(
+	const FDynamicMesh3* Mesh,
+	int32 VertexID,
+	int32 SplitEdgeID0, int32 SplitEdgeID1,
+	TArray<int32>& TriangleSet0, TArray<int32>& TriangleSet1)
+{
+	if (Mesh->IsVertex(VertexID) == false || Mesh->IsEdge(SplitEdgeID0) == false || Mesh->IsEdge(SplitEdgeID1) == false)
+	{
+		return false;
+	}
+
+	FIndex2i StartTris = Mesh->GetEdgeT(SplitEdgeID1);
+	if (StartTris.B <= 0)
+	{
+		return false;
+	}
+
+	for (int32 si = 0; si < 2; ++si)
+	{
+		TArray<int32>& CurrentTriangleSet = (si == 0) ? TriangleSet0 : TriangleSet1;
+		int32 StartTri = StartTris[si];
+		CurrentTriangleSet.Add(StartTri);
+
+		int32 EdgeOtherTri = StartTris[si == 0 ? 1 : 0];
+		int32 CurTri = StartTri;
+		int32 PrevTri = EdgeOtherTri;
+
+		bool bDone = false;
+		while (!bDone)
+		{
+			FIndex3i NextTri = UE::Geometry::FindNextAdjacentTriangleAroundVtx(Mesh, VertexID, CurTri, PrevTri,
+				[&](int32 Tri0, int32 Tri1, int32 Edge) { return Edge != SplitEdgeID0 && Edge != SplitEdgeID1; }
+			);
+			if (NextTri.A != IndexConstants::InvalidID)
+			{
+				if (NextTri.A == EdgeOtherTri)		// if we somehow looped around, then the arguments were bad and we need to abort
+				{
+					return false;
+				}
+
+				CurrentTriangleSet.Add(NextTri.A);
+				PrevTri = CurTri;
+				CurTri = NextTri.A;
+			}
+			else
+			{
+				bDone = true;
+			}
+		}
+	}
+
+	return (TriangleSet0.Num() > 0 && TriangleSet1.Num() > 0);
 }
\ No newline at end of file
diff --git a/Engine/Source/Runtime/GeometryCore/Public/DynamicMesh/MeshIndexUtil.h b/Engine/Source/Runtime/GeometryCore/Public/DynamicMesh/MeshIndexUtil.h
--- a/Engine/Source/Runtime/GeometryCore/Public/DynamicMesh/MeshIndexUtil.h
+++ b/Engine/Source/Runtime/GeometryCore/Public/DynamicMesh/MeshIndexUtil.h
@@ -6,83 +6,141 @@
 
 namespace UE
 {
-	namespace Geometry
+namespace Geometry
+{
+	/**
+	 * Find list of unique vertices that are contained in one or more triangles
+	 * @param Mesh input mesh
+	 * @param TriangleIDs list of triangle IDs of Mesh
+	 * @param VertexIDsOut list of vertices contained by triangles
+	 */
+	GEOMETRYCORE_API void TriangleToVertexIDs(const FDynamicMesh3* Mesh, const TArray<int>& TriangleIDs, TArray<int>& VertexIDsOut);
+
+
+	/**
+	 * Find all the triangles in all the one rings of a set of vertices
+	 * @param Mesh input mesh
+	 * @param VertexIDs list of Vertex IDs of Mesh
+	 * @param TriangleIDsOut list of triangle IDs where any vertex of triangle touches an element of VertexIDs
+	 */
+	GEOMETRYCORE_API void VertexToTriangleOneRing(const FDynamicMesh3* Mesh, const TArray<int>& VertexIDs, TSet<int>& TriangleIDsOut);
+
+
+	/**
+	 * Find the two edges in a Triangle that are connected to a given Vertex of the triangle
+	 * @return edge pair, or (IndexConstants::InvalidID, IndexConstants::InvalidID) on error
+	 */
+	GEOMETRYCORE_API FIndex2i FindVertexEdgesInTriangle(const FDynamicMesh3& Mesh, int32 TriangleID, int32 VertexID);
+
+
+	/**
+	 * @return the ID of the shared edge betwen two given triangles, or IndexConstants::InvalidID if not found / invalid input
+	 */
+	GEOMETRYCORE_API int32 FindSharedEdgeInTriangles(const FDynamicMesh3& Mesh, int32 Triangle0, int32 Triangle1);
+
+
+	/**
+	 * Call SetType.Add(VertexPosition) for all valid vertex indices in for_each(Enumeration)
+	 */
+	template<typename EnumeratorType, typename SetType>
+	void CollectVertexPositions(const FDynamicMesh3& Mesh, EnumeratorType Enumeration, SetType& Output);
+
+
+	/**
+	 * Walk around VertexID from FromTriangleID to next connected triangle if it exists, walking "away" from PrevTriangleID.
+	 * @param TrisConnectedFunc returns true if two triangles should be considered connected, to support breaking at seams/etc that are not in base mesh topology
+	 * @return triplet of values (FoundTriangleID, SharedEdgeID, IndexOfEdgeInFromTri), or all IndexConstants::InvalidID if not found
+	 */
+	template<typename TrisConnectedPredicate>
+	FIndex3i FindNextAdjacentTriangleAroundVtx(const FDynamicMesh3* Mesh,
+		int32 VertexID, int32 FromTriangleID, int32 PrevTriangleID,
+		TrisConnectedPredicate TrisConnectedTest);
+
+	/**
+	 * Split the triangle one-ring at VertexID into two triangle sets by "cutting" it with edge SplitEdgeID that is connected to VertexID.
+	 * VertexID must be a boundary-vertex of the mesh, otherwise one edge is not enough to "split" it. Bowtie vertices are not supported.
+	 * @return false if the split failed, ie the inputs are invalid or one of the output sets is empty 
+	 */
+	GEOMETRYCORE_API bool SplitBoundaryVertexTrianglesIntoSubsets(
+		const FDynamicMesh3* Mesh,
+		int32 VertexID,
+		int32 SplitEdgeID,
+		TArray<int32>& TriangleSet0, TArray<int32>& TriangleSet1);
+
+	/**
+	* Split the triangle one-ring at VertexID into two triangle sets by "cutting" it with two edges SplitEdgeID0 and SplitEdgeID1 that are both connected to VertexID.
+	* VertexID must be an "interior" vertex of the mesh (ie not on boundary). Bowtie vertices are not supported.
+	* @return false if the split failed, ie the inputs are invalid or one of the output sets is empty
+	*/
+	GEOMETRYCORE_API bool SplitInteriorVertexTrianglesIntoSubsets(
+		const FDynamicMesh3* Mesh,
+		int32 VertexID,
+		int32 SplitEdgeID0, int32 SplitEdgeID1,
+		TArray<int32>& TriangleSet0, TArray<int32>& TriangleSet1);
+
+}
+}
+
+
+
+// implementations of functions declared above
+
+namespace UE
+{
+namespace Geometry
+{
+
+template<typename EnumeratorType, typename SetType>
+void CollectVertexPositions(const FDynamicMesh3& Mesh, EnumeratorType Enumeration, SetType& Output)
+{
+	for (int32 vid : Enumeration)
 	{
-		/**
-		 * Find list of unique vertices that are contained in one or more triangles
-		 * @param Mesh input mesh
-		 * @param TriangleIDs list of triangle IDs of Mesh
-		 * @param VertexIDsOut list of vertices contained by triangles
-		 */
-		GEOMETRYCORE_API void TriangleToVertexIDs(const FDynamicMesh3* Mesh, const TArray<int>& TriangleIDs, TArray<int>& VertexIDsOut);
-
-
-		/**
-		 * Find all the triangles in all the one rings of a set of vertices
-		 * @param Mesh input mesh
-		 * @param VertexIDs list of Vertex IDs of Mesh
-		 * @param TriangleIDsOut list of triangle IDs where any vertex of triangle touches an element of VertexIDs
-		 */
-		GEOMETRYCORE_API void VertexToTriangleOneRing(const FDynamicMesh3* Mesh, const TArray<int>& VertexIDs, TSet<int>& TriangleIDsOut);
-
-
-		/**
-		 * Call SetType.Add(VertexPosition) for all valid vertex indices in for_each(Enumeration)
-		 */
-		template<typename EnumeratorType, typename SetType>
-		void CollectVertexPositions(const FDynamicMesh3& Mesh, EnumeratorType Enumeration, SetType& Output)
+		if (Mesh.IsVertex(vid))
 		{
-			for (int32 vid : Enumeration)
-			{
-				if (Mesh.IsVertex(vid))
-				{
-					Output.Add(Mesh.GetVertex(vid));
-				}
-			}
+			Output.Add(Mesh.GetVertex(vid));
 		}
+	}
+}
 
 
-		/**
-		 * Walk around VertexID from FromTriangleID to next connected triangle if it exists, walking "away" from PrevTriangleID.
-		 * @param TrisConnectedFunc returns true if two triangles should be considered connected, to support breaking at seams/etc that are not in base mesh topology
-		 * @return triplet of values (FoundTriangleID, SharedEdgeID, IndexOfEdgeInFromTri), or all IndexConstants::InvalidID if not found
-		 */
-		template<typename TrisConnectedPredicate>
-		FIndex3i FindNextAdjacentTriangleAroundVtx(const FDynamicMesh3* Mesh, 
-			int32 VertexID, int32 FromTriangleID, int32 PrevTriangleID,
-			TrisConnectedPredicate TrisConnectedTest)
-		{
-			check(Mesh);
+template<typename TrisConnectedPredicate>
+FIndex3i FindNextAdjacentTriangleAroundVtx(const FDynamicMesh3* Mesh, 
+	int32 VertexID, int32 FromTriangleID, int32 PrevTriangleID,
+	TrisConnectedPredicate TrisConnectedTest)
+{
+	check(Mesh);
 
-			// find neigbour edges and tris for triangle
-			FIndex3i TriEdges = Mesh->GetTriEdges(FromTriangleID);
-			FIndex3i TriNbrTris;
-			for (int32 j = 0; j < 3; ++j)
-			{
-				FIndex2i EdgeT = Mesh->GetEdgeT(TriEdges[j]);
-				TriNbrTris[j] = (EdgeT.A == FromTriangleID) ? EdgeT.B : EdgeT.A;
-			}
+	// find neigbour edges and tris for triangle
+	FIndex3i TriEdges = Mesh->GetTriEdges(FromTriangleID);
+	FIndex3i TriNbrTris;
+	for (int32 j = 0; j < 3; ++j)
+	{
+		FIndex2i EdgeT = Mesh->GetEdgeT(TriEdges[j]);
+		TriNbrTris[j] = (EdgeT.A == FromTriangleID) ? EdgeT.B : EdgeT.A;
+	}
 
-			// Search for the neighbour tri that is not PrevTriangleID, and is also connected to VertexID.
-			// This is our next triangle around the ring
-			for (int32 j = 0; j < 3; ++j)
+	// Search for the neighbour tri that is not PrevTriangleID, and is also connected to VertexID.
+	// This is our next triangle around the ring
+	for (int32 j = 0; j < 3; ++j)
+	{
+		if (TriNbrTris[j] != PrevTriangleID && Mesh->IsTriangle(TriNbrTris[j]))
+		{
+			FIndex3i TriVerts = Mesh->GetTriangle(TriNbrTris[j]);
+			if (TriVerts.A == VertexID || TriVerts.B == VertexID || TriVerts.C == VertexID)
 			{
-				if (TriNbrTris[j] != PrevTriangleID && Mesh->IsTriangle(TriNbrTris[j]))
+				// test if predicate allows this connection
+				if (TrisConnectedTest(FromTriangleID, TriNbrTris[j], TriEdges[j]) == false)
 				{
-					FIndex3i TriVerts = Mesh->GetTriangle(TriNbrTris[j]);
-					if (TriVerts.A == VertexID || TriVerts.B == VertexID || TriVerts.C == VertexID)
-					{
-						// test if predicate allows this connection
-						if (TrisConnectedTest(FromTriangleID, TriNbrTris[j], TriEdges[j]) == false)
-						{
-							break;
-						}
-						return FIndex3i(TriNbrTris[j], TriEdges[j], j);
-					}
+					break;
 				}
+				return FIndex3i(TriNbrTris[j], TriEdges[j], j);
 			}
-
-			return FIndex3i(IndexConstants::InvalidID, IndexConstants::InvalidID, IndexConstants::InvalidID);
 		}
 	}
+
+	return FIndex3i(IndexConstants::InvalidID, IndexConstants::InvalidID, IndexConstants::InvalidID);
+}
+
+
+}
 }
diff --git a/Engine/Source/Runtime/GeometryCore/Public/IndexTypes.h b/Engine/Source/Runtime/GeometryCore/Public/IndexTypes.h
--- a/Engine/Source/Runtime/GeometryCore/Public/IndexTypes.h
+++ b/Engine/Source/Runtime/GeometryCore/Public/IndexTypes.h
@@ -75,6 +75,23 @@ struct FIndex2i
 		return (A == Value) || (B == Value);
 	}
 
+	/** @return whichever of A or B is not Value, or IndexConstants::InvalidID if neither is Value */
+	int OtherElement(int Value) const
+	{
+		if (A == Value)
+		{
+			return B;
+		}
+		else if (B == Value)
+		{
+			return A;
+		}
+		else
+		{
+			return IndexConstants::InvalidID;
+		}
+	}
+
 	inline void Swap()
 	{
 		int tmp = A;
@@ -172,6 +189,25 @@ struct FIndex3i
 		return (A == Value) || (B == Value) || (C == Value);
 	}
 
+
+	/**
+	 * @return shifted triplet such that A=WantIndex0Value, and B,C values maintain the same relative ordering
+	 */
+	FIndex3i GetCycled(int32 WantIndex0Value) const
+	{
+		if (B == WantIndex0Value)
+		{
+			return FIndex3i(B, C, A);
+		}
+		else if (C == WantIndex0Value)
+		{
+			return FIndex3i(C, A, B);
+		}
+		return FIndex3i(A,B,C);
+	}
+
+
+
 	operator FIntVector() const
 	{
 		return FIntVector(A, B, C);
