GeometryProcessing: Refactor edge-loop inset calculation code from InsetMeshRegion into new utility functions in PolyEditingEdgeUtil.h. Use in FMeshBevel::DisplaceVertices() to improve edge insetting. Also add FMeshBevel::FixUpUnlinkedBevelEdges(), to repair vertex and edge ID lsits in unlinked open-span edges after unlinking vertices.
#rnx

diff --git a/Engine/Plugins/Runtime/GeometryProcessing/Source/DynamicMesh/Private/Operations/InsetMeshRegion.cpp b/Engine/Plugins/Runtime/GeometryProcessing/Source/DynamicMesh/Private/Operations/InsetMeshRegion.cpp
--- a/Engine/Plugins/Runtime/GeometryProcessing/Source/DynamicMesh/Private/Operations/InsetMeshRegion.cpp
+++ b/Engine/Plugins/Runtime/GeometryProcessing/Source/DynamicMesh/Private/Operations/InsetMeshRegion.cpp
@@ -11,6 +11,7 @@
 #include "Solvers/ConstrainedMeshDeformer.h"
 #include "DynamicMesh/DynamicMeshAABBTree3.h"
 #include "DynamicMesh/MeshTransforms.h"
+#include "Operations/PolyEditingEdgeUtil.h"
 
 #include "ExplicitUseGeometryMathTypes.h"		// using UE::Geometry::(math types)
 using namespace UE::Geometry;
@@ -102,53 +103,20 @@ bool FInsetMeshRegion::ApplyInset(FInsetInfo& Region, FMeshNormals* UseNormals)
 	for (FDynamicMeshEditor::FLoopPairSet& LoopPair : LoopPairs)
 	{
 		int32 NumEdges = LoopPair.InnerEdges.Num();
-		TArray<FLine3d> InsetLines;
-		InsetLines.SetNum(NumEdges);
 
-		for (int32 k = 0; k < NumEdges; ++k)
-		{
-			const FDynamicMesh3::FEdge EdgeVT = Mesh->GetEdge(LoopPair.InnerEdges[k]);
-			FVector3d A = Mesh->GetVertex(EdgeVT.Vert[0]);
-			FVector3d B = Mesh->GetVertex(EdgeVT.Vert[1]);
-			FVector3d EdgeDir = Normalized(A - B);
-			FVector3d Midpoint = (A + B) * 0.5;
-			int32 EdgeTri = EdgeVT.Tri[0];
-			FVector3d Normal, Centroid; double Area;
-			Mesh->GetTriInfo(EdgeTri, Normal, Area, Centroid);
-
-			FVector3d InsetDir = Normal.Cross(EdgeDir);
-			if ((Centroid - Midpoint).Dot(InsetDir) < 0)
-			{
-				InsetDir = -InsetDir;
-			}
+		TArray<FLine3d> InsetLines;
+		UE::Geometry::ComputeInsetLineSegmentsFromEdges(*Mesh, LoopPair.InnerEdges, InsetDistance, InsetLines);
 
-			InsetLines[k] = FLine3d(Midpoint + InsetDistance * InsetDir, EdgeDir);
-		}
+		TArray<FVector3d> NewPositions;
+		UE::Geometry::SolveInsetVertexPositionsFromInsetLines(*Mesh, InsetLines, LoopPair.InnerVertices, NewPositions, true);
 
-		int32 NumVertices = LoopPair.InnerVertices.Num();
-		for (int32 vi = 0; vi < NumVertices; ++vi)
+		if (NewPositions.Num() == LoopPair.InnerVertices.Num())
 		{
-			const FLine3d& PrevLine = InsetLines[(vi == 0) ? (NumEdges-1) : (vi-1)];
-			const FLine3d& NextLine = InsetLines[vi];
-
-			if (FMathd::Abs(PrevLine.Direction.Dot(NextLine.Direction)) > 0.999)
-			{
-				FVector3d CurPos = Mesh->GetVertex(LoopPair.InnerVertices[vi]);
-				FVector3d NewPos = NextLine.NearestPoint(CurPos);
-				Mesh->SetVertex(LoopPair.InnerVertices[vi], NewPos);
-			}
-			else
+			for (int32 k = 0; k < LoopPair.InnerVertices.Num(); ++k)
 			{
-				FDistLine3Line3d Distance(PrevLine, NextLine);
-				double DistSqr = Distance.GetSquared();
-
-				FVector3d CurPos = Mesh->GetVertex(LoopPair.InnerVertices[vi]);
-				FVector3d NewPos = 0.5 * (Distance.Line1ClosestPoint + Distance.Line2ClosestPoint);
-				Mesh->SetVertex(LoopPair.InnerVertices[vi], NewPos);
+				Mesh->SetVertex(LoopPair.InnerVertices[k], NewPositions[k]);
 			}
 		}
-
-
 	};
 
 	// stitch each loop
diff --git a/Engine/Plugins/Runtime/GeometryProcessing/Source/DynamicMesh/Private/Operations/MeshBevel.cpp b/Engine/Plugins/Runtime/GeometryProcessing/Source/DynamicMesh/Private/Operations/MeshBevel.cpp
--- a/Engine/Plugins/Runtime/GeometryProcessing/Source/DynamicMesh/Private/Operations/MeshBevel.cpp
+++ b/Engine/Plugins/Runtime/GeometryProcessing/Source/DynamicMesh/Private/Operations/MeshBevel.cpp
@@ -9,7 +9,9 @@
 #include "CompGeom/PolygonTriangulation.h"
 #include "DynamicMesh/MeshNormals.h"
 #include "DynamicMesh/MeshIndexUtil.h"
+#include "Operations/PolyEditingEdgeUtil.h"
 #include "Algo/Count.h"
+#include "Distance/DistLine3Line3.h"
 
 using namespace UE::Geometry;
 
@@ -91,6 +93,11 @@ bool FMeshBevel::Apply(FDynamicMesh3& Mesh, FDynamicMeshChangeTracker* ChangeTra
 	{
 		return false;
 	}
+	FixUpUnlinkedBevelEdges(Mesh);
+	if (ResultInfo.CheckAndSetCancelled(Progress))
+	{
+		return false;
+	}
 
 	// update vertex positions
 	DisplaceVertices(Mesh, InsetDistance);
@@ -656,7 +663,6 @@ void FMeshBevel::UnlinkLoops(FDynamicMesh3& Mesh, FDynamicMeshChangeTracker* Cha
 
 
 
-
 void FMeshBevel::UnlinkVertices(FDynamicMesh3& Mesh, FDynamicMeshChangeTracker* ChangeTracker)
 {
 	// TODO: currently have to do terminator vertices first because we do some of the 
@@ -754,65 +760,123 @@ void FMeshBevel::UnlinkTerminatorVertex(FDynamicMesh3& Mesh, FBevelVertex& Bevel
 	{
 		BevelVertex.Wedges[1].WedgeVertex = SplitInfo.NewVertex;
 	}
+}
+
+
+void FMeshBevel::FixUpUnlinkedBevelEdges(FDynamicMesh3& Mesh)
+{
+	// Rewrite vertex IDs in BevelEdge vertex lists to correctly match the vertices in the new unlinked wedges.
+	// We did not know these new vertices in UnlinkBevelEdgeInterior() because we didn't unlink the vertices
+	// into wedges until afterwards. 
+	for (FBevelEdge& Edge : Edges)
+	{
+		// If this is a one-mesh-edge edge, then Edge.NewMeshEdges is incorrect because we could
+		// not actually unlink the edge in UnlinkBevelEdgeInterior (since there were no interior vertices).
+		// But now that we have unlinked the vertices, the other edge now exists, and we can find it here.
+		if (Edge.MeshEdges.Num() == 1)
+		{
+			int32* FoundOtherEdge = MeshEdgePairs.Find(Edge.MeshEdges[0]);
+			ensure(FoundOtherEdge != nullptr);
+			if (FoundOtherEdge != nullptr)
+			{
+				Edge.NewMeshEdges[0] = *FoundOtherEdge;
+			}
+			else
+			{
+				continue;		// something went wrong, loop below will break things
+			}
+		}
 
+		// process start and end vertices of the path
+		for (int32 j = 0; j < 2; ++j)
+		{
+			int32 vi = (j == 0) ? 0 : (Edge.MeshVertices.Num()-1);
+			int32 ei = (j == 0) ? 0 : (Edge.MeshEdges.Num()-1);
+			const FBevelVertex* BevelVertex = GetBevelVertexFromVertexID(Edge.MeshVertices[vi]);
+			int32& V0 = Edge.MeshVertices[vi];
+			int32& V1 = Edge.NewMeshVertices[vi];
+			int32 E0 = Edge.MeshEdges[ei], E1 = Edge.NewMeshEdges[ei];
+			bool bFoundV0 = false, bFoundV1 = false;
+			for (const FOneRingWedge& Wedge : BevelVertex->Wedges)
+			{
+				for (int32 tid : Wedge.Triangles)
+				{
+					FIndex3i TriEdges = Mesh.GetTriEdges(tid);
+					if (TriEdges.Contains(E0) && bFoundV0 == false)
+					{
+						checkSlow( Mesh.GetEdgeV(E0).Contains(Wedge.WedgeVertex) );
+						V0 = Wedge.WedgeVertex;
+						bFoundV0 = true;
+						break;
+					}
+					else if (TriEdges.Contains(E1) && bFoundV1 == false)
+					{
+						checkSlow( Mesh.GetEdgeV(E1).Contains(Wedge.WedgeVertex) );
+						V1 = Wedge.WedgeVertex;
+						bFoundV1 = true;
+						break;
+					}
+				}
+			}
+		}
+	}
 }
 
 
 
 void FMeshBevel::DisplaceVertices(FDynamicMesh3& Mesh, double Distance)
 {
-	//
-	// This displacement method produces not-very-nice-looking bevels, improvements TBD
-	//
-
+	// fallback (very bad) technique to compute an inset vertex
 	auto GetDisplacedVertexPos = [Distance](const FDynamicMesh3& Mesh, int32 VertexID) -> FVector3d
 	{
 		FVector3d CurPos = Mesh.GetVertex(VertexID);
 		FVector3d Centroid = FMeshWeights::MeanValueCentroid(Mesh, VertexID);
-		FVector3d MoveDir = Normalized(Centroid - CurPos);
-		return CurPos + Distance * MoveDir;
+		return CurPos + Distance * Normalized(Centroid - CurPos);
 	};
 
+	// Basically want to inset any beveled edges inwards into the existing poly-faces. 
+	// To do this we will solve using our standard inset technique, eg similar to FInsetMeshRegion,
+	// which involves computing 'inset lines' for each edge and then finding nearest-points between
+	// pairs of lines (which will be the intersection point if the face is planar).
 
-	auto DisplacePairedVertexLists = [&GetDisplacedVertexPos](
-		const FDynamicMesh3& Mesh, 
-		TArray<int32>& Vertices0, TArray<int32>& Vertices1, 
-		TArray<FVector3d>& Positions0, TArray<FVector3d>& Positions1, int32 InsetStart, int32 InsetEnd)
+	// Need to keep track of the inset line sets for open paths because at the corner vertices we
+	// will need to combine data from multiple path-line-sets (possibly could do this more efficiently 
+	// as we only ever need the first and last...but small in context)
+	struct FEdgePathInsetLines
 	{
-		int32 NumVertices = Vertices0.Num();
-		if (NumVertices == Vertices1.Num())
-		{
-			Positions0.SetNum(NumVertices);
-			Positions1.SetNum(NumVertices);
-			int32 Stop = NumVertices - InsetEnd;
-			for (int32 k = InsetStart; k < Stop; ++k)
-			{
-				if (Vertices0[k] == Vertices1[k])
-				{
-					Positions0[k] = Positions1[k] = Mesh.GetVertex(Vertices0[k]);
-				}
-				else
-				{
-					Positions0[k] = GetDisplacedVertexPos(Mesh, Vertices0[k]);
-					Positions1[k] = GetDisplacedVertexPos(Mesh, Vertices1[k]);
-				}
-			}
-		}
+		TArray<FLine3d> InsetLines0;
+		TArray<FLine3d> InsetLines1;
 	};
+	TArray<FEdgePathInsetLines> AllInsetLines;
+	AllInsetLines.SetNum(Edges.Num());
 
-
-	for (FBevelEdge& Edge : Edges)
+	// solve open paths
+	for ( int32 k = 0; k < Edges.Num(); ++k)
 	{
-		DisplacePairedVertexLists(Mesh, Edge.MeshVertices, Edge.NewMeshVertices, Edge.NewPositions0, Edge.NewPositions1, 
-			Edge.bEndpointBoundaryFlag[0]?0:1, Edge.bEndpointBoundaryFlag[1]?0:1 );
+		FBevelEdge& Edge = Edges[k];
+		UE::Geometry::ComputeInsetLineSegmentsFromEdges(Mesh, Edge.MeshEdges, InsetDistance, AllInsetLines[k].InsetLines0);
+		UE::Geometry::SolveInsetVertexPositionsFromInsetLines(Mesh, AllInsetLines[k].InsetLines0, Edge.MeshVertices, Edge.NewPositions0, false);
+
+		UE::Geometry::ComputeInsetLineSegmentsFromEdges(Mesh, Edge.NewMeshEdges, InsetDistance, AllInsetLines[k].InsetLines1);
+		UE::Geometry::SolveInsetVertexPositionsFromInsetLines(Mesh, AllInsetLines[k].InsetLines1, Edge.NewMeshVertices, Edge.NewPositions1, false);
 	}
+
+	// solve loops
 	for (FBevelLoop& Loop : Loops)
 	{
-		DisplacePairedVertexLists(Mesh, Loop.MeshVertices, Loop.NewMeshVertices, Loop.NewPositions0, Loop.NewPositions1, 0, 0);
+		TArray<FLine3d> InsetLines;
+		UE::Geometry::ComputeInsetLineSegmentsFromEdges(Mesh, Loop.MeshEdges, InsetDistance, InsetLines);
+		UE::Geometry::SolveInsetVertexPositionsFromInsetLines(Mesh, InsetLines, Loop.MeshVertices, Loop.NewPositions0, true);
+
+		UE::Geometry::ComputeInsetLineSegmentsFromEdges(Mesh, Loop.NewMeshEdges, InsetDistance, InsetLines);
+		UE::Geometry::SolveInsetVertexPositionsFromInsetLines(Mesh, InsetLines, Loop.NewMeshVertices, Loop.NewPositions1, true);
 	}
 
 
-	// corners
+	// Now solve corners. For corners, we want to find the 1 or 2 inset-lines corresponding
+	// to the outgoing bevel-edges at each bevel-vertex-wedge. Unfortunately we do not have
+	// a precomputed mapping for this so we currently linear-search over the full edge set for 
+	// each wedge. Could do in parallel (eg make list of valid wedges first)
 	for (FBevelVertex& Vertex : Vertices)
 	{
 		if ( (Vertex.VertexType == EBevelVertexType::JunctionVertex)
@@ -822,7 +886,45 @@ void FMeshBevel::DisplaceVertices(FDynamicMesh3& Mesh, double Distance)
 			for (int32 k = 0; k < NumWedges; ++k)
 			{
 				FOneRingWedge& Wedge = Vertex.Wedges[k];
-				Wedge.NewPosition = GetDisplacedVertexPos(Mesh, Wedge.WedgeVertex);
+
+				// collect up set of inset lines relevant to this vertex
+				TArray<FLine3d> SolveLines;
+				for (int32 j = 0; j < Edges.Num(); ++j)
+				{
+					if (Edges[j].MeshVertices[0] == Wedge.WedgeVertex)
+					{
+						SolveLines.Add(AllInsetLines[j].InsetLines0[0]);
+					}
+					else if (Edges[j].MeshVertices.Last() == Wedge.WedgeVertex)
+					{
+						SolveLines.Add(AllInsetLines[j].InsetLines0.Last());
+					}
+					else if (Edges[j].NewMeshVertices[0] == Wedge.WedgeVertex)
+					{
+						SolveLines.Add(AllInsetLines[j].InsetLines1[0]);
+					}
+					else if (Edges[j].NewMeshVertices.Last() == Wedge.WedgeVertex)
+					{
+						SolveLines.Add(AllInsetLines[j].InsetLines1.Last());
+					}
+				}
+
+				// now that we have our line set, use it to solve inset position
+				FVector3d CurPos = Mesh.GetVertex(Wedge.WedgeVertex);
+				if (SolveLines.Num() == 1)
+				{
+					Wedge.NewPosition = SolveLines[0].NearestPoint(CurPos);
+				}
+				else if (SolveLines.Num() == 2)
+				{
+					Wedge.NewPosition = UE::Geometry::SolveInsetVertexPositionFromLinePair(CurPos, SolveLines[0], SolveLines[1]);
+				}
+				else
+				{
+					// Is this even possible? #SolveLines should equal #BoundaryEdges, how can we have more than 2 at a vertex??
+					// fall back to not-very-good inset technique
+					Wedge.NewPosition = GetDisplacedVertexPos(Mesh, Wedge.WedgeVertex);
+				}
 			}
 		}
 	}
diff --git a/Engine/Plugins/Runtime/GeometryProcessing/Source/DynamicMesh/Private/Operations/PolyEditingEdgeUtil.cpp b/Engine/Plugins/Runtime/GeometryProcessing/Source/DynamicMesh/Private/Operations/PolyEditingEdgeUtil.cpp
new file mode 100644
--- /dev/null
+++ b/Engine/Plugins/Runtime/GeometryProcessing/Source/DynamicMesh/Private/Operations/PolyEditingEdgeUtil.cpp
@@ -0,0 +1,98 @@
+// Copyright Epic Games, Inc. All Rights Reserved.
+
+#include "Operations/PolyEditingEdgeUtil.h"
+#include "DynamicMesh/DynamicMesh3.h"
+#include "Distance/DistLine3Line3.h"
+
+
+using namespace UE::Geometry;
+
+
+void UE::Geometry::ComputeInsetLineSegmentsFromEdges(
+	const FDynamicMesh3& Mesh,
+	const TArray<int32>& EdgeList,
+	double InsetDistance,
+	TArray<FLine3d>& InsetLinesOut)
+{
+	int32 NumEdges = EdgeList.Num();
+	InsetLinesOut.SetNum(NumEdges);
+	for (int32 k = 0; k < NumEdges; ++k)
+	{
+		const FDynamicMesh3::FEdge EdgeVT = Mesh.GetEdge(EdgeList[k]);
+		FVector3d A = Mesh.GetVertex(EdgeVT.Vert[0]);
+		FVector3d B = Mesh.GetVertex(EdgeVT.Vert[1]);
+		FVector3d EdgeDir = Normalized(A - B);
+		FVector3d Midpoint = (A + B) * 0.5;
+		int32 EdgeTri = EdgeVT.Tri[0];
+		FVector3d Normal, Centroid; double Area;
+		Mesh.GetTriInfo(EdgeTri, Normal, Area, Centroid);
+
+		FVector3d InsetDir = Normal.Cross(EdgeDir);
+		if ((Centroid - Midpoint).Dot(InsetDir) < 0)
+		{
+			InsetDir = -InsetDir;
+		}
+
+		InsetLinesOut[k] = FLine3d(Midpoint + InsetDistance * InsetDir, EdgeDir);
+	}
+}
+
+
+
+FVector3d UE::Geometry::SolveInsetVertexPositionFromLinePair(
+	const FVector3d& Position,
+	const FLine3d& InsetEdgeLine1,
+	const FLine3d& InsetEdgeLine2)
+{
+	FVector3d NewPos = Position;
+	if (FMathd::Abs(InsetEdgeLine1.Direction.Dot(InsetEdgeLine2.Direction)) > 0.999)
+	{
+		// in this case lines are parallel so intersection point is not useful, just use nearest point on either line
+		NewPos = InsetEdgeLine1.NearestPoint(Position);
+	}
+	else
+	{
+		// inset point to intersection point of the two lines
+		FDistLine3Line3d Distance(InsetEdgeLine1, InsetEdgeLine2);
+		double DistSqr = Distance.GetSquared();
+		NewPos = 0.5 * (Distance.Line1ClosestPoint + Distance.Line2ClosestPoint);
+	}
+	return NewPos;
+}
+
+
+void UE::Geometry::SolveInsetVertexPositionsFromInsetLines(
+	const FDynamicMesh3& Mesh,
+	const TArray<FLine3d>& InsetEdgeLines,
+	const TArray<int32>& VertexIDs,
+	TArray<FVector3d>& VertexPositionsOut,
+	bool bIsLoop)
+{
+	int32 NumVertices = VertexIDs.Num();
+	VertexPositionsOut.SetNum(NumVertices);
+
+	int32 StartIndex = 0, EndIndex= NumVertices;
+
+	// if it's an open vertex span, we don't have two lines to intersect at the start/end, so 
+	// just use the nearest points there (possibly something higher up will make a better decision)
+	if (bIsLoop == false)
+	{
+		StartIndex = 1;
+		EndIndex = NumVertices - 1;
+
+		FVector3d V0 = Mesh.GetVertex(VertexIDs[0]);
+		VertexPositionsOut[0] = InsetEdgeLines[0].NearestPoint(V0);				// should just be start point...
+		FVector3d VN = Mesh.GetVertex(VertexIDs.Last());
+		VertexPositionsOut.Last() = InsetEdgeLines.Last().NearestPoint(VN);		// should just be end point...
+	}
+
+	for (int32 vi = StartIndex; vi < EndIndex; ++vi)
+	{
+		const FLine3d& PrevLine = (vi == 0) ? InsetEdgeLines.Last() : InsetEdgeLines[vi-1];
+		const FLine3d& NextLine = InsetEdgeLines[vi];
+		FVector3d CurPos = Mesh.GetVertex(VertexIDs[vi]);
+		VertexPositionsOut[vi] = SolveInsetVertexPositionFromLinePair(CurPos, PrevLine, NextLine);
+	}
+
+
+}
diff --git a/Engine/Plugins/Runtime/GeometryProcessing/Source/DynamicMesh/Public/Operations/MeshBevel.h b/Engine/Plugins/Runtime/GeometryProcessing/Source/DynamicMesh/Public/Operations/MeshBevel.h
--- a/Engine/Plugins/Runtime/GeometryProcessing/Source/DynamicMesh/Public/Operations/MeshBevel.h
+++ b/Engine/Plugins/Runtime/GeometryProcessing/Source/DynamicMesh/Public/Operations/MeshBevel.h
@@ -222,6 +222,8 @@ protected:
 	void UnlinkJunctionVertex(FDynamicMesh3& Mesh, FBevelVertex& BevelVertex, FDynamicMeshChangeTracker* ChangeTracker);
 	void UnlinkTerminatorVertex(FDynamicMesh3& Mesh, FBevelVertex& BevelVertex, FDynamicMeshChangeTracker* ChangeTracker);
 
+	void FixUpUnlinkedBevelEdges(FDynamicMesh3& Mesh);
+
 	// Displace phase - move unlinked vertices to new positions
 
 	void DisplaceVertices(FDynamicMesh3& Mesh, double Distance);
diff --git a/Engine/Plugins/Runtime/GeometryProcessing/Source/DynamicMesh/Public/Operations/PolyEditingEdgeUtil.h b/Engine/Plugins/Runtime/GeometryProcessing/Source/DynamicMesh/Public/Operations/PolyEditingEdgeUtil.h
new file mode 100644
--- /dev/null
+++ b/Engine/Plugins/Runtime/GeometryProcessing/Source/DynamicMesh/Public/Operations/PolyEditingEdgeUtil.h
@@ -0,0 +1,58 @@
+// Copyright Epic Games, Inc. All Rights Reserved.
+
+#pragma once
+
+#include "MathUtil.h"
+#include "VectorTypes.h"
+#include "IntVectorTypes.h"
+#include "GeometryTypes.h"
+#include "LineTypes.h"
+
+class FProgressCancel;
+
+namespace UE
+{
+namespace Geometry
+{
+
+class FDynamicMesh3;
+
+/**
+ * For each edge in EdgeList, compute a line segment offset from the original edge by
+ * distance InsetDistance. Assumes that each Edge is a boundary edge and so Triangle0
+ * of that edge defines "inside" direction.
+ */
+DYNAMICMESH_API void ComputeInsetLineSegmentsFromEdges(
+	const FDynamicMesh3& Mesh, 
+	const TArray<int32>& EdgeList,
+	double InsetDistance,
+	TArray<FLine3d>& InsetLinesOut );
+
+/**
+ * For each vertex in VertexIDs, compute inset position based on list of precomputed
+ * inset lines and store in VertexPositionsOut. Assumption is that each vertex position 
+ * was at the intersection point of the pre-inset lines, ie sequential edge-line-pairs
+ * determine each sequential vertex position
+ * @param bIsLoop if true, vertex list is treated as a loop, otherwise as an open span
+ */
+DYNAMICMESH_API void SolveInsetVertexPositionsFromInsetLines(
+	const FDynamicMesh3& Mesh,
+	const TArray<FLine3d>& InsetEdgeLines,
+	const TArray<int32>& VertexIDs,
+	TArray<FVector3d>& VertexPositionsOut,
+	bool bIsLoop);
+
+/**
+ * Solve for a new inset position by finding the intersection point
+ * a pair of inset-lines (or handle parallel-lines case). Assumption
+ * is that Position is the location of the initial vertex on the
+ * pre-inset edge-line-pair, ie at their intersection point.
+ */
+DYNAMICMESH_API FVector3d SolveInsetVertexPositionFromLinePair(
+	const FVector3d& Position,
+	const FLine3d& InsetEdgeLine1,
+	const FLine3d& InsetEdgeLine2);
+
+
+} // end namespace UE::Geometry
+} // end namespace UE
\ No newline at end of file
