ModelingTools: Add cancel support to BakeVertex background eval.
Update BakeVertex tile size to 1024.
Improve BakeVertex sampling robustness.

#rnx

diff --git a/Engine/Plugins/Runtime/GeometryProcessing/Source/DynamicMesh/Private/Sampling/MeshVertexBaker.cpp b/Engine/Plugins/Runtime/GeometryProcessing/Source/DynamicMesh/Private/Sampling/MeshVertexBaker.cpp
--- a/Engine/Plugins/Runtime/GeometryProcessing/Source/DynamicMesh/Private/Sampling/MeshVertexBaker.cpp
+++ b/Engine/Plugins/Runtime/GeometryProcessing/Source/DynamicMesh/Private/Sampling/MeshVertexBaker.cpp
@@ -121,6 +121,11 @@ void FMeshVertexBaker::BakeImpl(void* Data)
 	auto SampleSurface = [Baker, Mesh, ColorOverlay, NormalOverlay, UseStrategy](int32 ElementIdx,
 	                                                                FMeshMapEvaluator::FCorrespondenceSample& ValueOut)
 	{
+		if (!ColorOverlay->IsElement(ElementIdx))
+		{
+			return;
+		}
+	
 		const int32 VertexId = ColorOverlay->GetParentVertex(ElementIdx);
 
 		// Compute ray direction
@@ -186,13 +191,18 @@ void FMeshVertexBaker::BakeImpl(void* Data)
 	};
 
 	// Perform bake
-	constexpr int32 TileWidth = 256;
+	constexpr int32 TileWidth = 1024;
 	constexpr int32 TileHeight = 1;
 	const FImageTiling Tiles(Baker->Dimensions, TileWidth, TileHeight);
 	const int32 NumTiles = Tiles.Num();
 	const int NumBakers = Baker->Bakers.Num();
 	ParallelFor(NumTiles, [Baker, &Tiles, TileWidth, NumBakers, &SampleSurface](const int32 TileIdx)
 	{
+		if (Baker->CancelF())
+		{
+			return;
+		}
+		
 		const FImageTile Tile = Tiles.GetTile(TileIdx);
 		const int Width = Tile.GetWidth();
 		for (int32 Idx = 0; Idx < Width; ++Idx)
