GeometryProcessing: add basic UV calculation and materialID support to MeshBevel. Add PolyEditingUVUtil file for UV computation utility functions in Poly-editing. Currently only ComputeArbitraryTrianglePatchUVs() function which uses expmap to compute UVs and (approximately) scales UVs to match UV scale on surrounding border triangles.
GeometryScript: add ApplyMeshPolygroupBevel operation
#rnx

diff --git a/Engine/Plugins/Experimental/GeometryScripting/Source/GeometryScriptingCore/Private/MeshModelingFunctions.cpp b/Engine/Plugins/Experimental/GeometryScripting/Source/GeometryScriptingCore/Private/MeshModelingFunctions.cpp
--- a/Engine/Plugins/Experimental/GeometryScripting/Source/GeometryScriptingCore/Private/MeshModelingFunctions.cpp
+++ b/Engine/Plugins/Experimental/GeometryScripting/Source/GeometryScriptingCore/Private/MeshModelingFunctions.cpp
@@ -11,8 +11,10 @@
 #include "Operations/JoinMeshLoops.h"
 #include "MeshBoundaryLoops.h"
 #include "DynamicMeshEditor.h"
+#include "GroupTopology.h"
 #include "DynamicMesh/DynamicMeshAABBTree3.h"
 #include "Operations/OffsetMeshRegion.h"
+#include "Operations/MeshBevel.h"
 #include "UDynamicMesh.h"
 
 #include "ExplicitUseGeometryMathTypes.h"		// using UE::Geometry::(math types)
@@ -425,4 +427,36 @@ UDynamicMesh* UGeometryScriptLibrary_MeshModelingFunctions::ApplyMeshExtrude(
 
 
 
+
+
+UDynamicMesh* UGeometryScriptLibrary_MeshModelingFunctions::ApplyMeshPolygroupBevel(
+	UDynamicMesh* TargetMesh,
+	FGeometryScriptMeshBevelOptions Options,
+	UGeometryScriptDebug* Debug)
+{
+	if (TargetMesh == nullptr)
+	{
+		UE::Geometry::AppendError(Debug, EGeometryScriptErrorType::InvalidInputs, LOCTEXT("ApplyMeshPolygroupBevel_InvalidInput", "ApplyMeshPolygroupBevel: TargetMesh is Null"));
+		return TargetMesh;
+	}
+
+	TargetMesh->EditMesh([&](FDynamicMesh3& EditMesh) 
+	{
+		UE::Geometry::FGroupTopology Topology(&EditMesh, true);
+
+		UE::Geometry::FMeshBevel Bevel;
+		Bevel.InsetDistance = Options.BevelDistance;
+		Bevel.MaterialIDMode = (Options.bInferMaterialID) ? FMeshBevel::EMaterialIDMode::InferMaterialID : FMeshBevel::EMaterialIDMode::ConstantMaterialID;
+		Bevel.SetConstantMaterialID = Options.SetMaterialID;
+		Bevel.InitializeFromGroupTopology(EditMesh, Topology);
+		Bevel.Apply(EditMesh, nullptr);
+
+	}, EDynamicMeshChangeType::GeneralEdit, EDynamicMeshAttributeChangeFlags::Unknown, false);
+
+	return TargetMesh;
+}
+
+
+
+
 #undef LOCTEXT_NAMESPACE
diff --git a/Engine/Plugins/Experimental/GeometryScripting/Source/GeometryScriptingCore/Public/GeometryScript/MeshModelingFunctions.h b/Engine/Plugins/Experimental/GeometryScripting/Source/GeometryScriptingCore/Public/GeometryScript/MeshModelingFunctions.h
--- a/Engine/Plugins/Experimental/GeometryScripting/Source/GeometryScriptingCore/Public/GeometryScript/MeshModelingFunctions.h
+++ b/Engine/Plugins/Experimental/GeometryScripting/Source/GeometryScriptingCore/Public/GeometryScript/MeshModelingFunctions.h
@@ -57,6 +57,23 @@ public:
 
 
 
+USTRUCT(BlueprintType)
+struct GEOMETRYSCRIPTINGCORE_API FGeometryScriptMeshBevelOptions
+{
+	GENERATED_BODY()
+public:
+	UPROPERTY(BlueprintReadWrite, Category = Options)
+	float BevelDistance = 1.0;
+
+	UPROPERTY(BlueprintReadWrite, Category = Options)
+	bool bInferMaterialID = false;
+
+	UPROPERTY(BlueprintReadWrite, Category = Options)
+	int SetMaterialID = 0;
+};
+
+
+
 
 UCLASS(meta = (ScriptName = "GeometryScript_MeshModeling"))
 class GEOMETRYSCRIPTINGCORE_API UGeometryScriptLibrary_MeshModelingFunctions : public UBlueprintFunctionLibrary
@@ -85,4 +102,11 @@ public:
 		FGeometryScriptMeshExtrudeOptions Options,
 		UGeometryScriptDebug* Debug = nullptr );
 
+	UFUNCTION(BlueprintCallable, Category = "GeometryScript|Modeling", meta=(ScriptMethod))
+	static UPARAM(DisplayName = "Target Mesh") UDynamicMesh* 
+	ApplyMeshPolygroupBevel(
+		UDynamicMesh* TargetMesh,
+		FGeometryScriptMeshBevelOptions Options,
+		UGeometryScriptDebug* Debug = nullptr );
+
 };
\ No newline at end of file
diff --git a/Engine/Plugins/Runtime/GeometryProcessing/Source/DynamicMesh/Private/Operations/MeshBevel.cpp b/Engine/Plugins/Runtime/GeometryProcessing/Source/DynamicMesh/Private/Operations/MeshBevel.cpp
--- a/Engine/Plugins/Runtime/GeometryProcessing/Source/DynamicMesh/Private/Operations/MeshBevel.cpp
+++ b/Engine/Plugins/Runtime/GeometryProcessing/Source/DynamicMesh/Private/Operations/MeshBevel.cpp
@@ -10,6 +10,7 @@
 #include "DynamicMesh/MeshNormals.h"
 #include "DynamicMesh/MeshIndexUtil.h"
 #include "Operations/PolyEditingEdgeUtil.h"
+#include "Operations/PolyEditingUVUtil.h"
 #include "Algo/Count.h"
 #include "Distance/DistLine3Line3.h"
 
@@ -17,6 +18,35 @@ using namespace UE::Geometry;
 
 
 
+
+namespace UELocal
+{
+	static void QuadsToTris(const FDynamicMesh3& Mesh, const TArray<FIndex2i>& Quads, TArray<int32>& TrisOut, bool bReset)
+	{
+		int32 N = Quads.Num();
+		if (bReset)
+		{
+			TrisOut.Reset();
+			TrisOut.Reserve(2 * N);
+		}
+		for (const FIndex2i& Quad : Quads)
+		{
+			if (Mesh.IsTriangle(Quad.A))
+			{
+				TrisOut.Add(Quad.A);
+			}
+			if (Mesh.IsTriangle(Quad.B))
+			{
+				TrisOut.Add(Quad.B);
+			}
+		}
+	};
+
+}
+
+
+
+
 void FMeshBevel::InitializeFromGroupTopology(const FDynamicMesh3& Mesh, const FGroupTopology& Topology)
 {
 	ResultInfo = FGeometryResult(EGeometryResultType::InProgress);
@@ -113,6 +143,22 @@ bool FMeshBevel::Apply(FDynamicMesh3& Mesh, FDynamicMeshChangeTracker* ChangeTra
 		return false;
 	}
 
+	// build output list of triangles so it can be re-used in operations below if useful
+	NewTriangles.Reset();
+	for (FBevelVertex& Vertex : Vertices)
+	{
+		NewTriangles.Append(Vertex.NewTriangles);
+	}
+	for (FBevelEdge& Edge : Edges)
+	{
+		UELocal::QuadsToTris(Mesh, Edge.StripQuads, NewTriangles, false);
+	}
+	for (FBevelLoop& Loop : Loops)
+	{
+		UELocal::QuadsToTris(Mesh, Loop.StripQuads, NewTriangles, false);
+	}
+
+
 	// compute normals
 	ComputeNormals(Mesh);
 	if (ResultInfo.CheckAndSetCancelled(Progress))
@@ -120,7 +166,20 @@ bool FMeshBevel::Apply(FDynamicMesh3& Mesh, FDynamicMeshChangeTracker* ChangeTra
 		return false;
 	}
 
-	// todo: compute UVs, other attribs
+	// compute UVs
+	ComputeUVs(Mesh);
+	if (ResultInfo.CheckAndSetCancelled(Progress))
+	{
+		return false;
+	}
+
+	ComputeMaterialIDs(Mesh);
+	if (ResultInfo.CheckAndSetCancelled(Progress))
+	{
+		return false;
+	}
+
+	// todo: compute other attribs
 
 	ResultInfo.SetSuccess(true, Progress);
 	return true;
@@ -181,6 +240,13 @@ void FMeshBevel::AddBevelGroupEdge(const FDynamicMesh3& Mesh, const FGroupTopolo
 	Edge.MeshVertices.Append(Topology.Edges[GroupEdgeID].Span.Vertices);
 	Edge.GroupEdgeID = GroupEdgeID;
 	Edge.GroupIDs = Topology.Edges[GroupEdgeID].Groups;
+
+	Edge.MeshEdgeTris.Reserve(Edge.MeshEdges.Num());
+	for (int32 eid : Edge.MeshEdges)
+	{
+		Edge.MeshEdgeTris.Add(Mesh.GetEdgeT(eid));
+	}
+
 	Edges.Add(MoveTemp(Edge));
 }
 
@@ -200,6 +266,12 @@ void FMeshBevel::AddBevelEdgeLoop(const FDynamicMesh3& Mesh, const FEdgeLoop& Me
 	Loop.MeshEdges = MeshEdgeLoop.Edges;
 	Loop.MeshVertices = MeshEdgeLoop.Vertices;
 
+	Loop.MeshEdgeTris.Reserve(Loop.MeshEdges.Num());
+	for (int32 eid : Loop.MeshEdges)
+	{
+		Loop.MeshEdgeTris.Add(Mesh.GetEdgeT(eid));
+	}
+
 	Loops.Add(Loop);
 }
 
@@ -1149,6 +1221,8 @@ void FMeshBevel::CreateBevelMeshing(FDynamicMesh3& Mesh)
 
 
 
+
+
 void FMeshBevel::ComputeNormals(FDynamicMesh3& Mesh)
 {
 	if (Mesh.HasAttributes() == false)
@@ -1165,25 +1239,6 @@ void FMeshBevel::ComputeNormals(FDynamicMesh3& Mesh)
 		}
 	};
 
-	auto QuadsToTris = [](const FDynamicMesh3& Mesh, const TArray<FIndex2i>& Quads, TArray<int32>& TrisOut)
-	{
-		int32 N = Quads.Num();
-		TrisOut.Reset();
-		TrisOut.Reserve(2 * N);
-		for (const FIndex2i& Quad : Quads)
-		{
-			if (Mesh.IsTriangle(Quad.A))
-			{
-				TrisOut.Add(Quad.A);
-			}
-			if (Mesh.IsTriangle(Quad.B))
-			{
-				TrisOut.Add(Quad.B);
-			}
-		}
-	};
-
-
 	for (FBevelVertex& Vertex : Vertices)
 	{
 		SetNormalsOnTriRegion(Vertex.NewTriangles);
@@ -1192,12 +1247,208 @@ void FMeshBevel::ComputeNormals(FDynamicMesh3& Mesh)
 	TArray<int32> TriList;
 	for (FBevelEdge& Edge : Edges)
 	{
-		QuadsToTris(Mesh, Edge.StripQuads, TriList);
+		UELocal::QuadsToTris(Mesh, Edge.StripQuads, TriList, true);
 		SetNormalsOnTriRegion(TriList);
 	}
 	for (FBevelLoop& Loop : Loops)
 	{
-		QuadsToTris(Mesh, Loop.StripQuads, TriList);
+		UELocal::QuadsToTris(Mesh, Loop.StripQuads, TriList, true);
 		SetNormalsOnTriRegion(TriList);
 	}
 }
+
+
+
+void FMeshBevel::ComputeUVs(FDynamicMesh3& Mesh)
+{
+	if (Mesh.HasAttributes() == false)
+	{
+		return;
+	}
+	FDynamicMeshUVOverlay* UVOverlay = Mesh.Attributes()->PrimaryUV();
+
+	auto SetUVsOnTriRegion = [&Mesh, UVOverlay](const TArray<int32>& Triangles)
+	{
+		if (Triangles.Num() > 0)
+		{
+			UE::Geometry::ComputeArbitraryTrianglePatchUVs(Mesh, *UVOverlay, Triangles);
+		}
+	};
+
+
+	TArray<int32> TriList;
+	for (FBevelEdge& Edge : Edges)
+	{
+		UELocal::QuadsToTris(Mesh, Edge.StripQuads, TriList, true);
+		SetUVsOnTriRegion(TriList);
+	}
+	for (FBevelLoop& Loop : Loops)
+	{
+		UELocal::QuadsToTris(Mesh, Loop.StripQuads, TriList, true);
+		SetUVsOnTriRegion(TriList);
+	}
+
+	// do vertices last because until edges have UVs, the vertex polygons have no neighbour UVs islands
+	for (FBevelVertex& Vertex : Vertices)
+	{
+		SetUVsOnTriRegion(Vertex.NewTriangles);
+	}
+}
+
+
+
+void FMeshBevel::ComputeMaterialIDs(FDynamicMesh3& Mesh)
+{
+	if (Mesh.HasAttributes() == false || Mesh.Attributes()->HasMaterialID() == false)
+	{
+		return;
+	}
+	FDynamicMeshMaterialAttribute* MaterialIDs = Mesh.Attributes()->GetMaterialID();
+
+	if (MaterialIDMode == EMaterialIDMode::ConstantMaterialID)
+	{
+		for (int32 tid : NewTriangles)
+		{
+			MaterialIDs->SetValue(tid, SetConstantMaterialID);
+		}
+	}
+	else
+	{
+		auto SetQuadMaterial = [&Mesh, MaterialIDs](const FIndex2i& Quad, int32 MaterialID)
+		{
+			if (Quad.A >= 0)
+			{
+				MaterialIDs->SetValue(Quad.A, MaterialID);
+			}
+			if (Quad.B >= 0)
+			{
+				MaterialIDs->SetValue(Quad.B, MaterialID);
+			}
+		};
+
+		// Try to set materials on the new triangles along a beveled edge based on the adjacent pre-bevel triangles.
+		// Could be improved to at least be more consistent in ambiguous cases
+		auto SetEdgeMaterials = [&Mesh, MaterialIDs, this, SetQuadMaterial](const TArray<FIndex2i>& StripQuads, const TArray<FIndex2i>& EdgeTris)
+		{
+			int32 NumEdges = EdgeTris.Num();
+			if (StripQuads.Num() == NumEdges)
+			{
+				TArray<int32> SawMaterialIDs;
+				TArray<int32> AmbiguousEdges;
+				for (int32 k = 0; k < StripQuads.Num(); ++k)
+				{
+					FIndex2i NbrTris = EdgeTris[k];
+					int MatIDA = MaterialIDs->GetValue(NbrTris.A);
+					int MatIDB = (NbrTris.A >= 0) ? MaterialIDs->GetValue(NbrTris.B) : MatIDA;
+					SawMaterialIDs.AddUnique(MatIDA);
+					SawMaterialIDs.AddUnique(MatIDB);
+					int SetMaterialID = (MatIDA == MatIDB) ? MatIDA : SetConstantMaterialID;
+					if (MatIDA != MatIDB)
+					{
+						if (MaterialIDMode == EMaterialIDMode::InferMaterialID_ConstantIfAmbiguous)
+						{
+							SetMaterialID = SetConstantMaterialID;
+						}
+						else
+						{
+							AmbiguousEdges.Add(k);
+						}
+					}
+					SetQuadMaterial(StripQuads[k], SetMaterialID);
+				}
+
+				if (AmbiguousEdges.Num() > 0)
+				{
+					SawMaterialIDs.Sort();
+					if (AmbiguousEdges.Num() == NumEdges)		// if all ambigous, just pick one
+					{
+						for (int32 k : AmbiguousEdges)
+						{
+							SetQuadMaterial(StripQuads[k], SawMaterialIDs[0]);
+						}
+					}
+					else
+					{
+						// TODO: what we probably want to do here is "infill" from known neighbours. 
+						// for now we will just punt and pick one
+						for (int32 k : AmbiguousEdges)
+						{
+							SetQuadMaterial(StripQuads[k], SawMaterialIDs[0]);
+						}
+					}
+				}
+
+			}
+			else
+			{
+				for (const FIndex2i& Quad : StripQuads)
+				{
+					SetQuadMaterial(Quad, SetConstantMaterialID);
+				}
+			}
+		};
+
+		for (FBevelEdge& Edge : Edges)
+		{
+			SetEdgeMaterials(Edge.StripQuads, Edge.MeshEdgeTris);
+		}
+		for (FBevelLoop& Loop : Loops)
+		{
+			SetEdgeMaterials(Loop.StripQuads, Loop.MeshEdgeTris);
+		}
+
+
+		// find all the unique material IDs of neighbours of the Triangles list (that are not in Triangles list) and
+		// return (MaterialID, NbrTriCount) tuples as a pair of lists
+		auto CountUniqueBorderMaterialIDs = [&](const FDynamicMesh3& Mesh, const FDynamicMeshMaterialAttribute& MaterialAttrib, const TArray<int32>& Triangles, TArray<int32>& MaterialIDs, TArray<int32>& Counts)
+		{
+			MaterialIDs.Reset(); Counts.Reset();
+			for (int32 tid : Triangles)
+			{
+				FIndex3i TriNbrs = Mesh.GetTriNeighbourTris(tid);
+				for (int32 j = 0; j < 3; ++j)
+				{
+					if (Triangles.Contains(TriNbrs[j]))
+					{
+						continue;
+					}
+					int MatID = MaterialAttrib.GetValue(TriNbrs[j]);
+					int32 Index = MaterialIDs.AddUnique(MatID);
+					if (Counts.Num() != MaterialIDs.Num())
+					{
+						Counts.Add(0);
+						Counts[Index]++;
+					}
+				}
+			}
+		};
+
+		// For each bevel-vertex-polygon, pick the nbr material ID that was most frequent.
+		// Terminator vertices are also handled this way, which is not ideal, should probably
+		// ignore the new 'edge' faces for the terminator vertex
+		for (FBevelVertex& Vertex : Vertices)
+		{
+			TArray<int32> NbrMaterialIDs, NbrMaterialIDCounts;
+			CountUniqueBorderMaterialIDs(Mesh, *MaterialIDs, Vertex.NewTriangles, NbrMaterialIDs, NbrMaterialIDCounts);
+			int32 SetMaterialID = SetConstantMaterialID;
+			if (NbrMaterialIDs.Num() > 0)
+			{
+				int32 MinIndex = 0;
+				for (int32 k = 1; k < NbrMaterialIDs.Num(); ++k)
+				{
+					if (NbrMaterialIDCounts[k] < NbrMaterialIDCounts[MinIndex])
+					{
+						MinIndex = k;
+					}
+				}
+				SetMaterialID = NbrMaterialIDs[MinIndex];
+			}
+			for (int32 tid : Vertex.NewTriangles)
+			{
+				MaterialIDs->SetValue(tid, SetMaterialID);
+			}
+
+		}
+	}
+
+}
\ No newline at end of file
diff --git a/Engine/Plugins/Runtime/GeometryProcessing/Source/DynamicMesh/Private/Operations/PolyEditingUVUtil.cpp b/Engine/Plugins/Runtime/GeometryProcessing/Source/DynamicMesh/Private/Operations/PolyEditingUVUtil.cpp
new file mode 100644
--- /dev/null
+++ b/Engine/Plugins/Runtime/GeometryProcessing/Source/DynamicMesh/Private/Operations/PolyEditingUVUtil.cpp
@@ -0,0 +1,50 @@
+// Copyright Epic Games, Inc. All Rights Reserved.
+
+#include "Operations/PolyEditingUVUtil.h"
+#include "DynamicMesh/DynamicMesh3.h"
+#include "Parameterization/DynamicMeshUVEditor.h"
+
+
+using namespace UE::Geometry;
+
+
+void UE::Geometry::ComputeArbitraryTrianglePatchUVs(
+	FDynamicMesh3& Mesh, 
+	FDynamicMeshUVOverlay& UVOverlay,
+	const TArray<int32>& TriangleSet)
+{
+	TArray<int32> NbrTriSet;
+	double NbrUVAreaSum = 0.0;
+	double Nbr3DAreaSum = 0.0;
+	for (int32 tid : TriangleSet)
+	{
+		FIndex3i NbrTris = Mesh.GetTriNeighbourTris(tid);
+		for (int32 j = 0; j < 3; ++j)
+		{
+			if (NbrTris[j] != IndexConstants::InvalidID && NbrTriSet.Contains(NbrTris[j]) == false)
+			{
+				NbrTriSet.Add(NbrTris[j]);
+				if (UVOverlay.IsSetTriangle(NbrTris[j]))
+				{
+					FVector3d A, B, C;
+					Mesh.GetTriVertices(NbrTris[j], A, B, C);
+					Nbr3DAreaSum += VectorUtil::Area(A,B,C);
+					FVector2f U, V, W;
+					UVOverlay.GetTriElements(NbrTris[j], U, V, W);
+					NbrUVAreaSum += (double)VectorUtil::Area(U, V, W);
+				}
+			}
+		}
+	}
+
+	double UseUVScale = FMathd::Max(FMathd::Sqrt(NbrUVAreaSum),0.0001) / FMathd::Max(FMathd::Sqrt(Nbr3DAreaSum),0.0001);
+
+	FDynamicMeshUVEditor UVEditor(&Mesh, &UVOverlay);
+	FUVEditResult UVEditResult;
+	UVEditor.SetTriangleUVsFromExpMap(TriangleSet, FDynamicMeshUVEditor::FExpMapOptions(), &UVEditResult);
+
+	UVEditor.TransformUVElements(UVEditResult.NewUVElements, [&](const FVector2f& UV)
+	{
+		return (float)UseUVScale * UV;
+	});
+}
\ No newline at end of file
diff --git a/Engine/Plugins/Runtime/GeometryProcessing/Source/DynamicMesh/Public/Operations/MeshBevel.h b/Engine/Plugins/Runtime/GeometryProcessing/Source/DynamicMesh/Public/Operations/MeshBevel.h
--- a/Engine/Plugins/Runtime/GeometryProcessing/Source/DynamicMesh/Public/Operations/MeshBevel.h
+++ b/Engine/Plugins/Runtime/GeometryProcessing/Source/DynamicMesh/Public/Operations/MeshBevel.h
@@ -46,12 +46,40 @@ class DYNAMICMESH_API FMeshBevel
 {
 public:
 
+	//
+	// Inputs
+	// 
+
 	/** Distance that bevel edges/vertices are inset from their initial position. Not guaranteed to hold for all vertices, though. */
 	double InsetDistance = 5.0;
+	
+	/** Options for MaterialID assignment on the new triangles generated for the bevel */
+	enum class EMaterialIDMode
+	{
+		ConstantMaterialID,
+		InferMaterialID,
+		InferMaterialID_ConstantIfAmbiguous
+	};
+	/** Which MaterialID assignment mode to use */
+	EMaterialIDMode MaterialIDMode = EMaterialIDMode::ConstantMaterialID;
+	/** Constant MaterialID used for various MaterialIDMode settings */
+	int32 SetConstantMaterialID = 0;
 
 	/** Set this member to support progress/cancel in the computations below */
 	FProgressCancel* Progress = nullptr;
 
+
+	//
+	// Outputs
+	//
+
+	/** list of all new triangles created by the operation */
+	TArray<int32> NewTriangles;
+
+	/** status of the operation, warnings/errors may be returned here */
+	FGeometryResult ResultInfo = FGeometryResult(EGeometryResultType::InProgress);
+
+public:
 	/**
 	 * Initialize the bevel with all edges of the given GroupTopology
 	 */
@@ -67,8 +95,7 @@ public:
 	 */
 	bool Apply(FDynamicMesh3& Mesh, FDynamicMeshChangeTracker* ChangeTracker);
 
-	/** status of the operation, warnings/errors may be returned here */
-	FGeometryResult ResultInfo = FGeometryResult(EGeometryResultType::InProgress);
+
 
 
 public:
@@ -94,6 +121,7 @@ public:
 		// initial topological information that defines what happens in unlink/displace/mesh steps
 		TArray<int32> MeshVertices;		// sequential list of mesh vertex IDs along edge loop
 		TArray<int32> MeshEdges;		// sequential list of mesh edge IDs along edge loop
+		TArray<FIndex2i> MeshEdgeTris;	// the one or two triangles associated w/ each MeshEdges element in the input mesh
 
 		// new mesh topology computed during unlink step
 		TArray<int32> NewMeshVertices;		// list of new vertices on "other" side of unlinked edge, 1-1 with MeshVertices
@@ -115,6 +143,7 @@ public:
 		// initial topological information that defines what happens in unlink/displace/mesh steps
 		TArray<int32> MeshVertices;		// sequential list of mesh vertex IDs along edge
 		TArray<int32> MeshEdges;		// sequential list of mesh edge IDs along edge
+		TArray<FIndex2i> MeshEdgeTris;	// the one or two triangles associated w/ each MeshEdges element in the input mesh
 		int32 GroupEdgeID;				// ID of this edge in external topology (eg FGroupTopology)
 		FIndex2i GroupIDs;				// topological IDs of groups on either side of topological edge
 		bool bEndpointBoundaryFlag[2];	// flag defining whether vertex at start/end of MeshVertices was a boundary vertex
@@ -240,7 +269,8 @@ protected:
 	// Normals phase - calculate normals for new geometry
 
 	void ComputeNormals(FDynamicMesh3& Mesh);
-
+	void ComputeUVs(FDynamicMesh3& Mesh);
+	void ComputeMaterialIDs(FDynamicMesh3& Mesh);
 };
 
 
diff --git a/Engine/Plugins/Runtime/GeometryProcessing/Source/DynamicMesh/Public/Operations/PolyEditingUVUtil.h b/Engine/Plugins/Runtime/GeometryProcessing/Source/DynamicMesh/Public/Operations/PolyEditingUVUtil.h
new file mode 100644
--- /dev/null
+++ b/Engine/Plugins/Runtime/GeometryProcessing/Source/DynamicMesh/Public/Operations/PolyEditingUVUtil.h
@@ -0,0 +1,25 @@
+// Copyright Epic Games, Inc. All Rights Reserved.
+
+#pragma once
+
+#include "MathUtil.h"
+#include "VectorTypes.h"
+#include "IntVectorTypes.h"
+#include "GeometryTypes.h"
+#include "DynamicMesh/DynamicMeshAttributeSet.h"
+
+class FProgressCancel;
+
+namespace UE
+{
+namespace Geometry
+{
+
+DYNAMICMESH_API void ComputeArbitraryTrianglePatchUVs(
+	FDynamicMesh3& Mesh, 
+	FDynamicMeshUVOverlay& UVOverlay,
+	const TArray<int32>& TriangleSet );
+
+
+} // end namespace UE::Geometry
+} // end namespace UE
\ No newline at end of file
