fix mesh boolean crash on trying to double-check coplanarity when the boolean only needed to process the first mesh, so there is nothing to double-check against

#rnx

diff --git a/Engine/Plugins/Runtime/GeometryProcessing/Source/DynamicMesh/Private/Operations/MeshBoolean.cpp b/Engine/Plugins/Runtime/GeometryProcessing/Source/DynamicMesh/Private/Operations/MeshBoolean.cpp
--- a/Engine/Plugins/Runtime/GeometryProcessing/Source/DynamicMesh/Private/Operations/MeshBoolean.cpp
+++ b/Engine/Plugins/Runtime/GeometryProcessing/Source/DynamicMesh/Private/Operations/MeshBoolean.cpp
@@ -362,7 +362,10 @@ bool FMeshBoolean::Compute()
 		// This array is used to double-check the assumption that we will delete the other surface when we keep a coplanar tri
 		// Note we only need it for mesh 0 (i.e., the mesh we try to keep triangles from when we preserve coplanar surfaces)
 		TArray<int32> DeleteIfOtherKept;
-		DeleteIfOtherKept.Init(-1, CutMesh[0]->MaxTriangleID());
+		if (NumMeshesToProcess > 1)
+		{
+			DeleteIfOtherKept.Init(-1, CutMesh[0]->MaxTriangleID());
+		}
 		for (int MeshIdx = 0; MeshIdx < NumMeshesToProcess; MeshIdx++)
 		{
 			TFastWindingTree<FDynamicMesh3> Winding(&Spatial[1 - MeshIdx]);
@@ -436,7 +439,7 @@ bool FMeshBoolean::Compute()
 									{
 										bool bKeep = DotNormals > 0 == bCoplanarKeepSameDir;
 										KeepTri[MeshIdx][TID] = bKeep;
-										if (bKeep)
+										if (NumMeshesToProcess > 1 && bKeep)
 										{
 											// If we kept this tri, remember the coplanar pair we expect to be deleted, in case
 											// it isn't deleted (e.g. because it wasn't coplanar); to then delete this one instead.
@@ -457,12 +460,15 @@ bool FMeshBoolean::Compute()
 		}
 
 		// Don't keep coplanar tris if the matched, second-mesh tri that we expected to delete was actually kept
-		for (int TID : CutMesh[0]->TriangleIndicesItr())
+		if (NumMeshesToProcess > 1)
 		{
-			int32 DeleteIfOtherKeptTID = DeleteIfOtherKept[TID];
-			if (DeleteIfOtherKeptTID > -1 && KeepTri[1][DeleteIfOtherKeptTID])
+			for (int TID : CutMesh[0]->TriangleIndicesItr())
 			{
-				KeepTri[0][TID] = false;
+				int32 DeleteIfOtherKeptTID = DeleteIfOtherKept[TID];
+				if (DeleteIfOtherKeptTID > -1 && KeepTri[1][DeleteIfOtherKeptTID])
+				{
+					KeepTri[0][TID] = false;
+				}
 			}
 		}
 
