add better comment + alternate call to the CDT helper function, to make it easier to use

#rnx

diff --git a/Engine/Plugins/Runtime/GeometryProcessing/Source/GeometryAlgorithms/Private/ConstrainedDelaunay2.cpp b/Engine/Plugins/Runtime/GeometryProcessing/Source/GeometryAlgorithms/Private/ConstrainedDelaunay2.cpp
--- a/Engine/Plugins/Runtime/GeometryProcessing/Source/GeometryAlgorithms/Private/ConstrainedDelaunay2.cpp
+++ b/Engine/Plugins/Runtime/GeometryProcessing/Source/GeometryAlgorithms/Private/ConstrainedDelaunay2.cpp
@@ -469,7 +469,7 @@ bool TConstrainedDelaunay2<RealType>::Triangulate()
 
 	return !bBoundaryTrackingFailure;
 }
-//
+
 template<typename RealType>
 TArray<FIndex3i> GEOMETRYALGORITHMS_API UE::Geometry::ConstrainedDelaunayTriangulate(const TGeneralPolygon2<RealType>& GeneralPolygon)
 {
@@ -479,6 +479,16 @@ TArray<FIndex3i> GEOMETRYALGORITHMS_API UE::Geometry::ConstrainedDelaunayTriangu
 	Triangulation.Triangulate();
 	return Triangulation.Triangles;
 }
+template<typename RealType>
+TArray<FIndex3i> GEOMETRYALGORITHMS_API UE::Geometry::ConstrainedDelaunayTriangulateWithVertices(const TGeneralPolygon2<RealType>& GeneralPolygon, TArray<TVector2<RealType>>& OutVertices)
+{
+	TConstrainedDelaunay2<RealType> Triangulation;
+	Triangulation.FillRule = TConstrainedDelaunay2<RealType>::EFillRule::Positive;
+	Triangulation.Add(GeneralPolygon);
+	Triangulation.Triangulate();
+	OutVertices = MoveTemp(Triangulation.Vertices);
+	return Triangulation.Triangles;
+}
 
 namespace UE
 {
@@ -487,6 +497,8 @@ namespace Geometry
 
 template TArray<FIndex3i> GEOMETRYALGORITHMS_API UE::Geometry::ConstrainedDelaunayTriangulate(const TGeneralPolygon2<double>& GeneralPolygon);
 template TArray<FIndex3i> GEOMETRYALGORITHMS_API UE::Geometry::ConstrainedDelaunayTriangulate(const TGeneralPolygon2<float>& GeneralPolygon);
+template TArray<FIndex3i> GEOMETRYALGORITHMS_API UE::Geometry::ConstrainedDelaunayTriangulateWithVertices(const TGeneralPolygon2<double>& GeneralPolygon, TArray<TVector2<double>>& Vertices);
+template TArray<FIndex3i> GEOMETRYALGORITHMS_API UE::Geometry::ConstrainedDelaunayTriangulateWithVertices(const TGeneralPolygon2<float>& GeneralPolygon, TArray<TVector2<float>>& Vertices);
 
 
 
diff --git a/Engine/Plugins/Runtime/GeometryProcessing/Source/GeometryAlgorithms/Public/ConstrainedDelaunay2.h b/Engine/Plugins/Runtime/GeometryProcessing/Source/GeometryAlgorithms/Public/ConstrainedDelaunay2.h
--- a/Engine/Plugins/Runtime/GeometryProcessing/Source/GeometryAlgorithms/Public/ConstrainedDelaunay2.h
+++ b/Engine/Plugins/Runtime/GeometryProcessing/Source/GeometryAlgorithms/Public/ConstrainedDelaunay2.h
@@ -146,10 +146,25 @@ struct TConstrainedDelaunay2
 	bool GEOMETRYALGORITHMS_API Triangulate(TFunctionRef<bool(const TArray<TVector2<RealType>>&, const FIndex3i&)> KeepTriangle);
 };
 
-/** Convenience function; invokes TConstrainedDelaunay2::Triangulate with most-often-used options */
+/**
+ * Convenience function; invokes TConstrainedDelaunay2::Triangulate with most-often-used options
+ * 
+ * @param GeneralPolygon A general polygon, which may include holes.
+ * @return An array of triangle indices.  Indices refer to a flat array w/ the outer polygon vertices first, and the hole vertices appended after.
+ */
 template<typename RealType>
 TArray<FIndex3i> GEOMETRYALGORITHMS_API ConstrainedDelaunayTriangulate(const TGeneralPolygon2<RealType>& GeneralPolygon);
 
+/**
+ * Convenience function; invokes TConstrainedDelaunay2::Triangulate with most-often-used options.
+ *
+ * @param GeneralPolygon A general polygon, which may include holes.
+ * @param OutVertices Output array of vertices used in the triangulation.
+ * @return An array of triangle indices.
+ */
+template<typename RealType>
+TArray<FIndex3i> GEOMETRYALGORITHMS_API ConstrainedDelaunayTriangulateWithVertices(const TGeneralPolygon2<RealType>& GeneralPolygon, TArray<TVector2<RealType>>& OutVertices);
+
 typedef TConstrainedDelaunay2<float> FConstrainedDelaunay2f;
 typedef TConstrainedDelaunay2<double> FConstrainedDelaunay2d;
 
