ModelingTools: Add additional trace events to the Bake tools.

#rnx

diff --git a/Engine/Plugins/Experimental/MeshModelingToolsetExp/Source/MeshModelingToolsExp/Private/BakeMeshAttributeMapsTool.cpp b/Engine/Plugins/Experimental/MeshModelingToolsetExp/Source/MeshModelingToolsExp/Private/BakeMeshAttributeMapsTool.cpp
--- a/Engine/Plugins/Experimental/MeshModelingToolsetExp/Source/MeshModelingToolsExp/Private/BakeMeshAttributeMapsTool.cpp
+++ b/Engine/Plugins/Experimental/MeshModelingToolsetExp/Source/MeshModelingToolsExp/Private/BakeMeshAttributeMapsTool.cpp
@@ -263,6 +263,8 @@ public:
 
 void UBakeMeshAttributeMapsTool::Setup()
 {
+	TRACE_CPUPROFILER_EVENT_SCOPE(UBakeMeshAttributeMapsTool::Setup);
+	
 	Super::Setup();
 
 	// Initialize preview mesh
@@ -448,6 +450,8 @@ TUniquePtr<UE::Geometry::TGenericDataOperator<FMeshMapBaker>> UBakeMeshAttribute
 
 void UBakeMeshAttributeMapsTool::Shutdown(EToolShutdownType ShutdownType)
 {
+	TRACE_CPUPROFILER_EVENT_SCOPE(UBakeMeshAttributeMapsTool::Shutdown);
+	
 	Super::Shutdown(ShutdownType);
 	
 	Settings->SaveProperties(this);
diff --git a/Engine/Plugins/Experimental/MeshModelingToolsetExp/Source/MeshModelingToolsExp/Private/BakeMeshAttributeVertexTool.cpp b/Engine/Plugins/Experimental/MeshModelingToolsetExp/Source/MeshModelingToolsExp/Private/BakeMeshAttributeVertexTool.cpp
--- a/Engine/Plugins/Experimental/MeshModelingToolsetExp/Source/MeshModelingToolsExp/Private/BakeMeshAttributeVertexTool.cpp
+++ b/Engine/Plugins/Experimental/MeshModelingToolsetExp/Source/MeshModelingToolsExp/Private/BakeMeshAttributeVertexTool.cpp
@@ -244,6 +244,7 @@ public:
 
 void UBakeMeshAttributeVertexTool::Setup()
 {
+	TRACE_CPUPROFILER_EVENT_SCOPE(UBakeMeshAttributeVertexTool::Setup);
 	Super::Setup();
 
 	UMaterial* Material = LoadObject<UMaterial>(nullptr, TEXT("/MeshModelingToolsetExp/Materials/MeshVertexColorMaterial"));
@@ -360,6 +361,8 @@ void UBakeMeshAttributeVertexTool::Setup()
 
 void UBakeMeshAttributeVertexTool::Shutdown(EToolShutdownType ShutdownType)
 {
+	TRACE_CPUPROFILER_EVENT_SCOPE(UBakeMeshAttributeVertexTool::Shutdown);
+	
 	Settings->SaveProperties(this);
 	ColorSettings->SaveProperties(this);
 	PerChannelSettings->SaveProperties(this);
diff --git a/Engine/Plugins/Experimental/MeshModelingToolsetExp/Source/MeshModelingToolsExp/Private/BakeMultiMeshAttributeMapsTool.cpp b/Engine/Plugins/Experimental/MeshModelingToolsetExp/Source/MeshModelingToolsExp/Private/BakeMultiMeshAttributeMapsTool.cpp
--- a/Engine/Plugins/Experimental/MeshModelingToolsetExp/Source/MeshModelingToolsExp/Private/BakeMultiMeshAttributeMapsTool.cpp
+++ b/Engine/Plugins/Experimental/MeshModelingToolsetExp/Source/MeshModelingToolsExp/Private/BakeMultiMeshAttributeMapsTool.cpp
@@ -395,6 +395,8 @@ public:
 
 void UBakeMultiMeshAttributeMapsTool::Setup()
 {
+	TRACE_CPUPROFILER_EVENT_SCOPE(UBakeMultiMeshAttributeMapsTool::Setup);
+	
 	Super::Setup();
 
 	// Initialize base mesh
@@ -533,6 +535,8 @@ TUniquePtr<UE::Geometry::TGenericDataOperator<FMeshMapBaker>> UBakeMultiMeshAttr
 
 void UBakeMultiMeshAttributeMapsTool::Shutdown(EToolShutdownType ShutdownType)
 {
+	TRACE_CPUPROFILER_EVENT_SCOPE(UBakeMultiMeshAttributeMapsTool::Shutdown);
+	
 	Super::Shutdown(ShutdownType);
 	
 	Settings->SaveProperties(this);
diff --git a/Engine/Plugins/Runtime/GeometryProcessing/Source/DynamicMesh/Private/Sampling/MeshMapBaker.cpp b/Engine/Plugins/Runtime/GeometryProcessing/Source/DynamicMesh/Private/Sampling/MeshMapBaker.cpp
--- a/Engine/Plugins/Runtime/GeometryProcessing/Source/DynamicMesh/Private/Sampling/MeshMapBaker.cpp
+++ b/Engine/Plugins/Runtime/GeometryProcessing/Source/DynamicMesh/Private/Sampling/MeshMapBaker.cpp
@@ -25,6 +25,8 @@ FMitchellNetravaliFilter FMeshMapBaker::MitchellNetravaliFilter(BCFilterRadius);
 
 void FMeshMapBaker::InitBake()
 {
+	TRACE_CPUPROFILER_EVENT_SCOPE(FMeshMapBaker::InitBake);
+	
 	// Retrieve evaluation contexts and cache:
 	// - index lists of accumulation modes (BakeAccumulateLists)
 	// - evaluator to bake result offsets (BakeOffsets)
@@ -127,22 +129,26 @@ void FMeshMapBaker::Bake()
 
 	InitBake();
 
-	// Generate UV space mesh
 	const FDynamicMesh3* Mesh = TargetMesh;
 	const FDynamicMeshUVOverlay* UVOverlay = GetTargetMeshUVs();
 	const FDynamicMeshNormalOverlay* NormalOverlay = GetTargetMeshNormals();
 
-	FlatMesh = FDynamicMesh3(EMeshComponents::FaceGroups);
-	for (int32 tid : Mesh->TriangleIndicesItr())
 	{
-		if (UVOverlay->IsSetTriangle(tid))
+		// Generate UV space mesh
+		TRACE_CPUPROFILER_EVENT_SCOPE(FMeshMapBaker::Bake_CreateUVMesh);
+		
+		FlatMesh = FDynamicMesh3(EMeshComponents::FaceGroups);
+		for (int32 tid : Mesh->TriangleIndicesItr())
 		{
-			FVector2f A, B, C;
-			UVOverlay->GetTriElements(tid, A, B, C);
-			const int32 VertA = FlatMesh.AppendVertex(FVector3d(A.X, A.Y, 0));
-			const int32 VertB = FlatMesh.AppendVertex(FVector3d(B.X, B.Y, 0));
-			const int32 VertC = FlatMesh.AppendVertex(FVector3d(C.X, C.Y, 0));
-			/*int32 NewTriID =*/ FlatMesh.AppendTriangle(VertA, VertB, VertC, tid);
+			if (UVOverlay->IsSetTriangle(tid))
+			{
+				FVector2f A, B, C;
+				UVOverlay->GetTriElements(tid, A, B, C);
+				const int32 VertA = FlatMesh.AppendVertex(FVector3d(A.X, A.Y, 0));
+				const int32 VertB = FlatMesh.AppendVertex(FVector3d(B.X, B.Y, 0));
+				const int32 VertC = FlatMesh.AppendVertex(FVector3d(C.X, C.Y, 0));
+				/*int32 NewTriID =*/ FlatMesh.AppendTriangle(VertA, VertB, VertC, tid);
+			}
 		}
 	}
 
@@ -225,7 +231,7 @@ void FMeshMapBaker::Bake()
 	// WriteToOutputBuffer transfers local tile data (TileBuffer) to the image output buffer (ImageTileBuffer).
 	auto WriteToOutputBuffer = [this, &ImageTileBuffer] (FMeshMapTileBuffer& TileBufferIn, const FImageTile& TargetTile, const TArray<int32>& BakeIds, auto&& Op, auto&& WeightOp)
 	{
-		TRACE_CPUPROFILER_EVENT_SCOPE(FMeshMapBaker::WriteToOutputBuffer);
+		TRACE_CPUPROFILER_EVENT_SCOPE(FMeshMapBaker::Bake_WriteToOutputBuffer);
 		
 		const int TargetTileWidth = TargetTile.GetWidth();
 		const int TargetTileHeight = TargetTile.GetHeight();
@@ -292,7 +298,8 @@ void FMeshMapBaker::Bake()
 	
 	ParallelFor(NumTiles, [this, &Tiles, &GutterTexelsPerTile, &OutputQueue, &WriteToOutputBuffer, &OverwriteFn, &NoopFn, &WriteQueuedOutput](int32 TileIdx)
 	{
-		TRACE_CPUPROFILER_EVENT_SCOPE(FMeshMapBaker::BakeTile);
+		TRACE_CPUPROFILER_EVENT_SCOPE(FMeshMapBaker::Bake_EvalTile);
+		
 		// Generate unpadded and padded tiles.
 		const FImageTile Tile = Tiles.GetTile(TileIdx);	// Image area to sample
 		const FImageTile PaddedTile = Tiles.GetTile(TileIdx, TilePadding); // Filtered image area
@@ -307,40 +314,44 @@ void FMeshMapBaker::Bake()
 
 		FMeshMapTileBuffer* TileBuffer = new FMeshMapTileBuffer(PaddedTile, BakeSampleBufferSize);
 
-		// Calculate interior texels
-		const int TileWidth = Tile.GetWidth();
-		const int TileHeight = Tile.GetHeight();
-		const int32 NumSamples = OccupancyMap.Multisampler.Num();
-		for (FVector2i TileCoords(0,0); TileCoords.Y < TileHeight; ++TileCoords.Y)
 		{
-			for (TileCoords.X = 0; TileCoords.X < TileWidth; ++TileCoords.X)
+			// Evaluate valid/interior samples
+			TRACE_CPUPROFILER_EVENT_SCOPE(FMeshMapBaker::Bake_EvalTileSamples);
+			
+			const int TileWidth = Tile.GetWidth();
+			const int TileHeight = Tile.GetHeight();
+			const int32 NumSamples = OccupancyMap.Multisampler.Num();
+			for (FVector2i TileCoords(0,0); TileCoords.Y < TileHeight; ++TileCoords.Y)
 			{
-				if (CancelF())
-				{
-					delete TileBuffer;
-					return;
-				}
-
-				const int64 TilePixelLinearIdx = Tile.GetIndex(TileCoords);
-				if (OccupancyMap.TexelNumSamples(TilePixelLinearIdx) == 0)
+				for (TileCoords.X = 0; TileCoords.X < TileWidth; ++TileCoords.X)
 				{
-					continue;
-				}
+					if (CancelF())
+					{
+						delete TileBuffer;
+						return;
+					}
 
-				const FVector2i ImageCoords = Tile.GetSourceCoords(TileCoords);
-				for (int32 SampleIdx = 0; SampleIdx < NumSamples; ++SampleIdx)
-				{
-					const int64 LinearIdx = TilePixelLinearIdx * NumSamples + SampleIdx;
-					if (OccupancyMap.IsInterior(LinearIdx))
+					const int64 TilePixelLinearIdx = Tile.GetIndex(TileCoords);
+					if (OccupancyMap.TexelNumSamples(TilePixelLinearIdx) == 0)
 					{
-						const FVector2d UVPosition = (FVector2d)OccupancyMap.TexelQueryUV[LinearIdx];
-						const int32 UVTriangleID = OccupancyMap.TexelQueryTriangle[LinearIdx];
+						continue;
+					}
 
-						FMeshMapEvaluator::FCorrespondenceSample Sample;
-						DetailCorrespondenceSampler.SampleUV(UVTriangleID, UVPosition, Sample);
-						if (Sample.DetailMesh && DetailSampler->IsTriangle(Sample.DetailMesh, Sample.DetailTriID))
+					const FVector2i ImageCoords = Tile.GetSourceCoords(TileCoords);
+					for (int32 SampleIdx = 0; SampleIdx < NumSamples; ++SampleIdx)
+					{
+						const int64 LinearIdx = TilePixelLinearIdx * NumSamples + SampleIdx;
+						if (OccupancyMap.IsInterior(LinearIdx))
 						{
-							BakeSample(*TileBuffer, Sample, UVPosition, ImageCoords);
+							const FVector2d UVPosition = (FVector2d)OccupancyMap.TexelQueryUV[LinearIdx];
+							const int32 UVTriangleID = OccupancyMap.TexelQueryTriangle[LinearIdx];
+
+							FMeshMapEvaluator::FCorrespondenceSample Sample;
+							DetailCorrespondenceSampler.SampleUV(UVTriangleID, UVPosition, Sample);
+							if (Sample.DetailMesh && DetailSampler->IsTriangle(Sample.DetailMesh, Sample.DetailTriID))
+							{
+								BakeSample(*TileBuffer, Sample, UVPosition, ImageCoords);
+							}
 						}
 					}
 				}
@@ -382,7 +393,7 @@ void FMeshMapBaker::Bake()
 		// Normalize and convert ImageTileBuffer data to color data.
 		ParallelFor(NumTiles, [this, &Tiles, &ImageTileBuffer](int32 TileIdx)
 		{
-			TRACE_CPUPROFILER_EVENT_SCOPE(FMeshMapBaker::WriteToImageBuffer);
+			TRACE_CPUPROFILER_EVENT_SCOPE(FMeshMapBaker::Bake_WriteToImageBuffer);
 		
 			const FImageTile Tile = Tiles.GetTile(TileIdx);
 			const int TileWidth = Tile.GetWidth();
@@ -446,7 +457,7 @@ void FMeshMapBaker::Bake()
 		const int32 NumResults = BakeResults.Num();
 		ParallelFor(NumTiles, [this, &NumResults, &GutterTexelsPerTile](int32 TileIdx)
 		{
-			TRACE_CPUPROFILER_EVENT_SCOPE(FMeshMapBaker::WriteGutterPixels);
+			TRACE_CPUPROFILER_EVENT_SCOPE(FMeshMapBaker::Bake_WriteGutterPixels);
 
 			const int NumGutter = GutterTexelsPerTile[TileIdx].Num();
 			for (int64 GutterIdx = 0; GutterIdx < NumGutter; ++GutterIdx)
diff --git a/Engine/Source/Runtime/GeometryCore/Public/Image/ImageOccupancyMap.h b/Engine/Source/Runtime/GeometryCore/Public/Image/ImageOccupancyMap.h
--- a/Engine/Source/Runtime/GeometryCore/Public/Image/ImageOccupancyMap.h
+++ b/Engine/Source/Runtime/GeometryCore/Public/Image/ImageOccupancyMap.h
@@ -103,6 +103,8 @@ public:
 	bool ComputeFromUVSpaceMesh(const MeshType& UVSpaceMesh, 
 		GetTriangleIDFuncType GetTriangleIDFunc = [](int32 TriangleID) { return TriangleID; } )
 	{
+		TRACE_CPUPROFILER_EVENT_SCOPE(FImageOccupancyMap::ComputeFromUVSpaceMesh);
+		
 		// make flat mesh
 		TMeshAABBTree3<MeshType> FlatSpatial(&UVSpaceMesh, true);
 
