ModelingTools: Fix FDynamicMeshUVEditor::CreateSeamsAtEdges creating accidental bowties.

#rnx

diff --git a/Engine/Plugins/Runtime/GeometryProcessing/Source/DynamicMesh/Private/Parameterization/DynamicMeshUVEditor.cpp b/Engine/Plugins/Runtime/GeometryProcessing/Source/DynamicMesh/Private/Parameterization/DynamicMeshUVEditor.cpp
--- a/Engine/Plugins/Runtime/GeometryProcessing/Source/DynamicMesh/Private/Parameterization/DynamicMeshUVEditor.cpp
+++ b/Engine/Plugins/Runtime/GeometryProcessing/Source/DynamicMesh/Private/Parameterization/DynamicMeshUVEditor.cpp
@@ -839,36 +839,50 @@ bool FDynamicMeshUVEditor::CreateSeamsAtEdges(const TSet<int32>& EidsToMakeIntoS
 
 	for (int32 Eid : EidsToMakeIntoSeams)
 	{
-		// If the edge is already a seam, don't need to do anything for it
-		if (UVOverlay->IsSeamEdge(Eid))
-		{
-			continue;
-		}
-
 		FIndex2i EdgeVids = Mesh->GetEdgeV(Eid);
-		TArray<int32> TidsToNextSeamA, TidsToNextSeamB;
-		bool bSeamCreated = false;
-
-		// Perform the two tests before we split either vert. A vert has to get split if it has a
-		// seam attached, because splitting the other vert would create a bowtie.
-		bool VertANeedsSplitting = DoesVertHaveAnotherSeamAttached(Eid, EdgeVids.A, TidsToNextSeamA);
-		bool VertBNeedsSplitting = DoesVertHaveAnotherSeamAttached(Eid, EdgeVids.B, TidsToNextSeamB);
+		bool VertNeedsSplitting[2] = { false, false };
+		TArray<int32> VertTidsToUseForSplit[2];
 
-		// If neither absolutely has to get split, then one needs to get split anyway so that we
-		// make the edge into a seam.
-		if (!VertANeedsSplitting && !VertBNeedsSplitting)
+		if (UVOverlay->IsSeamEdge(Eid))
 		{
-			VertANeedsSplitting = true;
-			TidsToNextSeamA.SetNum((TidsToNextSeamA.Num() + 1) / 2); // Go halfway around, rounding up
-		}
+			// If this is already a seam edge, we may have made it one by splitting the opposite vertex earlier
+			// (as part of an adjacent edge split). Make sure that we don't leave a bowtie.
+			for (int i = 0; i < 2; ++i)
+			{
+				if (UVOverlay->IsBowtieInOverlay(EdgeVids[i]))
+				{
+					VertNeedsSplitting[i] = true;
 
-		if (VertANeedsSplitting)
+					// Gather up the attached verts into VertTidsToUseForSplit
+					ensure(DoesVertHaveAnotherSeamAttached(Eid, EdgeVids[i], VertTidsToUseForSplit[i]));
+				}
+			}
+		}
+		else
 		{
-			SplitEdgeVertElement(Eid, EdgeVids.A, TidsToNextSeamA);
+			// If we're not already a seam, then it's going to become one. A vert needs to get split if it
+			// has a present or future seam attached, or else it will become a bowtie if we just split the
+			// other vert.
+			for (int i = 0; i < 2; ++i)
+			{
+				VertNeedsSplitting[i] = DoesVertHaveAnotherSeamAttached(Eid, EdgeVids[i], VertTidsToUseForSplit[i]);
+			}
+			
+			// If neither absolutely has to get split, then one needs to get split anyway so that we
+			// make the edge into a seam.
+			if (!VertNeedsSplitting[0] && !VertNeedsSplitting[1])
+			{
+				VertNeedsSplitting[0] = true;
+				VertTidsToUseForSplit[0].SetNum((VertTidsToUseForSplit[0].Num() + 1) / 2); // Go halfway around, rounding up
+			}
 		}
-		if (VertBNeedsSplitting)
+
+		for (int i = 0; i < 2; ++i)
 		{
-			SplitEdgeVertElement(Eid, EdgeVids.B, TidsToNextSeamB);
+			if (VertNeedsSplitting[i])
+			{
+				SplitEdgeVertElement(Eid, EdgeVids[i], VertTidsToUseForSplit[i]);
+			}
 		}
 	}
 
