Dynamic Mesh, Intrinsic Mesh - more robust calculation of intrinsic edges.
-  add integer-based normal coordinates for correspondence with the surface mesh and unify code that helps track the relationship between meshes.
-  rename FIntrinsicEdgeFlipMesh  FSimpleIntrinsicEdgeFlipMesh.
-  new FIntrinsicEdgeFlipMesh that uses the interger-based normal coordinates to track location of intrinsic edges relative to the surface mesh.

diff --git a/Engine/Plugins/Runtime/GeometryProcessing/Source/DynamicMesh/Private/Operations/IntrinsicCorrespondenceUtils.cpp b/Engine/Plugins/Runtime/GeometryProcessing/Source/DynamicMesh/Private/Operations/IntrinsicCorrespondenceUtils.cpp
new file mode 100644
--- /dev/null
+++ b/Engine/Plugins/Runtime/GeometryProcessing/Source/DynamicMesh/Private/Operations/IntrinsicCorrespondenceUtils.cpp
@@ -0,0 +1,559 @@
+// Copyright Epic Games, Inc. All Rights Reserved.
+
+#include "Operations/IntrinsicCorrespondenceUtils.h"
+#include "Operations/MeshGeodesicSurfaceTracer.h"
+#include "Async/ParallelFor.h"
+
+using namespace UE::Geometry;
+
+/**------------------------------------------------------------------------------
+* FSurfacePoint
+*-------------------------------------------------------------------------------*/
+
+IntrinsicCorrespondenceUtils::FSurfacePoint::FSurfacePoint()
+{
+	PositionType = EPositionType::Vertex;
+	Position.VertexPosition.VID = IndexConstants::InvalidID;
+}
+
+IntrinsicCorrespondenceUtils::FSurfacePoint::FSurfacePoint(int32 VID)
+{
+	PositionType = EPositionType::Vertex;
+	Position.VertexPosition.VID = VID;
+}
+
+IntrinsicCorrespondenceUtils::FSurfacePoint::FSurfacePoint(int32 EdgeID, double Alpha)
+{
+	PositionType = EPositionType::Edge;
+	Position.EdgePosition.EdgeID = EdgeID;
+	Position.EdgePosition.Alpha  = Alpha;
+}
+
+IntrinsicCorrespondenceUtils::FSurfacePoint::FSurfacePoint(int32 TriID, const FVector3d& BaryCentrics)
+{
+	PositionType = EPositionType::Triangle;
+	Position.TriPosition.TriID = TriID;
+	Position.TriPosition.BarycentricCoords[0] = BaryCentrics[0];
+	Position.TriPosition.BarycentricCoords[1] = BaryCentrics[1];
+	Position.TriPosition.BarycentricCoords[2] = BaryCentrics[2];
+}
+
+FVector3d IntrinsicCorrespondenceUtils::AsR3Position(const IntrinsicCorrespondenceUtils::FSurfacePoint& SurfacePoint, const FDynamicMesh3& Mesh, bool& bIsValid)
+{
+	const IntrinsicCorrespondenceUtils::FSurfacePoint::FSurfacePositionUnion& Position = SurfacePoint.Position;
+	FVector3d Result(0., 0., 0.);
+	switch (SurfacePoint.PositionType)
+	{
+		case FSurfacePoint::EPositionType::Vertex:
+		{
+			const int32 SurfaceVID = Position.VertexPosition.VID;
+			bIsValid               = Mesh.IsVertex(SurfaceVID);
+			if (bIsValid)
+			{
+				Result = Mesh.GetVertex(SurfaceVID);
+			}
+		}
+		break;
+		case FSurfacePoint::EPositionType::Edge:
+		{
+			const int32 SurfaceEID = Position.EdgePosition.EdgeID;
+			const double Alpha     = Position.EdgePosition.Alpha;
+			bIsValid               = Mesh.IsEdge(SurfaceEID);
+			if (bIsValid)
+			{
+				const FIndex2i SurfaceEdgeV = Mesh.GetEdgeV(SurfaceEID);
+				Result = Alpha * (Mesh.GetVertex(SurfaceEdgeV.A)) + (1. - Alpha) * (Mesh.GetVertex(SurfaceEdgeV.B));
+			}
+		}
+		break;
+		case FSurfacePoint::EPositionType::Triangle:
+		{
+			const int32 SurfaceTID = Position.TriPosition.TriID;
+			bIsValid = Mesh.IsTriangle(SurfaceTID);
+
+			if (bIsValid)
+			{
+				Result = Mesh.GetTriBaryPoint( SurfaceTID,
+											   Position.TriPosition.BarycentricCoords[0],
+											   Position.TriPosition.BarycentricCoords[1],
+											   Position.TriPosition.BarycentricCoords[2] );
+			}
+
+		}
+		break;
+		default:
+		{
+			bIsValid = false;
+			// shouldn't be able to reach this point
+			check(0);
+		}
+	}
+	return Result;
+}
+
+/**------------------------------------------------------------------------------
+* FMeshConnection
+*-------------------------------------------------------------------------------*/
+
+IntrinsicCorrespondenceUtils::FMeshConnection::FMeshConnection(const FDynamicMesh3& SurfMesh)
+	:SurfaceMesh(&SurfMesh)
+{
+	// pick local reference directions for each vertex and triangle.
+	const int32 MaxVertexID   = SurfMesh.MaxVertexID();
+	const int32 MaxTriangleID = SurfMesh.MaxTriangleID();
+
+	{
+		VIDToReferenceEID.AddUninitialized(MaxVertexID);
+		TIDToReferenceEID.AddUninitialized(MaxTriangleID);
+
+		for (int32 VID = 0; VID < MaxVertexID; ++VID)
+		{
+			if (SurfMesh.IsVertex(VID))
+			{
+				int32 RefEID = -1;
+				if (SurfMesh.IsBoundaryVertex(VID) && !SurfMesh.IsBowtieVertex(VID))
+				{
+					// vertex is on the mesh boundary, 
+					// choose the boundary edge such that traveling CCW (about the vertex) moves into the mesh
+					for (int32 EID : SurfMesh.VtxEdgesItr(VID))
+					{
+						if (SurfMesh.IsBoundaryEdge(EID))
+						{
+							const int32 TID     = SurfMesh.GetEdgeT(EID).A;
+							const int32 IndexOf = SurfMesh.GetTriEdges(TID).IndexOf(EID);
+							if (SurfMesh.GetTriangle(TID)[IndexOf] == VID)
+							{
+								RefEID = EID;
+								break;
+							}
+						}
+					}
+				}
+				else
+				{
+					for (int32 EID : SurfMesh.VtxEdgesItr(VID))
+					{
+						RefEID = EID;
+						break;
+					}
+				}
+				VIDToReferenceEID[VID] = RefEID;
+			}
+		}
+		for (int32 TID = 0; TID < MaxTriangleID; ++TID)
+		{
+			if (SurfMesh.IsTriangle(TID))
+			{
+				TIDToReferenceEID[TID] = SurfMesh.GetTriEdge(TID, 0);
+			}
+		}
+	}
+}
+
+
+/**------------------------------------------------------------------------------
+* FNormalCoordinates
+*-------------------------------------------------------------------------------*/
+
+IntrinsicCorrespondenceUtils::FNormalCoordinates::FNormalCoordinates(const FDynamicMesh3& SurfMesh)
+	: MyBase(SurfMesh)
+{
+	const int32 MaxEID = SurfMesh.MaxEdgeID();
+	const int32 MaxVID = SurfMesh.MaxVertexID();
+	const int32 MaxTID = SurfMesh.MaxTriangleID();
+
+	// the intrinsic mesh on construction will have the same triangles as the surface mesh,
+	// so no surface mesh edges will cross intrinsic edges.
+	NormalCoord.SetNumZeroed(MaxEID);
+
+	// populate Ref Vert Degree. 
+	RefVertDegree.SetNumZeroed(MaxVID);
+	for (int32 VID : SurfMesh.VertexIndicesItr())
+	{
+		int32 Degree = 0;
+		for (int32 EID : SurfMesh.VtxEdgesItr(VID))
+		{
+			Degree++;
+		}
+		RefVertDegree[VID] = Degree;
+	}
+
+	// init the roundabout data as invalid
+	RoundaboutOrder.SetNumUninitialized(MaxTID);
+	for (FIndex3i& PerTri : RoundaboutOrder)
+	{
+		PerTri = FIndex3i(FDynamicMesh3::InvalidID, FDynamicMesh3::InvalidID, FDynamicMesh3::InvalidID);
+	}
+
+	// correctly populate the roundabout data.  Note bow-tie vertices should prohibited at a higher level
+	for (int32 CVID : SurfMesh.VertexIndicesItr())
+	{
+
+		int32 Order = 0;
+		auto SetRoundaboutOrder = [&Order, this](int32 TID, int32 EID, int32 IdxOf)->bool
+									{
+										RoundaboutOrder[TID][IdxOf] = Order;
+										Order++;
+										constexpr bool bShouldBreak = false;
+										return bShouldBreak;
+									};
+		VisitVertexAdjacentElements(SurfMesh, CVID, VIDToReferenceEID[CVID], SetRoundaboutOrder);
+	}
+
+	// Edge Order is a copy of the roundabout before edge flips.  Edge order describes the surface mesh, and shouldn't be updated after construction
+	EdgeOrder = RoundaboutOrder;
+}
+
+bool  IntrinsicCorrespondenceUtils::FNormalCoordinates::OnFlipEdge(const int32 T0ID, const FIndex3i& BeforeT0EIDs, const int32 T0OppVID,
+																   const int32 T1ID, const FIndex3i& BeforeT1EIDs, const int32 T1OppVID,
+																   const int32 FlippedEID)
+{
+	// see Gillespi et al, eqn 15
+
+	// prior to flip of edge ij to become edge lk:
+	// TO triangle is a cyclic permutation of (i,j,k)  and likewise T1 is of (j,i,l)
+	// with this notation T0OppVID is the vertex at 'k' corner, and T1OppVID is the vertex at 'l' corner
+
+	// index of the edge to be flipped
+	const int32 T0IndexOf = BeforeT0EIDs.IndexOf(FlippedEID); // exits i in t0
+	const int32 T1IndexOf = BeforeT1EIDs.IndexOf(FlippedEID); // exits j in t1
+
+	checkSlow(FlippedEID < NormalCoord.Num());
+	if (T0IndexOf == FDynamicMesh3::InvalidID || T1IndexOf == FDynamicMesh3::InvalidID)
+	{
+		checkSlow(0);
+		return false;
+	}
+
+
+	// the corners of the two triangles before flip (alternately the edges)
+	const int32 T0_i = T0IndexOf;                // t0:eij
+	const int32 T0_j = (T0IndexOf + 1) % 3;      // t0:ejk
+	const int32 T0_k = (T0IndexOf + 2) % 3;      // t0:eki
+
+	const int32 T1_j = T1IndexOf;                // t1:eji
+	const int32 T1_i = (T1IndexOf + 1) % 3;      // t1:eil
+	const int32 T1_l = (T1IndexOf + 2) % 3;      // t1:elj
+
+	// ---- update the NormalCoord for the flipped edge
+	{
+		// corner crossings
+		const int32 Cji_l = NumCornerCrossingRefEdges(BeforeT1EIDs, T1_l);
+		const int32 Cij_k = NumCornerCrossingRefEdges(BeforeT0EIDs, T0_k);
+
+		const int32 Cil_j = NumCornerCrossingRefEdges(BeforeT1EIDs, T1_j);
+		const int32 Cki_j = NumCornerCrossingRefEdges(BeforeT0EIDs, T0_j);
+
+		const int32 Clj_i = NumCornerCrossingRefEdges(BeforeT1EIDs, T1_i);
+		const int32 Cjk_i = NumCornerCrossingRefEdges(BeforeT0EIDs, T0_i);
+
+		// corner emanating
+		const int32 Eji_l = NumCornerEmanatingRefEdges(BeforeT1EIDs, T1_l);
+		const int32 Eij_k = NumCornerEmanatingRefEdges(BeforeT0EIDs, T0_k);
+
+		const int32 Elj_i = NumCornerEmanatingRefEdges(BeforeT1EIDs, T1_i);
+		const int32 Ejk_i = NumCornerEmanatingRefEdges(BeforeT0EIDs, T0_i);
+
+		const int32 Eil_j = NumCornerEmanatingRefEdges(BeforeT1EIDs, T1_j);
+		const int32 Eki_j = NumCornerEmanatingRefEdges(BeforeT0EIDs, T0_j);
+
+		// if the edge being flipped aligns with a surface edge, then it will cross that surface edge after the  flip
+		const int32 Kronecker = (NormalCoord[FlippedEID] == 0) ? 1 : 0;
+
+		// total number of surfaces edges this edge crosses after the flip
+		const int32 N_kl = TMathUtil<int32>::Max(0, Cji_l + Cij_k + (TMathUtil<int32>::Abs(Cil_j - Cki_j) + TMathUtil<int32>::Abs(Clj_i - Cjk_i) - Eji_l - Eij_k) / 2
+			+ Elj_i + Ejk_i + Eil_j + Eki_j + Kronecker);
+
+		NormalCoord[FlippedEID] = N_kl;
+	}
+
+	// ---- update the roundabouts for the two directions of the flipped edge. Note this must follow the normal coord update as it relies on those results
+	{
+		// after edge-flip the T0 and T1 vertices and edges are re-ordered so the flipped edge is the first edge. The new order is given by
+		//             T0 edges =  ekl, elj, ejk
+		//             T1 edges =  elk, eki, eil
+		const int32 Edge_lj = BeforeT1EIDs[T1_l];
+		const int32 Edge_jk = BeforeT0EIDs[T0_j];
+		const int32 Edge_ki = BeforeT0EIDs[T0_k];
+		const int32 Edge_il = BeforeT1EIDs[T1_i];
+
+		const FIndex3i AfterT0EIDs(FlippedEID, Edge_lj, Edge_jk);
+		const FIndex3i AfterT1EIDs(FlippedEID, Edge_ki, Edge_il);
+
+		//  Roundabouts before flip ( respecting old edge order)
+		const FIndex3i T0Rab = RoundaboutOrder[T0ID];
+		const FIndex3i T1Rab = RoundaboutOrder[T1ID];
+
+		const int32 VID_k = T0OppVID;
+		const int32 VID_l = T1OppVID;
+
+		// the roundabouts for the outer half-edges (the edges that don't flip)
+		const int32 R_ki = T0Rab[T0_k];
+		const int32 R_jk = T0Rab[T0_j];
+		const int32 R_il = T1Rab[T1_i];
+		const int32 R_lj = T1Rab[T1_l];
+
+		// --compute the roundabouts for the two half-edges associated with the edge that does flip. (i.e  R_kl, R_lk)
+
+		const int32 Kronecker_ki = (NormalCoord[Edge_ki] == 0) ? 1 : 0;
+		const int32 Kronecker_lj = (NormalCoord[Edge_lj] == 0) ? 1 : 0;
+
+		const int32 Eil_k = NumCornerEmanatingRefEdges(AfterT1EIDs, 1/*k corner*/); // AfterT1 = (l, k, i)
+		const int32 Ejk_l = NumCornerEmanatingRefEdges(AfterT0EIDs, 1/*l corner*/); // AfterT0 = (k, l, j)
+
+		// Roundabout values for the two directions of the flipped edge ( i.e. as seen by the new T0 and T1)
+		const int32 R_kl = (R_ki + Eil_k + Kronecker_ki) % RefVertDegree[VID_k];
+		const int32 R_lk = (R_lj + Ejk_l + Kronecker_lj) % RefVertDegree[VID_l];
+
+
+		// update the roundabouts with new edge order and values for the flipped edge
+		RoundaboutOrder[T0ID] = FIndex3i(R_kl, R_lj, R_jk);
+		RoundaboutOrder[T1ID] = FIndex3i(R_lk, R_ki, R_il);
+	}
+
+	return true;
+}
+
+
+int32 IntrinsicCorrespondenceUtils::FNormalCoordinates::GetNthEdgeID(const int32 VID, const int32 N) const
+{
+
+	if (!SurfaceMesh->IsVertex(VID))
+	{
+		return -1;
+	}
+	const int32 Degree = RefVertDegree[VID];
+	if (N - 1 > Degree)
+	{
+		return -1;
+	}
+
+	const int32 RefEID = VIDToReferenceEID[VID];
+	int32 Count        = 0;
+	int32 ResultEID    = -1;
+	auto EdgeFinder = [N, &Count, &ResultEID](int32 TID, int32 EID, int32 IdxOf)->bool
+						{
+							bool bShouldBreak = false;
+							if (N == Count)
+							{
+								ResultEID = EID;
+								bShouldBreak = true;
+							}
+							Count++;
+							return bShouldBreak;
+						};
+	VisitVertexAdjacentElements(*SurfaceMesh, VID, RefEID, EdgeFinder);
+
+	return ResultEID;
+}
+
+
+
+int32 IntrinsicCorrespondenceUtils::FNormalCoordinates::GetEdgeOrder(const int32 VID, const int32 AdjEID) const
+{
+	if (!SurfaceMesh->IsVertex(VID) || !SurfaceMesh->IsEdge(AdjEID))
+	{
+		return -1;
+	}
+
+	FIndex2i EdgeV = SurfaceMesh->GetEdgeV(AdjEID);
+	if (EdgeV.A != VID && EdgeV.B != VID)
+	{
+		return -1;
+	}
+
+	FIndex2i EdgeT   = SurfaceMesh->GetEdgeT(AdjEID);
+	int32 TID        = EdgeT.A;
+	FIndex3i TriEIDs = SurfaceMesh->GetTriEdges(TID);
+	int32 IndexOf    = TriEIDs.IndexOf(AdjEID);
+
+	if (SurfaceMesh->GetTriangle(TID)[IndexOf] != VID)
+	{
+		TID = EdgeT.B;
+		if (TID != -1)
+		{
+			TriEIDs = SurfaceMesh->GetTriEdges(TID);
+			IndexOf = TriEIDs.IndexOf(AdjEID);
+			checkSlow(SurfaceMesh->GetTriangle(TID)[IndexOf] == VID);
+			return EdgeOrder[TID][IndexOf];
+		}
+		else
+		{
+			return RefVertDegree[VID] - 1; // we don't have an out-going edge from the vertex on this boundary, but we know it is the last edge.
+		}
+	}
+
+	return EdgeOrder[TID][IndexOf];
+}
+
+/**------------------------------------------------------------------------------
+* FSignpost
+*-------------------------------------------------------------------------------*/
+
+IntrinsicCorrespondenceUtils::FSignpost::FSignpost(const FDynamicMesh3& SrcMesh)
+	:MyBase(SrcMesh)
+{
+	// During construction the mesh is a deep copy of the SrcMesh ( i.e. has the same IDs)
+
+	const int32 MaxVertexID = SrcMesh.MaxVertexID();
+	const int32 MaxTriangleID = SrcMesh.MaxTriangleID();
+
+
+	// add surface position.
+	{
+		IntrinsicVertexPositions.SetNum(MaxVertexID);
+
+		// initialize: identify with vertex in Extrinsic Mesh
+		for (int32 VID = 0; VID < MaxVertexID; ++VID)
+		{
+			if (SrcMesh.IsVertex(VID))
+			{
+				FSurfacePoint VertexSurfacePoint(VID);
+				IntrinsicVertexPositions[VID] = VertexSurfacePoint;
+			}
+		}
+	}
+
+	// add edge directions and geometric info
+	{
+		const int32 MaxExtEdgeID = SrcMesh.MaxEdgeID();
+		IntrinsicEdgeAngles.SetNum(MaxTriangleID);
+
+		GeometricVertexInfo.SetNum(MaxVertexID);
+		IntrinsicVertexPositions.SetNum(MaxVertexID);
+
+		ParallelFor(MaxVertexID, [this, &SrcMesh](int32 VID)
+		{
+			if (!SrcMesh.IsVertex(VID))
+			{
+				return;
+			}
+
+			constexpr int32 PlusTwoModThree[3] = { 2, 0, 1 }; // PlusTwoModThree[i] = (i + 2)%3
+
+			// initialize by computing the angles
+			TArray<int32> Triangles;
+			TArray<int32> ContiguousGroupLengths;
+			TArray<bool> GroupIsLoop;
+
+			{
+
+				FGeometricInfo GeometricInfo;
+				// note: since the structure and IDs are a deep copy of the source mesh, we can use it to find the contiguous triangles
+				const EMeshResult Result = SrcMesh.GetVtxContiguousTriangles(VID, Triangles, ContiguousGroupLengths, GroupIsLoop);
+				if (Result == EMeshResult::Ok && GroupIsLoop.Num() == 1)
+				{
+					GeometricInfo.bIsInterior = (GroupIsLoop[0] == true);
+					// Contiguous tris could be cw or ccw?  Need them to  be ccw
+					if (Triangles.Num() > 1)
+					{
+						const int32 TID0     = Triangles[0];
+						FIndex3i TriVIDs     = SrcMesh.GetTriangle(TID0);
+						FIndex3i TriEIDs     = SrcMesh.GetTriEdges(TID0);
+						int32 VertSubIdx     = TriVIDs.IndexOf(VID);
+						const int32 EID      = TriEIDs[VertSubIdx];
+						const int32 NextEID  = TriEIDs[PlusTwoModThree[VertSubIdx]];
+						const int32 TID1     = Triangles[1];
+						TriVIDs              = SrcMesh.GetTriangle(TID1);
+						TriEIDs              = SrcMesh.GetTriEdges(TID1);
+						VertSubIdx           = TriVIDs.IndexOf(VID);
+
+						if (TriEIDs[VertSubIdx] != NextEID)
+						{
+							// need to reverse order to have CCW
+							Algo::Reverse(Triangles);
+						}
+					}
+
+					const int32 NumEdges = ContiguousGroupLengths[0];
+
+					double AngleOffset       = 0;
+					double WalkedAngle       = 0;
+					const int32 ReferenceEID = VIDToReferenceEID[VID];
+					checkSlow(ReferenceEID != FDynamicMesh3::InvalidID)
+
+					TMap<int32, int32> VisitedTriSubID;
+					for (int32 TID : Triangles)
+					{
+						const FIndex3i TriVIDs = SrcMesh.GetTriangle(TID);
+						const FIndex3i TriEIDs = SrcMesh.GetTriEdges(TID);
+						const int32 VertSubIdx = TriVIDs.IndexOf(VID);
+						IntrinsicEdgeAngles[TID][VertSubIdx] = WalkedAngle;
+
+						VisitedTriSubID.Add(TID, VertSubIdx);
+						const int32 EID = TriEIDs[VertSubIdx];
+						if (EID == ReferenceEID)
+						{
+							AngleOffset = WalkedAngle;
+						}
+
+						WalkedAngle += SrcMesh.GetTriInternalAngleR(TID, VertSubIdx);
+					}
+					// need to include last edge if the one ring of triangles isn't a closed loop, it could be the reference edge.
+					if (GroupIsLoop[0] == false)
+					{
+						const int32 TID        = Triangles.Last();
+						const FIndex3i TriVIDs = SrcMesh.GetTriangle(TID);
+						const FIndex3i TriEIDs = SrcMesh.GetTriEdges(TID);
+						const int32 VertSubIdx = TriVIDs.IndexOf(VID);
+						const int32 EID        = TriEIDs[PlusTwoModThree[VertSubIdx]];
+						if (EID == ReferenceEID)
+						{
+							AngleOffset = WalkedAngle;
+						}
+					}
+
+					// orient with the reference edge at zero (or pi) angle  and rescale to radians and use 0,2pi to define polar angle.
+					const double ToRadians = TMathUtilConstants<double>::TwoPi / WalkedAngle;
+					GeometricInfo.ToRadians = ToRadians;
+					double RefAngle = (SrcMesh.GetEdgeV(ReferenceEID).B == VID) ? TMathUtilConstants<double>::Pi : 0.;
+					for (TPair<int32, int32>& TriSubIDPair : VisitedTriSubID)
+					{
+						int32 TID     = TriSubIDPair.Key;
+						int32 SubID   = TriSubIDPair.Value;
+						double& Angle = IntrinsicEdgeAngles[TID][SubID];
+						Angle         = ToRadians * (Angle - AngleOffset) + RefAngle;
+						Angle         = AsZeroToTwoPi(Angle);
+					}
+				}
+				else // bow-tie or boundary
+				{
+					GeometricInfo.bIsInterior = false;
+				}
+
+				GeometricVertexInfo[VID] = GeometricInfo;
+			}
+		}, false /* bForceSingleThread*/);
+	}
+}
+
+void IntrinsicCorrespondenceUtils::FSignpost::OnFlipEdge(const int32 EID, const FIndex2i Tris, const FIndex2i OpposingVerts, const FIndex2i PreFlipIndexOf, double NewAngleAtOpp0, const double NewAngleAtOpp1)
+{
+	constexpr  int32 AddOneModThree[3] = { 1, 2, 0 };
+	constexpr  int32 AddTwoModThree[3] = { 2, 0, 1 };
+	const FVector3d PreFlipDirections[2] = { IntrinsicEdgeAngles[Tris[0]], IntrinsicEdgeAngles[Tris[1]] };
+
+	// polar directions for the four edges that don't change.
+	const double bcDir = PreFlipDirections[0][AddOneModThree[PreFlipIndexOf[0]]];
+	const double caDir = PreFlipDirections[0][AddTwoModThree[PreFlipIndexOf[0]]];
+
+	const double adDir = PreFlipDirections[1][AddOneModThree[PreFlipIndexOf[1]]];
+	const double dbDir = PreFlipDirections[1][AddTwoModThree[PreFlipIndexOf[1]]];
+
+
+
+	// compute the polar edge directions for the new edge (ie. dcDir and cdDir)
+	// dcDir 
+	const double newAngleAtD = NewAngleAtOpp1;
+	const double ToPolarAtD  = GeometricVertexInfo[OpposingVerts[1]].ToRadians;
+	const double dcDir       = AsZeroToTwoPi(dbDir + ToPolarAtD * newAngleAtD);
+
+	// cdDir 
+	const double newAngleAtC = NewAngleAtOpp0;
+	const double ToPolarAtC  = GeometricVertexInfo[OpposingVerts[0]].ToRadians;
+	const double cdDir       = AsZeroToTwoPi(caDir + ToPolarAtC * newAngleAtC);
+
+	IntrinsicEdgeAngles[Tris[0]] = FVector3d(cdDir, dbDir, bcDir); // edges leaving c, d, b
+	IntrinsicEdgeAngles[Tris[1]] = FVector3d(dcDir, caDir, adDir); // edges leaving d, c, a
+}
+
diff --git a/Engine/Plugins/Runtime/GeometryProcessing/Source/DynamicMesh/Private/Operations/IntrinsicTriangulationMesh.cpp b/Engine/Plugins/Runtime/GeometryProcessing/Source/DynamicMesh/Private/Operations/IntrinsicTriangulationMesh.cpp
--- a/Engine/Plugins/Runtime/GeometryProcessing/Source/DynamicMesh/Private/Operations/IntrinsicTriangulationMesh.cpp
+++ b/Engine/Plugins/Runtime/GeometryProcessing/Source/DynamicMesh/Private/Operations/IntrinsicTriangulationMesh.cpp
@@ -362,25 +362,746 @@ namespace
 
 };
 
+namespace SignpostSufaceTraceUtil
+{
+	using namespace IntrinsicCorrespondenceUtils;
+	/**
+	* helper code that uses FSignpost to trace from a specified intrinsic mesh vertex in a given direction.
+	* 
+	*Note: this doesn't check the validity of the start point - the calling code should do that
+	*/
+	UE::Geometry::FMeshGeodesicSurfaceTracer TraceFromIntrinsicVert( const FSignpost& SignpostData, const int32 TraceStartVID, 
+		                                                             const double TracePolarAngle, const double TraceDist )
+	{
+		const FDynamicMesh3* SurfaceMesh                            = SignpostData.SurfaceMesh; 
+		const FSurfacePoint& StartSurfacePoint                      = SignpostData.IntrinsicVertexPositions[TraceStartVID];
+		const FSurfacePoint::FSurfacePositionUnion& SurfacePosition = StartSurfacePoint.Position;
+
+		double ActualDist = 0.;
+		UE::Geometry::FMeshGeodesicSurfaceTracer SurfaceTracer(*SurfaceMesh);
+		switch (StartSurfacePoint.PositionType)
+		{
+		case FSurfacePoint::EPositionType::Vertex:
+		{
+			// convert vertex location
+			const int32 ExtrinsicVID    = StartSurfacePoint.Position.VertexPosition.VID;
+			const int32 ExtrinsicRefEID = SignpostData.VIDToReferenceEID[ExtrinsicVID];
+
+			const FMeshGeodesicSurfaceTracer::FMeshTangentDirection TangentDirection = { ExtrinsicVID, ExtrinsicRefEID, TracePolarAngle };
+			ActualDist = SurfaceTracer.TraceMeshFromVertex(TangentDirection, TraceDist);
+		}
+		break;
+		case FSurfacePoint::EPositionType::Edge:
+		{
+			// convert edge location -  Not used, and might have bugs
+			const int32 RefSurfaceEID   = SurfacePosition.EdgePosition.EdgeID;
+			double Alpha                = SurfacePosition.EdgePosition.Alpha;
+			// convert to BC
+			const int32 SurfaceTID      = SurfaceMesh->GetEdgeT(RefSurfaceEID).A;
+			const FIndex2i SurfaceEdgeV = SurfaceMesh->GetEdgeV(RefSurfaceEID);
+			const int32 IndexOf         = SurfaceMesh->GetTriEdges(SurfaceTID).IndexOf(RefSurfaceEID);
+			const bool bSameOrientation = (SurfaceMesh->GetTriangle(SurfaceTID)[IndexOf] == SurfaceEdgeV.A);
+			FVector3d BaryPoint(0., 0., 0.);
+			if (!bSameOrientation)
+			{
+				Alpha = 1. - Alpha;
+			}
+			BaryPoint[IndexOf] = Alpha;
+			BaryPoint[(IndexOf + 1) % 3] = 1. - Alpha;
+
+			FMeshGeodesicSurfaceTracer::FMeshSurfaceDirection DirectionOnSurface(RefSurfaceEID, TracePolarAngle);
+			ActualDist = SurfaceTracer.TraceMeshFromBaryPoint(SurfaceTID, BaryPoint, DirectionOnSurface, TraceDist);
+		}
+		break;
+		case FSurfacePoint::EPositionType::Triangle:
+		{
+			// trace from BC point in triangle
+			const int32 SurfaceTID = SurfacePosition.TriPosition.TriID;
+
+			const FVector3d BaryPoint( SurfacePosition.TriPosition.BarycentricCoords[0],
+									   SurfacePosition.TriPosition.BarycentricCoords[1],
+				                       SurfacePosition.TriPosition.BarycentricCoords[2] );
+			const int32 RefSurfaceEID = SignpostData.TIDToReferenceEID[SurfaceTID];
+
+			FMeshGeodesicSurfaceTracer::FMeshSurfaceDirection DirectionOnSurface(RefSurfaceEID, TracePolarAngle);
+			SurfaceTracer.TraceMeshFromBaryPoint(SurfaceTID, BaryPoint, DirectionOnSurface, TraceDist);
+		}
+		break;
+		default:
+		{
+			// not possible
+			check(0);
+		}
+		}
+
+		// RVO..
+		return SurfaceTracer;
+	}
+};
+
+namespace FNormalCoordSurfaceTraceImpl
+{
+	using namespace IntrinsicCorrespondenceUtils;
+	/**
+	* low-level code that uses normal coordinates to continues tracing a surface edge across an FIntrinsicEdgeFlipMesh.  
+	*
+	* This function is intended to be called directly after the first intrinsic edge crossing 
+	* at which time the 'Crossings' result array will hold the start vertex and the first edge crossing.
+	* 
+	* This function will populate the rest of the sequence of edge crossings for the surface mesh edge
+	* as a list of the edges crossed, and not the location on the individual edges.  
+	* To compute the those locations the resulting triangle strip should be unfolded and each location can be solved for
+	* using TraceEdgeOverHost below.
+	*/
+	void  ContinueTraceSurfaceEdge( const FIntrinsicEdgeFlipMesh& IntrinsicMesh, const int32 StartTID, const int32 StartEID, const int32 StartP,
+		                            TArray<FNormalCoordinates::FEdgeAndCrossingIdx>& Crossings )
+	{
+		const FNormalCoordinates& NCoords = IntrinsicMesh.GetNormalCoordinates();
+		if (Crossings.Num() == 0)
+		{
+			return;
+		}
+
+		if (!IntrinsicMesh.IsTriangle(StartTID))
+		{
+			return;
+		}
+
+		const FIndex3i StartTriEIDs = IntrinsicMesh.GetTriEdges(StartTID);
+		checkSlow(StartTriEIDs.IndexOf(StartEID) != -1);
+
+		const int32 NumXStartEID = NCoords.NormalCoord[StartEID];
+		if (StartP > NumXStartEID || StartP < 1) // note the actual permitted range of P is smaller.. 
+		{
+			checkSlow(0); // this shouldn't happen
+			return;
+		}
+
+		// This is adapted from Algorithm 2 of  Gillespi et al, 2020
+		auto GetNextCrossing = [&IntrinsicMesh, &NCoords](int32& TriID, int32& EdgeID, int32& P)
+		{
+			// get next tri
+			TriID = [&]
+				{
+					const FIndex2i EdgeT = IntrinsicMesh.GetEdgeT(EdgeID);
+					return (EdgeT.A == TriID) ? EdgeT.B : EdgeT.A;
+				}();
+			checkSlow(TriID != -1);
+
+			const FIndex3i TriEIDs = IntrinsicMesh.GetTriEdges(TriID);
+			const int32 IndexOf    = TriEIDs.IndexOf(EdgeID);
+
+			const int32 Edge_ji = EdgeID;
+			const int32 Edge_il = TriEIDs[(IndexOf + 1) % 3];
+			const int32 Edge_lj = TriEIDs[(IndexOf + 2) % 3];
+
+			const int32 N_ji = NCoords.NormalCoord[Edge_ji];
+			const int32 N_il = NCoords.NormalCoord[Edge_il];
+			const int32 N_lj = NCoords.NormalCoord[Edge_lj];
+
+			const int32 Pin = P;
+
+			if (N_ji > N_lj + N_il)  // case 1
+			{
+				// some of the N_ji edges that cross edge IJ must connect with the vertex at L 
+				if (P <= N_il)
+				{
+					// exits side il
+					EdgeID = Edge_il;
+					P      = Pin;
+				}
+				else if ((N_il < P) && (P <= N_ji - N_lj))
+				{
+					// this path terminates at vertex
+					// this is a slight abuse.  The convention, when CrossingIndex = 0, we specify the vertex by the edge that originates there.
+					EdgeID = Edge_lj; // GetTrianlge()[EdgeID_lj] will be the vertex.
+					P      = 0;
+				}
+				else
+				{
+					EdgeID = Edge_lj;
+					P      = Pin - (N_ji - N_lj);
+				}
+
+			}
+			else if (N_lj > N_ji + N_il)  // case 2
+			{
+				EdgeID = Edge_lj;
+				P      = Pin + (N_lj - N_ji);
+
+			}
+			else if (N_il > N_ji + N_lj) // case 3
+			{
+				EdgeID = Edge_il;
+				P      = Pin;
+
+			}
+			else  // case 4
+			{
+				const int32 Cij_l = (N_lj + N_il - N_ji) / 2;
+				if (P <= N_il - Cij_l)
+				{
+					EdgeID = Edge_il;
+					P      = Pin;
+
+				}
+				else
+				{
+					const int32 Clj_i = (N_il + N_ji - N_lj) / 2;
+					EdgeID = Edge_lj;
+					P      = Pin - Clj_i + Cij_l;
+				}
+			}
+		};
+
+		int32 P = StartP;
+		int32 EID = StartEID;
+		int32 TID = StartTID;
+		// record first crossing
+		auto& FirstXing = Crossings.AddZeroed_GetRef();
+		FirstXing.TID  = TID;
+		FirstXing.EID  = EID;
+		FirstXing.CIdx = P;
+		do
+		{
+			GetNextCrossing(TID, EID, P);
+			auto& Xing = Crossings.AddZeroed_GetRef();
+			Xing.TID   = TID;
+			Xing.EID   = EID;
+			Xing.CIdx  = P;
+			checkSlow(P > -1);
+		} while (P != 0); // P = 0 when the path terminates ( paths terminate at vertices only)
+
+
+	}
+
+	// trace the surface edge - note a surface edge along some surface mesh boundaries can't be specified by this API, but those edges will
+	// be the same in the intrinsic mesh since they can't flip.
+	TArray<FNormalCoordinates::FEdgeAndCrossingIdx> TraceSurfaceEdge( const FIntrinsicEdgeFlipMesh& IntrinsicMesh,
+		                                                              const int32 SurfaceTID, const int32 IndexOf )
+	{
+		const FNormalCoordinates& NCoords = IntrinsicMesh.GetNormalCoordinates();
+		const FDynamicMesh3& SurfaceMesh  = *NCoords.SurfaceMesh;
+
+		if (!SurfaceMesh.IsTriangle(SurfaceTID) || IndexOf > 2 || IndexOf < 0)
+		{
+			TArray< FNormalCoordinates::FEdgeAndCrossingIdx > EmptyCrossings;
+			return EmptyCrossings;
+		}
+
+
+		// the origin vertex
+		const int32 StartVID = SurfaceMesh.GetTriangle(SurfaceTID)[IndexOf];
+		// the surface edge we are tracing
+		const int32 TraceEID = SurfaceMesh.GetTriEdge(SurfaceTID, IndexOf);
+
+		checkSlow(!SurfaceMesh.IsBoundaryEdge(TraceEID));
+		const int32 RefEID = NCoords.VIDToReferenceEID[StartVID];
+		const int32 OrderOfTraceEID = NCoords.GetEdgeOrder(StartVID, TraceEID);
+		const int32 ValenceOfStartVID = NCoords.RefVertDegree[StartVID];
+		checkSlow(OrderOfTraceEID != -1);
+
+		// data to gather about the intrinsic triangle that where the trace starts.
+		struct
+		{
+			bool bIsAlsoSurfaceEdge = false;
+			int32 TID               = -1;
+			int32 EID               = -1;
+			int32 IdxOf             = -1;
+			int32 FirstRoundabout   = -1;
+			int32 SecondRoundabout  = -1;
+		} FinderInfo;
+
+		// Identify the target intrinsic triangle where this edge trace starts.
+		{
+
+			// where to start when visiting the intrinsic triangles that are adjacent to the startVID
+			const int32 VistorStartEID =[&] 
+				{	
+					if (SurfaceMesh.IsBoundaryVertex(StartVID))
+					{
+						// due to the way the intrinsic mesh is constructed we know the boundary edges will have the same EID in both meshes.
+						// recall, these can't flip, so they are fixed in the intrinsic mesh.
+						return RefEID;
+					}
+					else
+					{
+						for (int32 NbrEID : IntrinsicMesh.VtxEdgesItr(StartVID))
+						{
+							// get the first
+							return NbrEID;
+						}
+						return -1; // shouldn't reach here
+					}
+				}();
+
+			// when visiting each adjacent triangle (in CCW order) test if the edge we trace is inside this triangle.
+			auto EdgeFinder = [&](int32 IntrinsicTID, int32 IntrinsicEID, int32 IdxOf)->bool
+			{
+
+
+				const FIndex3i IntrinsicTriEIDs = IntrinsicMesh.GetTriEdges(IntrinsicTID);
+				const int32 ThisRoundabout      = NCoords.RoundaboutOrder[IntrinsicTID][IdxOf];
+				const bool bEquivalentToSurface = (NCoords.NormalCoord[IntrinsicEID] == 0) && (ThisRoundabout == OrderOfTraceEID);
+				int32 NextRoundabout = -1;
+
+				bool bShouldBreak = false;
+
+				if (bEquivalentToSurface) // test if the current intrinsic edge is the same as the surface mesh trace edge
+				{
+					bShouldBreak = true;
+				}
+				else // test if the edge starts in this intrinsic triangle
+				{
+					const int32 NextIntrinsicEID      = IntrinsicTriEIDs[(IdxOf + 2) % 3];
+					const FIndex2i NextIntrinsicEdgeT = IntrinsicMesh.GetEdgeT(NextIntrinsicEID);
+					const int32 NextIntrinsicTID      = (NextIntrinsicEdgeT.A == IntrinsicTID) ? NextIntrinsicEdgeT.B : NextIntrinsicEdgeT.A;
+
+					if (NextIntrinsicTID != -1)
+					{
+						const int32 NextIndexOf = IntrinsicMesh.GetTriEdges(NextIntrinsicTID).IndexOf(NextIntrinsicEID);
+						NextRoundabout          = NCoords.RoundaboutOrder[NextIntrinsicTID][NextIndexOf];
+						if (NextRoundabout < ThisRoundabout)
+						{
+							// we have crossed the zero roundabout cut
+							NextRoundabout += ValenceOfStartVID;
+						}
+						checkSlow(IntrinsicMesh.GetTriangle(NextIntrinsicTID)[NextIndexOf] == StartVID);
+						//checkSlow(NextRoundabout != 0 || ThisRoundabout != 0)
+						if ((NextRoundabout > OrderOfTraceEID) && (OrderOfTraceEID >= ThisRoundabout))
+						{
+							bShouldBreak = true;
+						}
+					}
+					else // mesh boundary case and we made it to the last triangle w/o finding this edge. it must be in this one.
+					{
+						bShouldBreak = true;
+					}
+				}
+
+				if (bShouldBreak)
+				{
+					FinderInfo.bIsAlsoSurfaceEdge = bEquivalentToSurface;
+					FinderInfo.TID                = IntrinsicTID;
+					FinderInfo.EID                = IntrinsicEID;
+					FinderInfo.IdxOf              = IdxOf;
+					FinderInfo.FirstRoundabout    = ThisRoundabout;
+					FinderInfo.SecondRoundabout   = NextRoundabout;
+				}
+
+				return bShouldBreak;
+			};
+			VisitVertexAdjacentElements(IntrinsicMesh, StartVID, VistorStartEID, EdgeFinder);
+
+			checkSlow(FinderInfo.TID != -1); // should have found something!
+		}
+
+
+		TArray<FNormalCoordinates::FEdgeAndCrossingIdx > Crossings;
+		// add start vertex location
+		auto& StartXing = Crossings.AddZeroed_GetRef();
+		StartXing.TID   = FinderInfo.TID;
+		StartXing.EID   = FinderInfo.IdxOf; // TriVIDs(IdxOf) = StartVID
+		StartXing.CIdx  = 0;
+		if (FinderInfo.bIsAlsoSurfaceEdge)
+		{
+			// only need to add end vertex location
+			auto& EndXing = Crossings.AddZeroed_GetRef();
+			EndXing.TID   = FinderInfo.TID;
+			EndXing.EID   = (FinderInfo.IdxOf + 1) % 3;  // TriVIDs( (IdxOf +1)%3) = EndVID
+			EndXing.CIdx  = 0;
+		}
+		else
+		{
+			// edge trace starts at StartVID and exits the opposite edge of intrinsic Tri TID (but it isn't an edge of the intrinsic tri).
+			const FIndex3i TriEIDs  = IntrinsicMesh.GetTriEdges(FinderInfo.TID);
+			const int32 OppEID      = TriEIDs[(FinderInfo.IdxOf + 1) % 3];
+			const int32 NormalCoord = NCoords.NormalCoord[FinderInfo.EID];
+			const int32 CrossingIdx = [&]
+				{
+					if (NormalCoord == 0)
+					{
+						return OrderOfTraceEID - FinderInfo.FirstRoundabout;
+					}
+					else
+					{
+						return OrderOfTraceEID - FinderInfo.FirstRoundabout + 1 + NormalCoord;
+					}
+				}();
+
+			checkSlow(CrossingIdx > 0); // would be zero if the edge was also a surface edge, but that case is handled above
+
+			ContinueTraceSurfaceEdge(IntrinsicMesh, FinderInfo.TID, OppEID, CrossingIdx, Crossings);
+		}
+
+		return MoveTemp(Crossings);
+	}
+
+
+
+	TArray<FNormalCoordinates::FEdgeAndCrossingIdx> TraceSurfaceEdge( const FIntrinsicEdgeFlipMesh& IntrinsicMesh,
+																	  const int32 SurfaceEID, const bool bReverse )
+	{
+		const FNormalCoordinates& NCoords = IntrinsicMesh.GetNormalCoordinates();
+		const FDynamicMesh3& SurfaceMesh  = *NCoords.SurfaceMesh;
+		if (!SurfaceMesh.IsEdge(SurfaceEID))
+		{
+			TArray< FNormalCoordinates::FEdgeAndCrossingIdx > Crossings;
+			return MoveTemp(Crossings);
+		}
+
+		const FIndex2i EdgeV = SurfaceMesh.GetEdgeV(SurfaceEID);
+		const int32 StartVID = (bReverse) ? EdgeV.B : EdgeV.A;
+
+		// Special case of tracing a surface edge on the mesh boundary.  This will be the same as an intrinsic edge.
+		if (SurfaceMesh.IsBoundaryEdge(SurfaceEID))
+		{
+			// boundary edges don't flip.  Will be the same as the intrinsic edge.
+
+			const int32 IntrinsicEID      = SurfaceEID;
+			const int32 IntrinsicTID      = IntrinsicMesh.GetEdgeT(IntrinsicEID).A;
+			const int32 IndexOf           = IntrinsicMesh.GetTriEdges(IntrinsicTID).IndexOf(IntrinsicEID);
+			const bool bRerverseIntrinsic = !(IntrinsicMesh.GetTriangle(IntrinsicTID)[IndexOf] == StartVID);
+
+
+			TArray< FNormalCoordinates::FEdgeAndCrossingIdx > Crossings;
+			Crossings.SetNumUninitialized(2);
+			auto& XingStart = Crossings[0];
+			XingStart.TID   = IntrinsicTID;
+			XingStart.EID   = (bRerverseIntrinsic) ? (IndexOf + 1) % 3 : IndexOf;
+			XingStart.CIdx  = 0;
+
+			auto& XingEnd = Crossings[1];
+			XingEnd.TID   = IntrinsicTID;
+			XingEnd.EID   = (bRerverseIntrinsic) ? IndexOf : (IndexOf + 1) % 3;
+			XingEnd.CIdx  = 0;
+
+			return MoveTemp(Crossings);
+		}
+		else  // the edge isn't a surface edge:  encode its direction by identifying it as an edge of a triangle and do the trace
+		{
+
+			const FIndex2i EdgeT = SurfaceMesh.GetEdgeT(SurfaceEID);
+			int32 TID            = EdgeT.A;
+			int32 IndexOf        = SurfaceMesh.GetTriEdges(TID).IndexOf(SurfaceEID);
+			if (SurfaceMesh.GetTriangle(TID)[IndexOf] != StartVID)
+			{
+				TID = EdgeT.B;
+				IndexOf = SurfaceMesh.GetTriEdges(TID).IndexOf(SurfaceEID);
+				checkSlow(SurfaceMesh.GetTriangle(TID)[IndexOf] == StartVID);
+			}
+			return TraceSurfaceEdge(IntrinsicMesh, TID, IndexOf);
+		}
+	}
+
+
+	/**
+	* Utility to trace an edge defined (by TraceEID) on the TraceMesh across the HostMesh, given a list of host mesh edges intersected by the trace mesh.
+	* note this assumes that the TraceMesh and HostMesh share the same vertex set
+	*/
+	template<typename HostMeshType, typename TraceMeshType>
+	TArray<FIntrinsicEdgeFlipMesh::FSurfacePoint> TraceEdgeOverHost( const int32 TraceEID, const TArray<int32>& HostXings, const HostMeshType& HostMesh,
+		                                                             const TraceMeshType& TraceMesh, double CoalesceThreshold, bool bReverse )
+	{
+		using FSurfacePoint = FIntrinsicEdgeFlipMesh::FSurfacePoint;
+		TArray<FSurfacePoint> EdgeTrace;
+
+		if (!TraceMesh.IsEdge(TraceEID))
+		{
+			// empty array..
+			return MoveTemp(EdgeTrace);
+		}
+
+		// utility to keep from adding duplicate VIDs.  Duplicate VIDs could result when Coalescing..
+		auto AddVIDToTrace = [&EdgeTrace](const int32 VID)
+		{
+			if (EdgeTrace.Num() == 0)
+			{
+				EdgeTrace.Emplace(VID);
+			}
+			else
+			{
+				const FSurfacePoint& LastPoint = EdgeTrace.Last();
+				const bool bSameAsLast         = ( LastPoint.PositionType == FSurfacePoint::EPositionType::Vertex 
+				                                && LastPoint.Position.VertexPosition.VID == VID );
+				if (!bSameAsLast)
+				{
+					EdgeTrace.Emplace(VID);
+				}
+			}
+		};
+
+
+		// an edge  w/o crossings is the same as the host mesh edge. the trace in this case requires only the start and edge vertex 
+		if (HostXings.Num() == 0)
+		{
+			const int32 AdjTID     = TraceMesh.GetEdgeT(TraceEID).A;
+			const int32 IndexOf    = TraceMesh.GetTriEdges(AdjTID).IndexOf(TraceEID);
+			const FIndex3i TriVIDs = TraceMesh.GetTriangle(AdjTID);
+
+			const int32 StartVID = TriVIDs[IndexOf];
+			const int32 EndVID   = TriVIDs[(IndexOf + 1) % 3];
+
+
+			EdgeTrace.Emplace(StartVID);
+			EdgeTrace.Emplace(EndVID);
+		}
+		else
+		{
+			// knowing the host mesh edges this trace edge crosses, we make a 2d triangle strip by unfolding the 3d triangles the edge crosses
+			// and solve a 2x2 problem for each 2d edge crossed.
+
+			// struct holds bare-bones 2d triangle strip and ability to map vertexIDs from the src mesh
+			struct
+			{
+				TMap<int32, int32> ToStripIndexMap;    // maps from mesh VID to triangle strip VID.
+				TArray<FVector2d> StripVertexBuffer;
+			} TriangleStrip;
+			TriangleStrip.StripVertexBuffer.Reserve(3 + HostXings.Num());
+
+
+			// Functor to add a triangle to the triangle strip. This assumes that at least one shared edge of this triangle already exits in the strip
+			auto AddTriangleToStrip = [&HostMesh, &TriangleStrip](const int32 HostTID)
+			{
+
+				FIndex3i TriVIDs = HostMesh.GetTriangle(HostTID);
+
+				// two of the tri vertices are already in the buffer.  Identify the new one.
+				int32 IndexOf = -1;
+				for (int32 i = 0; i < 3; ++i)
+				{
+					const int32 VID       = TriVIDs[i];
+					const int32* StripVID = TriangleStrip.ToStripIndexMap.Find(VID);
+					if (!StripVID)
+					{
+						IndexOf = i;
+						break;
+					}
+				}
+
+				// no need to do anything if all vertices had previously been added, 
+				if (IndexOf == -1)
+				{
+					return;
+				}
+
+
+				const FVector3d Verts[3] = { HostMesh.GetVertex(TriVIDs[0]),  HostMesh.GetVertex(TriVIDs[1]),  HostMesh.GetVertex(TriVIDs[2]) };
+				// with new vert last (i.e. V2).
+				const FIndex3i Reordered((IndexOf + 1) % 3, (IndexOf + 2) % 3, IndexOf);
+				// the spanning vectors 
+				const FVector3d E1 = Verts[Reordered[1]] - Verts[Reordered[0]];
+				const FVector3d E2 = Verts[Reordered[2]] - Verts[Reordered[0]];
+
+				// coordinates of V2 relative to the direction of E1, and its orthogonal complement.
+				const double E1DotE2       = E2.Dot(E1);
+				const double E1LengthSqr   = FMath::Max(E1.SizeSquared(), TMathUtilConstants<double>::ZeroTolerance);
+				const double E1Length      = FMath::Sqrt(E1LengthSqr);
+				const double E1Dist        = E2.Dot(E1) / E1Length;
+				const double OrthE1DistSqr = FMath::Max(0., E2.SizeSquared() - E1DotE2 * E1DotE2 / E1LengthSqr);
+				const double OrthE1Dist    = FMath::Sqrt(OrthE1DistSqr);
+
+				// 2d version of E1
+				const int32* StripTriV0 = TriangleStrip.ToStripIndexMap.Find(TriVIDs[Reordered[0]]);
+				const int32* StripTriV1 = TriangleStrip.ToStripIndexMap.Find(TriVIDs[Reordered[1]]);
+				checkSlow(StripTriV0); checkSlow(StripTriV1);
+
+				const FVector2d& StripTriVert0 = TriangleStrip.StripVertexBuffer[*StripTriV0];
+				const FVector2d& StripTriVert1 = TriangleStrip.StripVertexBuffer[*StripTriV1];
+				const FVector2d StripE1        = (StripTriVert1 - StripTriVert0);
+
+				const FVector2d StripE1Perp(-StripE1[1], StripE1[0]); // Rotate StripE1 90 CCW.
+
+																	  // new vertex 2d position
+				const FVector2d StripTriVert2 = StripTriVert0 + (StripE1)*E1DotE2 / E1LengthSqr + (StripE1Perp / E1Length) * OrthE1Dist;
+				TriangleStrip.StripVertexBuffer.Add(StripTriVert2);
+				const int32 StripVID          = TriangleStrip.StripVertexBuffer.Num() - 1;
+				const int32 SurfaceVID        = TriVIDs[IndexOf];
+				TriangleStrip.ToStripIndexMap.Add(SurfaceVID, StripVID);
+			};
+
+			// find the VID where we start the edge trace and add it to the trace.
+			const int32 StartVID = [&]
+				{
+					// by convention the intrinsic edge direction will be defined by the first adj triangle
+					const int32 AdjTID     = TraceMesh.GetEdgeT(TraceEID).A;
+					const int32 IndexOf    = TraceMesh.GetTriEdges(AdjTID).IndexOf(TraceEID);
+					const FIndex3i TriVIDs = TraceMesh.GetTriangle(AdjTID);
+					return TriVIDs[IndexOf];
+				}();
+			EdgeTrace.Emplace(StartVID);
+
+			// find host triangle that contains both the StartVID and the first host edge we cross.
+			int32 HostTID = [&]
+				{
+					const int32 HostEID      = HostXings[0];
+					const FIndex2i HostEdgeT = HostMesh.GetEdgeT(HostEID);
+					const FIndex3i TriAVIDs  = HostMesh.GetTriangle(HostEdgeT.A);
+					const FIndex3i TriBVIDs  = HostMesh.GetTriangle(HostEdgeT.B);
+
+					return (TriAVIDs.IndexOf(StartVID) != -1) ? HostEdgeT.A : HostEdgeT.B;
+				}();
+
+			// jump-start the process of making the triangle strip by adding the first 2 verts.
+			{
+				const FIndex3i TriVIDs = HostMesh.GetTriangle(HostTID);
+				const int32 IndexOf    = TriVIDs.IndexOf(StartVID);
+				checkSlow(IndexOf != -1);
+				const int32 NextVID    = TriVIDs[(IndexOf + 1) % 3];
+
+				const FVector3d Vert0 = HostMesh.GetVertex(StartVID);
+				const FVector3d Vert1 = HostMesh.GetVertex(NextVID);
+
+				const double LSqr = FMath::Max((Vert0 - Vert1).SizeSquared(), TMathUtilConstants<double>::ZeroTolerance);
+				const double L    = FMath::Sqrt(LSqr);
+
+				const FVector2d StripVert0(0., 0.);
+				const FVector2d StripVert1(L, 0.);
+
+				TriangleStrip.StripVertexBuffer.Add(StripVert0);
+				TriangleStrip.ToStripIndexMap.Add(StartVID, 0);
+
+				TriangleStrip.StripVertexBuffer.Add(StripVert1);
+				TriangleStrip.ToStripIndexMap.Add(NextVID, 1);
+			}
+
+			// unfold the triangle strip, add the first triangle and then all subsequent ones.
+			AddTriangleToStrip(HostTID);
+			for (int32 HostEID : HostXings)
+			{
+				const FIndex2i EdgeT = HostMesh.GetEdgeT(HostEID);
+				HostTID              = (EdgeT.A == HostTID) ? EdgeT.B : EdgeT.A;
+				AddTriangleToStrip(HostTID);
+			}
+
+			// get the 2d location of the end of the trace edge (i.e. at EndVID ). Note: the start location is (0,0) in 2d 
+			const FVector2d StripEndVertex = [&]
+				{
+					const FIndex2i TraceEdgeV = TraceMesh.GetEdgeV(TraceEID);
+					const int32 EndVID        = (TraceEdgeV.A == StartVID) ? TraceEdgeV.B : TraceEdgeV.A;
+					const int32* StripEndVID  = TriangleStrip.ToStripIndexMap.Find(EndVID);
+					checkSlow(StripEndVID);
+					return TriangleStrip.StripVertexBuffer[*StripEndVID];
+				}();
+
+			// loop over the two-d versions of the host edges the path crosses and find the intersection.
+			for (int32 HostEID : HostXings)
+			{
+				const FIndex2i EdgeV = HostMesh.GetEdgeV(HostEID);
+				const int32* StripA  = TriangleStrip.ToStripIndexMap.Find(EdgeV.A);
+				const int32* StripB  = TriangleStrip.ToStripIndexMap.Find(EdgeV.B);
+				checkSlow(StripA); checkSlow(StripB);
+				const FIndex2i StripEdgeV(*StripA, *StripB);
+
+				const FVector2d StripVertA = TriangleStrip.StripVertexBuffer[StripEdgeV.A];
+				const FVector2d StripVertB = TriangleStrip.StripVertexBuffer[StripEdgeV.B];
+
+				// solve Alpha*StripVertA + (1-Alpha)StripVertB = Gamma Start + (1-Gamma)End.
+				// i.e. 
+				// Alpha * (StripVertA - StripVertB) + Gamma * (End - Start) =  End - StripVertB.
+
+				// write this as 2x2 matrix problem M.x = b solving for unknown vector x=(alpha, gamma).
+
+				// matrix
+				double m[2][2];
+				m[0][0] = (StripVertA - StripVertB)[0];  m[0][1] = StripEndVertex[0];
+				m[1][0] = (StripVertA - StripVertB)[1];  m[1][1] = StripEndVertex[1];
+
+				// b-vector
+				const FVector2d b = StripEndVertex - StripVertB;
+
+				const double Det = m[0][0] * m[1][1] - m[0][1] * m[1][0];
+				// inverse: not yet scaled by 1/det
+				double invm[2][2];
+				invm[0][0] =  m[1][1];   invm[0][1] = -m[0][1];
+				invm[1][0] = -m[1][0];   invm[1][1] =  m[0][0];
+
+				// solve for alpha only ( don't care about gamma ) 
+				double alpha = invm[0][0] * b[0] + invm[0][1] * b[1];
+
+				if (FMath::Abs(Det) < TMathUtilConstants<double>::ZeroTolerance)
+				{
+					// this surface edge and the intrinsic edge that intersects it are nearly parallel.  
+					// Assume the crossing happens at the farther vertex.
+					const double dA = StripVertA.SquaredLength();
+					const double dB = StripVertB.SquaredLength();
+
+					alpha = (dA > dB) ? 1. : 0.;
+				}
+				else
+				{
+					alpha = FMath::Clamp(alpha / Det, 0., 1.);
+				}
+
+				// may want to convert this edge crossing to a vertex crossing, if it is close enough.
+				int32 CoalesceVID = -1;
+				bool bSnapToVert = false;
+
+				if (alpha < CoalesceThreshold)
+				{
+					CoalesceVID = EdgeV.B;
+					bSnapToVert = true;
+				}
+				else if ((1. - alpha) < CoalesceThreshold)
+				{
+					CoalesceVID = EdgeV.A;
+					bSnapToVert = true;
+				}
+
+				if (bSnapToVert)
+				{
+					AddVIDToTrace(CoalesceVID);
+				}
+				else
+				{
+					EdgeTrace.Emplace(HostEID, alpha);
+				}
+			}
+
+			// Add the end vertex to the EdgeTrace
+			{
+				const FIndex2i TraceEdgeV = TraceMesh.GetEdgeV(TraceEID);
+				const int32 EndVID        = (TraceEdgeV.A == StartVID) ? TraceEdgeV.B : TraceEdgeV.A;
+				AddVIDToTrace(EndVID);
+			}
+
+		}
+
+		// correct trace results order so it is either EdgeV order or reversed as requested
+		const bool bHasEdgeVOrder = (TraceMesh.GetEdgeV(TraceEID).A == EdgeTrace[0].Position.VertexPosition.VID);
+		const bool bNeedToReverse = (bHasEdgeVOrder && bReverse) || (!bHasEdgeVOrder && !bReverse);
+
+		if (bNeedToReverse)
+		{
+			Algo::Reverse(EdgeTrace);
+		}
+
+		return MoveTemp(EdgeTrace);
+	}
+};
 
 int32 UE::Geometry::FlipToDelaunay(FIntrinsicTriangulation& IntrinsicMesh, TSet<int32>& Uncorrected, const int32 MaxFlipCount)
 {
 	return FlipToDelaunayImpl(IntrinsicMesh, Uncorrected, MaxFlipCount);
 }
 
-int32 UE::Geometry::FlipToDelaunay(FIntrinsicEdgeFlipMesh& IntrinsicMesh, TSet<int32>& Uncorrected, const int32 MaxFlipCount)
+int32 UE::Geometry::FlipToDelaunay(FSimpleIntrinsicEdgeFlipMesh& IntrinsicMesh, TSet<int32>& Uncorrected, const int32 MaxFlipCount)
 {
 	return FlipToDelaunayImpl(IntrinsicMesh, Uncorrected, MaxFlipCount);
 }
 
+int32 UE::Geometry::FlipToDelaunay(FIntrinsicEdgeFlipMesh& IntrinsicMesh, TSet<int32>& Uncorrected, const int32 MaxFlipCount)
+{
+	return FlipToDelaunayImpl(IntrinsicMesh, Uncorrected, MaxFlipCount);
+}
 
 
 
 /**------------------------------------------------------------------------------
-* FIntrinsicEdgeFlipMesh Methods
+* FSimpleIntrinsicEdgeFlipMesh Methods
 *------------------------------------------------------------------------------ */
 
-FIntrinsicEdgeFlipMesh::FIntrinsicEdgeFlipMesh(const FDynamicMesh3& SrcMesh) :
+FSimpleIntrinsicEdgeFlipMesh::FSimpleIntrinsicEdgeFlipMesh(const FDynamicMesh3& SrcMesh) :
 	Vertices(SrcMesh.GetVerticesBuffer()),
 	VertexRefCounts{ SrcMesh.GetVerticesRefCounts() },
 	VertexEdgeLists{ SrcMesh.GetVertexEdges() },
@@ -420,7 +1141,7 @@ FIntrinsicEdgeFlipMesh::FIntrinsicEdgeFlipMesh(const FDynamicMesh3& SrcMesh) :
 }
 
 
-FIndex2i FIntrinsicEdgeFlipMesh::GetEdgeOpposingV(int32 EID) const
+FIndex2i FSimpleIntrinsicEdgeFlipMesh::GetEdgeOpposingV(int32 EID) const
 {
 	const FEdge& Edge = Edges[EID];
 	FIndex2i Result(InvalidID, InvalidID);
@@ -431,7 +1152,7 @@ FIndex2i FIntrinsicEdgeFlipMesh::GetEdgeOpposingV(int32 EID) const
 		if (TriID == InvalidID) continue;
 
 		const FIndex3i TriEIDs = GetTriEdges(TriID);
-		const int32 IndexOf = TriEIDs.IndexOf(EID);
+		const int32 IndexOf    = TriEIDs.IndexOf(EID);
 		const FIndex3i TriVIDs = GetTriangle(TriID);
 		Result[i] = TriVIDs[AddTwoModThree[IndexOf]];
 	}
@@ -440,7 +1161,7 @@ FIndex2i FIntrinsicEdgeFlipMesh::GetEdgeOpposingV(int32 EID) const
 }
 
 
-FIndex2i FIntrinsicEdgeFlipMesh::GetOrientedEdgeV(int32 EID, int32 TID) const
+FIndex2i FSimpleIntrinsicEdgeFlipMesh::GetOrientedEdgeV(int32 EID, int32 TID) const
 {
 	int32 IndexOf = GetTriEdges(TID).IndexOf(EID);
 	checkSlow(IndexOf != InvalidID);
@@ -449,7 +1170,7 @@ FIndex2i FIntrinsicEdgeFlipMesh::GetOrientedEdgeV(int32 EID, int32 TID) const
 }
 
 
-int32 FIntrinsicEdgeFlipMesh::ReplaceEdgeTriangle(int32 eID, int32 tOld, int32 tNew)
+int32 FSimpleIntrinsicEdgeFlipMesh::ReplaceEdgeTriangle(int32 eID, int32 tOld, int32 tNew)
 {
 	FIndex2i& Tris = Edges[eID].Tri;
 	int32 a = Tris[0], b = Tris[1];
@@ -477,7 +1198,7 @@ int32 FIntrinsicEdgeFlipMesh::ReplaceEdgeTriangle(int32 eID, int32 tOld, int32 t
 }
 
 
-EMeshResult FIntrinsicEdgeFlipMesh::FlipEdgeTopology(int32 eab, FEdgeFlipInfo& FlipInfo)
+EMeshResult FSimpleIntrinsicEdgeFlipMesh::FlipEdgeTopology(int32 eab, FEdgeFlipInfo& FlipInfo)
 {
 	if (!IsEdge(eab))
 	{
@@ -527,12 +1248,12 @@ EMeshResult FIntrinsicEdgeFlipMesh::FlipEdgeTopology(int32 eab, FEdgeFlipInfo& F
 	// update the two other edges whose triangle nbrs have changed
 	if (ReplaceEdgeTriangle(eca, t0, t1) == -1)
 	{
-		checkfSlow(false, TEXT("FIntrinsicEdgeFlipMesh.FlipEdge: first ReplaceEdgeTriangle failed"));
+		checkfSlow(false, TEXT("FSimpleIntrinsicEdgeFlipMesh.FlipEdge: first ReplaceEdgeTriangle failed"));
 		return EMeshResult::Failed_UnrecoverableError;
 	}
 	if (ReplaceEdgeTriangle(edb, t1, t0) == -1)
 	{
-		checkfSlow(false, TEXT("FIntrinsicEdgeFlipMesh.FlipEdge: second ReplaceEdgeTriangle failed"));
+		checkfSlow(false, TEXT("FSimpleIntrinsicEdgeFlipMesh.FlipEdge: second ReplaceEdgeTriangle failed"));
 		return EMeshResult::Failed_UnrecoverableError;
 	}
 
@@ -543,7 +1264,7 @@ EMeshResult FIntrinsicEdgeFlipMesh::FlipEdgeTopology(int32 eab, FEdgeFlipInfo& F
 	// remove old eab from verts a and b, and Decrement ref counts
 	if (VertexEdgeLists.Remove(a, eab) == false)
 	{
-		checkfSlow(false, TEXT("FIntrinsicEdgeFlipMesh.FlipEdge: first edge list remove failed"));
+		checkfSlow(false, TEXT("FSimpleIntrinsicEdgeFlipMesh.FlipEdge: first edge list remove failed"));
 		return EMeshResult::Failed_UnrecoverableError;
 	}
 	VertexRefCounts.Decrement(a);
@@ -551,14 +1272,14 @@ EMeshResult FIntrinsicEdgeFlipMesh::FlipEdgeTopology(int32 eab, FEdgeFlipInfo& F
 	{
 		if (VertexEdgeLists.Remove(b, eab) == false)
 		{
-			checkfSlow(false, TEXT("FIntrinsicEdgeFlipMesh.FlipEdge: second edge list remove failed"));
+			checkfSlow(false, TEXT("FSimpleIntrinsicEdgeFlipMesh.FlipEdge: second edge list remove failed"));
 			return EMeshResult::Failed_UnrecoverableError;
 		}
 		VertexRefCounts.Decrement(b);
 	}
 	if (IsVertex(a) == false || IsVertex(b) == false)
 	{
-		checkfSlow(false, TEXT("FIntrinsicEdgeFlipMesh.FlipEdge: either a or b is not a vertex?"));
+		checkfSlow(false, TEXT("FSimpleIntrinsicEdgeFlipMesh.FlipEdge: either a or b is not a vertex?"));
 		return EMeshResult::Failed_UnrecoverableError;
 	}
 
@@ -580,7 +1301,7 @@ EMeshResult FIntrinsicEdgeFlipMesh::FlipEdgeTopology(int32 eab, FEdgeFlipInfo& F
 }
 
 
-FVector3d FIntrinsicEdgeFlipMesh::ComputeTriInternalAnglesR(const int32 TID) const
+FVector3d FSimpleIntrinsicEdgeFlipMesh::ComputeTriInternalAnglesR(const int32 TID) const
 {
 	const FVector3d Lengths = GetTriEdgeLengths(TID);
 	FVector3d Angles;
@@ -593,9 +1314,9 @@ FVector3d FIntrinsicEdgeFlipMesh::ComputeTriInternalAnglesR(const int32 TID) con
 }
 
 
-double FIntrinsicEdgeFlipMesh::GetOpposingVerticesDistance(int32 EID) const
+double FSimpleIntrinsicEdgeFlipMesh::GetOpposingVerticesDistance(int32 EID) const
 {
-	const double OrgLength = GetEdgeLength(EID);
+	const double OrgLength  = GetEdgeLength(EID);
 	const FIndex2i EdgeTris = GetEdgeT(EID);
 	checkSlow(EdgeTris[1] != FDynamicMesh3::InvalidID);
 
@@ -617,7 +1338,7 @@ double FIntrinsicEdgeFlipMesh::GetOpposingVerticesDistance(int32 EID) const
 }
 
 
-EMeshResult FIntrinsicEdgeFlipMesh::FlipEdge(int32 EID, FEdgeFlipInfo& EdgeFlipInfo)
+EMeshResult FSimpleIntrinsicEdgeFlipMesh::FlipEdge(int32 EID, FEdgeFlipInfo& EdgeFlipInfo)
 {
 	if (!IsEdge(EID))
 	{
@@ -676,21 +1397,21 @@ EMeshResult FIntrinsicEdgeFlipMesh::FlipEdge(int32 EID, FEdgeFlipInfo& EdgeFlipI
 }
 
 
-FVector2d FIntrinsicEdgeFlipMesh::EdgeOpposingAngles(int32 EID) const
+FVector2d FSimpleIntrinsicEdgeFlipMesh::EdgeOpposingAngles(int32 EID) const
 {
 	FVector2d Result;
 	FDynamicMesh3::FEdge Edge = GetEdge(EID);
 	{
-		const int32 IndexOf = GetTriEdges(Edge.Tri.A).IndexOf(EID);
+		const int32 IndexOf    = GetTriEdges(Edge.Tri.A).IndexOf(EID);
 		const int32 IndexOfOpp = AddTwoModThree[IndexOf];
-		const double Angle = GetTriInternalAngleR(Edge.Tri.A, IndexOfOpp);
+		const double Angle     = GetTriInternalAngleR(Edge.Tri.A, IndexOfOpp);
 		Result[0] = Angle;
 	}
 	if (Edge.Tri.B != FDynamicMesh3::InvalidID)
 	{
-		const int32 IndexOf = GetTriEdges(Edge.Tri.B).IndexOf(EID);
+		const int32 IndexOf    = GetTriEdges(Edge.Tri.B).IndexOf(EID);
 		const int32 IndexOfOpp = AddTwoModThree[IndexOf];
-		const double Angle = GetTriInternalAngleR(Edge.Tri.B, IndexOfOpp);
+		const double Angle     = GetTriInternalAngleR(Edge.Tri.B, IndexOfOpp);
 		Result[1] = Angle;
 	}
 	else
@@ -702,13 +1423,13 @@ FVector2d FIntrinsicEdgeFlipMesh::EdgeOpposingAngles(int32 EID) const
 }
 
 
-double FIntrinsicEdgeFlipMesh::EdgeCotanWeight(int32 EID) const
+double FSimpleIntrinsicEdgeFlipMesh::EdgeCotanWeight(int32 EID) const
 {
 	auto ComputeCotanOppAngle = [this](int32 EID, int32 TID)->double
 	{
 		const FIndex3i TriEIDs = GetTriEdges(TID);
-		const int32 IOf = TriEIDs.IndexOf(EID);
-		const FVector3d TriLs = GetEdgeLengthTriple(TriEIDs);
+		const int32 IOf        = TriEIDs.IndexOf(EID);
+		const FVector3d TriLs  = GetEdgeLengthTriple(TriEIDs);
 
 		const FVector3d Ls(TriLs[IOf], TriLs[AddOneModThree[IOf]], TriLs[AddTwoModThree[IOf]]);
 		return ComputeCotangent(Ls[0], Ls[1], Ls[2]);
@@ -726,276 +1447,15 @@ double FIntrinsicEdgeFlipMesh::EdgeCotanWeight(int32 EID) const
 }
 
 
-
 /**------------------------------------------------------------------------------
 *  FIntrinsicTriangulation Methods
 *-------------------------------------------------------------------------------*/
 
-FIntrinsicTriangulation::FSurfacePoint::FSurfacePoint()
-{
-	PositionType = EPositionType::Vertex;
-	Position.VertexPosition.VID = FDynamicMesh3::InvalidID;
-}
-
-FIntrinsicTriangulation::FSurfacePoint::FSurfacePoint(int32 VID)
-{
-	PositionType = EPositionType::Vertex;
-	Position.VertexPosition.VID = VID;
-}
-
-FIntrinsicTriangulation::FSurfacePoint::FSurfacePoint(int32 EdgeID, double Alpha)
-{
-	PositionType = EPositionType::Edge;
-	Position.EdgePosition.EdgeID = EdgeID;
-	Position.EdgePosition.Alpha = Alpha;
-}
-
-FIntrinsicTriangulation::FSurfacePoint::FSurfacePoint(int32 TriID, const FVector3d& BaryCentrics)
-{
-	PositionType = EPositionType::Triangle;
-	Position.TriPosition.TriID = TriID;
-	Position.TriPosition.BarycentricCoords[0] = BaryCentrics[0];
-	Position.TriPosition.BarycentricCoords[1] = BaryCentrics[1];
-	Position.TriPosition.BarycentricCoords[2] = BaryCentrics[2];
-}
-
-FVector3d FIntrinsicTriangulation::FSurfacePoint::FSurfacePoint::AsR3Position(const FDynamicMesh3& Mesh, bool& bIsValid) const
-{
-	FVector3d Result(0., 0., 0.);
-	switch (PositionType)
-	{
-		case FSurfacePoint::EPositionType::Vertex:
-		{
-			const int32 SurfaceVID = Position.VertexPosition.VID;
-			bIsValid = Mesh.IsVertex(SurfaceVID);
-			if (bIsValid)
-			{
-				Result = Mesh.GetVertex(SurfaceVID);
-			}
-		}
-		break;
-		case FSurfacePoint::EPositionType::Edge:
-		{
-			const int32 SurfaceEID = Position.EdgePosition.EdgeID;
-			double Alpha = Position.EdgePosition.Alpha;
-			bIsValid = Mesh.IsEdge(SurfaceEID);
-			if (bIsValid)
-			{
-				const FIndex2i SurfaceEdgeV = Mesh.GetEdgeV(SurfaceEID);
-				Result = Alpha * (Mesh.GetVertex(SurfaceEdgeV.A)) + (1. - Alpha) * (Mesh.GetVertex(SurfaceEdgeV.B));
-			}
-		}
-		break;
-		case FSurfacePoint::EPositionType::Triangle:
-		{
-			const int32 SurfaceTID = Position.TriPosition.TriID;
-			bIsValid = Mesh.IsTriangle(SurfaceTID);
-
-			if (bIsValid)
-			{
-				Result = Mesh.GetTriBaryPoint(SurfaceTID,
-					Position.TriPosition.BarycentricCoords[0],
-					Position.TriPosition.BarycentricCoords[1],
-					Position.TriPosition.BarycentricCoords[2]);
-			}
-
-		}
-		break;
-		default:
-		{
-			bIsValid = false;
-			// shouldn't be able to reach this point
-			check(0);
-		}
-	}
-	return Result;
-}
 
 UE::Geometry::FIntrinsicTriangulation::FIntrinsicTriangulation(const FDynamicMesh3& SrcMesh)
 	: MyBase(SrcMesh)
-	, ExtrinsicMesh(&SrcMesh)
+	, SignpostData(SrcMesh)
 {
-	// During construction the mesh is a deep copy of the SrcMesh ( i.e. has the same IDs)
-
-	// pick local reference directions for each vertex and triangle.
-	const int32 MaxVertexID   = ExtrinsicMesh->MaxVertexID();
-	const int32 MaxTriangleID = ExtrinsicMesh->MaxTriangleID();
-
-	{
-		VIDToReferenceEID.AddUninitialized(MaxVertexID);
-		TIDToReferenceEID.AddUninitialized(MaxTriangleID);
-
-		for (int32 VID = 0; VID < MaxVertexID; ++VID)
-		{
-			if (ExtrinsicMesh->IsVertex(VID))
-			{
-				int32 RefEID = FDynamicMesh3::InvalidID;
-				if (ExtrinsicMesh->IsBoundaryVertex(VID) && !ExtrinsicMesh->IsBowtieVertex(VID))
-				{
-					// vertex is on the mesh boundary, 
-					// choose the boundary edge such that traveling CCW (about the vertex) moves into the mesh
-					for (int32 EID : ExtrinsicMesh->VtxEdgesItr(VID))
-					{
-						if (ExtrinsicMesh->IsBoundaryEdge(EID))
-						{
-							const int32 TID = ExtrinsicMesh->GetEdgeT(EID).A;
-							const int32 IndexOf = ExtrinsicMesh->GetTriEdges(TID).IndexOf(EID);
-							if (ExtrinsicMesh->GetTriangle(TID)[IndexOf] == VID)
-							{
-								RefEID = EID;
-								break;
-							}
-						}
-					}
-				}
-				else
-				{
-					// vertex isn't on a mesh boundary, just pick the first adj edge to be the reference edge
-					for (int32 EID : ExtrinsicMesh->VtxEdgesItr(VID))
-					{
-						RefEID = EID;
-						break;
-					}
-				}
-				VIDToReferenceEID[VID] = RefEID;
-			}
-		}
-		for (int32 TID = 0; TID < MaxTriangleID; ++TID)
-		{
-			if (ExtrinsicMesh->IsTriangle(TID))
-			{
-				TIDToReferenceEID[TID] = ExtrinsicMesh->GetTriEdge(TID, 0);
-			}
-		}
-	}
-
-	// add surface position.
-	{
-		IntrinsicVertexPositions.SetNum(MaxVertexID);
-
-		// initialize: identify with vertex in Extrinsic Mesh
-		for (int32 VID = 0; VID < MaxVertexID; ++VID)
-		{
-			if (IsVertex(VID))
-			{
-				FSurfacePoint VertexSurfacePoint(VID);
-				IntrinsicVertexPositions[VID] = VertexSurfacePoint;
-			}
-		}
-	}
-
-	// add edge directions
-	{
-		const int32 MaxExtEdgeID = ExtrinsicMesh->MaxEdgeID();
-		IntrinsicEdgeAngles.SetNum(MaxTriangleID);
-
-		GeometricVertexInfo.SetNum(MaxVertexID);
-		IntrinsicVertexPositions.SetNum(MaxVertexID);
-
-		ParallelFor(MaxVertexID, [this, &SrcMesh](int32 VID)
-		{
-			if (!ExtrinsicMesh->IsVertex(VID))
-			{
-				return;
-			}
-
-			const int32 PlusTwoModThree[3] = { 2, 0, 1 }; // PlusTwoModThree[i] = (i + 2)%3
-
-			// initialize by computing the angles
-			TArray<int32> Triangles;
-			TArray<int32> ContiguousGroupLengths;
-			TArray<bool> GroupIsLoop;
-
-			{
-
-				FGeometricInfo GeometricInfo;
-				// note: since the structure and IDs are a deep copy of the source mesh, we can use it to find the contiguous triangles
-				const EMeshResult Result = SrcMesh.GetVtxContiguousTriangles(VID, Triangles, ContiguousGroupLengths, GroupIsLoop);
-				if (Result == EMeshResult::Ok && GroupIsLoop.Num() == 1)
-				{
-					GeometricInfo.bIsInterior = (GroupIsLoop[0] == true);
-					// Contiguous tris could be cw or ccw?  Need them to  be ccw
-					if (Triangles.Num() > 1)
-					{
-						const int32 TID0 = Triangles[0];
-						FIndex3i TriVIDs = GetTriangle(TID0);
-						FIndex3i TriEIDs = GetTriEdges(TID0);
-						int32 VertSubIdx = TriVIDs.IndexOf(VID);
-						const int32 EID = TriEIDs[VertSubIdx];
-						const int32 NextEID = TriEIDs[PlusTwoModThree[VertSubIdx]];
-						const int32 TID1 = Triangles[1];
-						TriVIDs = GetTriangle(TID1);
-						TriEIDs = GetTriEdges(TID1);
-						VertSubIdx = TriVIDs.IndexOf(VID);
-
-						if (TriEIDs[VertSubIdx] != NextEID)
-						{
-							// need to reverse order to have CCW
-							Algo::Reverse(Triangles);
-						}
-					}
-
-					const int32 NumEdges = ContiguousGroupLengths[0];
-
-					double AngleOffset = 0;
-					double WalkedAngle = 0;
-					const int32 ReferenceEID = VIDToReferenceEID[VID];
-					checkSlow(ReferenceEID != FDynamicMesh3::InvalidID)
-
-					TMap<int32, int32> VisitedTriSubID;
-					for (int32 TID : Triangles)
-					{
-						const FIndex3i TriVIDs = GetTriangle(TID);
-						const FIndex3i TriEIDs = GetTriEdges(TID);
-						const int32 VertSubIdx = TriVIDs.IndexOf(VID);
-						IntrinsicEdgeAngles[TID][VertSubIdx] = WalkedAngle;
-
-						VisitedTriSubID.Add(TID, VertSubIdx);
-						const int32 EID = TriEIDs[VertSubIdx];
-						if (EID == ReferenceEID)
-						{
-							AngleOffset = WalkedAngle;
-						}
-						
-						WalkedAngle += GetTriInternalAngleR(TID, VertSubIdx);
-					}
-					// need to include last edge if the one ring of triangles isn't a closed loop, it could be the reference edge.
-					if (GroupIsLoop[0] == false)
-					{
-						const int32 TID = Triangles.Last();
-						const FIndex3i TriVIDs = GetTriangle(TID);
-						const FIndex3i TriEIDs = GetTriEdges(TID);
-						const int32 VertSubIdx = TriVIDs.IndexOf(VID);
-						const int32 EID = TriEIDs[PlusTwoModThree[VertSubIdx]];
-						if (EID == ReferenceEID)
-						{
-							AngleOffset = WalkedAngle;
-						}
-					}
-
-					// orient with the reference edge at zero (or pi) angle  and rescale to radians and use 0,2pi to define polar angle.
-					const double ToRadians = TMathUtilConstants<double>::TwoPi / WalkedAngle;
-					GeometricInfo.ToRadians = ToRadians;
-					double RefAngle = (ExtrinsicMesh->GetEdgeV(ReferenceEID).B == VID) ? TMathUtilConstants<double>::Pi : 0.;
-					for (TPair<int32, int32>& TriSubIDPair : VisitedTriSubID)
-					{
-						int32 TID = TriSubIDPair.Key;
-						int32 SubID = TriSubIDPair.Value;
-						double& Angle = IntrinsicEdgeAngles[TID][SubID];
-						Angle = ToRadians * (Angle - AngleOffset) + RefAngle;
-						Angle = AsZeroToTwoPi(Angle);
-					}
-				}
-				else // bow-tie or boundary
-				{
-					GeometricInfo.bIsInterior = false;
-				}
-
-				GeometricVertexInfo[VID] = GeometricInfo;
-			}
-		}, false /* force single thread*/);
-	}
-
 }
 
 TArray<FIntrinsicTriangulation::FSurfacePoint> UE::Geometry::FIntrinsicTriangulation::TraceEdge(int32 EID, double CoalesceThreshold, bool bReverse) const
@@ -1030,13 +1490,13 @@ TArray<FIntrinsicTriangulation::FSurfacePoint> UE::Geometry::FIntrinsicTriangula
 			if (GetTriangle(AdjTID)[IndexOf] == StartVID)
 			{
 				// IntrinsicEdgeAngles hold the local angle of each out-going edge relative to the vertex the edge exits
-				return IntrinsicEdgeAngles[AdjTID][IndexOf];
+				return SignpostData.IntrinsicEdgeAngles[AdjTID][IndexOf];
 			}
 			else
 			{
-				const double ToRadians = GeometricVertexInfo[StartVID].ToRadians;
+				const double ToRadians = SignpostData.GeometricVertexInfo[StartVID].ToRadians;
 				// polar angle of prev edge just before (clockwise) the edge we want
-				const double PrevPolarAngle = IntrinsicEdgeAngles[AdjTID][(IndexOf + 1) % 3];
+				const double PrevPolarAngle = SignpostData.IntrinsicEdgeAngles[AdjTID][(IndexOf + 1) % 3];
 				// angle between previous edge and this one
 				const double InternalAngle = InternalAngles[AdjTID][(IndexOf + 1) % 3];
 				// add the internal angle to rotate from Prev Edge to this edge
@@ -1046,64 +1506,9 @@ TArray<FIntrinsicTriangulation::FSurfacePoint> UE::Geometry::FIntrinsicTriangula
 
 	// Trace the surface mesh from the intrinsic StartVID, in the PolarAngle direction, a distance of IntrinsicEdgeLength.
 	// 
-	// Note: this intrisic vertex may or may not correspond to a vertex in the surface mesh if vertices were added to the intrinsic mesh
+	// Note: this intrinsic vertex may or may not correspond to a vertex in the surface mesh if vertices were added to the intrinsic mesh
 	// by doing an edge split or a triangle poke.
-	FMeshGeodesicSurfaceTracer SurfaceTracer(*SurfaceMesh);
-	{
-		const FSurfacePoint& StartSurfacePoint = GetVertexSurfacePoint(StartVID);
-		const FSurfacePoint::FSurfacePositionUnion& SurfacePosition = StartSurfacePoint.Position;
-		switch (StartSurfacePoint.PositionType)
-		{
-			case FSurfacePoint::EPositionType::Vertex:
-			{
-				const int32 SurfaceVID = SurfacePosition.VertexPosition.VID;
-				const int32 RefSurfaceEID = VIDToReferenceEID[SurfaceVID];
-
-				FMeshGeodesicSurfaceTracer::FMeshTangentDirection TangentAtVertex = { SurfaceVID, RefSurfaceEID, PolarAngle };
-				SurfaceTracer.TraceMeshFromVertex(TangentAtVertex, IntrinsicEdgeLength);
-			}
-			break;
-			case FSurfacePoint::EPositionType::Edge:
-			{
-				const int32 RefSurfaceEID = SurfacePosition.EdgePosition.EdgeID;
-				double Alpha = SurfacePosition.EdgePosition.Alpha;
-				// convert to BC
-				const int32 SurfaceTID = SurfaceMesh->GetEdgeT(RefSurfaceEID).A;
-				const FIndex2i SurfaceEdgeV = SurfaceMesh->GetEdgeV(RefSurfaceEID);
-				const int32 IndexOf = SurfaceMesh->GetTriEdges(SurfaceTID).IndexOf(RefSurfaceEID);
-				const bool bSameOrientation = (SurfaceMesh->GetTriangle(SurfaceTID)[IndexOf] == SurfaceEdgeV.A);
-				FVector3d BaryPoint(0., 0., 0.);
-				if (!bSameOrientation)
-				{
-					Alpha = 1. - Alpha;
-				}
-				BaryPoint[IndexOf] = Alpha;
-				BaryPoint[(IndexOf + 1) % 3] = 1. - Alpha;
-
-				FMeshGeodesicSurfaceTracer::FMeshSurfaceDirection DirectionOnSurface(RefSurfaceEID, PolarAngle);
-				SurfaceTracer.TraceMeshFromBaryPoint(SurfaceTID, BaryPoint, DirectionOnSurface, IntrinsicEdgeLength);
-			}
-			break;
-			case FSurfacePoint::EPositionType::Triangle:
-			{
-				const int32 SurfaceTID = SurfacePosition.TriPosition.TriID;
-
-				const FVector3d BaryPoint(SurfacePosition.TriPosition.BarycentricCoords[0],
-					SurfacePosition.TriPosition.BarycentricCoords[1],
-					SurfacePosition.TriPosition.BarycentricCoords[2]);
-				const int32 RefSurfaceEID = TIDToReferenceEID[SurfaceTID];
-
-				FMeshGeodesicSurfaceTracer::FMeshSurfaceDirection DirectionOnSurface(RefSurfaceEID, PolarAngle);
-				SurfaceTracer.TraceMeshFromBaryPoint(SurfaceTID, BaryPoint, DirectionOnSurface, IntrinsicEdgeLength);
-			}
-			break;
-			default:
-			{
-				// shouldn't be able to reach this point
-				check(0);
-			}
-		}
-	}
+	FMeshGeodesicSurfaceTracer SurfaceTracer =  SignpostSufaceTraceUtil::TraceFromIntrinsicVert(SignpostData, StartVID, PolarAngle, IntrinsicEdgeLength);
 
 	// util to convert the trace result to a surface point, potentially snapping to a vertex if within the coalesce threshold 
 	auto TraceResultToSurfacePoint = [CoalesceThreshold, SurfaceMesh](const FMeshGeodesicSurfaceTracer::FTraceResult& TraceResult)->FSurfacePoint
@@ -1183,20 +1588,9 @@ EMeshResult UE::Geometry::FIntrinsicTriangulation::FlipEdge(int32 EID, FEdgeFlip
 		return EMeshResult::Failed_IsBoundaryEdge;
 	}
 	// capture the state of the triangles before the flip.
-	
-	FIndex2i Tris = GetEdgeT(EID);
-	const int32 PreFlipIndexOf[2] = {GetTriEdges(Tris[0]).IndexOf(EID), GetTriEdges(Tris[1]).IndexOf(EID)};
-
-	// edge polar directions before the flip
-	const FVector3d PreFlipDirections[2] = { IntrinsicEdgeAngles[Tris[0]], IntrinsicEdgeAngles[Tris[1]] };
-	const FVector3d PreFlipInteriorAngles[2] = { InternalAngles[Tris[0]], InternalAngles[Tris[1]] };
-
-	// polar directions for the four edges that don't change.
-	const double bcDir = PreFlipDirections[0][AddOneModThree[PreFlipIndexOf[0]]];
-	const double caDir = PreFlipDirections[0][AddTwoModThree[PreFlipIndexOf[0]]];
+	const FIndex2i Tris = GetEdgeT(EID);
+	const FIndex2i PreFlipIndexOf(GetTriEdges(Tris[0]).IndexOf(EID), GetTriEdges(Tris[1]).IndexOf(EID));
 
-	const double adDir = PreFlipDirections[1][AddOneModThree[PreFlipIndexOf[1]]];
-	const double dbDir = PreFlipDirections[1][AddTwoModThree[PreFlipIndexOf[1]]];
 
 	// flip edge in the underlying mesh
 	// this updates the edge lengths and the interior angles.
@@ -1208,106 +1602,36 @@ EMeshResult UE::Geometry::FIntrinsicTriangulation::FlipEdge(int32 EID, FEdgeFlip
 		return MeshFlipResult;
 	}
 
-	// compute the polar edge directions for the new edge (ie. dcDir and cdDir)
-	// dcDir 
-	const double newAngleAtD = InternalAngles[Tris[0]][1];
-	const double ToPolarAtD = GeometricVertexInfo[EdgeFlipInfo.OpposingVerts[1]].ToRadians;
-	const double dcDir = AsZeroToTwoPi(dbDir + ToPolarAtD * newAngleAtD);
+	// internal angles at the verts that are now connected by the flipped edge
+	const double NewAngleAtC = InternalAngles[Tris[1]][1];
+	const double NewAngleAtD = InternalAngles[Tris[0]][1];
 
-	// cdDir 
-	const double newAngleAtC = InternalAngles[Tris[1]][1];
-	const double ToPolarAtC =  GeometricVertexInfo[EdgeFlipInfo.OpposingVerts[0]].ToRadians;
-	const double cdDir = AsZeroToTwoPi(caDir + ToPolarAtC * newAngleAtC);
+	// update signpost
+	SignpostData.OnFlipEdge(EID, Tris, EdgeFlipInfo.OpposingVerts, PreFlipIndexOf, NewAngleAtC, NewAngleAtD);
 
-	IntrinsicEdgeAngles[Tris[0]] = FVector3d(cdDir, dbDir, bcDir); // edges leaving c, d, b
-	IntrinsicEdgeAngles[Tris[1]] = FVector3d(dcDir, caDir, adDir); // edges leaving d, c, a
-	
 	return EMeshResult::Ok;
 }
 
 
 double UE::Geometry::FIntrinsicTriangulation::UpdateVertexByEdgeTrace(const int32 NewVID, const int32 TraceStartVID, const double TracePolarAngle, const double TraceDist)
 {
-	const FSurfacePoint& StartSurfacePoint = IntrinsicVertexPositions[TraceStartVID];
-
-	double ActualDist = 0;
-	FMeshGeodesicSurfaceTracer SurfaceTracer(*ExtrinsicMesh);
-	switch (StartSurfacePoint.PositionType)
-	{
-	case FSurfacePoint::EPositionType::Vertex:
-	{
-		// convert vertex location
-		const int32 ExtrinsicVID = StartSurfacePoint.Position.VertexPosition.VID;
-		const int32 ExtrinsicRefEID = VIDToReferenceEID[ExtrinsicVID];
-
-		const FMeshGeodesicSurfaceTracer::FMeshTangentDirection TangentDirection = { ExtrinsicVID, ExtrinsicRefEID, TracePolarAngle };
-		ActualDist = SurfaceTracer.TraceMeshFromVertex(TangentDirection, TraceDist);
-
-	}
-	break;
-	case FSurfacePoint::EPositionType::Edge:
-	{
-		// convert edge location -  Not used, and might have bugs
-		check(0);
-
-		const int32 ExtrinsicRefEID = StartSurfacePoint.Position.EdgePosition.EdgeID;
-		const double Alpha = StartSurfacePoint.Position.EdgePosition.Alpha;
-		const FDynamicMesh3::FEdge ExtEdge = ExtrinsicMesh->GetEdge(ExtrinsicRefEID);
-		const int32 StartExtrinsicTID = ExtEdge.Tri.A; // Q? should I inspect the angle to determine which triangle we are entering?
-		FIndex3i ExtrinsicTriVIDs = ExtrinsicMesh->GetTriangle(StartExtrinsicTID);
-		FVector3d StartBarycentricCoords(0.);
-		StartBarycentricCoords[ExtrinsicTriVIDs.IndexOf(ExtEdge.Vert.A)] = Alpha;
-		StartBarycentricCoords[ExtrinsicTriVIDs.IndexOf(ExtEdge.Vert.B)] = 1. - Alpha;
-
-		FMeshGeodesicSurfaceTracer::FMeshSurfaceDirection SurfaceDirection(ExtrinsicRefEID, TracePolarAngle);
-	}
-	break;
-	case FSurfacePoint::EPositionType::Triangle:
-	{
-		// trace from BC point in triangle
-		const int32 StartExtrinsicTID = StartSurfacePoint.Position.TriPosition.TriID;
-		const int32 ExtrinsicRefEID = TIDToReferenceEID[StartExtrinsicTID];
-
-		FVector3d StartBarycentricCoords(0.);
-		for (int32 i = 0; i < 3; ++i)
-		{
-			StartBarycentricCoords[i] = StartSurfacePoint.Position.TriPosition.BarycentricCoords[i];
-		}
-
-		const FMeshGeodesicSurfaceTracer::FMeshSurfaceDirection SurfaceDirection(ExtrinsicRefEID, TracePolarAngle);
-		ActualDist = SurfaceTracer.TraceMeshFromBaryPoint(StartExtrinsicTID, StartBarycentricCoords, SurfaceDirection, TraceDist);
-	}
-	break;
-	default:
-	{
-		// not possible
-		check(0);
-	}
-	}
-
-	const FMeshGeodesicSurfaceTracer::FTraceResult& TraceResult = SurfaceTracer.GetTraceResults().Last();
-	// update new vertex position in two ways: the r3 position and the intrinsic position 
-	// 1) the r3 position
-	FVector3d TraceResultPos = ExtrinsicMesh->GetTriBaryPoint(TraceResult.TriID, TraceResult.Barycentric.X, TraceResult.Barycentric.Y, TraceResult.Barycentric.Z);
-	Vertices[NewVID] = TraceResultPos;
-	// 2) the intrinsic position (i.e. relative to the surface)  
-	//  Use "Triangle" type position for result ( Q: should we classify as "Edge" position if very close to edge ?)
-	FSurfacePoint TraceResultPosition(TraceResult.TriID, TraceResult.Barycentric);
-	IntrinsicVertexPositions.InsertAt(TraceResultPosition, NewVID);
-
-
-	// at the new vertex location: update the polar angle for each new edge relative to the reference direction in this extrinsic triangle. 
+	using FSurfaceTraceResult = FSignpost::FSurfaceTraceResult;
 
+	FMeshGeodesicSurfaceTracer SurfaceTracer = SignpostSufaceTraceUtil::TraceFromIntrinsicVert(SignpostData, TraceStartVID, TracePolarAngle, TraceDist);
 
+	TArray<FMeshGeodesicSurfaceTracer::FTraceResult>& TraceResultArray = SurfaceTracer.GetTraceResults();
+	// need to convert the result of the trace into the correct form
 
+	const FMeshGeodesicSurfaceTracer::FTraceResult& TraceResult = TraceResultArray.Last();
+	const FSurfacePoint TraceResultPosition(TraceResult.TriID, TraceResult.Barycentric);
 	// fix directions relative to local reference edge on the extrinsic mesh 
 	// by finding direction of the TraceEID edge indecent on NewVID
-	double AngleOffset = 0;
+	const double AngleOffset = [&]
 	{
 		FVector2d Dir = TraceResult.SurfaceDirection.Dir;
 
 		// translate to Dir about the reference edge for this triangle
-		int32 EndRefEID = TIDToReferenceEID[TraceResult.TriID];
+		const int32 EndRefEID = SignpostData.TIDToReferenceEID[TraceResult.TriID];
 
 		const FMeshGeodesicSurfaceTracer::FTangentTri2& TangentTri2 = SurfaceTracer.GetLastTri();
 		// convert to local basis for first edge of tri2
@@ -1325,10 +1649,26 @@ double UE::Geometry::FIntrinsicTriangulation::UpdateVertexByEdgeTrace(const int3
 		const double AngleToNewVert = FMath::Atan2(DirRelToRefEID.Y, DirRelToRefEID.X);
 		// reverse direction of path because we NewVert is the local origin for polar angles around new vert
 		const double AngleFromNewVert = AngleToNewVert + TMathUtilConstants<double>::Pi;
-		AngleOffset = AsZeroToTwoPi(AngleFromNewVert);
-	}
+		return AsZeroToTwoPi(AngleFromNewVert);
+	}();
+
+	// Trace result as position and angle.
+	FSurfaceTraceResult SurfaceTraceResult = { TraceResultPosition, AngleOffset };
+
+	// As R3
+	bool bTmpIsValid;
+	const FVector3d TraceResultPos = AsR3Position(SurfaceTraceResult.SurfacePoint, *SignpostData.SurfaceMesh, bTmpIsValid);
+
+	// update the R3 for NewVID in the intrinsic mesh
+	Vertices[NewVID] = TraceResultPos;
+
+	// store the surface position for the intrinsic vertex in the signpost data
+	// Currently we are storing every new intrinsic surface position as a point on a surface tri
+	//  TODO: is there any advantage to classify as "Edge" position if very close to edge ?  
+	SignpostData.IntrinsicVertexPositions.InsertAt(SurfaceTraceResult.SurfacePoint, NewVID);
 
-	return AngleOffset;
+	// return relative angle of trace
+	return SurfaceTraceResult.Angle;
 }
 
 
@@ -1398,7 +1738,7 @@ UE::Geometry::EMeshResult UE::Geometry::FIntrinsicTriangulation::PokeTriangle(in
 	const FIndex3i OriginalVIDs = GetTriangle(TID);
 	const FIndex3i OriginalEdges = GetTriEdges(TID);
 	const FVector3d OriginalTriEdgeLengths = GetTriEdgeLengths(TID);
-	const FVector3d OriginalEdgeDirs = IntrinsicEdgeAngles[TID];
+	const FVector3d OriginalEdgeDirs = SignpostData.IntrinsicEdgeAngles[TID];
 	
 	// Add a new vertex and faces to the IntrinsicMesh.   
 	// Note: the r3 position will be wrong initially since this poke will just interpolate the corners of the tri.
@@ -1461,18 +1801,18 @@ UE::Geometry::EMeshResult UE::Geometry::FIntrinsicTriangulation::PokeTriangle(in
 	TriEdgeDir[2][0] = OriginalEdgeDirs[2]; // CtoA dir
 	
 	// edges from the corners of the original tri towards the new vertex at the "center"
-	const double ToRadiansAtA = GeometricVertexInfo[OriginalVIDs[0]].ToRadians;
-	const double ToRadiansAtB = GeometricVertexInfo[OriginalVIDs[1]].ToRadians;
-	const double ToRadiansAtC = GeometricVertexInfo[OriginalVIDs[2]].ToRadians;
+	const double ToRadiansAtA = SignpostData.GeometricVertexInfo[OriginalVIDs[0]].ToRadians;
+	const double ToRadiansAtB = SignpostData.GeometricVertexInfo[OriginalVIDs[1]].ToRadians;
+	const double ToRadiansAtC = SignpostData. GeometricVertexInfo[OriginalVIDs[2]].ToRadians;
 	TriEdgeDir[0][1] = AsZeroToTwoPi( OriginalEdgeDirs[1] + ToRadiansAtB * InternalAngles[NewTris[1]][0]);  // BtoNew dir
 	TriEdgeDir[1][1] = AsZeroToTwoPi( OriginalEdgeDirs[2] + ToRadiansAtC * InternalAngles[NewTris[2]][0]);  // CtoNew dir
 	TriEdgeDir[2][1] = AsZeroToTwoPi( OriginalEdgeDirs[0] + ToRadiansAtA * InternalAngles[NewTris[0]][0]);  // AtoNew dir   
 
 	// edges from new vertex to a, to b, and to c, using new-to-A as the zero direction.  These will be updated later when we learn the correct angle for new-to-A
-	TriEdgeDir[0][2] = 0.;                                                         // NewToA dir 
+	TriEdgeDir[0][2] = 0.;                                                                // NewToA dir 
 	TriEdgeDir[1][2] = InternalAngles[NewTris[0]][2];                                     // NewToB dir
 	TriEdgeDir[2][2] = InternalAngles[NewTris[0]][2] + InternalAngles[NewTris[1]][2];     // NewToC dir
-	GeometricVertexInfo.InsertAt(FGeometricInfo(), NewVID); // we want the default of false, and 1.
+	SignpostData.GeometricVertexInfo.InsertAt(FSignpost::FGeometricInfo(), NewVID); // we want the default of false, and 1.
 
 	// (4) update the surface point for the new vertex by tracing one of the new edges. 
 	// 
@@ -1495,9 +1835,9 @@ UE::Geometry::EMeshResult UE::Geometry::FIntrinsicTriangulation::PokeTriangle(in
 	}
 	
 	// record the new edge directions.
-	IntrinsicEdgeAngles.InsertAt(TriEdgeDir[2], NewTris[2]);
-	IntrinsicEdgeAngles.InsertAt(TriEdgeDir[1], NewTris[1]);
-	IntrinsicEdgeAngles[NewTris[0]] = TriEdgeDir[0];
+	SignpostData.IntrinsicEdgeAngles.InsertAt(TriEdgeDir[2], NewTris[2]);
+	SignpostData.IntrinsicEdgeAngles.InsertAt(TriEdgeDir[1], NewTris[1]);
+	SignpostData.IntrinsicEdgeAngles[NewTris[0]] = TriEdgeDir[0];
 	
 	// record the coordinate actually used. 
 	PokeInfo.BaryCoords = BaryCoordinates;
@@ -1694,7 +2034,7 @@ UE::Geometry::EMeshResult  UE::Geometry::FIntrinsicTriangulation::SplitEdge(int3
 		const FIndex3i OriginalT0Edges = Permute(IndexOfe, GetTriEdges(TID));  // as (ab, bc, ca)
 		
 		const FVector3d OriginalT0EdgeLengths = Permute(IndexOfe, GetTriEdgeLengths(TID));    // as (|ab|, |bc|, |ca|)
-		const FVector3d OriginalT0EdgeDirs    = Permute(IndexOfe, IntrinsicEdgeAngles[TID]);  // as ( aTob, bToc, cToa)
+		const FVector3d OriginalT0EdgeDirs    = Permute(IndexOfe, SignpostData.IntrinsicEdgeAngles[TID]);  // as ( aTob, bToc, cToa)
 
 		
 		// Update the connectivity with the edge split
@@ -1728,8 +2068,8 @@ UE::Geometry::EMeshResult  UE::Geometry::FIntrinsicTriangulation::SplitEdge(int3
 		const FVector3d TIDInternalAngles = Permute(IndexOfe, InternalAngles[TID]);
 		
 		// update the directions.  Say DirFtoB is zero
-		const double ToRadiansAtC = GeometricVertexInfo[OriginalT0VIDs[2]].ToRadians;
-		GeometricVertexInfo.InsertAt(FGeometricInfo(), NewVID); // we want the default of false, and 1.
+		const double ToRadiansAtC = SignpostData.GeometricVertexInfo[OriginalT0VIDs[2]].ToRadians;
+		SignpostData.GeometricVertexInfo.InsertAt(FSignpost::FGeometricInfo(), NewVID); // we want the default of false, and 1.
 
 		FVector3d TriEdgeDir[2]; // one for each tri in order (TID, NewTID )
 		// edges around the boundary of the original tri
@@ -1753,8 +2093,8 @@ UE::Geometry::EMeshResult  UE::Geometry::FIntrinsicTriangulation::SplitEdge(int3
 		TriEdgeDir[1][0] = AsZeroToTwoPi(TriEdgeDir[1][0] + ToLocalDir);
 		
 		// store angle results
-		IntrinsicEdgeAngles.InsertAt(TriEdgeDir[1], NewTID);
-		IntrinsicEdgeAngles[TID] = Permute((3 - IndexOfe)%3,  TriEdgeDir[0]);
+		SignpostData.IntrinsicEdgeAngles.InsertAt(TriEdgeDir[1], NewTID);
+		SignpostData.IntrinsicEdgeAngles[TID] = Permute((3 - IndexOfe)%3,  TriEdgeDir[0]);
 
 
 		return Result;
@@ -1770,7 +2110,7 @@ UE::Geometry::EMeshResult  UE::Geometry::FIntrinsicTriangulation::SplitEdge(int3
 		const FIndex3i OriginalT0Edges = Permute(IndexOfe, GetTriEdges(TID));  // as (ab, bc, ca)
 
 		const FVector3d OriginalT0EdgeLengths = Permute(IndexOfe, GetTriEdgeLengths(TID));    // as (|ab|, |bc|, |ca|)
-		const FVector3d OriginalT0EdgeDirs    = Permute(IndexOfe, IntrinsicEdgeAngles[TID]);  // as ( aTob, bToc, cToa)
+		const FVector3d OriginalT0EdgeDirs    = Permute(IndexOfe, SignpostData.IntrinsicEdgeAngles[TID]);  // as ( aTob, bToc, cToa)
 		
 		// Info about the original T0 tri, reordered to make the split edge the first edge..
 		const int32 TID1 = OriginalEdge.Tri[1];
@@ -1780,7 +2120,7 @@ UE::Geometry::EMeshResult  UE::Geometry::FIntrinsicTriangulation::SplitEdge(int3
 		const FIndex3i OriginalT1Edges = Permute(IndexOfT1e, GetTriEdges(TID1));  // as (ba, ad, db)
 		
 		const FVector3d OriginalT1EdgeLengths = Permute(IndexOfT1e, GetTriEdgeLengths(TID1));   // as (|ba|, |ad|, |db})
-		const FVector3d OriginalT1EdgeDirs    = Permute(IndexOfT1e, IntrinsicEdgeAngles[TID1]); // as (bToa, aTod, dTob) 
+		const FVector3d OriginalT1EdgeDirs    = Permute(IndexOfT1e, SignpostData.IntrinsicEdgeAngles[TID1]); // as (bToa, aTod, dTob) 
 		
 
 		// Update the connectivity with the edge split
@@ -1826,9 +2166,9 @@ UE::Geometry::EMeshResult  UE::Geometry::FIntrinsicTriangulation::SplitEdge(int3
 		InternalAngles[TID1] = ComputeTriInternalAnglesR(TID1);
 
 		// update the directions.  Say DirFtoB is zero
-		const double ToRadiansAtC = GeometricVertexInfo[OriginalT0VIDs[2]].ToRadians;
-		const double ToRadiansAtD = GeometricVertexInfo[OriginalT1VIDs[2]].ToRadians;
-		GeometricVertexInfo.InsertAt(FGeometricInfo(), NewVID); // we want the default of false, and 1.
+		const double ToRadiansAtC = SignpostData.GeometricVertexInfo[OriginalT0VIDs[2]].ToRadians;
+		const double ToRadiansAtD = SignpostData.GeometricVertexInfo[OriginalT1VIDs[2]].ToRadians;
+		SignpostData.GeometricVertexInfo.InsertAt(FSignpost::FGeometricInfo(), NewVID); // we want the default of false, and 1.
 
 		const FVector3d TIDInternalAngles  = Permute(IndexOfe, InternalAngles[TID]);
 		const FVector3d TID1InternalAngles = Permute(IndexOfT1e, InternalAngles[TID1]);
@@ -1866,12 +2206,148 @@ UE::Geometry::EMeshResult  UE::Geometry::FIntrinsicTriangulation::SplitEdge(int3
 		TriEdgeDir[3][0] = AsZeroToTwoPi(TriEdgeDir[3][0] + ToLocalDir);
 
 		// store angle results
-		IntrinsicEdgeAngles.InsertAt(TriEdgeDir[3], NewTID1);
-		IntrinsicEdgeAngles.InsertAt(TriEdgeDir[1], NewTID0);
-		IntrinsicEdgeAngles[TID]  = Permute((3 - IndexOfe) % 3, TriEdgeDir[0]);
-		IntrinsicEdgeAngles[TID1] = Permute((3 - IndexOfT1e) % 3, TriEdgeDir[2]);
+		SignpostData.IntrinsicEdgeAngles.InsertAt(TriEdgeDir[3], NewTID1);
+		SignpostData.IntrinsicEdgeAngles.InsertAt(TriEdgeDir[1], NewTID0);
+		SignpostData.IntrinsicEdgeAngles[TID]  = Permute((3 - IndexOfe) % 3, TriEdgeDir[0]);
+		SignpostData.IntrinsicEdgeAngles[TID1] = Permute((3 - IndexOfT1e) % 3, TriEdgeDir[2]);
 
 		return Result;
 	}
 }
 
+/**------------------------------------------------------------------------------
+*  FIntrinsicEdgeFlipMesh Methods
+*-------------------------------------------------------------------------------*/
+FIntrinsicEdgeFlipMesh::FIntrinsicEdgeFlipMesh(const FDynamicMesh3& SurfaceMesh)
+	: MyBase(SurfaceMesh)
+	, NormalCoordinates(SurfaceMesh)
+{
+}
+
+EMeshResult FIntrinsicEdgeFlipMesh::FlipEdge(int32 EID, FEdgeFlipInfo& EdgeFlipInfo)
+{
+	if (!IsEdge(EID))
+	{
+		return EMeshResult::Failed_NotAnEdge;
+	}
+
+	if (IsBoundaryEdge(EID))
+	{
+		return EMeshResult::Failed_IsBoundaryEdge;
+	}
+
+	// state before flip, needed when updating the normal coords after the flip
+	const FIndex2i EdgeT  = GetEdgeT(EID);
+	const FIndex3i TAEIDs = GetTriEdges(EdgeT.A);
+	const FIndex3i TBEIDs = GetTriEdges(EdgeT.B);
+	const FIndex2i OppVs  = GetEdgeOpposingV(EID);
+
+	EMeshResult FlipResult = MyBase::FlipEdge(EID, EdgeFlipInfo);
+
+	if (FlipResult == EMeshResult::Ok)
+	{
+		// update the normal coords
+		NormalCoordinates.OnFlipEdge(EdgeT.A, TAEIDs, OppVs.A, EdgeT.B, TBEIDs, OppVs.B, EID);
+	}
+
+	return FlipResult;
+}
+
+TArray<FIntrinsicEdgeFlipMesh::FEdgeAndCrossingIdx> FIntrinsicEdgeFlipMesh::GetImplicitEdgeCrossings(const int32 SurfaceEID, const bool bReverse) const
+{
+	return FNormalCoordSurfaceTraceImpl::TraceSurfaceEdge(*this, SurfaceEID, bReverse);
+}
+
+FIntrinsicEdgeFlipMesh::FEdgeCorrespondence::FEdgeCorrespondence(const FIntrinsicEdgeFlipMesh& Mesh)
+	: IntrinsicMesh(&Mesh)
+	, SurfaceMesh(Mesh.GetNormalCoordinates().SurfaceMesh)
+{
+	const int32 IntrinsicMaxEID = IntrinsicMesh->MaxEdgeID();
+	const int32 SurfaceMaxEID   = SurfaceMesh->MaxEdgeID();
+
+	SurfaceEdgesCrossed.SetNum(IntrinsicMaxEID);
+	const IntrinsicCorrespondenceUtils::FNormalCoordinates& NormalCoords = IntrinsicMesh->GetNormalCoordinates();
+
+	// allocate the SurfaceEdgesCrossed. From the normal coordinates we know how many surface edge crossings each intrinsic edge sees
+	for (int32 IntrinsicEID = 0; IntrinsicEID < IntrinsicMaxEID; ++IntrinsicEID)
+	{
+		if (!IntrinsicMesh->IsEdge(IntrinsicEID))
+		{
+			continue;
+		}
+		// number of times a surface edge crosses this intrinsic edge
+		int32 NumXings = NormalCoords.NormalCoord[IntrinsicEID];
+
+		if (NumXings > 0)
+		{
+			SurfaceEdgesCrossed[IntrinsicEID].SetNum(NumXings);
+
+		} // else don't bother making an entry when NumXings == 0 since that means the edges are the same on both meshes.
+	}
+
+	// trace each surface edge across the intrinsic mesh and construct an ordered list of surface edges crossing each intrinsic edge.
+	// the order of the crossings should be consistent with the edge direction relative to the first adjacent tri 
+	// (i.e. starting at the corner Mesh.GetTriEdges(GetEdgeT(EID).A).IndexOf(EID) ) 
+
+	for (int32 SurfaceEID = 0; SurfaceEID < SurfaceMaxEID; ++SurfaceEID)
+	{
+		if (!SurfaceMesh->IsEdge(SurfaceEID))
+		{
+			continue;
+		}
+		const TArray<FEdgeAndCrossingIdx> IntrinsicEdgeXings = IntrinsicMesh->GetImplicitEdgeCrossings(SurfaceEID, false /* = bReverseTrace*/);
+
+		for (int32 i = 0; i < IntrinsicEdgeXings.Num(); ++i)
+		{
+			const FEdgeAndCrossingIdx& EdgeXing = IntrinsicEdgeXings[i];
+			bool bIsEndVertex                   = (EdgeXing.CIdx == 0);
+
+			if (!bIsEndVertex)
+			{
+				const int32 IntrinsicEID      = EdgeXing.EID;
+				const FIndex2i IntrinsicEdgeT = IntrinsicMesh->GetEdgeT(IntrinsicEID);
+				checkSlow(IntrinsicEdgeT.A == EdgeXing.TID || IntrinsicEdgeT.B == EdgeXing.TID)
+
+					// array of surface edges this intrinsic edge crosses, these should be ordered relative to the direction of TriA.
+					TArray<int32>& XingSurfaceEdges = SurfaceEdgesCrossed[IntrinsicEID];
+
+				// crossings count from the bottom of the edge to the top relative to EdgeXing.TID 
+				const int32 XingID = (IntrinsicEdgeT.A == EdgeXing.TID) ? EdgeXing.CIdx - 1 : XingSurfaceEdges.Num() - EdgeXing.CIdx;
+
+				checkSlow(XingID > -1);
+				XingSurfaceEdges[XingID] = SurfaceEID;
+			}
+		}
+	}
+}
+
+TArray<FIntrinsicEdgeFlipMesh::FSurfacePoint>
+FIntrinsicEdgeFlipMesh::TraceSurfaceEdge(int32 SurfaceEID, double CoalesceThreshold, bool bReverse) const
+{
+	const FDynamicMesh3& HostMesh                   = *this->GetExtrinsicMesh();
+	const FIntrinsicEdgeFlipMesh& TraceMesh         = *this;
+	TArray<FEdgeAndCrossingIdx> EdgeAndCrossingIdxs = this->GetImplicitEdgeCrossings(SurfaceEID, bReverse);
+
+	TArray<int32> HostEdgeCrossings;
+	HostEdgeCrossings.Reserve(EdgeAndCrossingIdxs.Num() - 2); // EdgeAndCrossingsIdx include the start and end vertex.  don't need them.
+	for (FEdgeAndCrossingIdx EdgeAndCrossing : EdgeAndCrossingIdxs)
+	{
+		if (EdgeAndCrossing.CIdx != 0) // skip the start and end vertex
+		{
+			HostEdgeCrossings.Add(EdgeAndCrossing.EID);
+		}
+	}
+
+	return FNormalCoordSurfaceTraceImpl::TraceEdgeOverHost(SurfaceEID, HostEdgeCrossings, HostMesh, TraceMesh, CoalesceThreshold, bReverse);
+}
+
+TArray<FIntrinsicEdgeFlipMesh::FSurfacePoint>
+FIntrinsicEdgeFlipMesh::FEdgeCorrespondence::TraceEdge(int32 IntrinsicEID, double CoalesceThreshold, bool bReverse) const
+{
+	const FDynamicMesh3& HostMesh           = *SurfaceMesh;
+	const FIntrinsicEdgeFlipMesh& TraceMesh = *IntrinsicMesh;
+	const TArray<int32>& HostEdgeCrossings  = SurfaceEdgesCrossed[IntrinsicEID];
+
+	return FNormalCoordSurfaceTraceImpl::TraceEdgeOverHost(IntrinsicEID, HostEdgeCrossings, HostMesh, TraceMesh, CoalesceThreshold, bReverse);
+}
+
diff --git a/Engine/Plugins/Runtime/GeometryProcessing/Source/DynamicMesh/Public/Operations/IntrinsicCorrespondenceUtils.h b/Engine/Plugins/Runtime/GeometryProcessing/Source/DynamicMesh/Public/Operations/IntrinsicCorrespondenceUtils.h
new file mode 100644
--- /dev/null
+++ b/Engine/Plugins/Runtime/GeometryProcessing/Source/DynamicMesh/Public/Operations/IntrinsicCorrespondenceUtils.h
@@ -0,0 +1,301 @@
+// Copyright Epic Games, Inc. All Rights Reserved.
+
+#pragma once
+
+#include "CoreMinimal.h"
+#include "DynamicMesh/DynamicMesh3.h"
+#include "DynamicMesh/InfoTypes.h"
+#include "Util/DynamicVector.h"
+#include "VectorTypes.h"
+#include "MathUtil.h"
+
+
+namespace UE
+{
+namespace Geometry
+{
+
+// Utilities (primarily structs) used to form a relationship between an intrinsic mesh and the surface mesh that it lives on.
+namespace IntrinsicCorrespondenceUtils
+{
+	
+	// Location on an extrinsic (surface) mesh stored as a union.
+	struct DYNAMICMESH_API FSurfacePoint
+	{
+		enum class EPositionType
+		{
+			Vertex,
+			Edge,
+			Triangle
+		};
+
+		// constructors respect the correct position type
+		FSurfacePoint(); // defaults to invalid vertex position
+		FSurfacePoint(const int32 VID);
+		FSurfacePoint(const int32 EdgeID, const double Alpha);
+		FSurfacePoint(const int32 TriID, const FVector3d& BaryCentrics);
+
+		struct FVertexPosition
+		{
+			int32 VID;                      // Reference (extrinsic) vertex
+		};
+		struct FEdgePosition
+		{
+			int32 EdgeID;                  // Reference (extrinsic) edge
+			double Alpha;                  // Lerp coordinate along referenced edge as Pos(Edge.A) (Alpha) + (1-Alpha) Pos(Edge.B)
+		};
+		struct FTrianglePosition
+		{
+			int32 TriID;                    // Reference (extrinsic) triangle
+			double BarycentricCoords[3];    // Barycentric coordinates within referenced triangle
+		};
+		union FSurfacePositionUnion
+		{
+			FVertexPosition   VertexPosition;
+			FEdgePosition     EdgePosition;
+			FTrianglePosition TriPosition;
+		};
+
+		EPositionType PositionType;      // Position type, needed to correctly interpret the Position union
+		FSurfacePositionUnion Position;  // Position relative to vertex, edge, or triangle. 
+	};
+
+	/**
+	* @return the R3 position of this surface point as applied to the input Mesh
+	* on return bValidPoint will be false if the edge / vertex / triangle ID referenced by this surface
+	* point is not an element of the mesh in question
+	* 
+	* @param SurfacePoint - Position defined relative to the Mesh
+	* @param Mesh         - Dynamic mesh that defines a surface
+	* @param bValidPoint  - will be false if the vert, edge or tri that defines the surface point is not a member of Mesh
+	*/
+	FVector3d AsR3Position(const FSurfacePoint& SurfacePoint, const FDynamicMesh3& Mesh, bool& bValidPoint);
+
+	// A Connection stores the information needed to define a local reference direction for each vertex and triangle.
+	struct DYNAMICMESH_API FMeshConnection
+	{
+		FMeshConnection(const FDynamicMesh3& SurfMesh);
+		const FDynamicMesh3* SurfaceMesh;               // Pointer back to original mesh 
+		TArray<int32> VIDToReferenceEID;                // Reference (original mesh) Edge per (original mesh) Vertex, defines local zero angle on tangent plane
+		TArray<int32> TIDToReferenceEID;                // Reference (original mesh) Edge per (original mesh) Triangle, defines local zero angle on tangent plane
+	};
+
+	/**
+	* Normal Coordinates as defined by "Discrete Conformal Equivalence of Polyhedral Surfaces" - Gillespi, Springborn, Crane, TOG V40 No4, 2021
+	* This structure assumes that the surface mesh is fixed, and the intrinsic mesh shares the same vertex set as the surface mesh and
+	* is initialized with the same connectivity as the surface mesh. Note, it is assumed that the surface mesh has no bow-ties.
+	*/
+	struct DYNAMICMESH_API FNormalCoordinates : public FMeshConnection
+	{
+		typedef FMeshConnection  MyBase;
+
+		FNormalCoordinates(const FDynamicMesh3& SurfMesh);
+
+		struct FEdgeAndCrossingIdx
+		{
+			int32 TID;  // triangle traversed
+			int32 EID;  // edge crossed.  This will be one of the three edges in GetTriEdges(TID)
+			int32 CIdx; // index of curve relative to the edge crossed
+		};
+
+		/**
+		* @return number of surface edges within the specified intrinsic triangle that are adjacent to the specified triangle corner.
+		* This excludes any surface edges that are along the intrinsic triangle boundary edges.
+		* note it is assumed that IntrinsicTriEIDs are edge ID of an intrinsic triangle, but no checking is done.
+		*/
+		inline int32 NumCornerEmanatingRefEdges(const FIndex3i& IntrinsicTriEIDs, const int32 TriCornerIndex) const;
+
+		/**
+		* @return number of surface edges that cross corner the specified corner of the intrinsic triangle defined by the given edges.
+		* note, it is assumed that IntrinsicTriEIDs are edge ID of an intrinsic triangle, but no checking is done.
+		*/
+		inline int32 NumCornerCrossingRefEdges(const FIndex3i& IntrinsicTriEIDs, const int32 TriCornerIndex) const;
+
+		/**
+		* Update the NormalCoords and RoundaboutOrder after an intrinsic mesh edge flip. TriAEID, and TriBEID are the
+		* triangle edges of the two affected triangles prior (!) to flipping the specified edge.
+		* @return false if something failed
+		*/
+		bool  OnFlipEdge(const int32 T0ID, const FIndex3i& BeforeT0EIDs, const int32 T0OppVID,
+			             const int32 T1ID, const FIndex3i& BeforeT1EIDs, const int32 T1OppVID,
+			             const int32 FlippedEID);
+
+		/**
+		* @return the edgeID of the N-th out-going edge in the SurfaceMesh (counting CCW from the reference edge for this vertex).
+		*/
+		int32 GetNthEdgeID(const int32 VID, const int32 N) const;
+
+		/**
+		* @return the Order of the edge EID about VID (counting CCW from the surface reference edge for this vertex).
+		* Note: this assumes that edge EID is adjacent to vertex VID - otherwise -1 is returned.
+		*/
+		int32 GetEdgeOrder(const int32 VID, const int32 SurfaceEID) const;
+
+		TArray<int32>     NormalCoord;     // Per intrinsic edge: number of times src mesh edge crosses. NormalCoord[EID] is 0 if and only if same as edge on surface mesh   
+		TArray<int32>     RefVertDegree;   // Per vertex in ref mesh, caches the result of counting the edges in RefMesh.VtxEdgeIter().  
+		TArray<FIndex3i>  RoundaboutOrder; // Per intrinsic Triangle, per (directed) edge: indicates the closest (when traveling ccw) reference mesh edge  
+		TArray<FIndex3i>  EdgeOrder;       // Per surface triangle, per directed edge: number of edges past reference edge traveling ccw.
+	};
+
+
+	/**
+	* Signpost Coordinates inspired by "Navigating Intrinsic Triangulations" Sharp, Soliman and Crane [2019, ACM Transactions on Graphics]
+	* This structure assumes that the surface mesh is fixed, but allows the intrinsic mesh to support operations like edge splits and triangle pokes.
+	* is initialized with the same connectivity as the surface mesh.  Note, it is assumed that the surface mesh has no bow-ties.
+	*/
+	struct DYNAMICMESH_API FSignpost : public FMeshConnection
+	{
+		typedef FMeshConnection  MyBase;
+		FSignpost(const FDynamicMesh3& SurfaceMesh);
+
+		/**
+		* Update the signpost after an edge flip.  This requires information about the both the pre and post flip state.
+		* @param EID                  - edge that was flipped
+		* @param AdjTris              - IDs of the tris adjacent to flipped edge 
+		* @param OpposingVerts        - verts opposite the edge, pre-flip.  i.e. EdgeFlipInfo.OpposingVerts
+		* @param PreFlipIndexOf       - edge index in the EID in the pre-fliped tris. i.e. {GetTriEdges(Tris[0]).IndexOf(EID), GetTriEdges(Tris[1]).IndexOf(EID)}; 
+		* @param Opp0NewInternalAngle - internal angle of opposing vert 0 after the flip. i.e. InternalAngles[Tris[1]][1]
+		* @param Opp1NewInternalAngle - internal angle of opposing vert 1 after the flip. i.e  InternalAngles[Tris[0]][1];
+		*/
+		void OnFlipEdge(const int32 EID, const FIndex2i AdjTris, const FIndex2i OpposingVerts, const FIndex2i PreFlipIndexOf, 
+		                const double Opp0NewInternalAngle, const double Opp1NewInternalAngle);
+
+		struct FSurfaceTraceResult
+		{
+			FSurfacePoint SurfacePoint; // point on the surface mesh
+			double Angle;               // polar angle relative to reference edge (on surface mesh) in the neighborhood of the surface point.
+		};
+	
+		struct FGeometricInfo
+		{
+			double ToRadians = 1.;    // 2pi / sum of angles at this vertex.  Will be 1 if the Gaussian curvature is 0
+			bool bIsInterior = false; // false if on boundary or at a bow-tie.
+		};
+		TDynamicVector<FVector3d>            IntrinsicEdgeAngles;             // Per-intrinsic Triangle - polar angles of each edge relative to the vertex it exits (i.e. edge 0 relative to vertex 0)
+		TDynamicVector<FSurfacePoint>        IntrinsicVertexPositions;        // Per-intrinsic Vertex - Locates intrinsic mesh vertices relative to the surface defined by extrinsic mesh.
+		TDynamicVector<FGeometricInfo>       GeometricVertexInfo;             // Per-intrinsic Vertex - Simple geometric information about the extrinsic surface in the neighborhood of intrinsic vert 	
+	};
+
+
+
+
+
+	int32 FNormalCoordinates::NumCornerEmanatingRefEdges(const FIndex3i& IntrinsicTriEIDs, const int32 TriCornerIndex) const
+	{
+		// see Gillespi et al, eqn 13
+		const int32 k_index = TriCornerIndex;
+		const int32 i_index = (TriCornerIndex + 1) % 3;
+		const int32 j_index = (TriCornerIndex + 2) % 3;
+
+		const int32 Edgeij = IntrinsicTriEIDs[i_index];
+		const int32 Edgejk = IntrinsicTriEIDs[j_index];
+		const int32 Edgeki = IntrinsicTriEIDs[k_index];
+
+		const int32 N_ij = NormalCoord[Edgeij];
+		const int32 N_jk = NormalCoord[Edgejk];
+		const int32 N_ki = NormalCoord[Edgeki];
+
+
+		return TMathUtil<int32>::Max(0, N_ij - N_jk - N_ki);
+	}
+
+	int32 FNormalCoordinates::NumCornerCrossingRefEdges(const FIndex3i& IntrinsicTriEIDs, const int32 TriCornerIndex) const
+	{
+		const int32 k_index = TriCornerIndex;
+		const int32 i_index = (TriCornerIndex + 1) % 3;
+		const int32 j_index = (TriCornerIndex + 2) % 3;
+
+		const int32 Edgeij = IntrinsicTriEIDs[i_index];
+		const int32 Edgejk = IntrinsicTriEIDs[j_index];
+		const int32 Edgeki = IntrinsicTriEIDs[k_index];
+
+		const int32 N_ij = NormalCoord[Edgeij];
+		const int32 N_jk = NormalCoord[Edgejk];
+		const int32 N_ki = NormalCoord[Edgeki];
+
+		// see Gillespi et al, eqn 14
+
+		const int32 Ejk_i = NumCornerEmanatingRefEdges(IntrinsicTriEIDs, i_index);
+		const int32 Eki_j = NumCornerEmanatingRefEdges(IntrinsicTriEIDs, j_index);
+
+
+		return  (TMathUtil<int32>::Max(0, N_jk + N_ki - N_ij) - Ejk_i - Eki_j) / 2;
+	}
+
+
+	// Utility to walk around a vertex (VID) visiting the adjacent edges and triangles.
+	// Note: this isn't appropriate for bow-ties as it will terminate upon encountering a mesh edge.
+	// If the vertex is on the mesh boundary, StartEID should correspond to the first boundary edge allowing 
+	// for a ccw traversal of the vertex-adjacent faces.
+	template <typename MeshType, typename FunctorType>
+	bool VisitVertexAdjacentElements(const MeshType& Mesh, const int32 VID, const int32 StartEID, FunctorType& Functor)
+	{
+		constexpr static int InvalidID = IndexConstants::InvalidID;
+
+		if (!Mesh.IsEdge(StartEID) || !Mesh.IsVertex(VID))
+		{
+			return false;
+		}
+
+		const int32 EndEID = [&]
+			{
+				int32 EEID = StartEID;
+				for (int32 EID : Mesh.VtxEdgesItr(VID))
+				{
+					if (EID != StartEID && Mesh.IsBoundaryEdge(EID))
+					{
+						EEID = EID;
+						break;
+					}
+				}
+				return EEID;
+			}();
+
+		int32 CurEID        = StartEID;
+		FIndex2i CurEdgeT   = Mesh.GetEdgeT(CurEID);
+		int32 CurTID        = CurEdgeT[0];
+		FIndex3i CurTriEIDs = Mesh.GetTriEdges(CurTID);
+		int32 IndexOf       = CurTriEIDs.IndexOf(CurEID);
+
+		if (Mesh.GetTriangle(CurTID)[IndexOf] != VID)
+		{
+			CurTID = CurEdgeT[1];
+			if (CurTID == FDynamicMesh3::InvalidID)
+			{
+				checkSlow(0); // this shouldn't happen.  When the vertex is on a boundary, we should have selected the correct start edge. 
+				return false;
+			}
+			CurTriEIDs = Mesh.GetTriEdges(CurTID);
+			IndexOf = CurTriEIDs.IndexOf(StartEID);
+		}
+
+		const int32 NextEdgeOffset = 2;
+		do
+		{
+
+			if (Functor(CurTID, CurEID, IndexOf))
+			{
+				break;
+			}
+
+			// advance to next edge
+			CurEID   = CurTriEIDs[(NextEdgeOffset + IndexOf) % 3];
+			CurEdgeT = Mesh.GetEdgeT(CurEID);
+			CurTID   = (CurEdgeT[0] == CurTID) ? CurEdgeT[1] : CurEdgeT[0];
+
+			if (CurTID == InvalidID)
+			{
+				checkSlow(CurEID == EndEID);
+				break;
+			}
+
+			CurTriEIDs = Mesh.GetTriEdges(CurTID);
+			IndexOf    = CurTriEIDs.IndexOf(CurEID);
+
+		} while (CurEID != StartEID);
+
+		return true;
+	}
+};
+};
+};
\ No newline at end of file
diff --git a/Engine/Plugins/Runtime/GeometryProcessing/Source/DynamicMesh/Public/Operations/IntrinsicTriangulationMesh.h b/Engine/Plugins/Runtime/GeometryProcessing/Source/DynamicMesh/Public/Operations/IntrinsicTriangulationMesh.h
--- a/Engine/Plugins/Runtime/GeometryProcessing/Source/DynamicMesh/Public/Operations/IntrinsicTriangulationMesh.h
+++ b/Engine/Plugins/Runtime/GeometryProcessing/Source/DynamicMesh/Public/Operations/IntrinsicTriangulationMesh.h
@@ -6,6 +6,7 @@
 #include "DynamicMesh/DynamicMesh3.h"
 #include "DynamicMesh/InfoTypes.h"
 #include "Util/DynamicVector.h"
+#include "IntrinsicCorrespondenceUtils.h"
 #include "VectorTypes.h"
 #include "MathUtil.h"
 
@@ -15,28 +16,31 @@ namespace UE
 namespace Geometry
 {
 /**
-* FIntrinsicEdgeFlipMesh supports edge flips, but no other operations that change
+* FSimpleIntrinsicEdgeFlipMesh supports edge flips, but no other operations that change
 * the mesh connectivity, and no operations that affect the vertices.
 *
 * An Intrinsic Mesh can be thought of as a mesh that overlays another mesh, sharing
 * the original surface (and in this case vertex locations), but the edges of the intrinsic mesh
-* (while constrained to be on the mesh surface) need not align with the edges of the original mesh and the
+* (while constrained to be on the original mesh surface) need not align with the edges of the original mesh and the
 * lengths of these edges are measured on the surface of the original mesh (not the R3 distance).
 *
-* The FIntrinsicEdgeFlipMesh is designed to work with the function FlipToDelaunay() to make an
+* The FSimpleIntrinsicEdgeFlipMesh is designed to work with the function FlipToDelaunay() to make an
 * Intrinsic Delaunay Triangulation (IDT) of the same surface as a given FDynamicMesh,
-* allowing for more robust cotan-laplacians (see LaplacianMatrixAssembly.h)
+* allowing for more robust cotangent-Laplacians (see LaplacianMatrixAssembly.h)
+*
+* This simple mesh alone does not support computing surface positions from intrinsic mesh positions,
+* to do such computations FIntrinsicEdgeFlipMesh or FIntrinsicTriangulation should be used.
 *
-* Note: the implementation is a simple triangle-based mesh but there is no restriction that 
-* edge(a, b) is unique (e.g. multiple intrinsic edges may connect the same two vertices with different paths). 
-* In fact this mesh allows triangles and edges with repeated vertices. Such structures arise naturally 
+* Note: the implementation is a simple triangle-based mesh but there is no restriction that
+* edge(a, b) is unique (e.g. multiple intrinsic edges may connect the same two vertices with different paths).
+* In fact this mesh allows triangles and edges with repeated vertices. Such structures arise naturally
 * with intrinsic triangulation, for example an edge that starts and ends at the
 * same vertex may encircle a mesh.
 *
 * The API of FIntrinsicEdgeFlipMesh is similar to FDynamicMesh3, but only has a minimal subset of methods and some
 * such as ::FlipEdge() and ::GetEdgeOpposingV() have very different implementations
 */
-class DYNAMICMESH_API FIntrinsicEdgeFlipMesh
+class DYNAMICMESH_API FSimpleIntrinsicEdgeFlipMesh
 {
 public:
 	
@@ -46,10 +50,10 @@ public:
 	constexpr static int InvalidID = IndexConstants::InvalidID;
 
 	/** Constructor does ID-matching deep copy the basic mesh topology (but no attributes or groups) */
-	FIntrinsicEdgeFlipMesh(const FDynamicMesh3& SrcMesh);
+	FSimpleIntrinsicEdgeFlipMesh(const FDynamicMesh3& SrcMesh);
 
-	FIntrinsicEdgeFlipMesh() = delete;
-	FIntrinsicEdgeFlipMesh(const FIntrinsicEdgeFlipMesh&) = delete;
+	FSimpleIntrinsicEdgeFlipMesh() = delete;
+	FSimpleIntrinsicEdgeFlipMesh(const FSimpleIntrinsicEdgeFlipMesh&) = delete;
 
 	/**
 	* Flip a single edge in the Intrinsic Mesh.
@@ -328,21 +332,150 @@ protected:
 	// intrinsic quantities 
 	
 	TDynamicVector<double>            EdgeLengths;           // Length of each intrinsic edge in the Mesh: note these may differ from (Vert0 - Vert1).Lenght().
-	TDynamicVector<FVector3d>         InternalAngles;        // Interior Angles for each triangle: note may differ from FDynamicMesh::GetTriInternalAnglesR(TID) as they are computed from intrinsic edge lenghts
+	TDynamicVector<FVector3d>         InternalAngles;        // Interior Angles for each triangle: note may differ from FDynamicMesh::GetTriInternalAnglesR(TID) as they are computed from intrinsic edge lengths
 
 };
 
 
+/**
+* FIntrinsicEdgeFlipMesh supports edge flips, but no other operations that change
+* the mesh connectivity, and no operations that affect the vertices.
+*
+* Unlike, FSimpleIntrinsicEdgeFlipMesh from which it is derived, a correspondence can be 
+* formed between the two meshes: the surface mesh edges can be traced across this intrinsic mesh, 
+* and this in-turn can be used to trace the intrinsic mesh edges across the surface mesh.  
+*
+* The implementation uses integer coordinates to robustly construct these paths.
+* 
+* Note, because only edge flips are supported (and not edge split or triangle pokes), the
+* vertex set of the intrinsic mesh and the surface (i.e. extrinsic) mesh will be the same.
+* 
+* The expected use
+*		
+*		// generate the implicit mesh from a surface mesh
+*		FIntrinsicEdgeFlipMesh  ImplicitMesh( DynamicMesh );
+*		
+*			.. do some edge flips .. perhaps FlipToDelaunay( ImplicitMesh, ..)
+* 
+*      // create a correspondence that allows the edges of the implicit mesh to be traced on the surface mesh
+*      FIntrinsicEdgeFlipMesh::FEdgeCorrespondence  Correspondence = ImplicitMesh.ComputeEdgeCorrespondence();
+*
+*      // trace desired edge across the surface mesh
+*	   int32 IntrinsicEdgeID =  
+*      auto TracedIntrinsicEdge  =  Correspondence.TraceEdge(IntrinsicEdgeID);
+*
+* Note: this isn't currently expected to work with surface meshes that have bow-ties.
+*/
+class DYNAMICMESH_API FIntrinsicEdgeFlipMesh : public  FSimpleIntrinsicEdgeFlipMesh
+{
+public:
+	typedef FSimpleIntrinsicEdgeFlipMesh MyBase;
+	using FEdge               = FSimpleIntrinsicEdgeFlipMesh::FEdge;
+	using FEdgeFlipInfo       = FSimpleIntrinsicEdgeFlipMesh::FEdgeFlipInfo;
+	using FSurfacePoint       = IntrinsicCorrespondenceUtils::FSurfacePoint;
+	using FEdgeAndCrossingIdx = IntrinsicCorrespondenceUtils::FNormalCoordinates::FEdgeAndCrossingIdx;
+
+	FIntrinsicEdgeFlipMesh(const FDynamicMesh3& SurfaceMesh);
+
+	FIntrinsicEdgeFlipMesh() = delete;
+	FIntrinsicEdgeFlipMesh(const FIntrinsicEdgeFlipMesh&) = delete;
+
+	/** @return pointer the reference mesh, note the FIntrinsicTriangulation does not manage the lifetime of this mesh */
+	const FDynamicMesh3* GetExtrinsicMesh()  const
+	{
+		return NormalCoordinates.SurfaceMesh;
+	}
+
+	/**
+	* Flip a single edge in the Intrinsic Mesh.
+	* @param EdgeFlipInfo [out] populated on return.
+	*
+	* @return a success or failure info.
+	*/
+	EMeshResult FlipEdge(int32 EID, FEdgeFlipInfo& EdgeFlipInfo);
+
+
+	/**
+	* FEdgeCorrespondence allows intrinsic edges to be traced across the surface mesh.
+	* During construction a list of surface mesh edges crossed by each intrinsic edge is stored.
+	* The Trace operation for an intrinsic edge uses this information to unfold the crossed surface triangles
+	* into a 2d triangle strip upon which the actual trace takes place.
+	*/ 
+	struct DYNAMICMESH_API FEdgeCorrespondence
+	{
+
+		using FEdgeAndCrossingIdx = IntrinsicCorrespondenceUtils::FNormalCoordinates::FEdgeAndCrossingIdx;
+
+		FEdgeCorrespondence(const FIntrinsicEdgeFlipMesh& IntrinsicEdgeFlipTriangulation);
+
+		/**
+		* @return Array of FSurfacePoints relative to the ExtrinsicMesh that represent the specified intrinsic mesh edge.
+		* This is directed from vertex EdgeV.A to vertex EdgeV.B unless bReverse is true
+		*
+		* @param IntrinsicEID      - ID of the intrinsic mesh edge to be traced.
+		* @param CoalesceThreshold - in barycentric units [0,1], edge crossings within this threshold are snapped to the nearest vertex.
+		*                            and any resulting repetitions of vertex surface points are replaced with a single vertex surface point.
+		* @param bReverse          - if true, trace edge from EdgeV.B to EdgeV.A
+		*/
+		TArray<FSurfacePoint> TraceEdge(int32 IntrinsicEID, double CoalesceThreshold = 0., bool bReverse = false) const;
+
+
+		const FIntrinsicEdgeFlipMesh*          IntrinsicMesh;         // the intrinsic mesh
+		const FDynamicMesh3*                   SurfaceMesh;           // the surface mesh 
+		TArray<TArray<int32>>                  SurfaceEdgesCrossed;   // array of surface edge crossings per intrinsic edge.
+	};
+
+	FEdgeCorrespondence ComputeEdgeCorrespondence() const 
+	{
+		return FEdgeCorrespondence(*this);
+	};
+
+	/**
+	* @return NormalCoordinate data structure that counts regular surface (extrinsic) mesh edge-crossing for each intrinsic mesh edge. 
+	*/
+	const IntrinsicCorrespondenceUtils::FNormalCoordinates& GetNormalCoordinates() const
+	{
+		return NormalCoordinates;
+	}
+
+	/**
+	* Trace the specified surface (extrinsic) mesh edge across the intrinsic mesh.  By default this traces from vertex EdgeV.A to vertex EdgeV.B
+	* where EdgeV = SurfaceMesh.GetEdgeV(SurfaceEID).
+	* @param SurfaceEID   - the surface edge to be traced
+	* @param bReverse     - reverses the direction of the trace giving EdgeV.B -> EdgeV.A
+	* 
+	* @return array of intrinsic edges crossed between the specified vertices.  This can be used to compute the actual MeshPoints on the surface if desired.
+	*/ 
+	TArray<FEdgeAndCrossingIdx> GetImplicitEdgeCrossings(const int32 SurfaceEID, const bool bReverse = false) const;
+
+	/**
+	* Trace the specified surface (extrinsic) mesh edge across the intrinsic mesh.  By default this traces from vertex EdgeV.A to vertex EdgeV.B
+	* where EdgeV = SurfaceMesh.GetEdgeV(SurfaceEID).
+	* @param SurfaceEID   - the surface edge to be traced
+	* @param bReverse     - reverses the direction of the trace giving EdgeV.B -> EdgeV.A
+	* 
+	* @return array of surface points defined relative to the intrinsic mesh.  In particular, the array will be a vertex followed by a sequence of (intrinsic) edge crossings.
+	*/ 
+	TArray<FSurfacePoint> TraceSurfaceEdge(int32 SurfaceEID, double CoalesceThreshold, bool bReverse = false) const;
+
+protected:
+
+	using FNormalCoordinates = IntrinsicCorrespondenceUtils::FNormalCoordinates;
+	FNormalCoordinates  NormalCoordinates;      // connection and integer-based coordinates used in reconstructing {intrinsic, surface} edges on the {surface, intrinsic} mesh
+};
+
+
 /**
 * Class that manages the intrinsic triangulation of a given FDynamicMesh.
 *
+* This supports mesh operations that: SplitEdge, PokeTriangle, FlipEdge.
+*
 * An Intrinsic triangulation can be thought of as a mesh that overlays another mesh, sharing
 * the original surface, but the intrinsic triangles can fold over the edges of the original mesh
 * That is to say, the edges of the intrinsic mesh are still on the mesh surface
 * but need not align with the edges of the original mesh and the lengths of these edges are measured
 * on the surface of the original mesh (not the R3 distance).
 *
-* This supports mesh operations that: SplitEdge, PokeTriangle, FlipEdge.
 *
 * Both SplitEdge and PokeTriangle will generate new vertices and triangles in the
 * intrinsic triangulation.  These vertices will live on the surface of the original mesh
@@ -351,18 +484,19 @@ protected:
 *
 * NB: The lifetime of this structure should not exceed that of the original
 * FDynamicMesh as this class holds a pointer to that mesh to reference locations on its surface.
-*/
-class DYNAMICMESH_API FIntrinsicTriangulation : public FIntrinsicEdgeFlipMesh
+* Also, this class isn't currently expected to work with surface meshes that have bow-ties.
+*/ 
+class DYNAMICMESH_API FIntrinsicTriangulation : public FSimpleIntrinsicEdgeFlipMesh
 {
 
 public:
 
-	typedef FIntrinsicEdgeFlipMesh MyBase;
+	typedef FSimpleIntrinsicEdgeFlipMesh MyBase;
 
 	using FEdgeFlipInfo = DynamicMeshInfo::FEdgeFlipInfo;
 	using FEdgeSplitInfo = DynamicMeshInfo::FEdgeSplitInfo;
 	using FPokeTriangleInfo = DynamicMeshInfo::FPokeTriangleInfo;
-
+	using FSurfacePoint  = IntrinsicCorrespondenceUtils::FSurfacePoint;
 
 	FIntrinsicTriangulation(const FDynamicMesh3& SrcMesh);
 
@@ -373,11 +507,9 @@ public:
 	/** @return pointer the reference mesh, note the FIntrinsicTriangulation does not manage the lifetime of this mesh */
 	const FDynamicMesh3* GetExtrinsicMesh()  const
 	{
-		return ExtrinsicMesh;
+		return SignpostData.SurfaceMesh;
 	}
 	
-	/** Location relative to the surface mesh : a surface mesh vertex, surface mesh edge - crossing or barycentric coords */
-	struct FSurfacePoint;
 
 	/**
 	* @return Array of FSurfacePoints relative to the ExtrinsicMesh that represent the specified intrinsic mesh edge.
@@ -395,11 +527,11 @@ public:
 	const FIntrinsicTriangulation::FSurfacePoint& GetVertexSurfacePoint(int32 VID) const
 	{
 		checkSlow(IsVertex(VID));
-		return IntrinsicVertexPositions[VID];
+		return SignpostData.IntrinsicVertexPositions[VID];
 	}
 
 	/**
-	* return position in r3 of an intrinsic vertex.
+	* @return position in r3 of an intrinsic vertex.
 	*/
 	FVector3d GetVertexPosition(int32 VID) const
 	{
@@ -442,63 +574,6 @@ public:
 
 	
 
-public:
-
-
-	struct FGeometricInfo
-	{
-		double ToRadians = 1.;    // 2pi / sum of angles at this vertex.  Will be 1 if the Gaussian curvature is 0
-		bool bIsInterior = false; // false if on boundary or at a bow-tie.
-	};
-
-	// Location on (extrinsic) mesh stored as a union.
-	struct FSurfacePoint
-	{
-		enum class EPositionType
-		{
-			Vertex,
-			Edge,
-			Triangle
-		};
-
-		// constructors respect the correct position type
-		FSurfacePoint(); // defaults to invalid vertex position
-		FSurfacePoint(const int32 VID);
-		FSurfacePoint(const int32 EdgeID, const double Alpha);
-		FSurfacePoint(const int32 TriID, const FVector3d& BaryCentrics);
-
-		struct FVertexPosition
-		{
-			int32 VID;                      // Reference (extrinsic) vertex
-		};
-		struct FEdgePosition
-		{
-			int32 EdgeID;                  // Reference (extrinsic) edge
-			double Alpha;                  // Lerp coordinate along referenced edge as Pos(Edge.A) (Alpha) + (1-Alpha) Pos(Edge.B)
-		};
-		struct FTrianglePosition
-		{
-			int32 TriID;                    // Reference (extrinsic) triangle
-			double BarycentricCoords[3];    // Barycentric coordinates within referenced triangle
-		};
-		union FSurfacePositionUnion
-		{
-			FVertexPosition   VertexPosition;
-			FEdgePosition     EdgePosition;
-			FTrianglePosition TriPosition;
-		};
-
-		EPositionType PositionType;      // Position type, needed to correctly interpret the Position union
-		FSurfacePositionUnion Position;  // Position relative to vertex, edge, or triangle. 
-
-		/**
-		* @return the R3 postition of this surface point as applied to the input Mesh
-		* on return bValidPoint will be false if the edge / vertex / triangle ID referenced by this surface
-		* point is not an element of the mesh in question
-		*/
-		FVector3d AsR3Position(const FDynamicMesh3& Mesh, bool& bValidPoint) const;
-	};
-
 protected:
 
 	
@@ -573,15 +648,8 @@ protected:
 		return tid;
 	}
 protected:
-
-
-	const FDynamicMesh3* ExtrinsicMesh;                                   // Original mesh. The triangles of this mesh define the surface.  We don't manage the lifetime of this mesh
-	TArray<int32> VIDToReferenceEID;                                      // Reference (original mesh) Edge per (original mesh) Vertex, defines local zero angle on tangent plane
-	TArray<int32> TIDToReferenceEID;                                      // Reference (original mesh) Edge per (original mesh) Triangle, defines local zero angle on tangent plane
-	
-	TDynamicVector<FVector3d>            IntrinsicEdgeAngles;             // Per-intrinsic Triangle - polar angles of each edge relative to the vertex it exits (i.e. edge 0 relative to vertex 0)
-	TDynamicVector<FSurfacePoint>        IntrinsicVertexPositions;        // Per-intrinsic Vertex - Locates intrinsic mesh vertices relative to the surface defined by extrinsic mesh.
-	TDynamicVector<FGeometricInfo>       GeometricVertexInfo;             // Per-intrinsic Vertex - Simple geometric information about the extrinsic surface in the neighborhood of intrinsic vert 	
+	using FSignpost = IntrinsicCorrespondenceUtils::FSignpost;
+	FSignpost SignpostData;     // connection and directional-based coordinates used in reconstructing {intrinsic, surface} edges on the {surface, intrinsic} mesh vert 		
 
 };
 
@@ -594,6 +662,7 @@ protected:
 * 
 * @return  the number of flips.
 */
+int32 DYNAMICMESH_API FlipToDelaunay(FSimpleIntrinsicEdgeFlipMesh& IntrinsicMesh, TSet<int>& Uncorrected, const int32 MaxFlipCount = TMathUtilConstants<int>::MaxReal);
 int32 DYNAMICMESH_API FlipToDelaunay(FIntrinsicEdgeFlipMesh& IntrinsicMesh, TSet<int>& Uncorrected, const int32 MaxFlipCount = TMathUtilConstants<int>::MaxReal);
 int32 DYNAMICMESH_API FlipToDelaunay(FIntrinsicTriangulation& IntrinsicMesh, TSet<int>& Uncorrected, const int32 MaxFlipCount = TMathUtilConstants<int>::MaxReal);
 
diff --git a/Engine/Plugins/Runtime/GeometryProcessing/Source/DynamicMesh/Public/Solvers/LaplacianMatrixAssembly.h b/Engine/Plugins/Runtime/GeometryProcessing/Source/DynamicMesh/Public/Solvers/LaplacianMatrixAssembly.h
--- a/Engine/Plugins/Runtime/GeometryProcessing/Source/DynamicMesh/Public/Solvers/LaplacianMatrixAssembly.h
+++ b/Engine/Plugins/Runtime/GeometryProcessing/Source/DynamicMesh/Public/Solvers/LaplacianMatrixAssembly.h
@@ -706,7 +706,7 @@ void UE::MeshDeformation::ConstructIDTCotangentLaplacian(const FDynamicMesh3& Dy
 	const int32 NumInteriorVerts = NumVerts - NumBoundaryVerts;
 
 	// create intrinsic mesh with delaunay triangulation
-	UE::Geometry::FIntrinsicEdgeFlipMesh IntrinsicMesh(DynamicMesh);
+	UE::Geometry::FSimpleIntrinsicEdgeFlipMesh IntrinsicMesh(DynamicMesh);
 
 	TSet<int32> Uncorrected; // some edges can't be flipped.
 	const int32 NumFlips = UE::Geometry::FlipToDelaunay(IntrinsicMesh, Uncorrected);
@@ -925,7 +925,7 @@ void UE::MeshDeformation::ConstructFullIDTCotangentLaplacian(const FDynamicMesh3
 	const int32 NumVerts = VertexMap.NumVerts();
 
 	// create intrinsic mesh with delaunay triangulation
-	UE::Geometry::FIntrinsicTriangulation IntrinsicMesh(Mesh);
+	UE::Geometry::FSimpleIntrinsicEdgeFlipMesh IntrinsicMesh(Mesh);
 	TSet<int32> Uncorrected;
 	const int32 NumFlips = UE::Geometry::FlipToDelaunay(IntrinsicMesh, Uncorrected);
 
