Geometry tools:  Fix GeodesicPath compile error

diff --git a/Engine/Plugins/Runtime/GeometryProcessing/Source/DynamicMesh/Private/Operations/GeodesicPath.cpp b/Engine/Plugins/Runtime/GeometryProcessing/Source/DynamicMesh/Private/Operations/GeodesicPath.cpp
--- a/Engine/Plugins/Runtime/GeometryProcessing/Source/DynamicMesh/Private/Operations/GeodesicPath.cpp
+++ b/Engine/Plugins/Runtime/GeometryProcessing/Source/DynamicMesh/Private/Operations/GeodesicPath.cpp
@@ -146,10 +146,12 @@ void FEdgePath::RemoveSegment(int32 SID)
 *                           TriEIDs           - edge IDs for the triangle.
 *
 */
-template <typename FunctorType>
-bool VisitWedgeTriangles(const FIntrinsicEdgeFlipMesh& EdgeFlipMesh, const int32 StartEID, const int32 EndEID, const int32 CVID, 
+template <typename EdgeFlipMeshType, typename FunctorType>
+bool VisitWedgeTriangles(const EdgeFlipMeshType& EdgeFlipMesh, const int32 StartEID, const int32 EndEID, const int32 CVID,
                          FunctorType& Visitor, bool bClockwise = false)
 {
+	constexpr static int InvalidID = IndexConstants::InvalidID;
+
 	int32 CurEID = StartEID;
 
 	FIndex2i CurEdgeT   = EdgeFlipMesh.GetEdgeT(CurEID);
@@ -162,7 +164,7 @@ bool VisitWedgeTriangles(const FIntrinsicEdgeFlipMesh& EdgeFlipMesh, const int32
 		(!bClockwise && (EdgeFlipMesh.GetTriangle(CurTID)[IndexOf] != CVID)))
 	{
 		CurTID     = CurEdgeT[1];
-		if (CurTID == FIntrinsicEdgeFlipMesh::InvalidID)
+		if (CurTID == InvalidID)
 		{
 			return false;
 		}
@@ -180,7 +182,7 @@ bool VisitWedgeTriangles(const FIntrinsicEdgeFlipMesh& EdgeFlipMesh, const int32
 		CurEdgeT = EdgeFlipMesh.GetEdgeT(CurEID);
 		CurTID = (CurEdgeT[0] == CurTID) ? CurEdgeT[1] : CurEdgeT[0];
 
-		if (CurTID == FIntrinsicEdgeFlipMesh::InvalidID)
+		if (CurTID == InvalidID)
 		{
 			return (CurEID == EndEID);
 		}
@@ -566,7 +568,7 @@ bool FDeformableEdgePath::OuterArcFlipEdges(int32 StartEID, int32 EndEID, int32
 			TmpTID = (CurTID == PreFlipT[1]) ? PreFlipT[0] : PreFlipT[1];
 		}
 
-		UE::Geometry::FIntrinsicEdgeFlipMesh::FEdgeFlipInfo EdgeFlipInfo;
+		IntrinsicMeshType::FEdgeFlipInfo EdgeFlipInfo;
 		UE::Geometry::EMeshResult Result = EdgeFlipMesh.FlipEdge(NextEID, EdgeFlipInfo);
 
 		if (Result == UE::Geometry::EMeshResult::Failed_IsBoundaryEdge)
@@ -626,7 +628,7 @@ void FDeformableEdgePath::ComputeWedgeAngles(int32 IncomingEID,
 
 TArray<FDeformableEdgePath::FSurfacePoint> FDeformableEdgePath::AsSurfacePoints(double CoalesceThreshold) const
 {
-	TArray<FIntrinsicTriangulation::FSurfacePoint> PathSurfacePoints;
+	TArray<FSurfacePoint> PathSurfacePoints;
 
 	const int32 NumIntrinsicPathSegments = EdgePath.NumSegments();
 	
@@ -636,6 +638,8 @@ TArray<FDeformableEdgePath::FSurfacePoint> FDeformableEdgePath::AsSurfacePoints(
 		return PathSurfacePoints;
 	}
 
+	IntrinsicMeshType::FEdgeCorrespondence EdgeCorrespondence = EdgeFlipMesh.ComputeEdgeCorrespondence();
+
 	int32 SID = EdgePath.GetHeadSegmentID();
 	while (SID != InvalidID)
 	{ 
@@ -647,7 +651,7 @@ TArray<FDeformableEdgePath::FSurfacePoint> FDeformableEdgePath::AsSurfacePoints(
 		const bool bReverseEdge = (EdgePath.GetSegment(SID).HeadIndex == 0);
 
 		// get intrinsic edge as a sequence of surface points, note each segment starts and ends at a surface mesh vertex but may cross several surface mesh edges.
-		TArray<FIntrinsicTriangulation::FSurfacePoint> SegmentSurfacePoints = EdgeFlipMesh.TraceEdge(EID, CoalesceThreshold, bReverseEdge);
+		TArray<FSurfacePoint> SegmentSurfacePoints = EdgeCorrespondence.TraceEdge(EID, CoalesceThreshold, bReverseEdge);
 
 		PathSurfacePoints.Append(MoveTemp(SegmentSurfacePoints));
 		if (!bIsLastSegment)
@@ -666,7 +670,7 @@ double UE::Geometry::SumPathLength(const FDeformableEdgePath& DeformableEdgePath
 {
 	double TotalPathLength = 0;
 	const FEdgePath& EdgePath = DeformableEdgePath.GetEdgePath();
-	const FIntrinsicTriangulation& EdgeFlipMesh = DeformableEdgePath.GetIntrinsicMesh();
+	const FDeformableEdgePath::IntrinsicMeshType& EdgeFlipMesh = DeformableEdgePath.GetIntrinsicMesh();
 
 	const int32 NumIntrinsicPathSegments = EdgePath.NumSegments();
 
diff --git a/Engine/Plugins/Runtime/GeometryProcessing/Source/DynamicMesh/Public/Operations/GeodesicPath.h b/Engine/Plugins/Runtime/GeometryProcessing/Source/DynamicMesh/Public/Operations/GeodesicPath.h
--- a/Engine/Plugins/Runtime/GeometryProcessing/Source/DynamicMesh/Public/Operations/GeodesicPath.h
+++ b/Engine/Plugins/Runtime/GeometryProcessing/Source/DynamicMesh/Public/Operations/GeodesicPath.h
@@ -111,6 +111,8 @@ public:
 	
 	constexpr static int InvalidID = IndexConstants::InvalidID;
 	
+	typedef FIntrinsicEdgeFlipMesh  IntrinsicMeshType;
+
 	/**
 	* Constructor assumes the directed segments are ordered from tail (at index 0) to head
 	*/
@@ -153,12 +155,12 @@ public:
 	/**
 	* @return const reference to the intrinsic mesh on which the EdgePath is defined.
 	*/ 
-	inline const FIntrinsicTriangulation& GetIntrinsicMesh() const;
+	inline const FIntrinsicEdgeFlipMesh& GetIntrinsicMesh() const;
 
 	/**
 	* struct that references a point on a mesh, by vertex, by edge-crossing, or barycentric coords
 	*/ 
-	using FSurfacePoint =  FIntrinsicTriangulation::FSurfacePoint;
+	using FSurfacePoint =  IntrinsicCorrespondenceUtils::FSurfacePoint;
 	/**
 	* @return an array of surface points relative to the SurfaceMesh that define this path.
 	* Note the first and last surfacepoints correspond to the start and end vertex, but all 
@@ -259,7 +261,7 @@ protected:
 	
 protected:
      
-	FIntrinsicTriangulation        EdgeFlipMesh;        // Intrinsic mesh with the same vertices as the original surface mesh. The geodesic path is comprised of edges in this mesh.
+	IntrinsicMeshType              EdgeFlipMesh;        // Intrinsic mesh with the same vertices as the original surface mesh. The geodesic path is comprised of edges in this mesh.
 
 	double                         PathLength;          // Total length of the path
 	int32                          NumFlips;            // Count of the number of edge flips performed in shortening the path.
@@ -288,7 +290,7 @@ double  FDeformableEdgePath::GetPathLength() const
 	return PathLength;
 }
 
-const FIntrinsicTriangulation& FDeformableEdgePath::GetIntrinsicMesh() const
+const FIntrinsicEdgeFlipMesh& FDeformableEdgePath::GetIntrinsicMesh() const
 {
 	return EdgeFlipMesh;
 }
