ModelingTools: Default initialize MeshSurfaceSampler

#rnx

diff --git a/Engine/Plugins/Runtime/GeometryProcessing/Source/DynamicMesh/Public/Sampling/MeshSurfaceSampler.h b/Engine/Plugins/Runtime/GeometryProcessing/Source/DynamicMesh/Public/Sampling/MeshSurfaceSampler.h
--- a/Engine/Plugins/Runtime/GeometryProcessing/Source/DynamicMesh/Public/Sampling/MeshSurfaceSampler.h
+++ b/Engine/Plugins/Runtime/GeometryProcessing/Source/DynamicMesh/Public/Sampling/MeshSurfaceSampler.h
@@ -17,8 +17,8 @@ namespace Geometry
  */
 struct FDynamicMeshUVMesh
 {
-	const FDynamicMesh3* Mesh;
-	const FDynamicMeshUVOverlay* UV;
+	const FDynamicMesh3* Mesh = nullptr;
+	const FDynamicMeshUVOverlay* UV = nullptr;
 
 	inline bool IsTriangle(int32 TriangleIndex) const
 	{
@@ -83,22 +83,22 @@ struct FDynamicMeshUVMesh
 struct FMeshUVSampleInfo
 {
 	// Triangle containing the sample
-	int32 TriangleIndex;
+	int32 TriangleIndex = IndexConstants::InvalidID;
 
 	// 3D vertices
 	FIndex3i MeshVertices;
 	// 3D triangle
-	FTriangle3d Triangle3D;
+	FTriangle3d Triangle3D { FVector3d::ZeroVector, FVector3d::ZeroVector, FVector3d::ZeroVector };
 
 	// UV overlay vertices
 	FIndex3i UVVertices;
 	// 2D triangle
-	FTriangle2d TriangleUV;
+	FTriangle2d TriangleUV { FVector2d::ZeroVector, FVector2d::ZeroVector, FVector2d::ZeroVector };
 
 	// barycentric coords in triangle
-	FVector3d BaryCoords;
+	FVector3d BaryCoords = FVector3d::ZeroVector;
 	// surface point (lying in Triangle3D)
-	FVector3d SurfacePoint;
+	FVector3d SurfacePoint = FVector3d::ZeroVector;
 };
 
 
@@ -179,9 +179,9 @@ public:
 	virtual bool CachedSampleUV(const FMeshUVSampleInfo& CachedSampleInfo, SampleType& ResultOut);
 
 protected:
-	const FDynamicMesh3* Mesh;
-	const FDynamicMeshUVOverlay* UVOverlay;
-	EMeshSurfaceSamplerQueryType QueryType;
+	const FDynamicMesh3* Mesh = nullptr;
+	const FDynamicMeshUVOverlay* UVOverlay = nullptr;
+	EMeshSurfaceSamplerQueryType QueryType = EMeshSurfaceSamplerQueryType::TriangleAndUV;
 
 	TUniqueFunction<void(const FMeshUVSampleInfo& SampleInfo, SampleType& SampleValueOut)> ValueFunction;
 
