Fix TMeshSurfaceUVSampler compile error (temp fix until LWC FVector2)

#rnx

diff --git a/Engine/Plugins/Runtime/GeometryProcessing/Source/DynamicMesh/Public/Sampling/MeshSurfaceSampler.h b/Engine/Plugins/Runtime/GeometryProcessing/Source/DynamicMesh/Public/Sampling/MeshSurfaceSampler.h
--- a/Engine/Plugins/Runtime/GeometryProcessing/Source/DynamicMesh/Public/Sampling/MeshSurfaceSampler.h
+++ b/Engine/Plugins/Runtime/GeometryProcessing/Source/DynamicMesh/Public/Sampling/MeshSurfaceSampler.h
@@ -93,7 +93,7 @@ struct FMeshUVSampleInfo
 	// UV overlay vertices
 	FIndex3i UVVertices;
 	// 2D triangle
-	FTriangle2d TriangleUV { FVector2d::ZeroVector, FVector2d::ZeroVector, FVector2d::ZeroVector };
+	FTriangle2d TriangleUV { FVector2d(0.0,0.0), FVector2d(0.0,0.0), FVector2d(0.0,0.0) };
 
 	// barycentric coords in triangle
 	FVector3d BaryCoords = FVector3d::ZeroVector;
