Dynamic Mesh, Intrinsic Mesh - add a more efficient  normal-coordinate based method for tracing a small number of  intrinsic mesh edges.

diff --git a/Engine/Plugins/Runtime/GeometryProcessing/Source/DynamicMesh/Private/Operations/IntrinsicTriangulationMesh.cpp b/Engine/Plugins/Runtime/GeometryProcessing/Source/DynamicMesh/Private/Operations/IntrinsicTriangulationMesh.cpp
--- a/Engine/Plugins/Runtime/GeometryProcessing/Source/DynamicMesh/Private/Operations/IntrinsicTriangulationMesh.cpp
+++ b/Engine/Plugins/Runtime/GeometryProcessing/Source/DynamicMesh/Private/Operations/IntrinsicTriangulationMesh.cpp
@@ -443,10 +443,11 @@ namespace FNormalCoordSurfaceTraceImpl
 {
 	using namespace IntrinsicCorrespondenceUtils;
 	/**
-	* low-level code that uses normal coordinates to continues tracing a surface edge across an FIntrinsicEdgeFlipMesh.  
+	* low-level code that uses normal coordinates to continues tracing a surface edge across an FIntrinsicEdgeFlipMesh until it terminates at a vertex.
 	*
-	* This function is intended to be called directly after the first intrinsic edge crossing 
-	* at which time the 'Crossings' result array will hold the start vertex and the first edge crossing.
+	* In order to compute the full crossing history for a surface curve, this function could be called directly after 
+	* the first intrinsic edge crossing at which time the 'Crossings' result array will already hold the start vertex and the first edge crossing.
+	* 
 	* 
 	* This function will populate the rest of the sequence of edge crossings for the surface mesh edge
 	* as a list of the edges crossed, and not the location on the individual edges.  
@@ -457,10 +458,6 @@ namespace FNormalCoordSurfaceTraceImpl
 		                            TArray<FNormalCoordinates::FEdgeAndCrossingIdx>& Crossings )
 	{
 		const FNormalCoordinates& NCoords = IntrinsicMesh.GetNormalCoordinates();
-		if (Crossings.Num() == 0)
-		{
-			return;
-		}
 
 		if (!IntrinsicMesh.IsTriangle(StartTID))
 		{
@@ -514,7 +511,7 @@ namespace FNormalCoordSurfaceTraceImpl
 				{
 					// this path terminates at vertex
 					// this is a slight abuse.  The convention, when CrossingIndex = 0, we specify the vertex by the edge that originates there.
-					EdgeID = Edge_lj; // GetTrianlge()[EdgeID_lj] will be the vertex.
+					EdgeID = Edge_lj; // i = GetTriangle()[TriEIDs.IndexOf(EdgeID_lj)] will be the vertex.
 					P      = 0;
 				}
 				else
@@ -2253,6 +2250,71 @@ EMeshResult FIntrinsicEdgeFlipMesh::FlipEdge(int32 EID, FEdgeFlipInfo& EdgeFlipI
 	return FlipResult;
 }
 
+TArray<FIntrinsicEdgeFlipMesh::FSurfacePoint> FIntrinsicEdgeFlipMesh::TraceEdge(int32 IntrinsicEID, double CoalesceThreshold, bool bReverse) const
+{
+	TArray<FIntrinsicEdgeFlipMesh::FSurfacePoint> Result;
+	
+	if (!IsEdge(IntrinsicEID))
+	{
+		return MoveTemp(Result);
+	}
+
+	TArray<int32> SurfXings;
+	const int32 NumSurfXings = NormalCoordinates.NormalCoord[IntrinsicEID];
+	if (NumSurfXings > 0)
+	{
+		SurfXings.Reserve(NumSurfXings);
+	}
+
+	const FDynamicMesh3& SurfaceMesh = *GetExtrinsicMesh();
+	const FIntrinsicEdgeFlipMesh& IntrinsicMesh = *this;
+
+	const FIndex2i IntrinsicEdgeT = GetEdgeT(IntrinsicEID);
+
+	// need to create a list of surface mesh edges that cross this edge.  To do this we need to follow
+	// each curve that crosses this intrinsic edge to the vertices where it terminates ( thus identifying the surface edge )
+	for (int32 p = 1; p < NumSurfXings + 1; ++p)
+	{
+
+		const FIndex2i SurfaceEdgeV = [&]
+			{
+		
+				TArray<FNormalCoordinates::FEdgeAndCrossingIdx> Crossings;
+				FIndex2i Verts;
+				// follow surface curve (i.e. surface edge) forward to the end and identify the vertex
+				{
+					FNormalCoordSurfaceTraceImpl::ContinueTraceSurfaceEdge(IntrinsicMesh,  IntrinsicEdgeT.A, IntrinsicEID, p, Crossings);
+					const FNormalCoordinates::FEdgeAndCrossingIdx& LastXing = Crossings.Last();
+
+					const FIndex3i TriEIDs = GetTriEdges(LastXing.TID);
+					Verts.A                = GetTriangle(LastXing.TID)[TriEIDs.IndexOf(LastXing.EID)];
+				}
+				Crossings.Reset();
+
+				// follow surface curve (i.e. surface edge) backward to other end and identify the vertex
+				{
+					FNormalCoordSurfaceTraceImpl::ContinueTraceSurfaceEdge(IntrinsicMesh, IntrinsicEdgeT.B, IntrinsicEID, NumSurfXings + 1 - p, Crossings);
+					const FNormalCoordinates::FEdgeAndCrossingIdx& LastXing = Crossings.Last();
+
+					const FIndex3i TriEIDs = GetTriEdges(LastXing.TID);
+					Verts.B                = GetTriangle(LastXing.TID)[TriEIDs.IndexOf(LastXing.EID)];
+				}
+				return Verts;
+			}();
+
+		// identify the surface edge from its endpoints. 
+		const int32 SurfaceEID = SurfaceMesh.FindEdge(SurfaceEdgeV.A, SurfaceEdgeV.B);
+		checkSlow(SurfaceEID != InvalidID);
+		
+		SurfXings.Add(SurfaceEID);
+	}
+	
+	// do the actual trace - this identifies the surface mesh triangles crossed by the intrinsic edge and unfolds them into a triangle strip where the trace is performed
+	Result =  FNormalCoordSurfaceTraceImpl::TraceEdgeOverHost(IntrinsicEID, SurfXings, SurfaceMesh, IntrinsicMesh, CoalesceThreshold, bReverse);
+
+	return MoveTemp(Result);
+}
+
 TArray<FIntrinsicEdgeFlipMesh::FEdgeAndCrossingIdx> FIntrinsicEdgeFlipMesh::GetImplicitEdgeCrossings(const int32 SurfaceEID, const bool bReverse) const
 {
 	return FNormalCoordSurfaceTraceImpl::TraceSurfaceEdge(*this, SurfaceEID, bReverse);
@@ -2306,10 +2368,10 @@ FIntrinsicEdgeFlipMesh::FEdgeCorrespondence::FEdgeCorrespondence(const FIntrinsi
 			{
 				const int32 IntrinsicEID      = EdgeXing.EID;
 				const FIndex2i IntrinsicEdgeT = IntrinsicMesh->GetEdgeT(IntrinsicEID);
-				checkSlow(IntrinsicEdgeT.A == EdgeXing.TID || IntrinsicEdgeT.B == EdgeXing.TID)
+				checkSlow(IntrinsicEdgeT.A == EdgeXing.TID || IntrinsicEdgeT.B == EdgeXing.TID);
 
-					// array of surface edges this intrinsic edge crosses, these should be ordered relative to the direction of TriA.
-					TArray<int32>& XingSurfaceEdges = SurfaceEdgesCrossed[IntrinsicEID];
+				// array of surface edges this intrinsic edge crosses, these should be ordered relative to the direction of TriA.
+				TArray<int32>& XingSurfaceEdges = SurfaceEdgesCrossed[IntrinsicEID];
 
 				// crossings count from the bottom of the edge to the top relative to EdgeXing.TID 
 				const int32 XingID = (IntrinsicEdgeT.A == EdgeXing.TID) ? EdgeXing.CIdx - 1 : XingSurfaceEdges.Num() - EdgeXing.CIdx;
diff --git a/Engine/Plugins/Runtime/GeometryProcessing/Source/DynamicMesh/Public/Operations/IntrinsicTriangulationMesh.h b/Engine/Plugins/Runtime/GeometryProcessing/Source/DynamicMesh/Public/Operations/IntrinsicTriangulationMesh.h
--- a/Engine/Plugins/Runtime/GeometryProcessing/Source/DynamicMesh/Public/Operations/IntrinsicTriangulationMesh.h
+++ b/Engine/Plugins/Runtime/GeometryProcessing/Source/DynamicMesh/Public/Operations/IntrinsicTriangulationMesh.h
@@ -357,12 +357,20 @@ protected:
 *		
 *			.. do some edge flips .. perhaps FlipToDelaunay( ImplicitMesh, ..)
 * 
+*      // if only a small percentage of the intrinsic edges are to be traced, this can be done directly as
+*	   auto TracedIntrinsicEdge  =  ImplicitMesh.TraceEdge(IntrinsicEdgeID);
+* 
+*	   // alternately if a large number of edges will be traced, a correspondence will be more efficient. 
 *      // create a correspondence that allows the edges of the implicit mesh to be traced on the surface mesh
+*     
 *      FIntrinsicEdgeFlipMesh::FEdgeCorrespondence  Correspondence = ImplicitMesh.ComputeEdgeCorrespondence();
-*
-*      // trace desired edge across the surface mesh
-*	   int32 IntrinsicEdgeID =  
-*      auto TracedIntrinsicEdge  =  Correspondence.TraceEdge(IntrinsicEdgeID);
+*    
+*	   for (.. doing multiple traces..  ) 
+*      {
+*          // trace desired edge across the surface mesh
+*	       int32 IntrinsicEdgeID =  
+*          auto TracedIntrinsicEdge  =  Correspondence.TraceEdge(IntrinsicEdgeID);
+*      } 
 *
 * Note: this isn't currently expected to work with surface meshes that have bow-ties.
 */
@@ -386,6 +394,20 @@ public:
 		return NormalCoordinates.SurfaceMesh;
 	}
 
+	/**
+	* @return Array of FSurfacePoints relative to the surface mesh that represent the specified intrinsic mesh edge.
+	* This is directed relative to traversal of the first adjacent triangle (GetEdgeT().[0]), unless bReverse is true
+	*
+	* @param CoalesceThreshold - in barycentric units [0,1], edge crossings with in this threshold are snapped to the nearest vertex.
+	*                            and any resulting repetitions of vertex surface points are replaced with a single vertex surface point.
+	* @param bReverse          - if true, trace edge direction is reversed
+	* 
+	* Note: when tracing few edges, this method is preferred to using the FEdgeCorrespondence based tracing (see FEdgeCorrespondence below).  
+	* But when tracing a large percentage of the mesh edges the FEdgeCorrespondence will be faster as it reuses intermediate results.
+	*/
+	TArray<FSurfacePoint> TraceEdge(int32 IntrinsicEID, double CoalesceThreshold = 0., bool bReverse = false) const;
+
+
 	/**
 	* Flip a single edge in the Intrinsic Mesh.
 	* @param EdgeFlipInfo [out] populated on return.
@@ -410,12 +432,12 @@ public:
 
 		/**
 		* @return Array of FSurfacePoints relative to the ExtrinsicMesh that represent the specified intrinsic mesh edge.
-		* This is directed from vertex EdgeV.A to vertex EdgeV.B unless bReverse is true
+		* This is directed relative to traversal of the first adjacent triangle (GetEdgeT().[0]), unless bReverse is true
 		*
 		* @param IntrinsicEID      - ID of the intrinsic mesh edge to be traced.
 		* @param CoalesceThreshold - in barycentric units [0,1], edge crossings within this threshold are snapped to the nearest vertex.
 		*                            and any resulting repetitions of vertex surface points are replaced with a single vertex surface point.
-		* @param bReverse          - if true, trace edge from EdgeV.B to EdgeV.A
+		* @param bReverse          - if true, trace edge direction is reversed
 		*/
 		TArray<FSurfacePoint> TraceEdge(int32 IntrinsicEID, double CoalesceThreshold = 0., bool bReverse = false) const;
 
