Fixed a compilation error in CL18443426 caused by missing files.

#rnx

diff --git a/Engine/Plugins/Runtime/GeometryProcessing/Source/DynamicMesh/Private/Operations/PNTriangles.cpp b/Engine/Plugins/Runtime/GeometryProcessing/Source/DynamicMesh/Private/Operations/PNTriangles.cpp
--- a/Engine/Plugins/Runtime/GeometryProcessing/Source/DynamicMesh/Private/Operations/PNTriangles.cpp
+++ b/Engine/Plugins/Runtime/GeometryProcessing/Source/DynamicMesh/Private/Operations/PNTriangles.cpp
@@ -7,6 +7,7 @@
 #include "DynamicMesh/DynamicMesh3.h"
 #include "Async/ParallelFor.h"
 #include "Util/ProgressCancel.h"
+#include "Operations/UniformTesselate.h"
 
 using namespace UE::Geometry;
 
@@ -35,21 +36,6 @@ namespace FPNTrianglesLocals
 		TArray<FTriangleControlPoint> OnTriangles; // Map triangle ID to the triangle center control point
 	};
 
-	// Used to save just the connectivity information of a mesh
-	struct FMeshConnectivity
-	{
-		FMeshConnectivity(FDynamicMesh3& Mesh) 
-		{
-			Triangles = Mesh.GetTrianglesBuffer();
-			Edges = Mesh.GetEdgesBuffer();
-			TriangleEdges = Mesh.GetTriangleEdges();
-		}
-		
-		TDynamicVector<FIndex3i> Triangles;
-		TDynamicVector<FDynamicMesh3::FEdge> Edges;
-		TDynamicVector<FIndex3i> TriangleEdges;
-	};
-
 	/**
 	 *  Compute PN Triangle control points and optionally their normal component.
 	 * 
@@ -223,15 +209,17 @@ namespace FPNTrianglesLocals
 	 * @param Level How many times we are recursively subdividing the Mesh.
 	 * @param ProgressCancel Set this to be able to cancel running operation.
 	 * @param OutNewVertices Array of tuples of the new vertex ID and the original triangle ID the vertex belongs to.
+	 * @param OutMesh Result of the tesselation.
 	 * 
 	 * @return true if the operation succeeded, false if it failed or was canceled by the user.
 	 */
-	bool TesselateMesh(FDynamicMesh3& Mesh, 
+	bool TesselateMesh(const FDynamicMesh3& Mesh, 
+					   const int32 Level, 
+					   FProgressCancel* ProgressCancel,
 					   TArray<FIndex2i>& OutNewVertices,
-					   const int Level, 
-					   FProgressCancel* ProgressCancel) 
+					   FDynamicMesh3& OutMesh) 
 	{
-		check(Level >= 0);
+		checkSlow(Level >= 0);
 
 		if (Level < 0) 
 		{
@@ -242,96 +230,34 @@ namespace FPNTrianglesLocals
 			return true; // nothing to do
 		}
 
-		// Compute the number of new vertices we will generate with tesselation
-		int NumEdgeSegments = 1 << Level; // Number of the edges each original edge is split into
-		int NumEdgeVert = NumEdgeSegments + 1; // Number of the vertices along the edge after the tesselation
-		int NumNewEdgeVert = NumEdgeSegments - 1; // Number of the new vertices added per edge
-		// Triangular number series (n(n+1)/2) minus the 3 original vertices and new edge vertices
-		int NumNewTriangleVert = NumEdgeVert*(NumEdgeVert + 1)/2 - 3 - 3*NumNewEdgeVert;
-		int NumNewVertices = NumNewEdgeVert * Mesh.EdgeCount() + NumNewTriangleVert* Mesh.TriangleCount(); 
-		OutNewVertices.Reserve(NumNewVertices);
-
-		// The number of the new triangle IDs created with Level number of subdivisions 
-		int NumNewTrianglesIDs = (NumEdgeSegments * NumEdgeSegments - 1) * Mesh.TriangleCount();
-		
-		// When we perform edge splits, we want to keep track of the original triangles that 
-		// the new added triangles belong to (i.e parent triangle). 
-		TArray<int> ParentTriangles;
-		ParentTriangles.Init(FDynamicMesh3::InvalidID, Mesh.MaxTriangleID() + NumNewTrianglesIDs);
-		for (int TID : Mesh.TriangleIndicesItr()) 
-		{ 
-			ParentTriangles[TID] = TID; // Original triangles are parents of themselves.
+		FUniformTesselate Tesselator(&Mesh, &OutMesh);
+		Tesselator.TesselationNum = Level;
+		Tesselator.bComputeMappings = true;
+		if (Tesselator.Validate() != EOperationValidationResult::Ok)
+		{
+			return false;
 		}
 
-		// Recursively subdivide the mesh Level number of times, while keeping track of new 
-		// vertices added and which original triangles they belong to.
-		for (int CurLevel = 0; CurLevel < Level; ++CurLevel) 
+		if (Tesselator.Compute() == false) 
 		{
-			TArray<int> EdgesToProcess;
-			EdgesToProcess.Reserve(Mesh.EdgeCount());
-			for (int EID : Mesh.EdgeIndicesItr())
-			{
-				EdgesToProcess.Add(EID);
-			} 
-
-			int MaxTriangleID = Mesh.MaxTriangleID();
-
-			TArray<int> EdgesToFlip;
-			EdgesToFlip.Init(FDynamicMesh3::InvalidID, MaxTriangleID);
-
-			for (int EID : EdgesToProcess)
-			{
-				FIndex2i EdgeTris = Mesh.GetEdgeT(EID);
-
-				FDynamicMesh3::FEdgeSplitInfo SplitInfo;
-				EMeshResult Result = Mesh.SplitEdge(EID, SplitInfo);
-
-				if (Result == EMeshResult::Ok) 
-				{	  
-					ParentTriangles[SplitInfo.NewTriangles.A] = ParentTriangles[SplitInfo.OriginalTriangles.A]; 
-					if (SplitInfo.OriginalTriangles.B != FDynamicMesh3::InvalidID) 
-					{
-						ParentTriangles[SplitInfo.NewTriangles.B] = ParentTriangles[SplitInfo.OriginalTriangles.B]; 
-					}
-
-					OutNewVertices.Add(FIndex2i(SplitInfo.NewVertex, ParentTriangles[SplitInfo.OriginalTriangles.A]));
-
-					if (EdgeTris.A < MaxTriangleID && EdgesToFlip[EdgeTris.A] == FDynamicMesh3::InvalidID)
-					{
-						EdgesToFlip[EdgeTris.A] = SplitInfo.NewEdges.B;
-					}
-					if (EdgeTris.B != FDynamicMesh3::InvalidID)
-					{
-						if (EdgeTris.B < MaxTriangleID && EdgesToFlip[EdgeTris.B] == FDynamicMesh3::InvalidID)
-						{
-							EdgesToFlip[EdgeTris.B] = SplitInfo.NewEdges.C;
-						}
-					}
-				}   
+			return false;
+		} 
 
-			  	if (ProgressCancel && ProgressCancel->Cancelled())    
-				{
-					return false;
-				}
-			}
+		const int32 NewVertCount = Tesselator.VertexEdgeMap.Num() + Tesselator.VertexTriangleMap.Num();
+		OutNewVertices.Reserve(NewVertCount);
 
-			for (int EID : EdgesToFlip)
-			{
-				if (EID != FDynamicMesh3::InvalidID)
-				{
-					FDynamicMesh3::FEdgeFlipInfo FlipInfo;
-					Mesh.FlipEdge(EID, FlipInfo);
+		for (auto It = Tesselator.VertexEdgeMap.CreateConstIterator(); It; ++It)
+		{
+			FIndex2i EdgeTri = It.Value();
+			OutNewVertices.Add(FIndex2i(It.Key(), EdgeTri.A));
+		}
 
-					if (ProgressCancel && ProgressCancel->Cancelled()) 
-					{
-						return false;
-					}
-				}
-			}
+		for (auto It = Tesselator.VertexTriangleMap.CreateConstIterator(); It; ++It)
+		{
+			OutNewVertices.Add(FIndex2i(It.Key(), It.Value()));
 		}
 
-		checkSlow(OutNewVertices.Num() == NumNewVertices);
-		checkSlow(ParentTriangles.Num() == Mesh.MaxTriangleID());
+		checkSlow(OutNewVertices.Num() == NewVertCount);
 
 		return true;
 	}
@@ -340,22 +266,22 @@ namespace FPNTrianglesLocals
 	 * Displace the vertices created from the tesselation using the cubic patch formula based on their barycentric 
 	 * coordinates.
 	 * 
-	 * @param Mesh The tessellated mesh whose vertices we are displacing.
+	 * @param OriginalMesh The original mesh (before tesselation).
 	 * @param FControlPoints PN Triangle control points.
 	 * @param VerticesToDisplace Array of vertices into the Mesh that we are displacing.
 	 * @param bComputePNNormals Should we be computing quadratically varying normals using control points.
-	 * @param OriginalConnectivity Connectivity of the original mesh (before tesselation).
 	 * @param ProgressCancel Set this to be able to cancel running operation.
+	 * @param Mesh The tessellated mesh whose vertices we are displacing.
 	 * 
 	 * @return true if the operation succeeded, false if it failed or was canceled by the user.
 	 */
 
-	bool DisplaceAndSetNormals(FDynamicMesh3& Mesh,
+	bool DisplaceAndSetNormals(const FDynamicMesh3& OriginalMesh,
 					           const FControlPoints& FControlPoints,
 					           const TArray<FIndex2i>& VerticesToDisplace,
 					           const bool bComputePNNormals,
-					           const FMeshConnectivity& OriginalConnectivity,
-					           FProgressCancel* ProgressCancel) 
+							   FProgressCancel* ProgressCancel,
+							   FDynamicMesh3& Mesh)
 	{
 		bool bHasVertexNormals = Mesh.HasVertexNormals();
 		bool bHasAttributes = Mesh.HasAttributes();
@@ -373,8 +299,8 @@ namespace FPNTrianglesLocals
 			int OriginalTriangleID = NewVtx[1]; // ID of the original triangle new vertex belongs to
 			
 			// Get the topology information of the original triangle
-			FIndex3i TriVertex = OriginalConnectivity.Triangles[OriginalTriangleID];
-			FIndex3i TriEdges = OriginalConnectivity.TriangleEdges[OriginalTriangleID];
+			FIndex3i TriVertex = OriginalMesh.GetTriangle(OriginalTriangleID);
+			FIndex3i TriEdges = OriginalMesh.GetTriEdges(OriginalTriangleID);
 
 			FVector3d Bary = VectorUtil::BarycentricCoords(Mesh.GetVertexRef(VertexID), Mesh.GetVertexRef(TriVertex.A), 
 														   Mesh.GetVertexRef(TriVertex.B), Mesh.GetVertexRef(TriVertex.C));
@@ -390,7 +316,7 @@ namespace FPNTrianglesLocals
 			for (int EIDX = 0; EIDX < 3; ++EIDX) 
 			{   
 				int EID = TriEdges[EIDX];
-				FIndex2i EdgeV = OriginalConnectivity.Edges[EID].Vert;
+				FIndex2i EdgeV = OriginalMesh.GetEdgeRef(EID).Vert;
 				const FEdgeControlPoints& ControlPnts = FControlPoints.OnEdges[EID];
 
 				// Check that the orientation of the edge is consistent so we use the correct barycentric values
@@ -467,8 +393,8 @@ bool FPNTriangles::Compute()
 
 	// Compute per-vertex normals if no normals exist
 	FMeshNormals Normals;
-	bool bHasNormals = Mesh->HasVertexNormals() || Mesh->HasAttributes();
-	if (!bHasNormals)
+	bool bHasNormals = Mesh->HasVertexNormals() || (Mesh->HasAttributes() && Mesh->Attributes()->PrimaryNormals() != nullptr);
+	if (bHasNormals == false)
 	{
 		Normals = FMeshNormals(Mesh);
 		Normals.ComputeVertexNormals();
@@ -482,24 +408,24 @@ bool FPNTriangles::Compute()
 	{
 		return false;
 	}
-
-	// Save the topology information of the original mesh before we change it with tesselation
-	FMeshConnectivity OriginalConnectivity(*Mesh);
 	
 	// Tesselate the original mesh
 	TArray<FIndex2i> NewVertices;
-	bOk = TesselateMesh(*Mesh, NewVertices, TesselationLevel, Progress);
+	FDynamicMesh3 ResultMesh;
+	bOk = TesselateMesh(*Mesh, TesselationLevel, Progress, NewVertices, ResultMesh);
 	if (bOk == false) 
 	{
 		return false;
 	}
 
 	// Compute displacement and normals
-	bOk = DisplaceAndSetNormals(*Mesh, FControlPoints, NewVertices, bComputePNNormals, OriginalConnectivity, Progress);
+	bOk = DisplaceAndSetNormals(*Mesh, FControlPoints, NewVertices, bComputePNNormals, Progress, ResultMesh);
 	if (bOk == false) 
 	{
 		return false;
 	}
 
+	Mesh->Copy(ResultMesh);
+
 	return true;
 }
\ No newline at end of file
diff --git a/Engine/Plugins/Runtime/GeometryProcessing/Source/DynamicMesh/Private/Operations/UniformTesselate.cpp b/Engine/Plugins/Runtime/GeometryProcessing/Source/DynamicMesh/Private/Operations/UniformTesselate.cpp
--- a/Engine/Plugins/Runtime/GeometryProcessing/Source/DynamicMesh/Private/Operations/UniformTesselate.cpp
+++ b/Engine/Plugins/Runtime/GeometryProcessing/Source/DynamicMesh/Private/Operations/UniformTesselate.cpp
@@ -1,7 +1,6 @@
 // Copyright Epic Games, Inc. All Rights Reserved.
 
 #include "Operations/UniformTesselate.h"
-#include "IndexTypes.h"
 #include "VectorTypes.h"
 #include "DynamicMesh/DynamicMesh3.h"
 #include "DynamicMesh/MeshNormals.h"
@@ -281,11 +280,7 @@ namespace UniformTesselateLocals
 					return;
 				}
 
-				if (InMesh->IsEdge(EdgeID) == false) 
-				{
-					checkSlow(false); // Mesh must be compact
-					return;
-				}
+				checkSlow(InMesh->IsEdge(EdgeID));
 
 				const FIndex2i EdgeV = InMesh->GetEdgeV(EdgeID);
 				const FIndex2i EdgeTri = InMesh->GetEdgeT(EdgeID);
@@ -347,11 +342,7 @@ namespace UniformTesselateLocals
 					return;
 				}
 
-				if (InMesh->IsTriangle(TriangleID) == false) 
-				{
-					checkSlow(false); // Mesh must be compact
-					return;
-				}
+				checkSlow(InMesh->IsTriangle(TriangleID));
 
 				const FIndex3i TriVertices = InMesh->GetTriangle(TriangleID);
 				const FIndex3i TriEdges = InMesh->GetTriEdges(TriangleID);
@@ -688,12 +679,8 @@ namespace UniformTesselateLocals
 					return;
 				}
 
-				if (InMesh->IsTriangle(TriangleID) == false)
- 				{
-					 checkSlow(false); // Mesh must be compact
-					 return;
-				}
-
+				checkSlow(InMesh->IsTriangle(TriangleID));
+ 				
 				const FIndex3i TriVertices = InMesh->GetTriangle(TriangleID);
 				const FIndex3i TriEdges = InMesh->GetTriEdges(TriangleID);
 
@@ -839,11 +826,7 @@ namespace UniformTesselateLocals
 					return;
 				}
 
-				if (InMesh->IsEdge(EdgeID) == false) 
-				{
-					checkSlow(false); // Mesh must be compact
-					return;
-				}
+				checkSlow(InMesh->IsEdge(EdgeID));
 
 				const FIndex2i EdgeV = InMesh->GetEdgeV(EdgeID);
 				const FIndex2i EdgeTri = InMesh->GetEdgeT(EdgeID);
@@ -959,10 +942,23 @@ namespace UniformTesselateLocals
 		using BaseType::OutEdgeElementCount; 
 		using BaseType::OutInnerElementCount;
 		using BaseType::OutTriangleCount;
+
+		// Map vertex IDs inserted along the edges to the input mesh triangle ID(s) that share that edge.
+  		TMap<int32, FIndex2i>* VertexEdgeMap; 
+
+		// Map vertex IDs to IDs of the triangles those vertices inserted into.
+  		TMap<int32, int32>* VertexTriangleMap;
 		
-		FMeshVerticesGenerator(const FDynamicMesh3* Mesh, const int TesselationNum, FProgressCancel* InProgress, const bool bUseParallel) 
+		FMeshVerticesGenerator(const FDynamicMesh3* Mesh, 
+							   const int TesselationNum, 
+							   FProgressCancel* InProgress, 
+  							   TMap<int32, FIndex2i>* VertexEdgeMap,
+  							   TMap<int32, int32>* VertexTriangleMap,
+  							   const bool bUseParallel) 
 		:
-		FTrianglesGenerator<RealType, 3>(Mesh, TesselationNum, InProgress, bUseParallel)
+		FTrianglesGenerator<RealType, 3>(Mesh, TesselationNum, InProgress, bUseParallel), 
+  		VertexEdgeMap(VertexEdgeMap),
+  		VertexTriangleMap(VertexTriangleMap)
 		{
 		}
 			
@@ -987,6 +983,16 @@ namespace UniformTesselateLocals
 				GetInputMeshElementValue(VertexID, Value);
 				this->SetElement(VertexID, Value);
 			}
+
+			if (VertexEdgeMap) 
+			{
+				VertexEdgeMap->Reserve(OutEdgeElementCount);
+			}
+			
+			if (VertexTriangleMap) 
+			{
+				VertexTriangleMap->Reserve(OutInnerElementCount*InMesh->TriangleCount());
+			}
 		}
 
 		/**
@@ -1027,6 +1033,51 @@ namespace UniformTesselateLocals
 			GetInputMeshElementValue(VertexID, Value);
 		}
 
+		/** Compute the mappings between vertices and triangles of the input and the output meshes. */
+		bool ComputeTopologyMappings() 
+		{	
+			if (VertexEdgeMap) 
+			{
+				for (int32 EdgeID = 0; EdgeID < InMesh->MaxEdgeID(); ++EdgeID)
+				{
+					checkSlow(InMesh->IsEdge(EdgeID));
+					
+					const FIndex2i EdgeTri = InMesh->GetEdgeT(EdgeID);
+
+					for (int VertexOffset = 0; VertexOffset < TesselationNum; ++VertexOffset)
+					{	
+						const int VertexID = this->GetElementIDOnEdge(EdgeID, EdgeTri.A, VertexOffset);
+						VertexEdgeMap->Add(VertexID, EdgeTri);
+					}
+				}
+			}
+
+			if (this->Cancelled()) 
+			{
+				return false;
+			}
+
+			if (VertexTriangleMap) 
+			{			
+				for (int32 TriangleID = 0; TriangleID < InMesh->MaxTriangleID(); ++TriangleID)
+				{
+					checkSlow(InMesh->IsTriangle(TriangleID));
+					
+					int VertexIDCounter = 0;
+					for (int Level = 1; Level < TesselationNum; ++Level)
+					{
+						for (int VertexOffset = 0; VertexOffset < Level; ++VertexOffset, ++VertexIDCounter) 
+						{
+							const int VertexID = this->GetElementIDOnTriangle(TriangleID, VertexIDCounter);
+							VertexTriangleMap->Add(VertexID, TriangleID);
+						}
+					}
+				}
+			}
+			
+			return true;
+		}
+
 		virtual bool Generate() override
 		{	
 			if (this->Validate() == false) 
@@ -1050,6 +1101,11 @@ namespace UniformTesselateLocals
 			{
 				return false;
 			}
+
+			if (ComputeTopologyMappings() == false) 
+			{
+				return false;
+			}
 				
 			return true;
 		}
@@ -1455,11 +1511,7 @@ namespace UniformTesselateLocals
 					return;
 				}
 
-				if (InMesh->IsEdge(EdgeID) == false) 
-				{
-					checkSlow(false); // Mesh must be compact
-					return;
-				}
+				checkSlow(InMesh->IsEdge(EdgeID));
 
 				const FIndex2i EdgeV = InMesh->GetEdgeV(EdgeID);
 		
@@ -1499,11 +1551,7 @@ namespace UniformTesselateLocals
 					return;
 				}
 
-				if (InMesh->IsTriangle(TriangleID) == false) 
-				{
-					checkSlow(false); // Mesh must be compact
-					return;
-				}
+				checkSlow(InMesh->IsTriangle(TriangleID));
 
 				const FIndex3i TriVertices = InMesh->GetTriangle(TriangleID);
 				const FIndex3i TriEdges = InMesh->GetTriEdges(TriangleID);
@@ -1574,11 +1622,17 @@ namespace UniformTesselateLocals
 		}
 	};
 
-	bool Tesselate(const FDynamicMesh3* InMesh, const int TesselationNum, FProgressCancel* Progress, const bool bUseParallel, FDynamicMesh3* OutMesh) 
+	bool Tesselate(const FDynamicMesh3* InMesh, 
+				   const int TesselationNum, 
+				   FProgressCancel* Progress, 
+				   const bool bUseParallel,
+				   TMap<int32, FIndex2i>* OutVertexEdgeMap,
+				   TMap<int32, int32>* OutVertexTriangleMap, 
+				   FDynamicMesh3* OutMesh) 
 	{	
 		OutMesh->Clear();
 
-		FMeshVerticesGenerator<double> FTrianglesGenerator(InMesh, TesselationNum, Progress, bUseParallel);
+		FMeshVerticesGenerator<double> FTrianglesGenerator(InMesh, TesselationNum, Progress, OutVertexEdgeMap, OutVertexTriangleMap, bUseParallel);
 		if (FTrianglesGenerator.Generate() == false) 
 		{
 			return false;
@@ -1767,6 +1821,23 @@ namespace UniformTesselateLocals
 
 }
 
+int32 FUniformTesselate::ExpectedNumVertices(const FDynamicMesh3& Mesh, const int32 TesselationNum)
+{
+	int32 EdgeVertCount = TesselationNum * Mesh.EdgeCount();
+	int32 TriangleNumber = TesselationNum + 2; // plus 2 corner vertices on the edge ends
+	int32 InnerTriangleVertCount = UniformTesselateLocals::TriangularPatternUtils::NumInnerVertices(TriangleNumber);
+	int32 TotalVertCount = Mesh.VertexCount() + EdgeVertCount + InnerTriangleVertCount* Mesh.TriangleCount();
+	return TotalVertCount; 	
+}
+	
+int32 FUniformTesselate::ExpectedNumTriangles(const FDynamicMesh3& Mesh, const int32 TesselationNum)
+{
+	int32 TriangleNumber = TesselationNum + 2; // plus 2 corner vertices on the edge ends
+	int32 InnerTriangleCount = UniformTesselateLocals::TriangularPatternUtils::NumTriangles(TriangleNumber);
+	int32 TotalTriangleCount = InnerTriangleCount * Mesh.TriangleCount();
+	return TotalTriangleCount;
+}
+
 bool FUniformTesselate::Cancelled()
 {
 	return (Progress == nullptr) ? false : Progress->Cancelled();
@@ -1774,37 +1845,147 @@ bool FUniformTesselate::Cancelled()
 
 bool FUniformTesselate::Compute()
 {	
-	if (TesselationNum < 0 || Mesh == nullptr || ResultMesh == nullptr) 
-	{ 
+	if (Validate() != EOperationValidationResult::Ok) 
+	{
 		return false;
 	}
 
-	ResultMesh->Clear();
-
-	if (TesselationNum == 0 || Mesh->TriangleCount() == 0)
+	if (TesselationNum == 0)
 	{
-		ResultMesh->Copy(*Mesh); // Nothing to do. Copy the mesh over.
 		return true; 
 	}
+
+	// Check for the empty meshes
+	if (bInPlace && ResultMesh->TriangleCount() == 0) 
+	{
+		return true;
+	}
+
+	if (bInPlace == false && Mesh->TriangleCount() == 0) 
+	{
+		return true;
+	}
+
+	// Make a copy of the mesh since we are tesselating in place. 
+	// The copy will be used to restore the mesh to its original state in case the user cancelled the operation.
+	FDynamicMesh3 ResultMeshCopy;
+	if (bInPlace) 
+	{
+		ResultMeshCopy.Copy(*ResultMesh);
+		Mesh = &ResultMeshCopy;
+	}
+
+	ResultMesh->Clear();
+	VertexMap.Empty();
+	VertexEdgeMap.Empty();
+	VertexTriangleMap.Empty();
 	
 	bool bIsValidMesh = false;
 
 	if (Mesh->IsCompact()) 
 	{
-		bIsValidMesh = UniformTesselateLocals::Tesselate(Mesh, TesselationNum, Progress, bUseParallel, ResultMesh);
+		bIsValidMesh = UniformTesselateLocals::Tesselate(Mesh, 
+														 TesselationNum, 
+														 Progress, 
+														 bUseParallel, 
+														 bComputeMappings ? &VertexEdgeMap : nullptr, 
+														 bComputeMappings ? &VertexTriangleMap : nullptr, 
+														 ResultMesh);
+
+		if (bIsValidMesh && bComputeMappings) 
+		{
+			VertexMap.SetNumUninitialized(Mesh->MaxVertexID());
+			for (int32 VertexID = 0; VertexID < Mesh->MaxVertexID(); ++VertexID)
+			{
+				VertexMap[VertexID] = VertexID; 
+			}
+		}
 	}
 	else 
 	{
 		FDynamicMesh3 CompactMesh;
-		CompactMesh.CompactCopy(*Mesh);
-		bIsValidMesh = UniformTesselateLocals::Tesselate(&CompactMesh, TesselationNum, Progress, bUseParallel, ResultMesh);
+		FCompactMaps CompactInfo;
+		CompactMesh.CompactCopy(*Mesh, true, true, true, true, &CompactInfo);
+
+		bIsValidMesh = UniformTesselateLocals::Tesselate(&CompactMesh, 
+														 TesselationNum, 
+														 Progress, 
+														 bUseParallel, 
+														 bComputeMappings ? &VertexEdgeMap : nullptr, 
+														 bComputeMappings ? &VertexTriangleMap : nullptr, 
+														 ResultMesh);
+
+
+		if (bIsValidMesh && bComputeMappings) 
+		{	
+			VertexMap.SetNumUninitialized(Mesh->MaxVertexID());
+			for (int32 VertexID = 0; VertexID < Mesh->MaxVertexID(); ++VertexID)
+			{
+				VertexMap[VertexID] = CompactInfo.GetVertexMapping(VertexID); 
+			}
+			
+			// Compute an additional mapping from the compact mesh triangles to the input mesh triangles 
+			// i.e the inverse of the FCompactMaps
+			TMap<int32, int32> CompactToInputTriangles;
+			CompactToInputTriangles.Reserve(CompactMesh.TriangleCount());
+			for (int32 TriangleID : Mesh->TriangleIndicesItr()) 
+			{
+				const int32 CompactTriangleID = CompactInfo.GetTriangleMapping(TriangleID);
+				if (CompactTriangleID != FCompactMaps::InvalidID)
+				{
+					CompactToInputTriangles.Add(CompactTriangleID, TriangleID);
+				}
+			}
+			checkSlow(CompactToInputTriangles.Num() == CompactMesh.TriangleCount());
+				
+			// Chage the mapping [tesselated triangles --> compact mesh triangles]
+			// to the mapping [tesselated triangles --> input mesh triangles (before compacting)]
+			for (TMap<int32, FIndex2i>::TIterator It = VertexEdgeMap.CreateIterator(); It; ++It)
+			{
+				FIndex2i& EdgeTri = It.Value();
+				
+				// We should always be able to find the correct match, but double check in case the Tesselate method
+				// failed to return the correct mapping
+				int32* TID1= CompactToInputTriangles.Find(EdgeTri.A);
+				ensure(TID1 != nullptr);
+				EdgeTri.A = TID1 != nullptr ? *TID1 : EdgeTri.A;
+				
+				if (EdgeTri.B != FDynamicMesh3::InvalidID) 
+				{
+					int32* TID2 = CompactToInputTriangles.Find(EdgeTri.B);
+					ensure(TID2 != nullptr);
+					EdgeTri.B = TID2 != nullptr ? *TID2 : EdgeTri.B;
+				}
+			}
+
+			for (TMap<int32, int32>::TIterator It = VertexTriangleMap.CreateIterator(); It; ++It)
+			{
+				int32& CompactTriID = It.Value();
+				int32* TID1 = CompactToInputTriangles.Find(CompactTriID);
+				ensure(TID1 != nullptr);
+				CompactTriID = TID1 != nullptr ? *TID1 : CompactTriID;
+			}
+		}
 	}
 
-	if (!bIsValidMesh) 
+	if (bIsValidMesh == false || Cancelled()) 
 	{
+		if (bInPlace) 
+		{ 
+			// Restore the input mesh
+			ResultMesh->Copy(ResultMeshCopy);
+		}
+		else 
+		{
+			ResultMesh->Clear();
+		}
+
+		VertexMap.Empty();
+		VertexEdgeMap.Empty();
+		VertexTriangleMap.Empty();
+		
 		return false;
 	}
 	
-	return Cancelled() == false;
-}
-
+	return true;
+}
\ No newline at end of file
diff --git a/Engine/Plugins/Runtime/GeometryProcessing/Source/DynamicMesh/Public/Operations/PNTriangles.h b/Engine/Plugins/Runtime/GeometryProcessing/Source/DynamicMesh/Public/Operations/PNTriangles.h
--- a/Engine/Plugins/Runtime/GeometryProcessing/Source/DynamicMesh/Public/Operations/PNTriangles.h
+++ b/Engine/Plugins/Runtime/GeometryProcessing/Source/DynamicMesh/Public/Operations/PNTriangles.h
@@ -28,7 +28,7 @@ public:
 	// Input & Output
 	//
 
-	/** The mesh that we are modifying */
+	/** The mesh that we are modifying. If the operator fails or is canceled by the user, the mesh will not be changed. */
 	FDynamicMesh3* Mesh = nullptr;
 
 	//
@@ -39,7 +39,7 @@ public:
 	FProgressCancel* Progress = nullptr;
 	
 	/** How many times we are recursively subdividing triangles (loop style subdivision). */
-	int TesselationLevel = 1;
+	int32 TesselationLevel = 1;
 
 	/**
 	 * If true, use the quadratically varying normal computation. 
@@ -51,7 +51,10 @@ public:
 	FPNTriangles(FDynamicMesh3* Mesh) : Mesh(Mesh)
 	{
 	}
-	virtual ~FPNTriangles() {}
+
+	virtual ~FPNTriangles() 
+	{
+	}
 
 	/**
 	 * @return EOperationValidationResult::Ok if we can apply operation, or error code if we cannot.
diff --git a/Engine/Plugins/Runtime/GeometryProcessing/Source/DynamicMesh/Public/Operations/UniformTesselate.h b/Engine/Plugins/Runtime/GeometryProcessing/Source/DynamicMesh/Public/Operations/UniformTesselate.h
--- a/Engine/Plugins/Runtime/GeometryProcessing/Source/DynamicMesh/Public/Operations/UniformTesselate.h
+++ b/Engine/Plugins/Runtime/GeometryProcessing/Source/DynamicMesh/Public/Operations/UniformTesselate.h
@@ -3,6 +3,7 @@
 #pragma once
 
 #include "GeometryTypes.h"
+#include "IndexTypes.h"
 
 class FProgressCancel;
 
@@ -41,9 +42,6 @@ public:
 	// Inputs
 	//
 	
-	/** The input mesh to be tesselated. */
-	const FDynamicMesh3* Mesh = nullptr;
-
 	/** Set this to be able to cancel running operation. */
 	FProgressCancel* Progress = nullptr;
 
@@ -53,17 +51,67 @@ public:
 	/** Should multi-threading be enabled. */
 	bool bUseParallel = true;
 
+	/** If true, populate VertexMap, VertexEdgeMap, VertexTriangleMap. */
+	bool bComputeMappings = false;
+
 	//
-	// Output
+	// Input/Output
 	//
 
 	/** The tesselated mesh. */
-	FDynamicMesh3* ResultMesh;
+	FDynamicMesh3* ResultMesh = nullptr;
+	
+	//
+	// Output
+	//
+
+	/** 
+	 * Map the vertex ID in the input mesh to the corresponding vertex ID in the output mesh. 
+	 * If the vertex ID is Invalid in the input mesh, then it is mapped to InvalidID.
+	 */
+	TArray<int32> VertexMap;
+	
+	/**
+	 * Map the ID of the vertex inserted along the input mesh edge to the input mesh triangle IDs (t1,t2) that share 
+	 * that edge. If its a boundary edge then t2 is InvalidID.
+	 */
+	TMap<int32, FIndex2i> VertexEdgeMap;
+
+	/** Map the ID of the vertex inserted inside the input mesh triangle to the triangle's ID. */
+	TMap<int32, int32> VertexTriangleMap;
+
+protected:
+	
+	/** 
+	 * Points to either the input mesh to be used to generate the new tesselated mesh or will point to a copy of 
+	 * the mesh we are tesselating inplace.
+	 */
+	const FDynamicMesh3* Mesh = nullptr;
+
+	/** 
+	 * If true, the mesh to tesselate is contained in the ResultMesh and will be overwritten with the 
+	 * tesselation result. 
+	 */
+	bool bInPlace = false;
 
 public:
+
+	/** 
+	 * Tesselate the mesh inplace. If the tesselation fails or the user cancelled the operation, OutMesh will not 
+	 * be changed.
+	 */
+	FUniformTesselate(FDynamicMesh3* OutMesh) 
+	:
+	ResultMesh(OutMesh), bInPlace(true)
+	{
+	}
+
+	/** 
+	 * Tesselate the mesh and write the result into another mesh. This will overwrite any data stored in the OutMesh. 
+	 */
 	FUniformTesselate(const FDynamicMesh3* Mesh, FDynamicMesh3* OutMesh) 
 	: 
-	Mesh(Mesh), ResultMesh(OutMesh)
+	ResultMesh(OutMesh), Mesh(Mesh), bInPlace(false)
 	{
 	}
 	
@@ -71,14 +119,23 @@ public:
 	{
 	}
 
+	/** @return The number of vertices after the tesselation. */
+	static int32 ExpectedNumVertices(const FDynamicMesh3& Mesh, const int32 TesselationNum);
+	
+	/** @return The number of triangles after the tesselation. */
+	static int32 ExpectedNumTriangles(const FDynamicMesh3& Mesh, const int32 TesselationNum);
+
 	/**
 	 * @return EOperationValidationResult::Ok if we can apply operation, or error code if we cannot.
 	 */
 	virtual EOperationValidationResult Validate()
-	{
-		if (TesselationNum < 0 || Mesh == nullptr || ResultMesh == nullptr) 
+	{	
+		const bool bIsInvalid = bInPlace == false && (Mesh == nullptr || ResultMesh == nullptr);
+		const bool bIsInvalidInPlace = bInPlace && ResultMesh == nullptr;
+		
+		if (TesselationNum < 0 || bIsInvalid || bIsInvalidInPlace) 
 		{
-			return EOperationValidationResult::Failed_UnknownReason;    
+			return EOperationValidationResult::Failed_UnknownReason;
 		}
 
 		return EOperationValidationResult::Ok;
