GeometryProcessing: fix handling of single-edge case in MeshBevel. Need to properly compute Wedge border edge information so that the edge-pairs can later be found and added to the edge-correspondence map.
#rnx

diff --git a/Engine/Plugins/Runtime/GeometryProcessing/Source/DynamicMesh/Private/Operations/MeshBevel.cpp b/Engine/Plugins/Runtime/GeometryProcessing/Source/DynamicMesh/Private/Operations/MeshBevel.cpp
--- a/Engine/Plugins/Runtime/GeometryProcessing/Source/DynamicMesh/Private/Operations/MeshBevel.cpp
+++ b/Engine/Plugins/Runtime/GeometryProcessing/Source/DynamicMesh/Private/Operations/MeshBevel.cpp
@@ -529,21 +529,53 @@ void FMeshBevel::BuildTerminatorVertex(FBevelVertex& Vertex, const FDynamicMesh3
 		return;
 	}
 
+	// save edgeid/vertexid for the 'terminator edge' that we will disconnect at
 	FIndex2i SplitEdgeV = Mesh.GetEdgeV(RingSplitEdgeID);
 	Vertex.TerminatorInfo = FIndex2i(RingSplitEdgeID, SplitEdgeV.OtherElement(Vertex.VertexID));
 
+	// split the terminator vertex into two wedges
 	TArray<int32> SplitTriSets[2];
 	if (SplitInteriorVertexTrianglesIntoSubsets(&Mesh, Vertex.VertexID, IncomingEdgeID, RingSplitEdgeID, SplitTriSets[0], SplitTriSets[1]) == false)
 	{
 		return;
 	}
 
+	// make the wedge list
 	Vertex.Wedges.SetNum(2);
 	Vertex.Wedges[0].WedgeVertex = Vertex.VertexID;
 	Vertex.Wedges[0].Triangles.Append(SplitTriSets[0]);
 	Vertex.Wedges[1].WedgeVertex = Vertex.VertexID;
 	Vertex.Wedges[1].Triangles.Append(SplitTriSets[1]);
 
+	// We need to know the two border edges of each wedge, because we will use this information
+	// in later stages. Note that this block is not specific to the TerminatorVertex case
+	for (FOneRingWedge& Wedge : Vertex.Wedges)
+	{
+		int32 NumWedgeTris = Wedge.Triangles.Num();
+
+		FIndex2i VtxEdges0 = IndexUtil::FindVertexEdgesInTriangle(Mesh, Wedge.Triangles[0], Vertex.VertexID);
+		if (NumWedgeTris == 1)
+		{
+			// If the wedge only has one tri, both edges are the border edges. 
+			// Note that these are not sorted, ie B might not be the edge shared with the next Wedge. 
+			// Currently this does not matter but it might in the future?
+			Wedge.BorderEdges.A = VtxEdges0.A;
+			Wedge.BorderEdges.B = VtxEdges0.B;
+		}
+		else
+		{
+			// the wedge-border-edge is *not* the edge connected to the next triangle in the wedge-triangle-list
+			Wedge.BorderEdges.A = ( Mesh.GetEdgeT(VtxEdges0.A).Contains(Wedge.Triangles[1]) ) ? VtxEdges0.B : VtxEdges0.A;
+			// final wedge-border-edge is the same case, with the last and second-last tris
+			FIndex2i VtxEdges1 = IndexUtil::FindVertexEdgesInTriangle(Mesh, Wedge.Triangles[NumWedgeTris-1], Vertex.VertexID);
+			Wedge.BorderEdges.B = ( Mesh.GetEdgeT(VtxEdges1.A).Contains(Wedge.Triangles[NumWedgeTris-2]) ) ? VtxEdges1.B : VtxEdges1.A;
+		}
+		// save the index of the border edge in it's triangle, so that when we disconnect the wedges later,
+		// we can find the new edge ID
+		Wedge.BorderEdgeTriEdgeIndices.A = Mesh.GetTriEdges(Wedge.Triangles[0]).IndexOf(Wedge.BorderEdges.A);
+		Wedge.BorderEdgeTriEdgeIndices.B = Mesh.GetTriEdges(Wedge.Triangles.Last()).IndexOf(Wedge.BorderEdges.B);
+	}
+
 	Vertex.VertexType = EBevelVertexType::TerminatorVertex;
 }
 
@@ -840,7 +872,6 @@ void FMeshBevel::UnlinkJunctionVertex(FDynamicMesh3& Mesh, FBevelVertex& Vertex,
 			int32 OldWedgeEdgeIndex = Wedge.BorderEdgeTriEdgeIndices[j];
 			int32 TriangleID = (j == 0) ? Wedge.Triangles[0] : Wedge.Triangles.Last();
 			int32 CurWedgeEdgeID = Mesh.GetTriEdges(TriangleID)[OldWedgeEdgeIndex];
-			FIndex3i TriVerts = Mesh.GetTriangle(TriangleID);
 			if (OldWedgeEdgeID != CurWedgeEdgeID)
 			{
 				if (MeshEdgePairs.Contains(OldWedgeEdgeID) == false)
@@ -868,15 +899,40 @@ void FMeshBevel::UnlinkTerminatorVertex(FDynamicMesh3& Mesh, FBevelVertex& Bevel
 		ChangeTracker->SaveVertexOneRingTriangles(BevelVertex.VertexID, true);
 	}
 
-	// TODO: do we need to update Wedge BorderEdges here??
-
 	// split the vertex
 	FDynamicMesh3::FVertexSplitInfo SplitInfo;
 	EMeshResult Result = Mesh.SplitVertex(BevelVertex.VertexID, BevelVertex.Wedges[1].Triangles, SplitInfo);
 	if (Result == EMeshResult::Ok)
 	{
 		BevelVertex.Wedges[1].WedgeVertex = SplitInfo.NewVertex;
+
+		// update end start/end pairs for each wedge, and save in the edge correspondence map.
+		// Note that this is the same block as in UnlinkJunctionVertex
+		int32 NumWedges = BevelVertex.Wedges.Num();
+		for (int32 k = 0; k < NumWedges; ++k)
+		{
+			FOneRingWedge& Wedge = BevelVertex.Wedges[k];
+			for (int32 j = 0; j < 2; ++j)
+			{
+				int32 OldWedgeEdgeID = Wedge.BorderEdges[j];
+				int32 OldWedgeEdgeIndex = Wedge.BorderEdgeTriEdgeIndices[j];
+				int32 TriangleID = (j == 0) ? Wedge.Triangles[0] : Wedge.Triangles.Last();
+				int32 CurWedgeEdgeID = Mesh.GetTriEdges(TriangleID)[OldWedgeEdgeIndex];
+				if (OldWedgeEdgeID != CurWedgeEdgeID)
+				{
+					if (MeshEdgePairs.Contains(OldWedgeEdgeID) == false)
+					{
+						MeshEdgePairs.Add(OldWedgeEdgeID, CurWedgeEdgeID);
+						MeshEdgePairs.Add(CurWedgeEdgeID, OldWedgeEdgeID);
+					}
+					Wedge.BorderEdges[j] = CurWedgeEdgeID;
+				}
+			}
+		}
+
 	}
+
+
 }
 
 
@@ -1473,11 +1529,12 @@ void FMeshBevel::ComputeMaterialIDs(FDynamicMesh3& Mesh)
 				FIndex3i TriNbrs = Mesh.GetTriNeighbourTris(tid);
 				for (int32 j = 0; j < 3; ++j)
 				{
-					if (Triangles.Contains(TriNbrs[j]))
+					int32 NbrTriangleID = TriNbrs[j];
+					if ( Mesh.IsTriangle(NbrTriangleID) == false || Triangles.Contains(NbrTriangleID) )
 					{
 						continue;
 					}
-					int MatID = MaterialAttrib.GetValue(TriNbrs[j]);
+					int MatID = MaterialAttrib.GetValue(NbrTriangleID);
 					int32 Index = MaterialIDs.AddUnique(MatID);
 					if (Counts.Num() != MaterialIDs.Num())
 					{
