ModelingTools: Fix texture filtering into neighboring charts for tightly packed UV shells.

#rnx

diff --git a/Engine/Plugins/Experimental/MeshModelingToolsetExp/Source/MeshModelingToolsExp/Private/BakeMeshAttributeMapsTool.cpp b/Engine/Plugins/Experimental/MeshModelingToolsetExp/Source/MeshModelingToolsExp/Private/BakeMeshAttributeMapsTool.cpp
--- a/Engine/Plugins/Experimental/MeshModelingToolsetExp/Source/MeshModelingToolsExp/Private/BakeMeshAttributeMapsTool.cpp
+++ b/Engine/Plugins/Experimental/MeshModelingToolsetExp/Source/MeshModelingToolsExp/Private/BakeMeshAttributeMapsTool.cpp
@@ -103,6 +103,7 @@ public:
 	TUniquePtr<UE::Geometry::FMeshMapBaker> Baker;
 	UBakeMeshAttributeMapsTool::FBakeSettings BakeSettings;
 	TSharedPtr<UE::Geometry::TMeshTangents<double>, ESPMode::ThreadSafe> BaseMeshTangents;
+	TArray<int32>* BaseMeshUVCharts;
 
 	// Map Type settings
 	EBakeMapType Maps;
@@ -133,6 +134,7 @@ public:
 		Baker->SetDimensions(BakeSettings.Dimensions);
 		Baker->SetProjectionDistance(BakeSettings.ProjectionDistance);
 		Baker->SetSamplesPerPixel(BakeSettings.SamplesPerPixel);
+		Baker->SetTargetMeshUVCharts(BaseMeshUVCharts);
 		Baker->SetTargetMeshTangents(BaseMeshTangents);
 		
 		FMeshBakerDynamicMeshSampler DetailSampler(DetailMesh.Get(), DetailSpatial.Get(), DetailMeshTangents.Get());
@@ -398,6 +400,7 @@ TUniquePtr<UE::Geometry::TGenericDataOperator<FMeshMapBaker>> UBakeMeshAttribute
 	Op->DetailSpatial = DetailSpatial;
 	Op->BaseMesh = &TargetMesh;
 	Op->BakeSettings = CachedBakeSettings;
+	Op->BaseMeshUVCharts = &TargetMeshUVCharts;
 
 	constexpr EBakeMapType RequiresTangents = EBakeMapType::TangentSpaceNormal | EBakeMapType::BentNormal;
 	if ((bool)(CachedBakeSettings.BakeMapTypes & RequiresTangents))
diff --git a/Engine/Plugins/Experimental/MeshModelingToolsetExp/Source/MeshModelingToolsExp/Private/BakeMeshAttributeMapsToolBase.cpp b/Engine/Plugins/Experimental/MeshModelingToolsetExp/Source/MeshModelingToolsExp/Private/BakeMeshAttributeMapsToolBase.cpp
--- a/Engine/Plugins/Experimental/MeshModelingToolsetExp/Source/MeshModelingToolsExp/Private/BakeMeshAttributeMapsToolBase.cpp
+++ b/Engine/Plugins/Experimental/MeshModelingToolsetExp/Source/MeshModelingToolsExp/Private/BakeMeshAttributeMapsToolBase.cpp
@@ -79,6 +79,10 @@ void UBakeMeshAttributeMapsToolBase::PostSetup()
 	Compute->Setup(this);
 	Compute->OnResultUpdated.AddLambda([this](const TUniquePtr<FMeshMapBaker>& NewResult) { OnMapsUpdated(NewResult); });
 
+	// Initialize UV charts
+	// TODO: Compute UV charts asynchronously
+	FMeshMapBaker::ComputeUVCharts(TargetMesh, TargetMeshUVCharts);
+
 	GatherAnalytics(BakeAnalytics.MeshSettings);
 }
 
diff --git a/Engine/Plugins/Experimental/MeshModelingToolsetExp/Source/MeshModelingToolsExp/Private/BakeMultiMeshAttributeMapsTool.cpp b/Engine/Plugins/Experimental/MeshModelingToolsetExp/Source/MeshModelingToolsExp/Private/BakeMultiMeshAttributeMapsTool.cpp
--- a/Engine/Plugins/Experimental/MeshModelingToolsetExp/Source/MeshModelingToolsExp/Private/BakeMultiMeshAttributeMapsTool.cpp
+++ b/Engine/Plugins/Experimental/MeshModelingToolsetExp/Source/MeshModelingToolsExp/Private/BakeMultiMeshAttributeMapsTool.cpp
@@ -319,6 +319,7 @@ public:
 	TSharedPtr<UE::Geometry::FMeshTangentsd, ESPMode::ThreadSafe> BaseMeshTangents;
 	TUniquePtr<UE::Geometry::FMeshMapBaker> Baker;
 	UBakeMultiMeshAttributeMapsTool::FBakeSettings BakeSettings;
+	TArray<int32>* BaseMeshUVCharts;
 
 	// Detail bake data
 	TArray<TSharedPtr<UE::Geometry::TImageBuilder<FVector4f>>> CachedColorImages;
@@ -337,6 +338,7 @@ public:
 		Baker->SetProjectionDistance(BakeSettings.ProjectionDistance);
 		Baker->SetSamplesPerPixel(BakeSettings.SamplesPerPixel);
 		Baker->SetTargetMeshTangents(BaseMeshTangents);
+		Baker->SetTargetMeshUVCharts(BaseMeshUVCharts);
 		
 		FMeshBakerMeshSceneSampler DetailSampler(DetailMeshScene);
 		Baker->SetDetailSampler(&DetailSampler);
@@ -503,6 +505,7 @@ TUniquePtr<UE::Geometry::TGenericDataOperator<FMeshMapBaker>> UBakeMultiMeshAttr
 	TUniquePtr<FMultiMeshMapBakerOp> Op = MakeUnique<FMultiMeshMapBakerOp>();
 	Op->DetailMeshScene = &DetailMeshScene;
 	Op->BaseMesh = &TargetMesh;
+	Op->BaseMeshUVCharts = &TargetMeshUVCharts;
 	Op->BakeSettings = CachedBakeSettings;
 
 	constexpr EBakeMapType RequiresTangents = EBakeMapType::TangentSpaceNormal | EBakeMapType::BentNormal;
diff --git a/Engine/Plugins/Experimental/MeshModelingToolsetExp/Source/MeshModelingToolsExp/Public/BakeMeshAttributeMapsToolBase.h b/Engine/Plugins/Experimental/MeshModelingToolsetExp/Source/MeshModelingToolsExp/Public/BakeMeshAttributeMapsToolBase.h
--- a/Engine/Plugins/Experimental/MeshModelingToolsetExp/Source/MeshModelingToolsExp/Public/BakeMeshAttributeMapsToolBase.h
+++ b/Engine/Plugins/Experimental/MeshModelingToolsetExp/Source/MeshModelingToolsExp/Public/BakeMeshAttributeMapsToolBase.h
@@ -176,6 +176,8 @@ protected:
 	};
 	FBakeSettings CachedBakeSettings;
 
+	TArray<int32> TargetMeshUVCharts;
+
 	/**
 	 * To be invoked by client when bake map types change.
 	 * @param ResultMapTypes the requested map types to compute
diff --git a/Engine/Plugins/Runtime/GeometryProcessing/Source/DynamicMesh/Private/Sampling/MeshMapBaker.cpp b/Engine/Plugins/Runtime/GeometryProcessing/Source/DynamicMesh/Private/Sampling/MeshMapBaker.cpp
--- a/Engine/Plugins/Runtime/GeometryProcessing/Source/DynamicMesh/Private/Sampling/MeshMapBaker.cpp
+++ b/Engine/Plugins/Runtime/GeometryProcessing/Source/DynamicMesh/Private/Sampling/MeshMapBaker.cpp
@@ -5,7 +5,7 @@
 #include "Sampling/MeshMapBakerQueue.h"
 #include "Image/ImageOccupancyMap.h"
 #include "Image/ImageTile.h"
-#include "Algo/Count.h"
+#include "Selections/MeshConnectedComponents.h"
 #include "ProfilingDebugging/ScopedTimers.h"
 
 using namespace UE::Geometry;
@@ -81,6 +81,13 @@ void FMeshMapBaker::InitBake()
 	}
 
 	InitFilter();
+
+	// Compute UV charts if null or invalid.
+	if (!TargetMeshUVCharts || !ensure(TargetMeshUVCharts->Num() == TargetMesh->TriangleCount()))
+	{
+		ComputeUVCharts(*TargetMesh, TargetMeshUVChartsLocal);
+		TargetMeshUVCharts = &TargetMeshUVChartsLocal;
+	}
 }
 
 void FMeshMapBaker::InitBakeDefaults()
@@ -121,7 +128,7 @@ void FMeshMapBaker::Bake()
 	BakeAnalytics.Reset();
 	FScopedDurationTimer TotalBakeTimer(BakeAnalytics.TotalBakeDuration);
 	
-	if (Bakers.IsEmpty())
+	if (Bakers.IsEmpty() || !TargetMesh)
 	{
 		return;
 	}
@@ -305,11 +312,17 @@ void FMeshMapBaker::Bake()
 
 		FImageOccupancyMap OccupancyMap;
 		OccupancyMap.GutterSize = GutterSize;
-		OccupancyMap.Initialize(Dimensions, Tile, SamplesPerPixel);
-		OccupancyMap.ComputeFromUVSpaceMesh(FlatMesh, [this](int32 TriangleID) { return FlatMesh.GetTriangleGroup(TriangleID); });
+		OccupancyMap.Initialize(Dimensions, PaddedTile, SamplesPerPixel);
+		OccupancyMap.ComputeFromUVSpaceMesh(FlatMesh, [this](int32 TriangleID) { return FlatMesh.GetTriangleGroup(TriangleID); }, TargetMeshUVCharts);
 		GutterTexelsPerTile[TileIdx] = OccupancyMap.GutterTexels;
 
-		BakeAnalytics.NumSamplePixels += Algo::CountIf(OccupancyMap.TexelInteriorSamples, [](const int32 Value){ return Value > 0; });
+		const int64 NumTilePixels = Tile.Num();
+		for (int64 TilePixelIdx = 0; TilePixelIdx < NumTilePixels; ++TilePixelIdx)
+		{
+			const FVector2i SourceCoords = Tile.GetSourceCoords(TilePixelIdx);
+			const int64 OccupancyMapIdx = OccupancyMap.Tile.GetIndexFromSourceCoords(SourceCoords);
+			BakeAnalytics.NumSamplePixels += OccupancyMap.TexelInteriorSamples[OccupancyMapIdx];; 
+		}
 
 		FMeshMapTileBuffer* TileBuffer = new FMeshMapTileBuffer(PaddedTile, BakeSampleBufferSize);
 
@@ -330,16 +343,16 @@ void FMeshMapBaker::Bake()
 						return;
 					}
 
-					const int64 TilePixelLinearIdx = Tile.GetIndex(TileCoords);
-					if (OccupancyMap.TexelNumSamples(TilePixelLinearIdx) == 0)
+					const FVector2i ImageCoords = Tile.GetSourceCoords(TileCoords);
+					const int64 OccupancyMapLinearIdx = OccupancyMap.Tile.GetIndexFromSourceCoords(ImageCoords);
+					if (OccupancyMap.TexelNumSamples(OccupancyMapLinearIdx) == 0)
 					{
 						continue;
 					}
 
-					const FVector2i ImageCoords = Tile.GetSourceCoords(TileCoords);
 					for (int32 SampleIdx = 0; SampleIdx < NumSamples; ++SampleIdx)
 					{
-						const int64 LinearIdx = TilePixelLinearIdx * NumSamples + SampleIdx;
+						const int64 LinearIdx = OccupancyMapLinearIdx * NumSamples + SampleIdx;
 						if (OccupancyMap.IsInterior(LinearIdx))
 						{
 							const FVector2d UVPosition = (FVector2d)OccupancyMap.TexelQueryUV[LinearIdx];
@@ -349,7 +362,7 @@ void FMeshMapBaker::Bake()
 							DetailCorrespondenceSampler.SampleUV(UVTriangleID, UVPosition, Sample);
 							if (Sample.DetailMesh && DetailSampler->IsTriangle(Sample.DetailMesh, Sample.DetailTriID))
 							{
-								BakeSample(*TileBuffer, Sample, UVPosition, ImageCoords);
+								BakeSample(*TileBuffer, Sample, UVPosition, ImageCoords, OccupancyMap);
 							}
 						}
 					}
@@ -479,7 +492,8 @@ void FMeshMapBaker::BakeSample(
 	FMeshMapTileBuffer& TileBuffer,
 	const FMeshMapEvaluator::FCorrespondenceSample& Sample,
 	const FVector2d& UVPosition,
-	const FVector2i& ImageCoords)
+	const FVector2i& ImageCoords,
+	const FImageOccupancyMap& OccupancyMap)
 {
 	// Evaluate each baker into stack allocated float buffer
 	float* Buffer = static_cast<float*>(FMemory_Alloca(sizeof(float) * BakeSampleBufferSize));
@@ -493,7 +507,10 @@ void FMeshMapBaker::BakeSample(
 
 	const FImageTile& Tile = TileBuffer.GetTile();
 
-	auto AddFn = [this, &ImageCoords, &UVPosition, &Tile, &TileBuffer, Buffer](const TArray<int32>& BakeIds) -> void
+	const int64 OccupancyMapSampleIdx = OccupancyMap.Tile.GetIndexFromSourceCoords(ImageCoords);
+	const int32 SampleUVChart = OccupancyMap.TexelQueryUVChart[OccupancyMapSampleIdx];
+
+	auto AddFn = [this, &ImageCoords, &UVPosition, &Tile, &TileBuffer, Buffer, &OccupancyMap, SampleUVChart](const TArray<int32>& BakeIds) -> void
 	{
 		const FVector2i BoxFilterStart(
 			FMath::Clamp(ImageCoords.X - FilterKernelSize, 0, Dimensions.GetWidth()),
@@ -509,6 +526,8 @@ void FMeshMapBaker::BakeSample(
 		{
 			const FVector2i SourceCoords = BoxFilterTile.GetSourceCoords(FilterIdx);
 			const int64 BufferTilePixelLinearIdx = Tile.GetIndexFromSourceCoords(SourceCoords);
+			const int64 OccupancyMapFilterIdx = OccupancyMap.Tile.GetIndexFromSourceCoords(SourceCoords);
+			const int32 BufferTilePixelUVChart = OccupancyMap.TexelQueryUVChart[OccupancyMapFilterIdx];
 			float* PixelBuffer = TileBuffer.GetPixel(BufferTilePixelLinearIdx);
 			float& PixelWeight = TileBuffer.GetPixelWeight(BufferTilePixelLinearIdx);
 
@@ -517,7 +536,8 @@ void FMeshMapBaker::BakeSample(
 			TexelDistance.X *= Dimensions.GetWidth();
 			TexelDistance.Y *= Dimensions.GetHeight();
 
-			const float FilterWeight = TextureFilterEval(TexelDistance);
+			float FilterWeight = TextureFilterEval(TexelDistance);
+			FilterWeight *= (SampleUVChart == BufferTilePixelUVChart);
 			PixelWeight += FilterWeight;
 			for (const int32 BakeIdx : BakeIds)
 			{
@@ -670,4 +690,27 @@ float FMeshMapBaker::EvaluateFilter(const FVector2d& Dist)
 }
 
 
+void FMeshMapBaker::ComputeUVCharts(const FDynamicMesh3& Mesh, TArray<int32>& MeshUVCharts)
+{
+	MeshUVCharts.SetNumZeroed(Mesh.TriangleCount());
+	if (const FDynamicMeshUVOverlay* UVOverlay = Mesh.Attributes() ? Mesh.Attributes()->PrimaryUV() : nullptr)
+	{
+		FMeshConnectedComponents UVComponents(&Mesh);
+		UVComponents.FindConnectedTriangles();
+		UVComponents.FindConnectedTriangles([UVOverlay](int32 Triangle0, int32 Triangle1) {
+			return UVOverlay ? UVOverlay->AreTrianglesConnected(Triangle0, Triangle1) : false;
+		});
+		const int32 NumComponents = UVComponents.Num();
+		for (int32 ComponentId = 0; ComponentId < NumComponents; ++ComponentId)
+		{
+			const FMeshConnectedComponents::FComponent& UVComp = UVComponents.GetComponent(ComponentId);
+			for (const int32 TriId : UVComp.Indices)
+			{
+				MeshUVCharts[TriId] = ComponentId;
+			}
+		}
+	}
+}
+
+
 
diff --git a/Engine/Plugins/Runtime/GeometryProcessing/Source/DynamicMesh/Public/Sampling/MeshMapBaker.h b/Engine/Plugins/Runtime/GeometryProcessing/Source/DynamicMesh/Public/Sampling/MeshMapBaker.h
--- a/Engine/Plugins/Runtime/GeometryProcessing/Source/DynamicMesh/Public/Sampling/MeshMapBaker.h
+++ b/Engine/Plugins/Runtime/GeometryProcessing/Source/DynamicMesh/Public/Sampling/MeshMapBaker.h
@@ -77,6 +77,36 @@ public:
 	EBakeFilterType GetFilter() const { return FilterType; }
 	int32 GetTileSize() const { return TileSize; }
 
+	/**
+	 * Computes the connected UV triangles and returns an array containing
+	 * the mapping from triangle ID to unique UV chart ID. If the mesh
+	 * has no UVs, the UVCharts will be initialized to 0.
+	 *
+	 * @param Mesh the mesh to compute UV charts.
+	 * @param MeshUVCharts the triangle ID to UV Chart ID array.
+	 */
+	static void ComputeUVCharts(const FDynamicMesh3& Mesh, TArray<int32>& MeshUVCharts);
+
+	/**
+	 * Set an a Triangle ID to UV Chart ID array for TargetMesh.
+	 * If this is not set, then the baker will compute it as part of Bake().
+	 * Since ComputeUVCharts() is non-trivial, this method is intended
+	 * to allow a client to externally cache the result of ComputeUVCharts
+	 * to minimize the overhead per bake.
+	 * 
+	 * @param UVChartsIn the TriID to UVChartID map
+	 */
+	void SetTargetMeshUVCharts(TArray<int32>* UVChartsIn)
+	{
+		TargetMeshUVCharts = UVChartsIn;
+	}
+
+	/** @return the Triangle ID to UV Chart ID mapping */
+	const TArray<int32>* GetTargetMeshUVCharts() const
+	{
+		return TargetMeshUVCharts;
+	}
+
 	//
 	// Analytics
 	//
@@ -105,7 +135,8 @@ protected:
 		FMeshMapTileBuffer& TileBuffer,
 		const FMeshMapEvaluator::FCorrespondenceSample& Sample,
 		const FVector2d& UVPosition,
-		const FVector2i& ImageCoords);
+		const FVector2i& ImageCoords,
+		const FImageOccupancyMap& OccupancyMap);
 
 	/** Initialize evaluation contexts and precompute data for bake evaluation. */
 	void InitBake();
@@ -184,6 +215,20 @@ protected:
 
 	/** Array of bake result images. */
 	TArray<TUniquePtr<TImageBuilder<FVector4f>>> BakeResults;
+
+	/**
+	 * Array of TargetMesh triangle ID to UV chart ID mapping.
+	 * Can be optionally provided by the client. If not provided,
+	 * will be computed as part of the bake.
+	 */
+	TArray<int32>* TargetMeshUVCharts = nullptr;
+
+	/**
+	 * Local Array of TargetMesh triangle ID to UV chart ID mapping.
+	 * This will be populated only if not provided by the client via
+	 * TargetMeshUVCharts.
+	 */
+	TArray<int32> TargetMeshUVChartsLocal;
 };
 
 } // end namespace UE::Geometry
diff --git a/Engine/Source/Runtime/GeometryCore/Public/Image/ImageOccupancyMap.h b/Engine/Source/Runtime/GeometryCore/Public/Image/ImageOccupancyMap.h
--- a/Engine/Source/Runtime/GeometryCore/Public/Image/ImageOccupancyMap.h
+++ b/Engine/Source/Runtime/GeometryCore/Public/Image/ImageOccupancyMap.h
@@ -51,6 +51,8 @@ public:
 	TArray64<FVector2f> TexelQueryUV;
 	/** integer/Triangle ID for each texel in image. Only set for Interior texels. */
 	TArray64<int32> TexelQueryTriangle;
+	/** UV Chart ID for each texel in image. Only set if UVSpaceMeshTriCharts is provided during ComputeFromUVSpaceMesh */
+	TArray64<int32> TexelQueryUVChart;
 
 	/** Set of Gutter Texels. Pair is <LinearIndexOfGutterTexel, LinearIndexOfNearestInteriorTexel>, so
 	    Gutter can be filled by directly copying from source to target. */
@@ -99,9 +101,17 @@ public:
 	//}
 
 
+	/**
+	 * Computes the image occupancy map from a UV space mesh.
+	 * 
+	 * @param UVSpaceMesh UV space mesh to compute the occupancy map from.
+	 * @param GetTriangleIDFunc Lambda to remap a texel's nearest triangle ID
+	 * @param UVSpaceMeshTriCharts Optional UVSpaceMesh triangle ID to UV chart map 
+	 */
 	template<typename MeshType, typename GetTriangleIDFuncType>
 	bool ComputeFromUVSpaceMesh(const MeshType& UVSpaceMesh, 
-		GetTriangleIDFuncType GetTriangleIDFunc = [](int32 TriangleID) { return TriangleID; } )
+		GetTriangleIDFuncType GetTriangleIDFunc = [](int32 TriangleID) { return TriangleID; },
+		const TArray<int32>* UVSpaceMeshTriCharts = nullptr)
 	{
 		TRACE_CPUPROFILER_EVENT_SCOPE(FImageOccupancyMap::ComputeFromUVSpaceMesh);
 		
@@ -114,6 +124,7 @@ public:
 		TexelInteriorSamples.Init(0, LinearImageSize);
 		TexelQueryUV.Init(FVector2f::Zero(), LinearSampleSize);
 		TexelQueryTriangle.Init(IndexConstants::InvalidID, LinearSampleSize);
+		TexelQueryUVChart.Init(IndexConstants::InvalidID, LinearImageSize);
 
 		const FVector2d TexelSize = Dimensions.GetTexelSize();
 		const double TexelDiag = TexelSize.Length();
@@ -126,8 +137,10 @@ public:
 		TArray< TArray64<TTuple<int64, int64>> > GutterTexelsPerScanline;
 		GutterTexelsPerScanline.SetNum(Tile.GetHeight());
 
+		const int32 NumUVSpaceMeshTriCharts = UVSpaceMeshTriCharts ? UVSpaceMeshTriCharts->Num() : 0;
+
 		ParallelFor(Tile.GetHeight(),
-					[this, &UVSpaceMesh, &GetTriangleIDFunc, &FlatSpatial, TexelDiag, &QueryOptions, &TexelSize, &TotalGutterCounter, &GutterTexelsPerScanline]
+					[this, &UVSpaceMesh, &GetTriangleIDFunc, &FlatSpatial, TexelDiag, &QueryOptions, &TexelSize, &TotalGutterCounter, &GutterTexelsPerScanline, UVSpaceMeshTriCharts, NumUVSpaceMeshTriCharts]
 		(int32 ImgY)
 		{
 			for (int32 ImgX = 0; ImgX < Tile.GetWidth(); ++ImgX)
@@ -159,11 +172,12 @@ public:
 						UVSpaceMesh.GetTriVertices(NearestTriID, A, B, C);
 						const FTriangle2d UVTriangle(GetXY(A), GetXY(B), GetXY(C));
 
+						const int32 TriId = GetTriangleIDFunc(NearestTriID);
 						if (UVTriangle.IsInsideOrOn(UVPoint))
 						{
 							TexelType[SampleLinearIdx] = InteriorTexel;
 							TexelQueryUV[SampleLinearIdx] = (FVector2f)UVPoint;
-							TexelQueryTriangle[SampleLinearIdx] = GetTriangleIDFunc(NearestTriID);
+							TexelQueryTriangle[SampleLinearIdx] = TriId;
 							++TexelInteriorSamples[TexelLinearIdx];
 							bIsGutterTexel = false;
 						}
@@ -176,7 +190,7 @@ public:
 
 							TexelType[SampleLinearIdx] = InteriorTexel;
 							TexelQueryUV[SampleLinearIdx] = (FVector2f)NearestUV;
-							TexelQueryTriangle[SampleLinearIdx] = GetTriangleIDFunc(NearestTriID);
+							TexelQueryTriangle[SampleLinearIdx] = TriId;
 							++TexelInteriorSamples[TexelLinearIdx];
 							bIsGutterTexel = false;
 						}
@@ -193,6 +207,10 @@ public:
 							const int64 NearestLinearIdx = Dimensions.GetIndex(NearestCoords);
 							GutterNearestTexel = TTuple<int64, int64>(SourceTexelLinearIdx, NearestLinearIdx);
 						}
+						if (TriId < NumUVSpaceMeshTriCharts && TexelQueryUVChart[TexelLinearIdx] == IndexConstants::InvalidID)
+						{
+							TexelQueryUVChart[TexelLinearIdx] = (*UVSpaceMeshTriCharts)[TriId];
+						}
 						bHitTriID = true;
 					}
 				}
