ModelingTools: Fix XAtlas AutoUV placing UVs outside of the unit square.

#rnx

diff --git a/Engine/Plugins/Runtime/GeometryProcessing/Source/GeometryAlgorithms/Private/XAtlasWrapper.cpp b/Engine/Plugins/Runtime/GeometryProcessing/Source/GeometryAlgorithms/Private/XAtlasWrapper.cpp
--- a/Engine/Plugins/Runtime/GeometryProcessing/Source/GeometryAlgorithms/Private/XAtlasWrapper.cpp
+++ b/Engine/Plugins/Runtime/GeometryProcessing/Source/GeometryAlgorithms/Private/XAtlasWrapper.cpp
@@ -88,8 +88,8 @@ namespace XAtlasWrapper
 		UVIndexBuffer.Empty();
 		VertexRemapArray.Empty();
 
-		float W = Atlas->width;
-		float H = Atlas->height;
+		const float W = Atlas->width;
+		const float H = Atlas->height;
 
 		for (uint32 OutMeshID = 0; OutMeshID < Atlas->meshCount; ++OutMeshID)
 		{
@@ -98,7 +98,7 @@ namespace XAtlasWrapper
 			for (uint32 VertexID = 0; VertexID < CurrentMesh.vertexCount; ++VertexID)
 			{
 				xatlas::Vertex UVVert = CurrentMesh.vertexArray[VertexID];
-				UVVertexBuffer.Add(FVector2D{ UVVert.uv[0] / H, UVVert.uv[1] / W});
+				UVVertexBuffer.Add(FVector2D{ UVVert.uv[0] / W, UVVert.uv[1] / H});
 
 				uint32 OriginalVertexIndex = UVVert.xref;
 				VertexRemapArray.Add(OriginalVertexIndex);
