ModelingTools: Fix lost poly group data during TriSel duplicate/separate.

#rnx

diff --git a/Engine/Plugins/Runtime/GeometryProcessing/Source/DynamicMesh/Private/DynamicMeshEditor.cpp b/Engine/Plugins/Runtime/GeometryProcessing/Source/DynamicMesh/Private/DynamicMeshEditor.cpp
--- a/Engine/Plugins/Runtime/GeometryProcessing/Source/DynamicMesh/Private/DynamicMeshEditor.cpp
+++ b/Engine/Plugins/Runtime/GeometryProcessing/Source/DynamicMesh/Private/DynamicMeshEditor.cpp
@@ -1929,6 +1929,11 @@ void FDynamicMeshEditor::AppendTriangles(const FDynamicMesh3* SourceMesh, const
 	ResultOut.Reset();
 	IndexMaps.Initialize(Mesh);
 
+	if (SourceMesh->HasTriangleGroups())
+	{
+		Mesh->EnableTriangleGroups();
+	}
+
 	for (int SourceTriangleID : SourceTriangles)
 	{
 		check(SourceMesh->IsTriangle(SourceTriangleID));
