ModelingTools: Initialize triangle groups at the call site and disable compute of group mappings unless requested.

#rnx

diff --git a/Engine/Plugins/Experimental/MeshModelingToolsetExp/Source/MeshModelingToolsExp/Private/MeshSelectionTool.cpp b/Engine/Plugins/Experimental/MeshModelingToolsetExp/Source/MeshModelingToolsExp/Private/MeshSelectionTool.cpp
--- a/Engine/Plugins/Experimental/MeshModelingToolsetExp/Source/MeshModelingToolsExp/Private/MeshSelectionTool.cpp
+++ b/Engine/Plugins/Experimental/MeshModelingToolsetExp/Source/MeshModelingToolsExp/Private/MeshSelectionTool.cpp
@@ -1471,6 +1471,7 @@ void UMeshSelectionTool::SeparateSelectedTriangles(bool bDeleteSelected)
 
 	// extract copy of triangles
 	FDynamicMesh3 SeparatedMesh;
+	SeparatedMesh.EnableTriangleGroups();
 	SeparatedMesh.EnableAttributes();
 	SeparatedMesh.Attributes()->EnableMatchingAttributes(*SourceMesh->Attributes());
 	FDynamicMeshEditor Editor(&SeparatedMesh);
diff --git a/Engine/Plugins/Runtime/GeometryProcessing/Source/DynamicMesh/Private/DynamicMeshEditor.cpp b/Engine/Plugins/Runtime/GeometryProcessing/Source/DynamicMesh/Private/DynamicMeshEditor.cpp
--- a/Engine/Plugins/Runtime/GeometryProcessing/Source/DynamicMesh/Private/DynamicMeshEditor.cpp
+++ b/Engine/Plugins/Runtime/GeometryProcessing/Source/DynamicMesh/Private/DynamicMeshEditor.cpp
@@ -1929,11 +1929,7 @@ void FDynamicMeshEditor::AppendTriangles(const FDynamicMesh3* SourceMesh, const
 	ResultOut.Reset();
 	IndexMaps.Initialize(Mesh);
 
-	if (SourceMesh->HasTriangleGroups())
-	{
-		Mesh->EnableTriangleGroups();
-	}
-
+	int DefaultGroupID = FDynamicMesh3::InvalidID;
 	for (int SourceTriangleID : SourceTriangles)
 	{
 		check(SourceMesh->IsTriangle(SourceTriangleID));
@@ -1946,18 +1942,32 @@ void FDynamicMeshEditor::AppendTriangles(const FDynamicMesh3* SourceMesh, const
 
 		// FindOrCreateDuplicateGroup
 		int NewGroupID = FDynamicMesh3::InvalidID;
-		if (SourceMesh->HasTriangleGroups())
+		if (Mesh->HasTriangleGroups())
 		{
-			int SourceGroupID = SourceMesh->GetTriangleGroup(SourceTriangleID);
-			if (SourceGroupID >= 0)
+			if (SourceMesh->HasTriangleGroups())
+			{
+				int SourceGroupID = SourceMesh->GetTriangleGroup(SourceTriangleID);
+				if (SourceGroupID >= 0)
+				{
+					NewGroupID = IndexMaps.GetNewGroup(SourceGroupID);
+					if (NewGroupID == IndexMaps.InvalidID())
+					{
+						NewGroupID = Mesh->AllocateTriangleGroup();
+						IndexMaps.SetGroup(SourceGroupID, NewGroupID);
+						ResultOut.NewGroups.Add(NewGroupID);
+					}
+				}
+			}
+			else
 			{
-				NewGroupID = IndexMaps.GetNewGroup(SourceGroupID);
-				if (NewGroupID == IndexMaps.InvalidID())
+				// If the source mesh does not have triangle groups, but the destination
+				// mesh does, create a default group for all triangles.
+				if (DefaultGroupID == FDynamicMesh3::InvalidID)
 				{
-					NewGroupID = Mesh->AllocateTriangleGroup();
-					IndexMaps.SetGroup(SourceGroupID, NewGroupID);
-					ResultOut.NewGroups.Add(NewGroupID);
+					DefaultGroupID = Mesh->AllocateTriangleGroup();
+					ResultOut.NewGroups.Add(DefaultGroupID);
 				}
+				NewGroupID = DefaultGroupID;
 			}
 		}
 
