MeshModelingTools: Make group edge inserter not break in various ways with high vertex tolerance or very thin sliver triangles. Make edge loop insertion not hit ensure when unable to insert edge with plane cut.

#rnx

diff --git a/Engine/Plugins/Runtime/GeometryProcessing/Source/DynamicMesh/Private/Operations/GroupEdgeInserter.cpp b/Engine/Plugins/Runtime/GeometryProcessing/Source/DynamicMesh/Private/Operations/GroupEdgeInserter.cpp
--- a/Engine/Plugins/Runtime/GeometryProcessing/Source/DynamicMesh/Private/Operations/GroupEdgeInserter.cpp
+++ b/Engine/Plugins/Runtime/GeometryProcessing/Source/DynamicMesh/Private/Operations/GroupEdgeInserter.cpp
@@ -72,7 +72,7 @@ bool ConnectMultipleUsingPlaneCut(FDynamicMesh3& Mesh,
 	const TArray <FGroupEdgeInserter::FGroupEdgeSplitPoint>& EndPoints,
 	double VertexTolerance, int32& NumGroupsCreated, 
 	FGroupEdgeInserter::FOptionalOutputParams& OptionalOut, FProgressCancel* Progress);
-bool EmbedPlaneCutPath(FDynamicMesh3& Mesh, int32 GroupID,
+bool EmbedPlaneCutPath(FDynamicMesh3& Mesh, const FGroupTopology& Topology, int32 GroupID,
 	const FGroupEdgeInserter::FGroupEdgeSplitPoint& StartPoint,
 	const FGroupEdgeInserter::FGroupEdgeSplitPoint& EndPoint,
 	double VertexTolerance, TSet<int32>& PathEidsOut, 
@@ -81,7 +81,8 @@ bool CreateNewGroups(FDynamicMesh3& Mesh, TSet<int32>& PathEids, int32 OriginalG
 	FGroupEdgeInserter::FOptionalOutputParams& OptionalOut, FProgressCancel* Progress);
 bool GetPlaneCutPath(const FDynamicMesh3& Mesh, int32 GroupID,
 	const FGroupEdgeInserter::FGroupEdgeSplitPoint& StartPoint, const FGroupEdgeInserter::FGroupEdgeSplitPoint& EndPoint,
-	TArray<TPair<FMeshSurfacePoint, int>>& OutputPath, double VertexCutTolerance, FProgressCancel* Progress);
+	TArray<TPair<FMeshSurfacePoint, int>>& OutputPath, double VertexCutTolerance, 
+	const TSet<int32>& DisallowedVids, FProgressCancel* Progress);
 
 bool InsertSingleWithRetriangulation(FDynamicMesh3& Mesh, FGroupTopology& Topology,
 	int32 GroupID, int32 BoundaryIndex,
@@ -154,8 +155,14 @@ bool FGroupEdgeInserter::InsertEdgeLoops(const FEdgeLoopInsertionParams& Params,
 		bSuccess = InsertEdgeLoopEdgesInDirection(Params, StartEndpoints, ForwardGroupID, ForwardEdgeID,
 			ForwardCornerID, ForwardBoundaryIndex, AlteredGroups, TotalNumInserted, OptionalOut, Progress);
 	}
-	if (bHaveBackwardEdge)
+	if (bSuccess && bHaveBackwardEdge)
 	{
+		// TODO: The behavior here is not ideal in that if we fail to insert edges across a group we fail
+		// the entire insertion. For instance, in a curved staircase, where the bottom is a C-shape but otherwise
+		// a quad, we may fail to perform a plane cut, but we should still insert edges on the steps, where it
+		// doesn't fail. Unfortunately, being tolerant of insertion failures requires us to be able to revert
+		// the failed group to its original state on failure. We need to add support for that.
+
 		int32 NumInserted = 0;
 		bSuccess = InsertEdgeLoopEdgesInDirection(Params, StartEndpoints, BackwardGroupID, BackwardEdgeID,
 			BackwardCornerID, BackwardBoundaryIndex, AlteredGroups, NumInserted, OptionalOut, Progress) && bSuccess;
@@ -412,7 +419,7 @@ bool InsertNewVertexEndpoints(
 		}
 
 		// Advance until the next vertex would overshoot the target length.
-		while (NextIndex < PerVertexLengths.Num() && PerVertexLengths[NextIndex] <= TargetLength + Params.VertexTolerance)
+		while (NextIndex < PerVertexLengths.Num() && PerVertexLengths[NextIndex] <= TargetLength)
 		{
 			CurrentVid = SpanVids[NextIndex];
 			CurrentArcLength = PerVertexLengths[NextIndex];
@@ -422,22 +429,22 @@ bool InsertNewVertexEndpoints(
 		// The point is now either at the current vertex, or on the edge going forward (if there is one)
 		FGroupEdgeInserter::FGroupEdgeSplitPoint SplitPoint;
 
-		// See if the target is at the current vertex
-		if (abs(TargetLength - CurrentArcLength) <= Params.VertexTolerance)
+		// Used if we find ourselves needing to snap to one of the endpoints of the edge
+		auto SetSplitPointToVertex = [&SplitPoint, &SpanVids, &Params, NextIndex](int32 Vid)
 		{
-			SplitPoint.ElementID = CurrentVid;
+			SplitPoint.ElementID = Vid;
 			SplitPoint.bIsVertex = true;
 			
 			// Get the tangent vector. We haven't been keeping track of any previous verts we inserted,
 			// but inserted verts must be on an edge and must have the forward edge as their tangent.
-			bool bVertexIsOriginal = (CurrentVid == SpanVids[NextIndex - 1]);
+			bool bVertexIsOriginal = (Vid == SpanVids[NextIndex - 1]);
 			if (!bVertexIsOriginal)
 			{
-				SplitPoint.Tangent = Normalized(Params.Mesh->GetVertex(SpanVids[NextIndex]) - Params.Mesh->GetVertex(CurrentVid));
+				SplitPoint.Tangent = Normalized(Params.Mesh->GetVertex(SpanVids[NextIndex]) - Params.Mesh->GetVertex(Vid));
 			}
 			else
 			{
-				FVector3d VertexPosition = Params.Mesh->GetVertex(CurrentVid);
+				FVector3d VertexPosition = Params.Mesh->GetVertex(Vid);
 				SplitPoint.Tangent = FVector3d::Zero();
 				if (NextIndex > 1)
 				{
@@ -449,6 +456,17 @@ bool InsertNewVertexEndpoints(
 				}
 				Normalize(SplitPoint.Tangent);
 			}
+		};
+
+		// See if the target is at the current vertex
+		if (TargetLength - CurrentArcLength <= Params.VertexTolerance)
+		{
+			SetSplitPointToVertex(CurrentVid);
+		}
+		else if (NextIndex < PerVertexLengths.Num()
+			&& PerVertexLengths[NextIndex] - CurrentArcLength <= Params.VertexTolerance)
+		{
+			SetSplitPointToVertex(SpanVids[NextIndex]);
 		}
 		else
 		{
@@ -917,7 +935,7 @@ bool ConnectMultipleUsingPlaneCut(FDynamicMesh3& Mesh,
 	TSet<int32> PathsEids;
 	for (int32 i = 0; i < NumEdgesToInsert; ++i)
 	{
-		bool bSuccess = EmbedPlaneCutPath(Mesh, GroupID, StartPoints[i], EndPoints[i], 
+		bool bSuccess = EmbedPlaneCutPath(Mesh, Topology, GroupID, StartPoints[i], EndPoints[i],
 			VertexTolerance, PathsEids, OptionalOut.ChangedTidsOut, Progress);
 
 		if (!bSuccess || (Progress && Progress->Cancelled()))
@@ -939,8 +957,8 @@ bool ConnectMultipleUsingPlaneCut(FDynamicMesh3& Mesh,
  * Places a plane path connecting the endpoints into the mesh, but does not give the triangles new groups yet.
  * However, outputs the path edge ID's so that can be done later.
  */
-bool EmbedPlaneCutPath(FDynamicMesh3& Mesh, int32 GroupID,
-	const FGroupEdgeInserter::FGroupEdgeSplitPoint& StartPoint,
+bool EmbedPlaneCutPath(FDynamicMesh3& Mesh, const FGroupTopology& Topology,
+	int32 GroupID, const FGroupEdgeInserter::FGroupEdgeSplitPoint& StartPoint,
 	const FGroupEdgeInserter::FGroupEdgeSplitPoint& EndPoint,
 	double VertexTolerance, TSet<int32>& PathEidsOut, 
 	TSet<int32>* ChangedTrisOut, FProgressCancel* Progress)
@@ -950,10 +968,34 @@ bool EmbedPlaneCutPath(FDynamicMesh3& Mesh, int32 GroupID,
 		return false;
 	}
 
+	// We don't allow snapping to any of the boundary vertices via plane distance because this can lead
+	// to a situation where we snap via plane cut but not via absolute distance (depending on plane
+	// orientation relative boundary), which results in us arriving at the boundary at a nearby vertex
+	// and then walking along the boundary to the destination. Aside from looking/behaving weirdly, 
+	// this is very bad for edge loops, where doing so can join the edge at a different point from the
+	// one at which it continues on the other side, which means that quad-like topology gets
+	// inadvertantly broken without the user realizing why.
+	TSet<int32> DisallowedVids;
+	const FGroupTopology::FGroup* Group = Topology.FindGroupByID(GroupID);
+	if (ensure(Group))
+	{
+		for (const FGroupTopology::FGroupBoundary& Boundary : Group->Boundaries)
+		{
+			for (int32 GroupEdgeID : Boundary.GroupEdges)
+			{
+				if (ensure(GroupEdgeID < Topology.Edges.Num()))
+				{
+					DisallowedVids.Append(Topology.Edges[GroupEdgeID].Span.Vertices);
+				}
+			}
+		}
+	}
+	
+
 	// Find the path we're going to take.
 	TArray<TPair<FMeshSurfacePoint, int>> CutPath;
 	bool bSuccess = GetPlaneCutPath(Mesh, GroupID, StartPoint, EndPoint,
-		CutPath, VertexTolerance, Progress);
+		CutPath, VertexTolerance, DisallowedVids, Progress);
 	if (!bSuccess || (Progress && Progress->Cancelled()))
 	{
 		return false;
@@ -991,7 +1033,11 @@ bool EmbedPlaneCutPath(FDynamicMesh3& Mesh, int32 GroupID,
 
 	for (int32 i = 1; i < PathVertices.Num(); ++i)
 	{
-		PathEidsOut.Add(Mesh.FindEdge(PathVertices[i - 1], PathVertices[i]));
+		int32 Eid = Mesh.FindEdge(PathVertices[i - 1], PathVertices[i]);
+		if (ensure(Eid >= 0))
+		{
+			PathEidsOut.Add(Eid);
+		}
 	}
 
 	return true;
@@ -1097,7 +1143,7 @@ bool FGroupEdgeInserter::InsertGroupEdge(FGroupEdgeInsertionParams& Params, FGro
 		TSet<int32> TempNewEids;
 		TSet<int32>* NewEids = OptionalOut.NewEidsOut ? OptionalOut.NewEidsOut : &TempNewEids;
 
-		bool bSuccess = EmbedPlaneCutPath(*Params.Mesh, Params.GroupID, Params.StartPoint,
+		bool bSuccess = EmbedPlaneCutPath(*Params.Mesh, *Params.Topology, Params.GroupID, Params.StartPoint,
 			Params.EndPoint, Params.VertexTolerance, *NewEids, OptionalOut.ChangedTidsOut, Progress);
 		if (!bSuccess || (Progress && Progress->Cancelled()))
 		{
@@ -1242,13 +1288,21 @@ bool InsertSingleWithRetriangulation(FDynamicMesh3& Mesh, FGroupTopology& Topolo
  */
 bool GetPlaneCutPath(const FDynamicMesh3& Mesh, int32 GroupID,
 	const FGroupEdgeInserter::FGroupEdgeSplitPoint& StartPoint, const FGroupEdgeInserter::FGroupEdgeSplitPoint& EndPoint,
-	TArray<TPair<FMeshSurfacePoint, int>>& OutputPath, double VertexCutTolerance, FProgressCancel* Progress)
+	TArray<TPair<FMeshSurfacePoint, int>>& OutputPath, double VertexCutTolerance, 
+	const TSet<int32>& DisallowedVids, FProgressCancel* Progress)
 {
 	if (Progress && Progress->Cancelled())
 	{
 		return false;
 	}
 
+	// Used to make sure we don't end up walking in a loop or backwards. This could happen with a high vertex cut
+	// tolerance and pathological topology via vertices. It shouldn't be possible via edges for the cases where
+	// we use this function (between points of a contiguous group boundary), but we'll be safe.
+	TSet<int32> CrossedVids;
+	TSet<int32> CrossedEids;
+
+
 	// Start by determining the plane we will use. 
 	FVector3d StartPosition = StartPoint.bIsVertex ? Mesh.GetVertex(StartPoint.ElementID)
 		: Mesh.GetEdgePoint(StartPoint.ElementID, StartPoint.EdgeTValue);
@@ -1258,8 +1312,19 @@ bool GetPlaneCutPath(const FDynamicMesh3& Mesh, int32 GroupID,
 	FVector3d InPlaneVector = Normalized(EndPosition - StartPosition);
 
 	// Get components of the two tangents that are orthogonal to the vector between the points.
-	FVector3d NormalA = Normalized(StartPoint.Tangent - StartPoint.Tangent.Dot(InPlaneVector) * InPlaneVector);
-	FVector3d NormalB = Normalized(StartPoint.Tangent - StartPoint.Tangent.Dot(InPlaneVector) * InPlaneVector);
+	FVector3d NormalA = Normalized(StartPoint.Tangent - StartPoint.Tangent.Dot(InPlaneVector) * InPlaneVector, 
+		static_cast<double>(KINDA_SMALL_NUMBER));
+	FVector3d NormalB = Normalized(EndPoint.Tangent - EndPoint.Tangent.Dot(InPlaneVector) * InPlaneVector,
+		static_cast<double>(KINDA_SMALL_NUMBER));
+
+	if (NormalA.IsZero() || NormalB.IsZero())
+	{
+		// One or both of the tangents were pointing directly toward the destination, and in such
+		// cases the cutting plane we fit is likely to be nonsense (along the edge of a cube,
+		// for instance, we may end up trying to use a plane roughly coplanar with the face we're
+		// trying to "cut").
+		return false;
+	}
 
 	// Make the vectors be in the same half space so that their average represents the closer average of the
 	// corresponding lines.
@@ -1272,10 +1337,9 @@ bool GetPlaneCutPath(const FDynamicMesh3& Mesh, int32 GroupID,
 	FVector CutPlaneNormal = (FVector)Normalized(NormalA + NormalB);
 	if (CutPlaneNormal.IsZero())
 	{
-		// This likely shouldn't happen, since it means that the two endpoint tangents are colinear
-		// with the vector between them. Let's say that the normal doesn't matter then as long as it's
-		// orthogonal.
-		CutPlaneNormal = FVector(InPlaneVector.Y, -InPlaneVector.X, InPlaneVector.Z);
+		// This shouldn't happen because we already checked for both of them being zero, and we made them be in
+		// the same half space.
+		return ensure(false);
 	}
 
 	FVector CutPlaneOrigin = (FVector)StartPosition;
@@ -1284,47 +1348,30 @@ bool GetPlaneCutPath(const FDynamicMesh3& Mesh, int32 GroupID,
 	// finding the next point.
 	float CurrentEdgeVertPlaneDistances[2];
 
-	// Prep the first point
+	// Prep the first point.
+	OutputPath.Empty();
 	if (StartPoint.bIsVertex)
 	{
 		OutputPath.Emplace(FMeshSurfacePoint(StartPoint.ElementID), FDynamicMesh3::InvalidID);
 	}
 	else
 	{
+		// Note that we do not clamp the endpoints in this function because clamping based
+		// on plane distance can result in different behavior depending on which way the plane is oriented, which
+		// could end up clamping to different endpoints as we connect multiple paths through the same start/end
+		// point (such as when following a loop)
 		FIndex2i EdgeVids = Mesh.GetEdgeV(StartPoint.ElementID);
 		CurrentEdgeVertPlaneDistances[0] = FVector::PointPlaneDist((FVector)Mesh.GetVertex(EdgeVids.A), CutPlaneOrigin, CutPlaneNormal);
 		CurrentEdgeVertPlaneDistances[1] = FVector::PointPlaneDist((FVector)Mesh.GetVertex(EdgeVids.B), CutPlaneOrigin, CutPlaneNormal);
 
-		if (abs(CurrentEdgeVertPlaneDistances[0]) <= VertexCutTolerance)
-		{
-			if (abs(CurrentEdgeVertPlaneDistances[1]) <= VertexCutTolerance)
-			{
-				// This will happen if the first point is on an edge colinear with the direction to the endpoint,
-				// and it's not worth dealing with, since the path is quite unlikely to be reasonable anyway.
-				return false;
-			}
-			OutputPath.Emplace(FMeshSurfacePoint(EdgeVids.A), FDynamicMesh3::InvalidID);
-		}
-		else if (abs(CurrentEdgeVertPlaneDistances[1]) <= VertexCutTolerance)
-		{
-			OutputPath.Emplace(FMeshSurfacePoint(EdgeVids.B), FDynamicMesh3::InvalidID);
-		}
-		else
-		{
 			OutputPath.Emplace(FMeshSurfacePoint(StartPoint.ElementID, StartPoint.EdgeTValue), FDynamicMesh3::InvalidID);
 		}
-	}
 	check(OutputPath.Num() == 1);
 
 	// Set up a few more variables we'll need as we walk from start to end
 	bool bCurrentPointIsVertex = (OutputPath[0].Key.PointType == ESurfacePointType::Vertex);
 	int32 CurrentElementID = OutputPath[0].Key.ElementID;
 
-	// These help us avoid backtracking.
-	int32 PreviousTid = FDynamicMesh3::InvalidID; // if we walked across a triangle
-	int32 PreviousVid = FDynamicMesh3::InvalidID; // if we walked from a vertex
-
-
 	int32 PointCount = 1; // Used as a sanity check
 
 	// Tracks which triangle (if not an edge) we traversed to get to the next point, so we can avoid
@@ -1340,88 +1387,110 @@ bool GetPlaneCutPath(const FDynamicMesh3& Mesh, int32 GroupID,
 			return false;
 		}
 
+		// Prevent ourselves from going in a loop
+		if (bCurrentPointIsVertex)
+		{
+			if (CrossedVids.Contains(CurrentElementID))
+			{
+				return false;
+			}
+			CrossedVids.Add(CurrentElementID);
+		}
+		else
+		{
+			if (CrossedEids.Contains(CurrentElementID))
+			{
+				return false;
+			}
+			CrossedEids.Add(CurrentElementID);
+		}
+
 		check(PointCount < Mesh.EdgeCount()); // sanity check to avoid infinite loop
 
 		if (bCurrentPointIsVertex)
 		{
 			FMeshSurfacePoint NextPoint(FDynamicMesh3::InvalidID);
+			FVector3d CurrentPosition = OutputPath.Last().Key.Pos(&Mesh);
 
 			// Look through the surrounding triangles that have the group we want and find one that intersects the plane
+			int32 CandidateTraversedTid = FDynamicMesh3::InvalidID;
 			for (int32 Tid : Mesh.VtxTrianglesItr(CurrentElementID))
 			{
-				if (Tid == PreviousTid || Mesh.GetTriangleGroup(Tid) != GroupID)
+				if (Tid == TraversedTid || Mesh.GetTriangleGroup(Tid) != GroupID)
 				{
 					continue;
 				}
 
-				const FIndex3i& TriangleVids = Mesh.GetTriangle(Tid);
-				int32 VertA = (TriangleVids.A == CurrentElementID) ? TriangleVids.C : TriangleVids.A;
-				int32 VertB = (TriangleVids.B == CurrentElementID) ? TriangleVids.C : TriangleVids.B;
-				if (VertA == PreviousVid || VertB == PreviousVid)
+				// See if one of the triangle edges has the endpoint, in which case we can go straight there.
+				if (!EndPoint.bIsVertex)
 				{
-					// Our path can't go through this triangle, since we already went down one of its sides and
-					// since we bail when we see triangles that are coplanar with the cut plane.
-					continue;
+					const FIndex3i& TriangleEids = Mesh.GetTriEdges(Tid);
+					for (int32 i = 0; i < 3; ++i)
+					{
+						if (EndPoint.ElementID == TriangleEids[i])
+						{
+							// Got to the end
+							OutputPath.Emplace(FMeshSurfacePoint(EndPoint.ElementID, EndPoint.EdgeTValue), FDynamicMesh3::InvalidID);
+							return true;
+						}
+					}
 				}
 
 				// Check to see if one of the triangle vertices is the end
+				const FIndex3i& TriangleVids = Mesh.GetTriangle(Tid);
+				int32 VertA = (TriangleVids.A == CurrentElementID) ? TriangleVids.C : TriangleVids.A;
+				int32 VertB = (TriangleVids.B == CurrentElementID) ? TriangleVids.C : TriangleVids.B;
 				if (EndPoint.bIsVertex && (EndPoint.ElementID == VertA || EndPoint.ElementID == VertB))
 				{
-					// TODO: This ought to be cleaned up because there is inconsistency in how we deal ambiguity
-					// at the end of our path (we may accept or not accept ambiguities).
-					NextPoint = FMeshSurfacePoint(EndPoint.ElementID);
-					break;
+					OutputPath.Emplace(FMeshSurfacePoint(EndPoint.ElementID), FDynamicMesh3::InvalidID);
+					return true;
 				}
 
+				// See if one of the other vertices is on the plane (and is therefore the next destination)
 				float PlaneDistanceA = FVector::PointPlaneDist((FVector)Mesh.GetVertex(VertA), CutPlaneOrigin, CutPlaneNormal);
 				float PlaneDistanceB = FVector::PointPlaneDist((FVector)Mesh.GetVertex(VertB), CutPlaneOrigin, CutPlaneNormal);
 				bool bVertAIsOnPlane = abs(PlaneDistanceA) <= VertexCutTolerance;
 				bool bVertBIsOnPlane = abs(PlaneDistanceB) <= VertexCutTolerance;
 
-				if (bVertAIsOnPlane && bVertBIsOnPlane)
+				auto UpdateNextPoint = [&InPlaneVector, &Mesh, &CurrentPosition, &NextPoint](const FMeshSurfacePoint& CandidateSurfacePoint)
 				{
-					return false; // Triangle coplanar with the cut plane. We're not going to try to handle that.
-				}
+					if (NextPoint.ElementID == CandidateSurfacePoint.ElementID && NextPoint.PointType == CandidateSurfacePoint.PointType)
+					{
+						// We're looking at the same point from an adjacent triangle
+						return false;
+					}
 
-				if (bVertAIsOnPlane || bVertBIsOnPlane)
-				{
-					int32 CandidateVert = bVertAIsOnPlane ? VertA : VertB;
-					if (NextPoint.ElementID != FDynamicMesh3::InvalidID)
+					// Update the next point if the candidate moves more directly toward the destination. 
+					// TOOD: This hack is necessary to deal with some common ambiguous cases (such as starting from a 
+					// vertex in the concave region of a nonconvex planar surface), but the proper solution is to alter 
+					// and use EmbedSurfacePath.cpp::WalkMeshPlanar
+					if (NextPoint.ElementID == FDynamicMesh3::InvalidID
+						|| (InPlaneVector.Dot(NextPoint.Pos(&Mesh) - CurrentPosition) < 
+							InPlaneVector.Dot(CandidateSurfacePoint.Pos(&Mesh) - CurrentPosition)))
 					{
-						// We already found a destination point. See if we're looking at the same one (from an adjacent triangle).
-						if (NextPoint.PointType == ESurfacePointType::Vertex && NextPoint.ElementID == CandidateVert)
-						{
-							continue; // Ok, same point
-						}
-						else
-						{
-							// Another point, so ambiguous. Accept it only if the previous moved away from the destination
-							// and the current is moving toward. This is a hack to deal with nonconvex planar surfaces nicely,
-							// but the proper solution is to alter EmbedSurfacePath.cpp::WalkMeshPlanar instead.
-							FVector3d CurrentPosition = OutputPath.Last().Key.Pos(&Mesh);
-							if (!(InPlaneVector.Dot(NextPoint.Pos(&Mesh) - CurrentPosition) < 0
-								&& InPlaneVector.Dot(Mesh.GetVertex(CandidateVert) - CurrentPosition) > 0))
-							{
-								continue;
-							}
-						}
+						NextPoint = CandidateSurfacePoint;
+						return true;
 					}
+					return false;
+				};
 
-					// Save point
-					NextPoint = FMeshSurfacePoint(CandidateVert);
-				}//end if one of the points is on the plane
-				else if (PlaneDistanceA * PlaneDistanceB < 0)
+				bool bEdgeVertIsPreferred = false;
+				if (bVertAIsOnPlane && !DisallowedVids.Contains(VertA))
+				{
+					bEdgeVertIsPreferred = true;
+					UpdateNextPoint(FMeshSurfacePoint(VertA));
+				}
+				if (bVertBIsOnPlane && !DisallowedVids.Contains(VertB))
 				{
-					// The triangle's opposite edge crosses the plane.
-					int32 Eid = Mesh.FindEdge(VertA, VertB);
+					bEdgeVertIsPreferred = true;
+					UpdateNextPoint(FMeshSurfacePoint(VertB));
+				}
+				if (!bEdgeVertIsPreferred && PlaneDistanceA * PlaneDistanceB < 0)
+				{
+					// The triangle's opposite edge crosses the plane, and the edge verts are not
+					// valid snap targets
+					int32 Eid = Mesh.FindEdgeFromTri(VertA, VertB, Tid);
 
-					if (!EndPoint.bIsVertex && EndPoint.ElementID == Eid)
-					{
-						// Got to the end
-						NextPoint = FMeshSurfacePoint(EndPoint.ElementID, EndPoint.EdgeTValue);
-						break;
-					}
-			
 					double EdgeTValue = PlaneDistanceA / (PlaneDistanceA - PlaneDistanceB);
 
 					if (VertA != Mesh.GetEdgeV(Eid).A)
@@ -1429,29 +1498,16 @@ bool GetPlaneCutPath(const FDynamicMesh3& Mesh, int32 GroupID,
 						EdgeTValue = 1 - EdgeTValue;
 					}
 
-					if (NextPoint.ElementID != FDynamicMesh3::InvalidID)
+					if (UpdateNextPoint(FMeshSurfacePoint(Eid, EdgeTValue)))
 					{
-						// Looks like there are multiple intersections with the plane, so ambigous.
-						// Accept new point only if the previous moved away from the destination
-						// and the current is moving toward. This is a hack to deal with nonconvex planar surfaces nicely,
-						// but the proper solution is to alter EmbedSurfacePath.cpp::WalkMeshPlanar instead.
-						FVector3d CurrentPosition = OutputPath.Last().Key.Pos(&Mesh);
-						if (!(InPlaneVector.Dot(NextPoint.Pos(&Mesh) - CurrentPosition) < 0
-							&& InPlaneVector.Dot(Mesh.GetEdgePoint(Eid, EdgeTValue) - CurrentPosition) > 0))
+						CurrentEdgeVertPlaneDistances[0] = PlaneDistanceA;
+						CurrentEdgeVertPlaneDistances[1] = PlaneDistanceB;
+						if (VertA != Mesh.GetEdgeV(Eid).A)
 						{
-							continue;
+							Swap(CurrentEdgeVertPlaneDistances[0], CurrentEdgeVertPlaneDistances[1]);
 						}
+						CandidateTraversedTid = Tid;
 					}
-
-					// Save the edge point
-					CurrentEdgeVertPlaneDistances[0] = PlaneDistanceA;
-					CurrentEdgeVertPlaneDistances[1] = PlaneDistanceB;
-					if (VertA != Mesh.GetEdgeV(Eid).A)
-					{
-						Swap(CurrentEdgeVertPlaneDistances[0], CurrentEdgeVertPlaneDistances[1]);
-					}
-					NextPoint = FMeshSurfacePoint(Eid, EdgeTValue);
-					TraversedTid = Tid;
 				}
 			}//end going through triangles
 
@@ -1463,6 +1519,9 @@ bool GetPlaneCutPath(const FDynamicMesh3& Mesh, int32 GroupID,
 			else
 			{
 				OutputPath.Emplace(NextPoint, FDynamicMesh3::InvalidID);
+
+				TraversedTid = (NextPoint.PointType == ESurfacePointType::Edge) ? CandidateTraversedTid
+					: FDynamicMesh3::InvalidID;
 			}
 		}//end if at vert
 		else
@@ -1471,11 +1530,11 @@ bool GetPlaneCutPath(const FDynamicMesh3& Mesh, int32 GroupID,
 
 			// We're starting from an edge. Get the triangle that we're dealing with.
 			int32 NextTid;
-			if (Edge.Tri.A == PreviousTid)
+			if (Edge.Tri.A == TraversedTid)
 			{
 				NextTid = Edge.Tri.B;
 			}
-			else if (Edge.Tri.B == PreviousTid)
+			else if (Edge.Tri.B == TraversedTid)
 			{
 				NextTid = Edge.Tri.A;
 			}
@@ -1486,21 +1545,37 @@ bool GetPlaneCutPath(const FDynamicMesh3& Mesh, int32 GroupID,
 
 			if (NextTid == FDynamicMesh3::InvalidID || Mesh.GetTriangleGroup(NextTid) != GroupID)
 			{
-				// We dead ended before getting to the end.
+				// We hit a dead end before getting to the end.
 				return false;
 			}
 			TraversedTid = NextTid;
 
-			// Get the vertex opposite the current edge and it's location relative to the cut plane
+			// See if the opposite vert is our endpoint.
 			int32 OppositeVert = IndexUtil::FindTriOtherVtx(Edge.Vert.A, Edge.Vert.B, Mesh.GetTriangle(NextTid));
-			float OppositeVertPlaneDistance = FVector::PointPlaneDist((FVector)Mesh.GetVertex(OppositeVert), CutPlaneOrigin, CutPlaneNormal);
-
 			if (EndPoint.bIsVertex && EndPoint.ElementID == OppositeVert)
 			{
-				// Got to the end
 				OutputPath.Emplace(FMeshSurfacePoint(EndPoint.ElementID), FDynamicMesh3::InvalidID);
-			} 
-			else if (abs(OppositeVertPlaneDistance) <= VertexCutTolerance)
+				return true;
+			}
+
+			// See if one of the triangle edges has the endpoint, in which case we can go straight there.
+			if (!EndPoint.bIsVertex)
+			{
+				const FIndex3i& TriangleEids = Mesh.GetTriEdges(NextTid);
+				for (int32 i = 0; i < 3; ++i)
+				{
+					if (EndPoint.ElementID == TriangleEids[i])
+					{
+						// Got to the end
+						OutputPath.Emplace(FMeshSurfacePoint(EndPoint.ElementID, EndPoint.EdgeTValue), FDynamicMesh3::InvalidID);
+						return true;
+					}
+				}
+			}
+
+			// We'll keep going. Get the placement of the opposite vert relative to the plane.
+			float OppositeVertPlaneDistance = FVector::PointPlaneDist((FVector)Mesh.GetVertex(OppositeVert), CutPlaneOrigin, CutPlaneNormal);
+			if (abs(OppositeVertPlaneDistance) <= VertexCutTolerance && !DisallowedVids.Contains(OppositeVert))
 			{
 				// We are cutting through a vertex
 				OutputPath.Emplace(FMeshSurfacePoint(OppositeVert), FDynamicMesh3::InvalidID);
@@ -1515,47 +1590,45 @@ bool GetPlaneCutPath(const FDynamicMesh3& Mesh, int32 GroupID,
 					SecondVertOfNextEdge = Edge.Vert.A;
 					SecondPlaneDistance = CurrentEdgeVertPlaneDistances[0];
 				}
-				else
+				else if (CurrentEdgeVertPlaneDistances[1] * OppositeVertPlaneDistance < 0)
 				{
-					// Impossible that all vertices are on the same side of plane, if we started from an edge that was 
-					// cut by the plane.
-					check(CurrentEdgeVertPlaneDistances[1] * OppositeVertPlaneDistance < 0);
 					SecondVertOfNextEdge = Edge.Vert.B;
 					SecondPlaneDistance = CurrentEdgeVertPlaneDistances[1];
 				}
+				else
+				{
+					// For neither of the edge vertices to be on the opposite side of the plane from
+					// the opposite vertex, the edge must lie in the plane. This can only happen if
+					// we started on an edge and then fit a bad cutting plane (otherwise we would have
+					// snapped to a vertex instead). This should usually not happen since we check the
+					// tangents, but we use a different type of check/tolerance there, and it could 
+					// still happen if endpoints weren't snapped for some reason and the edge is super short.
+					return false;
+				}
 
 				// Add the edge point to output
 				int32 Eid = Mesh.FindEdge(OppositeVert, SecondVertOfNextEdge);
 
-				if (!EndPoint.bIsVertex && EndPoint.ElementID == Eid)
-				{
-					// Got to the end
-					OutputPath.Emplace(FMeshSurfacePoint(EndPoint.ElementID, EndPoint.EdgeTValue), FDynamicMesh3::InvalidID);
-				}
-				else
-				{
-					double EdgeTValue = OppositeVertPlaneDistance / (OppositeVertPlaneDistance - SecondPlaneDistance);
-
-					CurrentEdgeVertPlaneDistances[0] = OppositeVertPlaneDistance;
-					CurrentEdgeVertPlaneDistances[1] = SecondPlaneDistance;
+				// The division here is safe because SecondPlaneDistance is guaranteed to have an opposite sign to
+				// OppositeVertPlaneDistance by the if-else statements above.
+				double EdgeTValue = OppositeVertPlaneDistance / (OppositeVertPlaneDistance - SecondPlaneDistance);
 
-					if (OppositeVert != Mesh.GetEdgeV(Eid).A)
-					{
-						EdgeTValue = 1 - EdgeTValue;
-						Swap(CurrentEdgeVertPlaneDistances[0], CurrentEdgeVertPlaneDistances[1]);
-					}
+				CurrentEdgeVertPlaneDistances[0] = OppositeVertPlaneDistance;
+				CurrentEdgeVertPlaneDistances[1] = SecondPlaneDistance;
 
-					OutputPath.Emplace(FMeshSurfacePoint(Eid, EdgeTValue), FDynamicMesh3::InvalidID);
+				if (OppositeVert != Mesh.GetEdgeV(Eid).A)
+				{
+					EdgeTValue = 1 - EdgeTValue;
+					Swap(CurrentEdgeVertPlaneDistances[0], CurrentEdgeVertPlaneDistances[1]);
 				}
+
+				OutputPath.Emplace(FMeshSurfacePoint(Eid, EdgeTValue), FDynamicMesh3::InvalidID);
 			}//end cutting through edge
 		}//end if last point was edge
 
 		++PointCount;
 		check(PointCount == OutputPath.Num()) // Another sanity check, to make sure we're always advancing
 
-		PreviousTid = TraversedTid;
-		PreviousVid = bCurrentPointIsVertex ? CurrentElementID : FDynamicMesh3::InvalidID;
-
 		CurrentElementID = OutputPath.Last().Key.ElementID;
 		bCurrentPointIsVertex = (OutputPath.Last().Key.PointType == ESurfacePointType::Vertex);
 	}//end until we get to end
diff --git a/Engine/Plugins/Runtime/MeshModelingToolset/Source/MeshModelingTools/Private/ToolActivities/PolyEditInsertEdgeActivity.cpp b/Engine/Plugins/Runtime/MeshModelingToolset/Source/MeshModelingTools/Private/ToolActivities/PolyEditInsertEdgeActivity.cpp
--- a/Engine/Plugins/Runtime/MeshModelingToolset/Source/MeshModelingTools/Private/ToolActivities/PolyEditInsertEdgeActivity.cpp
+++ b/Engine/Plugins/Runtime/MeshModelingToolset/Source/MeshModelingTools/Private/ToolActivities/PolyEditInsertEdgeActivity.cpp
@@ -336,19 +336,6 @@ void UPolyEditInsertEdgeActivity::ConditionallyUpdatePreview(
 	const FGroupEdgeInserter::FGroupEdgeSplitPoint& NewEndPoint, int32 NewEndTopologyID, bool bNewEndIsCorner,
 	int32 NewCommonGroupID, int32 NewBoundaryIndex)
 {
-	//if (InputGroupEdgeID != NewGroupID
-	//	|| (NewInteractiveInputLength && Settings->PositionMode != EEdgeLoopPositioningMode::Even
-	//		&& *NewInteractiveInputLength != InteractiveInputLength))
-	//{
-	//	InputGroupEdgeID = NewGroupID;
-	//	if (NewInteractiveInputLength)
-	//	{
-	//		InteractiveInputLength = *NewInteractiveInputLength;
-	//	}
-	//	PreviewEdges.Reset();
-	//	ActivityContext->Preview->InvalidateResult();
-	//}
-
 	if (bShowingBaseMesh 
 		|| bEndIsCorner != bNewEndIsCorner || EndTopologyID != NewEndTopologyID
 		|| EndPoint.bIsVertex != NewEndPoint.bIsVertex || EndPoint.ElementID != NewEndPoint.ElementID
@@ -697,13 +684,12 @@ bool UPolyEditInsertEdgeActivity::GetHoveredItem(const FRay& WorldRay,
 		FRay EdgeRay((FVector)StartVert, (FVector)PointOut.Tangent, true);
 		float DistDownEdge = EdgeRay.GetParameter((FVector)Position);
 
-		PositionOut = (FVector3d)EdgeRay.PointAt(DistDownEdge);
-
 		// See if the point is at a vertex in the group edge span.
 		if (DistDownEdge <= Settings->VertexTolerance)
 		{
 			PointOut.bIsVertex = true;
 			PointOut.ElementID = StartVid;
+			PositionOut = StartVert;
 			if (EdgeSegmentID > 0)
 			{
 				// Average with previous normalized edge vector
@@ -716,6 +702,7 @@ bool UPolyEditInsertEdgeActivity::GetHoveredItem(const FRay& WorldRay,
 		{
 			PointOut.bIsVertex = true;
 			PointOut.ElementID = EndVid;
+			PositionOut = EndVert;
 			if (EdgeSegmentID + 2 < GroupEdge.Span.Vertices.Num())
 			{
 				PointOut.Tangent += UE::Geometry::Normalized(
@@ -728,6 +715,7 @@ bool UPolyEditInsertEdgeActivity::GetHoveredItem(const FRay& WorldRay,
 			PointOut.bIsVertex = false;
 			PointOut.ElementID = Eid;
 			PointOut.EdgeTValue = DistDownEdge / EdgeLength;
+			PositionOut = (FVector3d)EdgeRay.PointAt(DistDownEdge);
 			if (ActivityContext->CurrentMesh->GetEdgeV(Eid).A != StartVid)
 			{
 				PointOut.EdgeTValue = 1 - PointOut.EdgeTValue;
