Fixed a crash bug in the UniformTesselation and PNTriangles where we did not handle accessing overlay elements of the unset triangles.
Added a unit test entry to test unset overlay triangle case.

#rnx

diff --git a/Engine/Plugins/Runtime/GeometryProcessing/Source/DynamicMesh/Private/Operations/PNTriangles.cpp b/Engine/Plugins/Runtime/GeometryProcessing/Source/DynamicMesh/Private/Operations/PNTriangles.cpp
--- a/Engine/Plugins/Runtime/GeometryProcessing/Source/DynamicMesh/Private/Operations/PNTriangles.cpp
+++ b/Engine/Plugins/Runtime/GeometryProcessing/Source/DynamicMesh/Private/Operations/PNTriangles.cpp
@@ -113,8 +113,16 @@ namespace FPNTrianglesLocals
 				if (Mesh.HasAttributes()) 
 				{
 					const FDynamicMeshNormalOverlay* NormalOverlay = Mesh.Attributes()->PrimaryNormals();
-					NormalOverlay->GetElementAtVertex(EdgeTri.A, EdgeV.A, Normal1);
-					NormalOverlay->GetElementAtVertex(EdgeTri.A, EdgeV.B, Normal2);
+					if (NormalOverlay->IsSetTriangle(EdgeTri.A)) 
+					{
+						NormalOverlay->GetElementAtVertex(EdgeTri.A, EdgeV.A, Normal1);
+						NormalOverlay->GetElementAtVertex(EdgeTri.A, EdgeV.B, Normal2);
+					}
+					else
+					{
+						Normal1 = static_cast<FVector3f>(Mesh.GetTriNormal(EdgeTri.A));
+						Normal2 = Normal1;
+					}
 				}
 				else 
 				{
@@ -132,13 +140,22 @@ namespace FPNTrianglesLocals
 				FVector3d Point1ForTriA = ComputeControlPoint(Vertex1, Vertex2, Normal1);
 				FVector3d Point2ForTriA = ComputeControlPoint(Vertex2, Vertex1, Normal2);
 
-				if (Mesh.HasAttributes()) 
+				if (Mesh.HasAttributes() && Mesh.Attributes()->PrimaryNormals()) 
 				{
 					const FDynamicMeshNormalOverlay* NormalOverlay = Mesh.Attributes()->PrimaryNormals();
 					if (EdgeTri.B != FDynamicMesh3::InvalidID && NormalOverlay->IsSeamEdge(EID)) 
 					{	
-						NormalOverlay->GetElementAtVertex(EdgeTri.B, EdgeV.A, Normal1);
-						NormalOverlay->GetElementAtVertex(EdgeTri.B, EdgeV.B, Normal2); 
+						if (NormalOverlay->IsSetTriangle(EdgeTri.B)) 
+						{
+							NormalOverlay->GetElementAtVertex(EdgeTri.B, EdgeV.A, Normal1);
+							NormalOverlay->GetElementAtVertex(EdgeTri.B, EdgeV.B, Normal2);
+						}
+						else 
+						{
+							Normal1 = static_cast<FVector3f>(Mesh.GetTriNormal(EdgeTri.B));
+							Normal2 = Normal1;
+						}
+
 						FVector3d Point1ForTriB = ComputeControlPoint(Vertex1, Vertex2, Normal1);
 						FVector3d Point2ForTriB = ComputeControlPoint(Vertex2, Vertex1, Normal2);
 
@@ -302,15 +319,15 @@ namespace FPNTrianglesLocals
 			FIndex3i TriVertex = OriginalMesh.GetTriangle(OriginalTriangleID);
 			FIndex3i TriEdges = OriginalMesh.GetTriEdges(OriginalTriangleID);
 
-			FVector3d Bary = VectorUtil::BarycentricCoords(Mesh.GetVertexRef(VertexID), Mesh.GetVertexRef(TriVertex.A), 
-														   Mesh.GetVertexRef(TriVertex.B), Mesh.GetVertexRef(TriVertex.C));
+			FVector3d Bary = VectorUtil::BarycentricCoords(Mesh.GetVertexRef(VertexID), OriginalMesh.GetVertexRef(TriVertex.A),
+														   OriginalMesh.GetVertexRef(TriVertex.B), OriginalMesh.GetVertexRef(TriVertex.C));
 														   
 			FVector3d BarySquared = Bary*Bary;
 
 			// Displaced vertex. First compute contribution of the original control points at the original vertices
-			FVector3d NewPos = Bary[0]*BarySquared[0]*Mesh.GetVertexRef(TriVertex.A) + 
-							   Bary[1]*BarySquared[1]*Mesh.GetVertexRef(TriVertex.B) + 
-							   Bary[2]*BarySquared[2]*Mesh.GetVertexRef(TriVertex.C);
+			FVector3d NewPos = Bary[0]*BarySquared[0]*OriginalMesh.GetVertexRef(TriVertex.A) + 
+							   Bary[1]*BarySquared[1]*OriginalMesh.GetVertexRef(TriVertex.B) + 
+							   Bary[2]*BarySquared[2]*OriginalMesh.GetVertexRef(TriVertex.C);
 			
 			// Compute contribution of the control points at the edges
 			for (int EIDX = 0; EIDX < 3; ++EIDX) 
@@ -341,9 +358,9 @@ namespace FPNTrianglesLocals
 				if (bHasVertexNormals) 
 				{
 					// Compute contribution of the normal control points at the original vertices
-					FVector3d NewNormal = BarySquared[0]*Mesh.GetVertexNormal(TriVertex.A) + 
-										  BarySquared[1]*Mesh.GetVertexNormal(TriVertex.B) +
-										  BarySquared[2]*Mesh.GetVertexNormal(TriVertex.C);
+					FVector3d NewNormal = BarySquared[0]*OriginalMesh.GetVertexNormal(TriVertex.A) + 
+										  BarySquared[1]*OriginalMesh.GetVertexNormal(TriVertex.B) +
+										  BarySquared[2]*OriginalMesh.GetVertexNormal(TriVertex.C);
 					
 					// Compute contribution of the normal control points at the edges
 					for (int EIDX = 0; EIDX < 3; ++EIDX) 
diff --git a/Engine/Plugins/Runtime/GeometryProcessing/Source/DynamicMesh/Private/Operations/UniformTesselate.cpp b/Engine/Plugins/Runtime/GeometryProcessing/Source/DynamicMesh/Private/Operations/UniformTesselate.cpp
--- a/Engine/Plugins/Runtime/GeometryProcessing/Source/DynamicMesh/Private/Operations/UniformTesselate.cpp
+++ b/Engine/Plugins/Runtime/GeometryProcessing/Source/DynamicMesh/Private/Operations/UniformTesselate.cpp
@@ -136,6 +136,9 @@ namespace UniformTesselateLocals
 		// ---------------------------------------------------------------------
 		TArray<RealType> Elements;
 
+		// Track which elements we set in the Elements array. 
+		TArray<bool> ElementIsSet; 
+
 		FElementsGenerator(const FDynamicMesh3* Mesh, const int TesselationNum, FProgressCancel* InProgress, const bool bUseParallel) 
 		:
 		FBaseGenerator(Mesh, TesselationNum, InProgress, bUseParallel)
@@ -146,6 +149,13 @@ namespace UniformTesselateLocals
 		{
 		}
 
+		virtual void SetBufferSizes(const int NumElements)
+		{
+			checkSlow(NumElements * ElementSize > 0);
+			Elements.SetNum(NumElements * ElementSize);
+			ElementIsSet.Init(false, NumElements);
+		}
+
 		// Child classes overwrite this to initialize element counts (InElementCount, OutElementCount ... )
 		virtual void Initialize() = 0;
 
@@ -162,7 +172,7 @@ namespace UniformTesselateLocals
 
 		virtual void SetElementOnEdge(const int InEdgeID, const int InTriangleID, const int VertexOffset, const RealType* Value) 
 		{
-			const int Index = GetElementIDOnEdge(InEdgeID, InTriangleID, VertexOffset);
+			const int Index = this->GetElementIDOnEdge(InEdgeID, InTriangleID, VertexOffset);
 			checkSlow(Index >= InElementCount && Index < InElementCount + OutEdgeElementCount); 
 			SetElement(Index, Value); 
 		}
@@ -204,6 +214,7 @@ namespace UniformTesselateLocals
 		void SetElement(const int Index, const RealType* Value)
 		{	
 			const int Offset = Index * ElementSize;
+			ElementIsSet[Index] = true;
 			for (int Idx = 0; Idx < ElementSize; ++Idx)
 			{
 				Elements[Offset + Idx] = Value[Idx];
@@ -214,6 +225,7 @@ namespace UniformTesselateLocals
 		void SetElement(const int Index, const AsType& TypeValue)
 		{	
 			const int Offset = Index * ElementSize;
+			ElementIsSet[Index] = true;
 			for (int Idx = 0; Idx < ElementSize; ++Idx)
 			{
 				Elements[Offset + Idx] = TypeValue[Idx];
@@ -222,6 +234,12 @@ namespace UniformTesselateLocals
 
 		void GetElement(const int Index, RealType* Value) const
 		{
+			if (ElementIsSet[Index] == false) 
+			{
+				checkSlow(false); // we are trying to access an element that wasn't set
+				return;
+			}
+
 			const int Offset = Index * ElementSize;
 			checkSlow(Offset + ElementSize <= Elements.Num());
 			for (int Idx = 0; Idx < ElementSize; ++Idx) 
@@ -233,6 +251,12 @@ namespace UniformTesselateLocals
 		template<typename AsType>
 		void GetElement(const int Index, AsType& TypeValue) const
 		{
+			if (ElementIsSet[Index] == false)
+			{
+				checkSlow(false); // we are trying to access an element that wasn't set
+				return;
+			}
+
 			const int Offset = Index * ElementSize;
 			checkSlow(Offset + ElementSize <= Elements.Num());
 			for (int Idx = 0; Idx < ElementSize; ++Idx) 
@@ -341,8 +365,11 @@ namespace UniformTesselateLocals
 				{
 					return;
 				}
-
-				checkSlow(InMesh->IsTriangle(TriangleID));
+				
+				if (this->IsValidTriangle(TriangleID) == false) 
+				{
+					return;
+				}
 
 				const FIndex3i TriVertices = InMesh->GetTriangle(TriangleID);
 				const FIndex3i TriEdges = InMesh->GetTriEdges(TriangleID);
@@ -355,8 +382,8 @@ namespace UniformTesselateLocals
 				{
 					// The number of new vertices added along each Level is same as the Level number
 					const int NumNewLevelVertices = Level; 
-					const int ID1 = GetElementIDOnEdge(TriEdges[0], TriangleID, OrderedEdgeOffset(IsEdgeReversed[0], Level));  
-					const int ID2 = GetElementIDOnEdge(TriEdges[2], TriangleID, OrderedEdgeOffset(IsEdgeReversed[2], Level));
+					const int ID1 = this->GetElementIDOnEdge(TriEdges[0], TriangleID, OrderedEdgeOffset(IsEdgeReversed[0], Level));  
+					const int ID2 = this->GetElementIDOnEdge(TriEdges[2], TriangleID, OrderedEdgeOffset(IsEdgeReversed[2], Level));
 
 					RealType Element1[ElementSize];
 					RealType Element2[ElementSize];
@@ -368,7 +395,7 @@ namespace UniformTesselateLocals
 					{
 						const RealType Alpha = (RealType)(VertexOffset + 1) / (NumNewLevelVertices + 1);
 						LerpElements(Element1, Element2, OutElement, Alpha);
-						SetElementOnTriangle(TriangleID, ElementIDCounter, OutElement);
+						this->SetElementOnTriangle(TriangleID, ElementIDCounter, OutElement);
 						ElementIDCounter++;
 					}
 				}
@@ -381,6 +408,12 @@ namespace UniformTesselateLocals
 
 			return true;
 		}
+
+
+		virtual bool IsValidTriangle(int32 TriangleID) const
+		{
+			return InMesh->IsTriangle(TriangleID);
+		}
 	};
 
 	/**
@@ -394,9 +427,11 @@ namespace UniformTesselateLocals
 		using BaseType = FElementsGenerator<RealType, ElementSize>;
 		using BaseType::InMesh;
 		using BaseType::Elements;
+		using BaseType::ElementIsSet;
 		using BaseType::TesselationNum;
 		using BaseType::OrderedEdgeOffset;
 		using BaseType::GetEdgeOrder;
+		using BaseType::SetBufferSizes;
 
 		int OutTriangleCount = 0; // Total number of triangle triplets in the output mesh
 		int OutInnerTriangleCount = 0; // Number of triangle triplets introduced per input triangle after tesselation
@@ -423,11 +458,11 @@ namespace UniformTesselateLocals
 			Triangles[Index] = FIndex3i(A, B, C);
 		}
 
-		void SetBufferSizes(const int NumTriangles, const int NumElements)
+		virtual void SetBufferSizes(const int NumTriangles, const int NumElements)
 		{
-			checkSlow(NumTriangles > 0 && NumElements * ElementSize > 0);
+			this->SetBufferSizes(NumElements);
+			checkSlow(NumTriangles > 0);
 			Triangles.SetNum(NumTriangles);
-			Elements.SetNum(NumElements * ElementSize);
 		}
 
 		/**
@@ -679,8 +714,11 @@ namespace UniformTesselateLocals
 					return;
 				}
 
-				checkSlow(InMesh->IsTriangle(TriangleID));
- 				
+				if (this->IsValidTriangle(TriangleID) == false) 
+				{
+					return;
+				}
+				 				
 				const FIndex3i TriVertices = InMesh->GetTriangle(TriangleID);
 				const FIndex3i TriEdges = InMesh->GetTriEdges(TriangleID);
 
@@ -724,6 +762,8 @@ namespace UniformTesselateLocals
 		using BaseType::OutInnerElementCount;
 		using BaseType::Triangles;
 		using BaseType::OutTriangleCount;
+		using BaseType::OutInnerTriangleCount;
+		using BaseType::ElementIsSet;
 		
 		// Track the offset into the Elements array for a block of data containing elements along the edge 
 		// (1 or more for each vertex added along the edge). We can have a variable number of elements per-vertex 
@@ -735,10 +775,10 @@ namespace UniformTesselateLocals
 		const TDynamicMeshOverlay<RealType, ElementSize>* Overlay;
 		
 		FOverlayGenerator(const FDynamicMesh3* Mesh, 
-						const int TesselationNum,
-						FProgressCancel* InProgress, 
-						const bool bUseParallel,
-						const TDynamicMeshOverlay<RealType, ElementSize>* InOverlay) 
+						  const int TesselationNum,
+						  FProgressCancel* InProgress, 
+						  const bool bUseParallel,
+						  const TDynamicMeshOverlay<RealType, ElementSize>* InOverlay) 
 		:
 		FTrianglesGenerator<RealType, ElementSize>(Mesh, TesselationNum, InProgress, bUseParallel), Overlay(InOverlay)
 		{
@@ -766,7 +806,7 @@ namespace UniformTesselateLocals
 				LastOffset += bIsSeam ? 2*TesselationNum : TesselationNum;	
 			}
 
-			InElementCount = Overlay->ElementCount();
+			InElementCount = Overlay->MaxElementID();
 			OutEdgeElementCount = LastOffset;
 
 			const int TriangleNumber = TesselationNum + 2; // plus 2 vertices on the ends
@@ -776,7 +816,7 @@ namespace UniformTesselateLocals
 			this->SetBufferSizes(OutTriangleCount, OutElementCount);
 
 			RealType ElementData[ElementSize];
-			for (int ElementID = 0; ElementID < Overlay->ElementCount(); ++ElementID)
+			for (int32 ElementID : Overlay->ElementIndicesItr())
 			{
 				Overlay->GetElement(ElementID, ElementData);
 				this->SetElement(ElementID, ElementData);
@@ -833,27 +873,30 @@ namespace UniformTesselateLocals
 
 				RealType Element1[ElementSize]; 
 				RealType Element2[ElementSize]; 
-				GetInputMeshElementValue(EdgeTri.A, EdgeV.A, Element1);
-				GetInputMeshElementValue(EdgeTri.A, EdgeV.B, Element2);
-
-				RealType Element3[ElementSize]; 
-				RealType Element4[ElementSize]; 
-				if (SeamEdges[EdgeID]) 
+				RealType Out[ElementSize];
+				
+				if (Overlay->IsSetTriangle(EdgeTri.A)) 
 				{
-					GetInputMeshElementValue(EdgeTri.B, EdgeV.A, Element3);
-					GetInputMeshElementValue(EdgeTri.B, EdgeV.B, Element4);
+					GetInputMeshElementValue(EdgeTri.A, EdgeV.A, Element1);
+					GetInputMeshElementValue(EdgeTri.A, EdgeV.B, Element2);
+
+					for (int VertexOffset = 0; VertexOffset < TesselationNum; ++VertexOffset)
+					{
+						const RealType Tau = (RealType)(VertexOffset + 1) / (TesselationNum + 1);
+						this->LerpElements(Element1, Element2, Out, Tau);
+						this->SetElementOnEdge(EdgeID, EdgeTri.A, VertexOffset, Out);	
+					}
 				}
 
-				RealType Out[ElementSize];
-				for (int VertexOffset = 0; VertexOffset < TesselationNum; ++VertexOffset)
+				if (SeamEdges[EdgeID] && EdgeTri.B != FDynamicMesh3::InvalidID && Overlay->IsSetTriangle(EdgeTri.B))
 				{
-					const RealType Tau = (RealType)(VertexOffset + 1) / (TesselationNum + 1);
-					this->LerpElements(Element1, Element2, Out, Tau);
-					this->SetElementOnEdge(EdgeID, EdgeTri.A, VertexOffset, Out);	
-
-					if (SeamEdges[EdgeID]) 
+					GetInputMeshElementValue(EdgeTri.B, EdgeV.A, Element1);
+					GetInputMeshElementValue(EdgeTri.B, EdgeV.B, Element2);
+					
+					for (int VertexOffset = 0; VertexOffset < TesselationNum; ++VertexOffset)
 					{
-						this->LerpElements(Element3, Element4, Out, Tau);
+						const RealType Tau = (RealType)(VertexOffset + 1) / (TesselationNum + 1);
+						this->LerpElements(Element1, Element2, Out, Tau);
 						this->SetElementOnEdge(EdgeID, EdgeTri.B, VertexOffset, Out);
 					}
 				}
@@ -864,28 +907,58 @@ namespace UniformTesselateLocals
 				return false;
 			}
 
-
 			return true;
 		}
 
 		// Copy the content of the Generator (elements and triangles) to the overlay.
 		void CopyToAttribute(TDynamicMeshOverlay<RealType, ElementSize>* OutOverlay) 
 		{
-			if (IsOverlayValid() == false) 
+			if (IsOverlayValid() == false || OutOverlay == nullptr) 
 			{
-				return; // empty input overlay, nothing to do
+				return; // empty input or output overlay, nothing to do
 			}
-
+			
+			OutOverlay->ClearElements();
+			
+			TArray<int32> ElementMap;
+			ElementMap.SetNumUninitialized(OutElementCount);
 			RealType Value[ElementSize];
-			for (int EID = 0; EID < OutElementCount; ++EID)
+			for (int32 EID = 0; EID < OutElementCount; ++EID)
 			{
-				this->GetElement(EID, Value);
-				OutOverlay->AppendElement(Value);
+				if (ElementIsSet[EID]) 
+				{
+					this->GetElement(EID, Value);
+					ElementMap[EID] = OutOverlay->AppendElement(Value);
+				}
+				else 
+				{
+					ElementMap[EID] = FDynamicMesh3::InvalidID;
+				}
 			}
 			
-			for (int TID = 0; TID < Triangles.Num(); ++TID)
-			{
-				ensure(OutOverlay->SetTriangle(TID, Triangles[TID]) == EMeshResult::Ok);
+			for (int32 TriangleID = 0; TriangleID < InMesh->MaxTriangleID(); ++TriangleID) 
+			{	
+				// Check if the triangle is valid and set in the overlay
+				if (this->IsValidTriangle(TriangleID)) 
+				{
+					// Iterate over each subtriangle we generated for each input mesh triangle
+					for (int TriangleOffset = 0; TriangleOffset < OutInnerTriangleCount; ++TriangleOffset)
+					{	
+						int32 OutTriangleID = this->GetOutTriangleID(TriangleID, TriangleOffset);
+						FIndex3i ElementIDs = Triangles[OutTriangleID];
+						FIndex3i MappedElementIDs = FIndex3i(ElementMap[ElementIDs[0]], ElementMap[ElementIDs[1]], ElementMap[ElementIDs[2]]);
+						
+						bool bIsValidMapping = MappedElementIDs[0] != FDynamicMesh3::InvalidID &&
+							      			   MappedElementIDs[1] != FDynamicMesh3::InvalidID && 
+							      			   MappedElementIDs[2] != FDynamicMesh3::InvalidID;
+						checkSlow(bIsValidMapping);
+
+						if (bIsValidMapping) 
+						{
+							ensure(OutOverlay->SetTriangle(OutTriangleID, MappedElementIDs) == EMeshResult::Ok);
+						}
+					}
+				}
 			}
 		}	
 		
@@ -925,10 +998,16 @@ namespace UniformTesselateLocals
 				
 			return true;
 		}
+
+		
+		virtual bool IsValidTriangle(int32 TriangleID) const override
+		{
+			return InMesh->IsTriangle(TriangleID) && Overlay->IsSetTriangle(TriangleID);
+		}
 	};
 
 	/**
-	 * Generator to handle the interpolation of the vertex position data and generation of the new the triangle topology.
+	 * Generator to handle the interpolation of the vertex position data and generation of the new triangle topology.
 	 */
 	template<typename RealType>
 	class FMeshVerticesGenerator : public FTrianglesGenerator<RealType, 3>
@@ -1152,8 +1231,8 @@ namespace UniformTesselateLocals
 			OutInnerElementCount = TriangularPatternUtils::NumInnerVertices(TriangleNumber); 
 			OutElementCount = InElementCount + OutEdgeElementCount + OutInnerElementCount*InMesh->TriangleCount();	
 
-			Elements.SetNum(OutElementCount*ElementSize);
-			
+			this->SetBufferSizes(OutElementCount);
+
 			RealType Value[ElementSize];
 			for (int VertexID = 0; VertexID < InMesh->VertexCount(); ++VertexID)
 			{
@@ -1289,10 +1368,10 @@ namespace UniformTesselateLocals
 		TFunction<void(const int TriangleID, RealType* OutValue)> DataFunction;
 		
 		FTriangleAttributeGenerator(const FDynamicMesh3* Mesh, 
-								const int TesselationNum,
-								FProgressCancel* InProgress, 
-								const bool bUseParallel,
-								const TFunction<void(int TriangleID, RealType* Value)>& InDataFunction) 
+								    const int TesselationNum,
+								    FProgressCancel* InProgress, 
+								    const bool bUseParallel,
+								    const TFunction<void(int TriangleID, RealType* Value)>& InDataFunction) 
 		:
 		FBaseGenerator(Mesh, TesselationNum, InProgress, bUseParallel), DataFunction(InDataFunction)
 		{
